Switched UAT build process on Mac to use msbuild.
Enabled parallel msbuild for Mac
Added HashCollection class that can hold and serialize a collection of content hashes
Added quick dependency checking to UAT build by hashing csproj, input and output files for all modules
Reduced some logging when not running with -verbose

diff --git a/Engine/Build/BatchFiles/RunUAT.sh b/Engine/Build/BatchFiles/RunUAT.sh
--- a/Engine/Build/BatchFiles/RunUAT.sh
+++ b/Engine/Build/BatchFiles/RunUAT.sh
@@ -71,12 +71,11 @@ if [ "$UATCompileArg" = "-compile" ]; then
 		echo No project to compile, attempting to use precompiled AutomationTool
 		UATCompileArg=
 	else
-		echo Compiling AutomationTool with xbuild
 
 		# make sure the UBT project has references to auto-discovered platform extension source files
 		"${SCRIPT_DIR}/FindPlatformExtensionSources.sh"
 
-		ARGS="/p:Configuration=Development /p:Platform=AnyCPU /verbosity:quiet /nologo /p:NoWarn=1591"
+		ARGS="/p:Configuration=Development /p:Platform=AnyCPU /verbosity:quiet /nologo /p:NoWarn=1591 /property:AutomationToolProjectOnly=true"
 		ARGS="${ARGS} /p:TargetFrameworkProfile="
 
 		echo "xbuild Source/Programs/AutomationTool/AutomationTool.csproj $ARGS"
@@ -85,8 +84,6 @@ if [ "$UATCompileArg" = "-compile" ]; then
 		if [ $? -ne 0 ]; then
 			echo RunUAT ERROR: AutomationTool failed to compile.
 			exit 1
-		else
-			echo Compilation Succeeded
 		fi
 	fi
 fi
diff --git a/Engine/Source/Programs/AutomationTool/AutomationUtils/Automation.cs b/Engine/Source/Programs/AutomationTool/AutomationUtils/Automation.cs
--- a/Engine/Source/Programs/AutomationTool/AutomationUtils/Automation.cs
+++ b/Engine/Source/Programs/AutomationTool/AutomationUtils/Automation.cs
@@ -103,13 +103,14 @@ namespace AutomationTool
 		public static CommandLineArg NoP4 = new CommandLineArg("-NoP4");
 		public static CommandLineArg P4 = new CommandLineArg("-P4");
         public static CommandLineArg Compile = new CommandLineArg("-Compile");
-        /// <summary>
-        /// This command is LEGACY because we used to run UAT.exe to compile scripts by default.
-        /// Now we only compile by default when run via RunUAT.bat, which still understands -nocompile.
-        /// However, the batch file simply passes on all arguments, so UAT will choke when encountering -nocompile.
-        /// Keep this CommandLineArg around so that doesn't happen.
-        /// </summary>
-        public static CommandLineArg NoCompileLegacyDontUse = new CommandLineArg("-NoCompile");
+		public static CommandLineArg IgnoreDependencies = new CommandLineArg("-IgnoreDependencies");
+		/// <summary>
+		/// This command is LEGACY because we used to run UAT.exe to compile scripts by default.
+		/// Now we only compile by default when run via RunUAT.bat, which still understands -nocompile.
+		/// However, the batch file simply passes on all arguments, so UAT will choke when encountering -nocompile.
+		/// Keep this CommandLineArg around so that doesn't happen.
+		/// </summary>
+		public static CommandLineArg NoCompileLegacyDontUse = new CommandLineArg("-NoCompile");
         public static CommandLineArg NoCompileEditor = new CommandLineArg("-NoCompileEditor");
 		public static CommandLineArg Help = new CommandLineArg("-Help");
 		public static CommandLineArg List = new CommandLineArg("-List");
diff --git a/Engine/Source/Programs/AutomationTool/AutomationUtils/MacHostPlatform.cs b/Engine/Source/Programs/AutomationTool/AutomationUtils/MacHostPlatform.cs
--- a/Engine/Source/Programs/AutomationTool/AutomationUtils/MacHostPlatform.cs
+++ b/Engine/Source/Programs/AutomationTool/AutomationUtils/MacHostPlatform.cs
@@ -14,7 +14,7 @@ namespace AutomationTool
 	{
 		public override string GetMsBuildExe()
 		{
-			return "xbuild";
+			return "msbuild";
 		}
 
 		public override string RelativeBinariesFolder
diff --git a/Engine/Source/Programs/AutomationTool/AutomationUtils/ScriptCompiler.cs b/Engine/Source/Programs/AutomationTool/AutomationUtils/ScriptCompiler.cs
--- a/Engine/Source/Programs/AutomationTool/AutomationUtils/ScriptCompiler.cs
+++ b/Engine/Source/Programs/AutomationTool/AutomationUtils/ScriptCompiler.cs
@@ -1,4 +1,4 @@
-// Copyright 1998-2019 Epic Games, Inc. All Rights Reserved.
+﻿// Copyright 1998-2019 Epic Games, Inc. All Rights Reserved.
 using System;
 using System.Collections.Generic;
 using System.Linq;
@@ -73,7 +73,6 @@ namespace AutomationTool
 			// Compile only if not disallowed.
 			if (DoCompile && !String.IsNullOrEmpty(CommandUtils.CmdEnv.MsBuildExe))
 			{
-				CleanupScriptsAssemblies();
 				FindAndCompileScriptModules(ScriptsForProjectFileName, AdditionalScriptsFolders);
 			}
 
@@ -128,8 +127,6 @@ namespace AutomationTool
 
 		private static void FindAndCompileScriptModules(string ScriptsForProjectFileName, List<string> AdditionalScriptsFolders)
 		{
-			Log.TraceInformation("Compiling scripts.");
-
 			var OldCWD = Environment.CurrentDirectory;
 			var UnrealBuildToolCWD = CommandUtils.CombinePaths(CommandUtils.CmdEnv.LocalRoot, "Engine", "Source");
 
@@ -138,13 +135,13 @@ namespace AutomationTool
 			// Configure the rules compiler
 			// Get all game folders and convert them to build subfolders.
 			List<DirectoryReference> AllGameFolders;
-			if(ScriptsForProjectFileName == null)
+			if (ScriptsForProjectFileName == null)
 			{
 				AllGameFolders = NativeProjects.EnumerateProjectFiles().Select(x => x.Directory).ToList();
 			}
 			else
 			{
-				AllGameFolders = new List<DirectoryReference>{ new DirectoryReference(Path.GetDirectoryName(ScriptsForProjectFileName)) };
+				AllGameFolders = new List<DirectoryReference> { new DirectoryReference(Path.GetDirectoryName(ScriptsForProjectFileName)) };
 			}
 
 			var AllAdditionalScriptFolders = new List<DirectoryReference>(AdditionalScriptsFolders.Select(x => new DirectoryReference(x)));
@@ -172,42 +169,159 @@ namespace AutomationTool
 					CommandUtils.LogVerbose("Script module {0} filtered by the Host Platform and will not be compiled.", ModuleFilename);
 				}
 			}
+			
+			CompileModules(ModulesToCompile);
+
+			Environment.CurrentDirectory = OldCWD;
+		}
 
-			if ((UnrealBuildTool.BuildHostPlatform.Current.Platform == UnrealBuildTool.UnrealTargetPlatform.Win64) ||
-				(UnrealBuildTool.BuildHostPlatform.Current.Platform == UnrealBuildTool.UnrealTargetPlatform.Win32))
+		/// <summary>
+		/// Creates a hash collection that represents the state of all modules in this list. If 
+		/// a module has no current output or is missing files then a null collection will be returned
+		/// </summary>
+		/// <param name="Modules"></param>
+		/// <returns></returns>
+		private static HashCollection HashModules(IEnumerable<string> Modules, bool WarnOnFailure=true)
+		{
+			HashCollection Hashes = new HashCollection();
+
+			foreach (string Module in Modules)
 			{
-				string Modules = string.Join(";", ModulesToCompile.ToArray());
-				var UATProj = CommandUtils.CombinePaths(CommandUtils.CmdEnv.LocalRoot, @"Engine\Source\Programs\AutomationTool\Scripts\UAT.proj");
-				var CmdLine = String.Format("\"{0}\" /p:Modules=\"{1}\" /p:Configuration={2} /verbosity:minimal /nologo", UATProj, Modules, BuildConfig);
-				// suppress the run command because it can be long and intimidating, making the logs around this code harder to read.
-				var Result = CommandUtils.Run(CommandUtils.CmdEnv.MsBuildExe, CmdLine, Options: CommandUtils.ERunOptions.Default | CommandUtils.ERunOptions.NoLoggingOfRunCommand | CommandUtils.ERunOptions.LoggingOfRunDuration);
-				if (Result.ExitCode != 0)
+				CsProjectInfo Proj;
+
+				Dictionary<string, string> Properties = new Dictionary<string, string>();
+				Properties.Add("Platform", "AnyCPU");
+				Properties.Add("Configuration", "Development");
+
+				FileReference ModuleFile = new FileReference(Module);
+
+				if (!CsProjectInfo.TryRead(ModuleFile, Properties, out Proj))
+				{
+					if (WarnOnFailure)
+					{
+						Log.TraceWarning("Failed to read file {0}", ModuleFile);
+					}
+					return null;
+				}
+
+				if (!Hashes.AddCsProjectInfo(Proj, HashCollection.HashType.MetaData))
+				{
+					if (WarnOnFailure)
+					{
+						Log.TraceWarning("Failed to hash file {0}", ModuleFile);
+					}
+					return null;
+				}
+			}
+
+			return Hashes;
+		}
+
+		/// <summary>
+		/// Pulls all dependencies from the specified module list, gathers the state of them, and
+		/// compares it to the state expressed in the provided file from a previous run
+		/// </summary>
+		/// <param name="ModuleList"></param>
+		/// <param name="DependencyFile"></param>
+		/// <returns></returns>
+		private static bool AreDependenciesUpToDate(IEnumerable<string> ModuleList, string DependencyFile)
+		{
+			bool UpToDate = false;
+
+			if (File.Exists(DependencyFile))
+			{
+				Log.TraceVerbose("Read dependency file at {0}", DependencyFile);
+
+				HashCollection OldHashes = HashCollection.CreateFromFile(DependencyFile);
+
+				if (OldHashes != null)
+				{
+					HashCollection CurrentHashes = HashModules(ModuleList, false);
+
+					if (OldHashes.Equals(CurrentHashes))
+					{
+						UpToDate = true;
+					}
+					else
+					{
+						if (Log.OutputLevel >= LogEventType.VeryVerbose)
+						{
+							CurrentHashes.LogDifferences(OldHashes);
+						}
+					}
+				}
+				else
 				{
-					throw new AutomationException(String.Format("Failed to build \"{0}\":{1}{2}", UATProj, Environment.NewLine, Result.Output));
+					Log.TraceInformation("Failed to read dependency info!");
 				}
 			}
 			else
 			{
-				CompileModules(ModulesToCompile);
+				Log.TraceVerbose("No dependency file exists at {0}. Will do full build.", DependencyFile);
 			}
-			
 
-			Environment.CurrentDirectory = OldCWD;
+			return UpToDate;
 		}
-
+		
 		/// <summary>
 		/// Compiles all script modules.
 		/// </summary>
 		/// <param name="Modules">Module project filenames.</param>
-		/// <param name="CompiledModuleFilenames">The resulting compiled module assembly filenames.</param>
 		private static void CompileModules(List<string> Modules)
-		{
-			Log.TraceInformation("Building script modules");
-			// Make sure DefaultScriptsDLLName is compiled first
-			var DefaultScriptsProjName = Path.ChangeExtension(DefaultScriptsDLLName, "csproj");
-			foreach (var ModuleName in Modules)
+		{			
+			string DependencyFile = Path.Combine(CommandUtils.CmdEnv.EngineSavedFolder, "UATModuleHashes.xml");
+			
+			if (AreDependenciesUpToDate(Modules, DependencyFile) && !GlobalCommandLine.IgnoreDependencies)
+			{
+				Log.TraceInformation("Dependencies are up to date. Skipping compile.");
+				return;
+			}
+
+			Log.TraceInformation("Dependencies are out of date. Compiling scripts.");
+
+			// clean old assemblies
+			CleanupScriptsAssemblies();
+
+			DateTime StartTime = DateTime.Now;
+
+			bool UseParallelMsBuild = UnrealBuildTool.BuildHostPlatform.Current.Platform == UnrealBuildTool.UnrealTargetPlatform.Win64
+				|| UnrealBuildTool.BuildHostPlatform.Current.Platform == UnrealBuildTool.UnrealTargetPlatform.Win32
+				|| UnrealBuildTool.BuildHostPlatform.Current.Platform == UnrealBuildTool.UnrealTargetPlatform.Mac;
+
+			if (UseParallelMsBuild)
 			{
-				if (ModuleName.IndexOf(DefaultScriptsProjName, StringComparison.InvariantCultureIgnoreCase) >= 0)
+				string ModulesList = string.Join(";", Modules);
+
+				// Mono has an issue where arugments with semicolons or commas can't be passed through to
+				// as arguments so we need to manually construct a temp file with the list of modules
+				// see ([REDACTED_PROFILE_URL]/msbuild/issues/471)
+				var UATProjTemplate = CommandUtils.CombinePaths(CommandUtils.CmdEnv.LocalRoot, @"Engine\Source\Programs\AutomationTool\Scripts\UAT.proj");
+				var UATProjFile = Path.Combine(CommandUtils.CmdEnv.EngineSavedFolder, "UATTempProj.proj");
+
+				string ProjContents = File.ReadAllText(UATProjTemplate);
+				ProjContents = ProjContents.Replace("$(Modules)", ModulesList);
+
+				File.WriteAllText(UATProjFile, ProjContents);
+
+				string MsBuildVerbosity = Log.OutputLevel >= LogEventType.Verbose ? "minimal" : "quiet";
+
+				var CmdLine = String.Format("\"{0}\" /p:Configuration={1} /verbosity:{2} /nologo", UATProjFile, BuildConfig, MsBuildVerbosity);
+				// suppress the run command because it can be long and intimidating, making the logs around this code harder to read.
+				var Result = CommandUtils.Run(CommandUtils.CmdEnv.MsBuildExe, CmdLine, Options: CommandUtils.ERunOptions.Default | CommandUtils.ERunOptions.NoLoggingOfRunCommand | CommandUtils.ERunOptions.LoggingOfRunDuration);
+				if (Result.ExitCode != 0)
+				{
+					throw new AutomationException(String.Format("Failed to build \"{0}\":{1}{2}", UATProjFile, Environment.NewLine, Result.Output));
+				}
+			}
+			else
+			{
+				// Make sure DefaultScriptsDLLName is compiled first
+				var DefaultScriptsProjName = Path.ChangeExtension(DefaultScriptsDLLName, "csproj");
+
+				// Primary modules must be built first
+				List<string> PrimaryModules = Modules.Where(M => M.IndexOf(DefaultScriptsProjName, StringComparison.InvariantCultureIgnoreCase) >= 0).ToList();
+
+				foreach (var ModuleName in PrimaryModules)
 				{
 					Log.TraceInformation("Building script module: {0}", ModuleName);
 					try
@@ -220,13 +334,14 @@ namespace AutomationTool
 						throw new AutomationException("Failed to compile module {0}", ModuleName);
 					}
 					break;
+
 				}
-			}
 
-			// Second pass, compile everything else
-			foreach (var ModuleName in Modules)
-			{
-				if (ModuleName.IndexOf(DefaultScriptsProjName, StringComparison.InvariantCultureIgnoreCase) < 0)
+				// Second pass, compile everything else
+				List<string> SecondaryModules = Modules.Where(M => !PrimaryModules.Contains(M)).ToList();
+
+				// Non-parallel method
+				foreach (var ModuleName in SecondaryModules)
 				{
 					Log.TraceInformation("Building script module: {0}", ModuleName);
 					try
@@ -239,11 +354,27 @@ namespace AutomationTool
 						throw new AutomationException("Failed to compile module {0}", ModuleName);
 					}
 				}
+			}			
+			
+			TimeSpan Duration = DateTime.Now - StartTime;
+			Log.TraceInformation("Compiled {0} modules in {1} secs", Modules.Count, Duration.TotalSeconds);
+
+			HashCollection NewHashes = HashModules(Modules);
+
+			if (NewHashes == null)
+			{
+				Log.TraceWarning("Failed to save dependency info!");
+			}
+			else
+			{
+				NewHashes.SaveToFile(DependencyFile);
+				Log.TraceVerbose("Wrote depencencies to {0}", DependencyFile);
 			}
 		}
 
 		/// <summary>
-		/// Builds a script module (csproj file)
+		/// Starts compiling the provided project file and returns the process. Caller should check HasExited
+		/// or call WaitForExit() to get results
 		/// </summary>
 		/// <param name="ProjectFile"></param>
 		/// <returns></returns>
@@ -263,6 +394,7 @@ namespace AutomationTool
 
 			// Compile the project
 			var Result = CommandUtils.Run(CommandUtils.CmdEnv.MsBuildExe, CmdLine);
+
 			if (Result.ExitCode != 0)
 			{
 				throw new AutomationException(String.Format("Failed to build \"{0}\":{1}{2}", ProjectFile, Environment.NewLine, Result.Output));
@@ -328,7 +460,7 @@ namespace AutomationTool
 			}
 		}
 
-		private void CleanupScriptsAssemblies()
+		private static void CleanupScriptsAssemblies()
 		{
 			CommandUtils.LogVerbose("Cleaning up script DLL folder");
 			CommandUtils.DeleteDirectory(GetScriptAssemblyFolder());
diff --git a/Engine/Source/Programs/DotNETCommon/DotNETUtilities/CsProjectInfo.cs b/Engine/Source/Programs/DotNETCommon/DotNETUtilities/CsProjectInfo.cs
--- a/Engine/Source/Programs/DotNETCommon/DotNETUtilities/CsProjectInfo.cs
+++ b/Engine/Source/Programs/DotNETCommon/DotNETUtilities/CsProjectInfo.cs
@@ -29,17 +29,28 @@ namespace Tools.DotNETCommon
 		/// </summary>
 		public Dictionary<FileReference, bool> ProjectReferences = new Dictionary<FileReference, bool>();
 
+		/// <summary>
+		/// List of compile references in the project.
+		/// </summary>
+		public List<FileReference> CompileReferences = new List<FileReference>();
+
 		/// <summary>
 		/// Mapping of content IF they are flagged Always or Newer
 		/// </summary>
 		public Dictionary<FileReference, bool> ContentReferences = new Dictionary<FileReference, bool>();
 
+		/// <summary>
+		/// Path to the CSProject file
+		/// </summary>
+		public FileReference ProjectPath;
+
 		/// <summary>
 		/// Constructor
 		/// </summary>
 		/// <param name="InProperties">Initial mapping of property names to values</param>
-		CsProjectInfo(Dictionary<string, string> InProperties)
+		CsProjectInfo(Dictionary<string, string> InProperties, FileReference InProjectPath)
 		{
+			ProjectPath = InProjectPath;
 			Properties = new Dictionary<string, string>(InProperties);
 		}
 
@@ -61,6 +72,17 @@ namespace Tools.DotNETCommon
 			}
 		}
 
+		/// <summary>
+		/// Returns the assembly name used by this project
+		/// </summary>
+		/// <returns></returns>
+		public string GetAssemblyName()
+		{
+			string Output = "";
+			Properties.TryGetValue("AssemblyName", out Output);
+			return Output;
+		}
+
 		/// <summary>
 		/// Finds all build products from this project. This includes content and other assemblies marked to be copied local.
 		/// </summary>
@@ -234,7 +256,7 @@ namespace Tools.DotNETCommon
 			}
 
 			// Parse the basic structure of the document, updating properties and recursing into other referenced projects as we go
-			CsProjectInfo ProjectInfo = new CsProjectInfo(Properties);
+			CsProjectInfo ProjectInfo = new CsProjectInfo(Properties, File);
 			foreach (XmlElement Element in Document.DocumentElement.ChildNodes.OfType<XmlElement>())
 			{
 				switch (Element.Name)
@@ -303,6 +325,13 @@ namespace Tools.DotNETCommon
 							ParseProjectReference(BaseDirectory, ItemElement, ProjectInfo.ProjectReferences);
 						}
 						break;
+					case "Compile":
+						// Reference to another project
+						if (EvaluateCondition(ItemElement, ProjectInfo.Properties))
+						{
+							ParseCompileReference(BaseDirectory, ItemElement, ProjectInfo.CompileReferences);
+						}
+						break;
 					case "Content":
 					case "None":
 						// Reference to another project
@@ -354,6 +383,22 @@ namespace Tools.DotNETCommon
 			}
 		}
 
+		/// <summary>
+		/// Parses a project reference from a given 'ProjectReference' element
+		/// </summary>
+		/// <param name="BaseDirectory">Directory to resolve relative paths against</param>
+		/// <param name="ParentElement">The parent 'ProjectReference' element</param>
+		/// <param name="CompileReferences">List of source files.</param>
+		static void ParseCompileReference(DirectoryReference BaseDirectory, XmlElement ParentElement, List<FileReference> CompileReferences)
+		{
+			string IncludePath = UnescapeString(ParentElement.GetAttribute("Include"));
+			if (!String.IsNullOrEmpty(IncludePath))
+			{
+				FileReference SourceFile = FileReference.Combine(BaseDirectory, IncludePath);
+				CompileReferences.Add(SourceFile);
+			}
+		}
+
 		/// <summary>
 		/// Parses an assembly reference from a given 'Content' element
 		/// </summary>
@@ -694,4 +739,37 @@ namespace Tools.DotNETCommon
 			return NewText;
 		}
 	}
+
+	/// <summary>
+	/// Extension methods for CsProject support
+	/// </summary>
+	public static class CsProjectInfoExtensionMethods
+	{
+		/// <summary>
+		/// Adds aall input/output properties of a CSProject to a hash collection
+		/// </summary>
+		/// <param name="Hasher"></param>
+		/// <param name="Project"></param>
+		/// <returns></returns>
+		public static bool AddCsProjectInfo(this HashCollection Hasher, CsProjectInfo Project, HashCollection.HashType HashType)
+		{
+			// Get the output assembly and pdb file
+			DirectoryReference ProjectDirectory = Project.ProjectPath.Directory;
+			DirectoryReference OutputDir = Project.GetOutputDir(ProjectDirectory);
+			FileReference OutputFile = FileReference.Combine(OutputDir, Project.GetAssemblyName() + ".dll");
+			FileReference DebugFile = OutputFile.ChangeExtension("pdb");
+
+			// build a list of all input and output files from this module
+			List<FileReference> DependentFiles = new List<FileReference> { Project.ProjectPath, OutputFile, DebugFile };
+
+			DependentFiles.AddRange(Project.CompileReferences);
+	
+			if (!Hasher.AddFiles(DependentFiles, HashType))
+			{
+				return false;
+			}
+
+			return true;
+		}
+	}
 }
diff --git a/Engine/Source/Programs/DotNETCommon/DotNETUtilities/DotNETUtilities.csproj b/Engine/Source/Programs/DotNETCommon/DotNETUtilities/DotNETUtilities.csproj
--- a/Engine/Source/Programs/DotNETCommon/DotNETUtilities/DotNETUtilities.csproj
+++ b/Engine/Source/Programs/DotNETCommon/DotNETUtilities/DotNETUtilities.csproj
@@ -161,6 +161,7 @@
     <Compile Include="ResXResource\UEResXWriter.cs" />
     <Compile Include="StringUtils.cs" />
     <Compile Include="ThreadPoolWorkQueue.cs" />
+    <Compile Include="HashCollection.cs" />
   </ItemGroup>
   <Import Project="$(MSBuildToolsPath)\Microsoft.CSharp.targets" />
   <!-- To modify your build process, add your task inside one of the targets below and uncomment it. 
diff --git a/Engine/Source/Programs/DotNETCommon/DotNETUtilities/HashCollection.cs b/Engine/Source/Programs/DotNETCommon/DotNETUtilities/HashCollection.cs
new file mode 100644
--- /dev/null
+++ b/Engine/Source/Programs/DotNETCommon/DotNETUtilities/HashCollection.cs
@@ -0,0 +1,250 @@
+// Copyright 1998-2019 Epic Games, Inc. All Rights Reserved.
+
+using System;
+using System.Collections.Generic;
+using System.IO;
+using System.Security.Cryptography;
+using System.Text;
+using System.Xml;
+using System.Xml.Serialization;
+using System.Linq;
+using Tools.DotNETCommon;
+
+namespace Tools.DotNETCommon
+{
+
+	/// <summary>
+	/// Class that holds a single hash entry. Only the hash is used for comparisons but the
+	/// name and meta fields can hold related data for debugging etc.
+	/// </summary>
+	public class HashEntry : IEquatable<HashEntry>
+	{
+		public string Name { get; set; }
+		public string MetaData { get; set; }
+
+		[XmlIgnore]
+		public ContentHash Hash { get; set; }
+		
+		[XmlElement(ElementName = "Hash", DataType = "hexBinary")]
+		public byte[] HashBytes
+		{
+			get { return Hash.Bytes; }
+			set { Hash = new ContentHash(value); }
+		}
+
+		public HashEntry(string InName, ContentHash InHash, string InMeta)
+		{
+			Name = InName;
+			Hash = InHash;
+			MetaData = InMeta;
+		}
+
+		// parameterless constructor for serialization
+		public HashEntry()
+		{
+		}
+
+		// Equality check using the underlying content hash
+		public bool Equals(HashEntry RHS)
+		{
+			return Hash == RHS.Hash;
+		}
+
+		public override int GetHashCode()
+		{
+			return Hash.GetHashCode();
+		}
+
+		public override string ToString()
+		{
+			return Hash.ToString();
+		}
+	}
+
+	/// <summary>
+	/// Class that holds a collection of hashes with helpers for adding files, comparisions, and
+	/// serialization to and from files
+	/// </summary>
+	public class HashCollection
+	{
+		public enum HashType
+		{
+			MetaData,
+			Content
+		}
+
+		// underlying hashset that holds our data
+		public HashSet<HashEntry> Hashes { get; protected set; }
+
+		public HashCollection()
+		{
+			Hashes = new HashSet<HashEntry>();
+		}
+
+		/// <summary>
+		/// Adds a file to our hash collection. If the specified path does not exist an exception
+		/// will be thrown
+		/// </summary>
+		/// <param name="InPath"></param>
+		/// <param name="HashType"></param>
+		public bool AddFile(FileReference InPath, HashCollection.HashType HashType)
+		{
+			var Fi = new FileInfo(InPath.FullName);
+
+			if (!Fi.Exists)
+			{
+				return false;
+			}
+
+			string MetaData = string.Format("File={0} Size={1} LastWrite={2}", InPath, Fi.Length, Fi.LastWriteTimeUtc.ToString());
+
+			// Hash metadata or content
+			ContentHash Hash = HashType == HashCollection.HashType.MetaData ? ContentHash.SHA1(MetaData) : ContentHash.SHA1(InPath);
+
+			HashEntry Entry = new HashEntry(InPath.FullName, Hash, MetaData);
+
+			Hashes.Add(Entry);
+
+			return true;
+		}
+
+		/// <summary>
+		/// Convenience function that adds a range of files
+		/// </summary>
+		/// <param name="Files"></param>
+		/// <param name="HashType"></param>
+		public bool AddFiles(IEnumerable<FileReference> Files, HashCollection.HashType HashType)
+		{
+			foreach (var File in Files)
+			{
+				if (!AddFile(File, HashType))
+				{
+					return false;
+				}
+			}
+
+			return true;
+		}
+
+		/// <summary>
+		/// Compares two collections by comparing the underlying hashsets
+		/// </summary>
+		/// <param name="obj"></param>
+		/// <returns></returns>
+		public override bool Equals(object obj)
+		{
+			HashCollection RHS = obj as HashCollection;
+
+			if (RHS == null)
+			{
+				return false;
+
+			}
+			return Hashes.SetEquals(RHS.Hashes);
+		}
+
+		/// <summary>
+		/// Debugging helper that returns true if there's a hash with the specified name
+		/// </summary>
+		/// <param name="Name"></param>
+		/// <returns></returns>
+		public bool HasEntryForName(string Name)
+		{
+			return Hashes.Where(E => E.Name == Name).Any();
+		}
+
+		/// <summary>
+		/// Debugging helper that returns the hash entry for the specified name
+		/// </summary>
+		/// <param name="Input"></param>
+		/// <returns></returns>
+		public HashEntry GetEntryForName(string Input)
+		{
+			return Hashes.Where(E => E.Name == Input).FirstOrDefault();
+		}
+
+		/// <summary>
+		/// Logs differences between two collections for debugging
+		/// </summary>
+		/// <param name="RHS"></param>
+		public void LogDifferences(HashCollection RHS)
+		{
+			foreach (var Entry in Hashes)
+			{
+				if (!RHS.HasEntryForName(Entry.Name))
+				{
+					Log.TraceInformation("RHS does not have an entry for {0}", Entry.Name);
+				}
+				else if (!RHS.Hashes.Contains(Entry))
+				{
+					HashEntry RHSEntry = RHS.GetEntryForName(Entry.Name);
+					Log.TraceInformation("RHS hash mismatch for {0}", Entry.Name);
+					Log.TraceInformation("Current:\n\t{0}\n\t{1}\n\t{2}", Entry.Hash, Entry.Name, Entry.MetaData);
+					Log.TraceInformation("RHS:\n\t{0}\n\t{1}\n\t{2}", RHSEntry.Hash, RHSEntry.Name, RHSEntry.MetaData);
+				}
+			}
+		}
+
+		/// <summary>
+		/// 
+		/// </summary>
+		/// <returns></returns>
+		public override int GetHashCode()
+		{
+			return EqualityComparer<HashSet<HashEntry>>.Default.GetHashCode(Hashes);
+		}
+
+		/// <summary>
+		/// Creates a hash collection from a previously serialized file
+		/// </summary>
+		/// <param name="InPath"></param>
+		/// <returns></returns>
+		public static HashCollection CreateFromFile(string InPath)
+		{
+			try
+			{
+				using (var stream = System.IO.File.OpenRead(InPath))
+				{
+					var Serializer = new XmlSerializer(typeof(HashCollection));
+					return Serializer.Deserialize(stream) as HashCollection;
+				}
+			}
+			catch (Exception Ex)
+			{
+				Log.TraceWarning("Failed to load HashCollection from {0}. {1}", InPath, Ex.Message);
+			}
+
+			return null;
+		}
+
+		/// <summary>
+		/// Saves a hash collection to the specified file path as XML.
+		/// </summary>
+		/// <param name="OutPath"></param>
+		public void SaveToFile(string OutPath)
+		{
+			XmlSerializer XS = new XmlSerializer(this.GetType());
+
+			//first serialize the object to memory stream,
+			//in case of exception, the original file is not corrupted
+			using (MemoryStream MS = new MemoryStream())
+			{
+				XmlWriterSettings Settings = new XmlWriterSettings();
+				Settings.Indent = true;
+				Settings.IndentChars = ("\t");
+				Settings.Encoding = Encoding.UTF8;
+
+				using (XmlWriter Writer = XmlWriter.Create(MS, Settings))
+				{
+					XS.Serialize(Writer, this);
+				}
+
+				MS.Flush();
+
+				string String = Encoding.UTF8.GetString(MS.ToArray());
+
+				File.WriteAllText(OutPath, String);
+			}
+		}
+	}
+}
