Added missing copyright notice

diff --git a/Engine/Source/Programs/AutomationTool/AutomationUtils/P4WriteConfig.cs b/Engine/Source/Programs/AutomationTool/AutomationUtils/P4WriteConfig.cs
--- a/Engine/Source/Programs/AutomationTool/AutomationUtils/P4WriteConfig.cs
+++ b/Engine/Source/Programs/AutomationTool/AutomationUtils/P4WriteConfig.cs
@@ -1,3 +1,5 @@
+// Copyright 1998-2019 Epic Games, Inc. All Rights Reserved.
+
 using System;
 using System.Text;
 using System.IO;
