Fixed potential issue when UAT is run on a clean workspace

diff --git a/Engine/Source/Programs/AutomationTool/AutomationUtils/ScriptCompiler.cs b/Engine/Source/Programs/AutomationTool/AutomationUtils/ScriptCompiler.cs
--- a/Engine/Source/Programs/AutomationTool/AutomationUtils/ScriptCompiler.cs
+++ b/Engine/Source/Programs/AutomationTool/AutomationUtils/ScriptCompiler.cs
@@ -277,7 +277,7 @@ namespace AutomationTool
 				return;
 			}
 
-			Log.TraceInformation("Dependencies are out of date. Compiling scripts.");
+			Log.TraceInformation("Dependencies are out of date. Compiling scripts....");
 
 			// clean old assemblies
 			CleanupScriptsAssemblies();
@@ -301,8 +301,9 @@ namespace AutomationTool
 				string ProjContents = File.ReadAllText(UATProjTemplate);
 				ProjContents = ProjContents.Replace("$(Modules)", ModulesList);
 
+				Directory.CreateDirectory(Path.GetDirectoryName(UATProjFile));
 				File.WriteAllText(UATProjFile, ProjContents);
-
+				
 				string MsBuildVerbosity = Log.OutputLevel >= LogEventType.Verbose ? "minimal" : "quiet";
 
 				var CmdLine = String.Format("\"{0}\" /p:Configuration={1} /verbosity:{2} /nologo", UATProjFile, BuildConfig, MsBuildVerbosity);
