Updated SyncedProject script with some minor QOL improvements

 - Added help documentation
 - Added list of files that should be clobbered to avoid workspace settings
 - Use helper to autodetect project path if a full path is not specified
 - Added skipbuild option
 - Added preview option
 - Added parallel sync (defaults to 2: 30% faster on internal network)
 - Added option to force sync
 - Added option to sync explicit paths (last two are for writable files)
 - Fix issue if buildinfo file is missing

diff --git a/Engine/Source/Programs/AutomationTool/AutomationUtils/P4Utils.cs b/Engine/Source/Programs/AutomationTool/AutomationUtils/P4Utils.cs
--- a/Engine/Source/Programs/AutomationTool/AutomationUtils/P4Utils.cs
+++ b/Engine/Source/Programs/AutomationTool/AutomationUtils/P4Utils.cs
@@ -1914,9 +1914,21 @@ namespace AutomationTool
 		/// Invokes p4 sync command.
 		/// </summary>
 		/// <param name="CommandLine">CommandLine to pass on to the command.</param>
-        public void Sync(string CommandLine, bool AllowSpew = true, bool SpewIsVerbose = false)
+        public void Sync(string CommandLine, bool AllowSpew = true, bool SpewIsVerbose = false, int Retries=0, int MaxWait=0)
 		{
-			LogP4("sync " + CommandLine, null, AllowSpew, SpewIsVerbose:SpewIsVerbose);
+			string SyncCommandLine = "sync " + CommandLine;
+
+			if (MaxWait > 0)
+			{
+				SyncCommandLine = string.Format("-vnet.maxwait={0} {1}", MaxWait, SyncCommandLine);
+			}
+
+			if (Retries > 0)
+			{
+				SyncCommandLine = string.Format("-r{0} {1}", Retries, SyncCommandLine);
+			}
+
+			LogP4(SyncCommandLine, null, AllowSpew, SpewIsVerbose:SpewIsVerbose);
 		}
 
 		/// <summary>
diff --git a/Engine/Source/Programs/AutomationTool/AutomationUtils/ProjectUtils.cs b/Engine/Source/Programs/AutomationTool/AutomationUtils/ProjectUtils.cs
--- a/Engine/Source/Programs/AutomationTool/AutomationUtils/ProjectUtils.cs
+++ b/Engine/Source/Programs/AutomationTool/AutomationUtils/ProjectUtils.cs
@@ -645,6 +645,45 @@ namespace AutomationTool
 				CommandUtils.DeleteDirectoryContents(RulesFolder);
 			}
 		}
+
+		/// <summary>
+		/// Takes a game name (e.g "ShooterGame") and tries to find the path to the project file
+		/// </summary>
+		/// <param name="GameName"></param>
+		/// <returns></returns>
+		public static FileReference FindProjectFileFromName(string GameName)
+		{
+			string ProjectFile = GameName;
+
+			// Look for GameName.uproject
+			if (string.IsNullOrEmpty(Path.GetExtension(ProjectFile)))
+			{
+				// if project was specified but had no extension then just add it.
+				ProjectFile = Path.ChangeExtension(GameName, ".uproject");
+			}
+
+			// easy!
+			if (File.Exists(ProjectFile))
+			{
+				return new FileReference(ProjectFile);
+			}
+
+			// check for sibling to engine
+			string SiblingPath = Path.Combine(Environment.CurrentDirectory, GameName, ProjectFile);
+
+			if (File.Exists(SiblingPath))
+			{
+				return new FileReference(SiblingPath);
+			}
+
+			// check projectfiles paths.
+			IEnumerable<FileReference> Projects = NativeProjects.EnumerateProjectFiles();
+
+			FileReference ProjectPath = Projects.Where(R => string.Equals(R.GetFileName(), ProjectFile, StringComparison.OrdinalIgnoreCase)).FirstOrDefault();
+						
+			// either valid or we're out of ideas...
+			return ProjectPath;
+		}
 	}
 
     public class BranchInfo
diff --git a/Engine/Source/Programs/AutomationTool/AutomationUtils/UE4Build.cs b/Engine/Source/Programs/AutomationTool/AutomationUtils/UE4Build.cs
--- a/Engine/Source/Programs/AutomationTool/AutomationUtils/UE4Build.cs
+++ b/Engine/Source/Programs/AutomationTool/AutomationUtils/UE4Build.cs
@@ -370,7 +370,10 @@ namespace AutomationTool
 					Version.BranchName = Branch;
 					Version.BuildVersionString = Build;
 
-					VersionFileUpdater.MakeFileWriteable(BuildVersionFile.FullName);
+					if (File.Exists(BuildVersionFile.FullName))
+					{
+						VersionFileUpdater.MakeFileWriteable(BuildVersionFile.FullName);
+					}
 
 					Version.Write(BuildVersionFile);
 				}
