Suppress warning for missing DebugNonUFS files when staging installed build.

diff --git a/Engine/Source/Programs/AutomationTool/BuildGraph/Tasks/TagReceiptTask.cs b/Engine/Source/Programs/AutomationTool/BuildGraph/Tasks/TagReceiptTask.cs
--- a/Engine/Source/Programs/AutomationTool/BuildGraph/Tasks/TagReceiptTask.cs
+++ b/Engine/Source/Programs/AutomationTool/BuildGraph/Tasks/TagReceiptTask.cs
@@ -169,13 +169,13 @@ namespace AutomationTool.Tasks
 							continue;
 						}
 
-						// Only add files that exist as dependencies are assumed to always exist
+						// Check which files exist, and warn about any that don't. Ignore debug files, as they are frequently excluded for size (eg. UE4 on GitHub). This matches logic during staging.
 						FileReference DependencyPath = RuntimeDependency.Path;
 						if (FileReference.Exists(DependencyPath))
 						{
 							Files.Add(DependencyPath);
 						}
-						else
+						else if(RuntimeDependency.Type != UnrealBuildTool.StagedFileType.DebugNonUFS)
 						{
 							CommandUtils.LogWarning("File listed as RuntimeDependency in {0} does not exist ({1})", TargetFile.FullName, DependencyPath.FullName);
 						}
