UAT: iterative deploy no longer prevents cleaning of staging directory
(copy 4369732 from Dev-Build)

Partial fix for

Might need additional work to reduce file copies from Build to Stage.

diff --git a/Engine/Source/Programs/AutomationTool/Scripts/CopyBuildToStagingDirectory.Automation.cs b/Engine/Source/Programs/AutomationTool/Scripts/CopyBuildToStagingDirectory.Automation.cs
--- a/Engine/Source/Programs/AutomationTool/Scripts/CopyBuildToStagingDirectory.Automation.cs
+++ b/Engine/Source/Programs/AutomationTool/Scripts/CopyBuildToStagingDirectory.Automation.cs
@@ -2210,7 +2210,7 @@ public partial class Project : CommandUtils
 		{
 			LogInformation("Skipping clean of staging directory due to -NoCleanStage argument.");
 		}
-		else if (SC.Stage && !Params.SkipStage && !Params.IterativeDeploy)
+		else if (SC.Stage && !Params.SkipStage)
 		{
 			if(Params.SkipPak)
 			{
@@ -2219,21 +2219,21 @@ public partial class Project : CommandUtils
 			}
 			else
 			{
-		LogInformation("Cleaning Stage Directory: {0}", SC.StageDirectory.FullName);
-			try
-			{
-				DeleteDirectory(SC.StageDirectory.FullName);
-			}
-			catch (Exception Ex)
-			{
-				// Delete cooked data (if any) as it may be incomplete / corrupted.
-				throw new AutomationException(ExitCode.Error_FailedToDeleteStagingDirectory, Ex, "Stage Failed. Failed to delete staging directory " + SC.StageDirectory.FullName);
+				LogInformation("Cleaning Stage Directory: {0}", SC.StageDirectory.FullName);
+				try
+				{
+					DeleteDirectory(SC.StageDirectory.FullName);
+				}
+				catch (Exception Ex)
+				{
+					// Delete cooked data (if any) as it may be incomplete / corrupted.
+					throw new AutomationException(ExitCode.Error_FailedToDeleteStagingDirectory, Ex, "Stage Failed. Failed to delete staging directory " + SC.StageDirectory.FullName);
+				}
 			}
 		}
-		}
 		else
 		{
-			LogInformation("Cleaning pak files: {0}", SC.StageDirectory.FullName);
+			LogInformation("Cleaning PAK files in stage directory: {0}", SC.StageDirectory.FullName);
 			try
 			{
 				// delete old pak files
