UAT: Create manifest directory, before writing to it

diff --git a/Engine/Source/Programs/AutomationTool/Scripts/CopyBuildToStagingDirectory.Automation.cs b/Engine/Source/Programs/AutomationTool/Scripts/CopyBuildToStagingDirectory.Automation.cs
--- a/Engine/Source/Programs/AutomationTool/Scripts/CopyBuildToStagingDirectory.Automation.cs
+++ b/Engine/Source/Programs/AutomationTool/Scripts/CopyBuildToStagingDirectory.Automation.cs
@@ -1243,6 +1243,7 @@ public partial class Project : CommandUtils
 
         EncryptionAndSigning.CryptoSettings PakCryptoSettings = EncryptionAndSigning.ParseCryptoSettings(DirectoryReference.FromFile(Params.RawProjectPath), SC.StageTargetPlatform.IniPlatformType);
 		FileReference CryptoKeysCacheFilename = FileReference.Combine(SC.MetadataDir, "Crypto.json");
+		DirectoryReference.CreateDirectory(SC.MetadataDir);
 		PakCryptoSettings.Save(CryptoKeysCacheFilename);
 
 		List<CreatePakParams> PakInputs = new List<CreatePakParams>();
@@ -1869,6 +1870,7 @@ public partial class Project : CommandUtils
 		// Parse and cache crypto settings from INI file
 		EncryptionAndSigning.CryptoSettings PakCryptoSettings = EncryptionAndSigning.ParseCryptoSettings(DirectoryReference.FromFile(Params.RawProjectPath), SC.StageTargetPlatform.IniPlatformType);
 		FileReference CryptoKeysCacheFilename = FileReference.Combine(SC.MetadataDir, "Crypto.json");
+		DirectoryReference.CreateDirectory(SC.MetadataDir);
 		PakCryptoSettings.Save(CryptoKeysCacheFilename);
 
 		// We still want to have a list of all files to stage. We will use the chunk manifests
