UAT: Aways create directory before saving encryption parameters

diff --git a/Engine/Source/Programs/AutomationTool/Scripts/CopyBuildToStagingDirectory.Automation.cs b/Engine/Source/Programs/AutomationTool/Scripts/CopyBuildToStagingDirectory.Automation.cs
--- a/Engine/Source/Programs/AutomationTool/Scripts/CopyBuildToStagingDirectory.Automation.cs
+++ b/Engine/Source/Programs/AutomationTool/Scripts/CopyBuildToStagingDirectory.Automation.cs
@@ -1241,9 +1241,8 @@ public partial class Project : CommandUtils
 
 		var UnrealPakResponseFile = CreatePakResponseFileFromStagingManifest(SC, SC.FilesToStage.UFSFiles);
 
-        EncryptionAndSigning.CryptoSettings PakCryptoSettings = EncryptionAndSigning.ParseCryptoSettings(DirectoryReference.FromFile(Params.RawProjectPath), SC.StageTargetPlatform.IniPlatformType);
+		EncryptionAndSigning.CryptoSettings PakCryptoSettings = EncryptionAndSigning.ParseCryptoSettings(DirectoryReference.FromFile(Params.RawProjectPath), SC.StageTargetPlatform.IniPlatformType);
 		FileReference CryptoKeysCacheFilename = FileReference.Combine(SC.MetadataDir, "Crypto.json");
-		DirectoryReference.CreateDirectory(SC.MetadataDir);
 		PakCryptoSettings.Save(CryptoKeysCacheFilename);
 
 		List<CreatePakParams> PakInputs = new List<CreatePakParams>();
@@ -1870,7 +1869,6 @@ public partial class Project : CommandUtils
 		// Parse and cache crypto settings from INI file
 		EncryptionAndSigning.CryptoSettings PakCryptoSettings = EncryptionAndSigning.ParseCryptoSettings(DirectoryReference.FromFile(Params.RawProjectPath), SC.StageTargetPlatform.IniPlatformType);
 		FileReference CryptoKeysCacheFilename = FileReference.Combine(SC.MetadataDir, "Crypto.json");
-		DirectoryReference.CreateDirectory(SC.MetadataDir);
 		PakCryptoSettings.Save(CryptoKeysCacheFilename);
 
 		// We still want to have a list of all files to stage. We will use the chunk manifests
diff --git a/Engine/Source/Programs/UnrealBuildTool/System/EncryptionAndSigning.cs b/Engine/Source/Programs/UnrealBuildTool/System/EncryptionAndSigning.cs
--- a/Engine/Source/Programs/UnrealBuildTool/System/EncryptionAndSigning.cs
+++ b/Engine/Source/Programs/UnrealBuildTool/System/EncryptionAndSigning.cs
@@ -130,6 +130,7 @@ namespace UnrealBuildTool
 			/// </summary>
 			public void Save(FileReference InFile)
 			{
+				DirectoryReference.CreateDirectory(InFile.Directory);
 				FileReference.WriteAllText(InFile, fastJSON.JSON.Instance.ToJSON(this, new fastJSON.JSONParameters {}));
 			}
 		}
