Fix compile warning when converting a blueprint project to native.

diff --git a/Engine/Source/Programs/AutomationTool/AutomationUtils/ProjectUtils.cs b/Engine/Source/Programs/AutomationTool/AutomationUtils/ProjectUtils.cs
--- a/Engine/Source/Programs/AutomationTool/AutomationUtils/ProjectUtils.cs
+++ b/Engine/Source/Programs/AutomationTool/AutomationUtils/ProjectUtils.cs
@@ -292,6 +292,8 @@ namespace AutomationTool
 				Writer.WriteLine("{");
 				Writer.WriteLine("\tpublic {0}(ReadOnlyTargetRules Target) : base(Target)", ProjectName);
 				Writer.WriteLine("\t{");
+				Writer.WriteLine("\t\tPCHUsage = PCHUsageMode.UseExplicitOrSharedPCHs;");
+				Writer.WriteLine();
 				Writer.WriteLine("\t\tPrivateDependencyModuleNames.Add(\"Core\");");
 				Writer.WriteLine("\t\tPrivateDependencyModuleNames.Add(\"Core\");");
 				Writer.WriteLine("\t}");
