Fix UAT compile error.

diff --git a/Engine/Source/Programs/AutomationTool/Scripts/BuildPluginCommand.Automation.cs b/Engine/Source/Programs/AutomationTool/Scripts/BuildPluginCommand.Automation.cs
--- a/Engine/Source/Programs/AutomationTool/Scripts/BuildPluginCommand.Automation.cs
+++ b/Engine/Source/Programs/AutomationTool/Scripts/BuildPluginCommand.Automation.cs
@@ -128,7 +128,7 @@ class BuildPlugin : BuildCommand
 			CommandUtils.LogInformation("Building plugin for host platforms: {0}", String.Join(", ", HostPlatforms));
 			foreach (UnrealTargetPlatform HostPlatform in HostPlatforms)
 			{
-				if (Plugin.bCanBeUsedWithUnrealHeaderTool)
+				if (Plugin.SupportedPrograms != null && Plugin.SupportedPrograms.Contains("UnrealHeaderTool"))
 				{
 					CompilePluginWithUBT(null, HostProjectPluginFile, Plugin, "UnrealHeaderTool", TargetType.Program, HostPlatform, UnrealTargetConfiguration.Development, ReceiptFileNames, String.Format("{0} -plugin={1}", AdditionalArgs, CommandUtils.MakePathSafeToUseWithCommandLine(HostProjectPluginFile.FullName)));
 				}
