PR #5151: Fix BuildPlugin automation command to include Shaders folder (Contributed by [REDACTED_PERSON])

diff --git a/Engine/Source/Programs/AutomationTool/Scripts/BuildPluginCommand.Automation.cs b/Engine/Source/Programs/AutomationTool/Scripts/BuildPluginCommand.Automation.cs
--- a/Engine/Source/Programs/AutomationTool/Scripts/BuildPluginCommand.Automation.cs
+++ b/Engine/Source/Programs/AutomationTool/Scripts/BuildPluginCommand.Automation.cs
@@ -239,6 +239,7 @@ class BuildPlugin : BuildCommand
 		Filter.Include("/Resources/...");
 		Filter.Include("/Content/...");
 		Filter.Include("/Intermediate/Build/.../Inc/...");
+		Filter.Include("/Shaders/...");
 		Filter.Include("/Source/...");
 
 		// Add custom rules for each platform
