Rework the way that custom icons are used on Windows.

* For code projects, the custom icon is embedded into the resource when compiled from UBT. This ensures the custom icon is set even before packaging, and removes the need for a deployment step.
* For content projects, the custom icon is set after packaging using the previous method.

diff --git a/Engine/Source/Programs/AutomationTool/Win/Win.Automation.csproj b/Engine/Source/Programs/AutomationTool/Win/Win.Automation.csproj
--- a/Engine/Source/Programs/AutomationTool/Win/Win.Automation.csproj
+++ b/Engine/Source/Programs/AutomationTool/Win/Win.Automation.csproj
@@ -52,6 +52,7 @@
     <Compile Include="Properties\AssemblyInfo.cs" />
     <Compile Include="Tasks\SrcSrvTask.cs" />
     <Compile Include="WinPlatform.Automation.cs" />
+    <Compile Include="WinResources.cs" />
   </ItemGroup>
   <ItemGroup>
     <ProjectReference Include="..\..\DotNETCommon\DotNETUtilities\DotNETUtilities.csproj">
diff --git a/Engine/Source/Programs/AutomationTool/Win/WinPlatform.Automation.cs b/Engine/Source/Programs/AutomationTool/Win/WinPlatform.Automation.cs
--- a/Engine/Source/Programs/AutomationTool/Win/WinPlatform.Automation.cs
+++ b/Engine/Source/Programs/AutomationTool/Win/WinPlatform.Automation.cs
@@ -184,17 +184,31 @@ public abstract class BaseWinPlatform : Platform
 
 	public override void Package(ProjectParams Params, DeploymentContext SC, int WorkingCL)
 	{
-        List<FileReference> ExeNames = GetExecutableNames(SC);
+		// If this is a content-only project and there's a custom icon, update the executable
+		if (!Params.HasDLCName && !Params.IsCodeBasedProject)
+		{
+			FileReference IconFile = FileReference.Combine(Params.RawProjectPath.Directory, "Build", "Windows", "Application.ico");
+			if(FileReference.Exists(IconFile))
+			{
+				CommandUtils.LogInformation("Updating executable with custom icon from {0}", IconFile);
 
-        // Select target configurations based on the exe list returned from GetExecutableNames
-        List<UnrealTargetConfiguration> TargetConfigs = SC.StageTargetConfigurations.GetRange(0, ExeNames.Count);
+				GroupIconResource GroupIcon = GroupIconResource.FromIco(IconFile.FullName);
 
-		if (!Params.HasDLCName)
-		{
-			WindowsExports.SetApplicationIcon(Params.RawProjectPath, Params.ShortProjectName, SC.ProjectRoot, TargetConfigs, ExeNames, SC.EngineRoot);
+				List<FileReference> ExecutablePaths = GetExecutableNames(SC);
+				foreach (FileReference ExecutablePath in ExecutablePaths)
+				{
+					using (ModuleResourceUpdate Update = new ModuleResourceUpdate(ExecutablePath.FullName, false))
+					{
+						const int IconResourceId = 123; // As defined in Engine\Source\Runtime\Launch\Resources\Windows\resource.h
+						if (GroupIcon != null)
+						{
+							Update.SetIcons(IconResourceId, GroupIcon);
+						}
+					}
+				}
+			}
 		}
 
-		// package up the program, potentially with an installer for Windows
 		PrintRunTime();
 	}
 
diff --git a/Engine/Source/Programs/AutomationTool/Win/WinResources.cs b/Engine/Source/Programs/AutomationTool/Win/WinResources.cs
new file mode 100644
--- /dev/null
+++ b/Engine/Source/Programs/AutomationTool/Win/WinResources.cs
@@ -0,0 +1,292 @@
+// Copyright 1998-2018 Epic Games, Inc. All Rights Reserved.
+using System;
+using System.Collections.Generic;
+using System.Linq;
+using System.Text;
+using System.IO;
+using System.Runtime.InteropServices;
+
+#if !__MonoCS__
+#pragma warning disable CS1591
+#endif
+
+public enum ResourceType
+{
+	Icon = 3,
+	RawData = 10,
+	GroupIcon = 14,
+	Version = 16,
+}
+
+public class IconResource
+{
+	public byte Width;
+	public byte Height;
+	public byte ColorCount;
+	public ushort Planes;
+	public ushort BitCount;
+	public byte[] Data;
+
+	public IconResource(byte InWidth, byte InHeight, byte InColorCount, ushort InPlanes, ushort InBitCount, byte[] InData)
+	{
+		Width = InWidth;
+		Height = InHeight;
+		ColorCount = InColorCount;
+		Planes = InPlanes;
+		BitCount = InBitCount;
+		Data = InData;
+	}
+
+	public override string ToString()
+	{
+		return String.Format("{0}x{1}@{2}bpp", Width, Height, BitCount);
+	}
+}
+
+public class GroupIconResource
+{
+	public IconResource[] Icons;
+
+	public GroupIconResource(IconResource[] InIcons)
+	{
+		Icons = InIcons;
+	}
+
+	public static GroupIconResource FromIco(string FileName)
+	{
+		using(FileStream Stream = new FileStream(FileName, FileMode.Open, FileAccess.Read))
+		{
+			BinaryReader Reader = new BinaryReader(Stream);
+
+			Reader.ReadUInt16(); // Reserved
+			Reader.ReadUInt16(); // Type
+			ushort Count = Reader.ReadUInt16();
+
+			IconResource[] Icons = new IconResource[Count];
+			uint[] Offsets = new uint[Count];
+
+			for(int Idx = 0; Idx < Count; Idx++)
+			{
+				Icons[Idx] = ReadIconHeader(Reader);
+				Offsets[Idx] = Reader.ReadUInt32();
+			}
+
+			for(int Idx = 0; Idx < Icons.Length; Idx++)
+			{
+				Stream.Seek(Offsets[Idx], SeekOrigin.Begin);
+				Stream.Read(Icons[Idx].Data, 0, Icons[Idx].Data.Length);
+			}
+
+			return new GroupIconResource(Icons);
+		}
+	}
+
+	public static GroupIconResource FromExe(string FileName)
+	{
+		using(ModuleResourceLibary Module = new ModuleResourceLibary(FileName))
+		{
+			byte[] GroupData = Module.ReadFirstResource(ResourceType.GroupIcon);
+			return FromByteArray(GroupData, Module);
+		}
+	}
+
+	public static GroupIconResource FromByteArray(byte[] GroupData, ModuleResourceLibary Module)
+	{
+		if (GroupData != null)
+		{
+			using (MemoryStream Input = new MemoryStream(GroupData))
+			{
+				BinaryReader Reader = new BinaryReader(Input);
+
+				Reader.ReadUInt16(); // Reserved
+				Reader.ReadUInt16(); // Type
+				ushort Count = Reader.ReadUInt16();
+
+				IconResource[] Icons = new IconResource[Count];
+				for (int Idx = 0; Idx < Count; Idx++)
+				{
+					Icons[Idx] = ReadIconHeader(Reader);
+					int IconId = Reader.ReadUInt16();
+					Icons[Idx].Data = Module.ReadResource(IconId, ResourceType.Icon);
+				}
+				return new GroupIconResource(Icons);
+			}
+		}
+		else
+		{
+			return null;
+		}
+	}
+
+	public byte[] ToByteArray()
+	{
+		using(MemoryStream Output = new MemoryStream())
+		{
+			BinaryWriter Writer = new BinaryWriter(Output);
+
+			Writer.Write((ushort)0); // Reserved
+			Writer.Write((ushort)1); // Type (1 = icons)
+			Writer.Write((ushort)Icons.Length); 
+
+			for(int Idx = 0; Idx < Icons.Length; Idx++)
+			{
+				WriteIconHeader(Writer, Icons[Idx]);
+				Writer.Write((ushort)(Idx + 1));
+			}
+
+			return Output.ToArray();
+		}
+	}
+
+	public static IconResource ReadIconHeader(BinaryReader Reader)
+	{
+		byte Width = Reader.ReadByte();
+		byte Height = Reader.ReadByte();
+		byte ColorCount = Reader.ReadByte();
+		if(Reader.ReadByte() != 0) throw new InvalidDataException();
+		ushort Planes = Reader.ReadUInt16();
+		ushort BitCount = Reader.ReadUInt16();
+		byte[] Data = new byte[Reader.ReadUInt32()];
+
+		return new IconResource(Width, Height, ColorCount, Planes, BitCount, Data);
+	}
+
+	public void WriteIconHeader(BinaryWriter Writer, IconResource Icon)
+	{
+		Writer.Write(Icon.Width);
+		Writer.Write(Icon.Height);
+		Writer.Write(Icon.ColorCount);
+		Writer.Write((byte)0);
+		Writer.Write(Icon.Planes);
+		Writer.Write(Icon.BitCount);
+		Writer.Write(Icon.Data.Length);
+	}
+}
+
+public class ModuleResourceLibary : IDisposable
+{
+	const uint LOAD_LIBRARY_AS_DATAFILE = 0x00000002; 
+
+	delegate bool EnumResourceNamesDelegate(IntPtr ModuleHandle, IntPtr Type, IntPtr Name, IntPtr Param);
+
+	[DllImport("kernel32.dll", SetLastError=true)]
+	extern static IntPtr LoadLibraryEx(string FileName, IntPtr FileHandle, uint Flags);
+
+	[DllImport("kernel32.dll", SetLastError=true)]
+	extern static bool FreeLibrary(IntPtr Handle);
+
+	[DllImport("kernel32.dll", SetLastError=true)]
+	static extern bool EnumResourceNames(IntPtr ModuleHandle, IntPtr Type, EnumResourceNamesDelegate EnumDelegate, IntPtr Param);
+
+	[DllImport("kernel32.dll", SetLastError=true)]
+	static extern IntPtr FindResource(IntPtr ModuleHandle, IntPtr Name, IntPtr Type);
+
+	[DllImport("kernel32.dll", SetLastError=true)]
+	static extern IntPtr LoadResource(IntPtr ModuleHandle, IntPtr ResourceHandle);
+
+	[DllImport("kernel32.dll", SetLastError=true)]
+	static extern IntPtr LockResource(IntPtr ResourceDataHandle);
+
+	[DllImport("kernel32.dll", SetLastError=true)]
+	static extern uint SizeofResource(IntPtr ModuleHandle, IntPtr ResourceDataHandle);
+
+	IntPtr ModuleHandle;
+
+	public ModuleResourceLibary(string FileName)
+	{
+		ModuleHandle = LoadLibraryEx(FileName, new IntPtr(0), LOAD_LIBRARY_AS_DATAFILE);
+	}
+
+	public void Dispose()
+	{
+		FreeLibrary(ModuleHandle);
+	}
+
+	public byte[] ReadResource(int ResourceId, ResourceType Type)
+	{
+		return ReadResourceInternal(new IntPtr(ResourceId), new IntPtr((int)Type));
+	}
+
+	public byte[] ReadFirstResource(ResourceType Type)
+	{
+		byte[] Result = null;
+		EnumResourceNames(ModuleHandle, new IntPtr((int)Type), (Handle, FoundType, Name, Param) => { Result = ReadResourceInternal(Name, FoundType); return false; }, new IntPtr(0));
+		return Result;
+	}
+
+	byte[] ReadResourceInternal(IntPtr Name, IntPtr Type)
+	{
+		IntPtr ResourceHandle = FindResource(ModuleHandle, Name, Type);
+
+		IntPtr ResourceDataHandle = LoadResource(ModuleHandle, ResourceHandle);
+		IntPtr ResourceData = LockResource(ResourceDataHandle);
+	
+		uint ResourceSize = SizeofResource(ModuleHandle, ResourceHandle);
+
+		byte[] ResourceDataArray = new byte[ResourceSize];
+		Marshal.Copy(ResourceData, ResourceDataArray, 0, (int)ResourceSize);
+
+		return ResourceDataArray;
+	}
+}
+
+public class ModuleResourceUpdate : IDisposable
+{
+	[DllImport("kernel32.dll", SetLastError=true)]
+	static extern IntPtr BeginUpdateResource(string pFileName, [MarshalAs(UnmanagedType.Bool)]bool bDeleteExistingResources);
+
+	[DllImport("kernel32.dll", SetLastError=true)]
+	static extern bool UpdateResource(IntPtr hUpdate, IntPtr lpType, IntPtr lpName, ushort wLanguage, IntPtr lpData, uint cbData);
+
+	[DllImport("kernel32.dll", SetLastError=true)]
+	static extern bool EndUpdateResource(IntPtr hUpdate, bool fDiscard);
+
+	const ushort DefaultLanguage = 1033; // en-us
+
+	IntPtr UpdateHandle;
+	List<IntPtr> UnmanagedPointers = new List<IntPtr>();
+
+	public ModuleResourceUpdate(string OutputFile, bool bRemoveExisting)
+	{
+		UpdateHandle = BeginUpdateResource(OutputFile, bRemoveExisting);
+	}
+
+	public void SetData(int ResourceId, ResourceType Type, byte[] Data)
+	{
+		IntPtr UnmanagedPointer = Marshal.AllocHGlobal(Data.Length);
+		UnmanagedPointers.Add(UnmanagedPointer);
+
+		Marshal.Copy(Data, 0, UnmanagedPointer, Data.Length);
+
+		if(!UpdateResource(UpdateHandle, new IntPtr((int)Type), new IntPtr(ResourceId), DefaultLanguage, UnmanagedPointer, (uint)Data.Length))
+		{
+			throw new Exception("Couldn't update resource");
+		}
+	}
+
+	public void SetIcons(int ResourceId, GroupIconResource GroupIcon)
+	{
+		// Write the icon group resource
+		SetData(ResourceId, ResourceType.GroupIcon, GroupIcon.ToByteArray());
+
+		// Write the individual icons
+		for(int Idx = 0; Idx < GroupIcon.Icons.Length; Idx++)
+		{
+			SetData(Idx + 1, ResourceType.Icon, GroupIcon.Icons[Idx].Data);
+		}
+	}
+
+	public void Dispose()
+	{
+		EndUpdateResource(UpdateHandle, false);
+		foreach(IntPtr UnmanagedPointer in UnmanagedPointers)
+		{
+			Marshal.FreeHGlobal(UnmanagedPointer);
+		}
+		UnmanagedPointers.Clear();
+	}
+}
+
+#if !__MonoCS__
+#pragma warning restore CS1591
+#endif
diff --git a/Engine/Source/Programs/UnrealBuildTool/Configuration/UEBuildModuleCPP.cs b/Engine/Source/Programs/UnrealBuildTool/Configuration/UEBuildModuleCPP.cs
--- a/Engine/Source/Programs/UnrealBuildTool/Configuration/UEBuildModuleCPP.cs
+++ b/Engine/Source/Programs/UnrealBuildTool/Configuration/UEBuildModuleCPP.cs
@@ -565,7 +565,20 @@ namespace UnrealBuildTool
 			// must be specified relative to the resource file itself or Engine/Source.
 			if(SourceFilesToBuild.RCFiles.Count > 0)
 			{
+				// Figure the icon to use. We can only use a custom icon when compiling to a project-specific intemediate directory (and not for the shared editor executable, for example).
+				FileReference IconFile;
+				if(Target.ProjectFile != null && IntermediateDirectory.IsUnderDirectory(Target.ProjectFile.Directory))
+				{
+					IconFile = WindowsPlatform.GetApplicationIcon(Target.ProjectFile);
+				}
+				else
+				{
+					IconFile = WindowsPlatform.GetApplicationIcon(null);
+				}
+
+				// Setup the compile environment, setting the icon to use via a macro. This is used in PCLaunch.rc2.
 				CppCompileEnvironment ResourceCompileEnvironment = new CppCompileEnvironment(BinaryCompileEnvironment);
+				ResourceCompileEnvironment.Definitions.Add(String.Format("BUILD_ICON_FILE_NAME=\"\\\"{0}\\\"\"", IconFile.FullName.Replace("\\", "\\\\")));
 				LinkInputFiles.AddRange(ToolChain.CompileRCFiles(ResourceCompileEnvironment, SourceFilesToBuild.RCFiles, IntermediateDirectory, ActionGraph).ObjectFiles);
 			}
 
diff --git a/Engine/Source/Programs/UnrealBuildTool/Configuration/UEBuildPlatform.cs b/Engine/Source/Programs/UnrealBuildTool/Configuration/UEBuildPlatform.cs
--- a/Engine/Source/Programs/UnrealBuildTool/Configuration/UEBuildPlatform.cs
+++ b/Engine/Source/Programs/UnrealBuildTool/Configuration/UEBuildPlatform.cs
@@ -656,30 +656,6 @@ namespace UnrealBuildTool
 			return TempList;
 		}
 
-		/// <summary>
-		/// Return whether this platform has uniquely named binaries across multiple games
-		/// </summary>
-		public virtual bool HasUniqueBinaries()
-		{
-			return true;
-		}
-
-		/// <summary>
-		/// Return whether we wish to have this platform's binaries in our builds
-		/// </summary>
-		public virtual bool IsBuildRequired()
-		{
-			return true;
-		}
-
-		/// <summary>
-		/// Return whether we wish to have this platform's binaries in our CIS tests
-		/// </summary>
-		public virtual bool IsCISRequired()
-		{
-			return true;
-		}
-
 		/// <summary>
 		/// Whether the build platform requires deployment prep
 		/// </summary>
diff --git a/Engine/Source/Programs/UnrealBuildTool/Platform/Mac/UEBuildMac.cs b/Engine/Source/Programs/UnrealBuildTool/Platform/Mac/UEBuildMac.cs
--- a/Engine/Source/Programs/UnrealBuildTool/Platform/Mac/UEBuildMac.cs
+++ b/Engine/Source/Programs/UnrealBuildTool/Platform/Mac/UEBuildMac.cs
@@ -232,22 +232,6 @@ namespace UnrealBuildTool
 		{
 		}
 
-		/// <summary>
-		/// Return whether we wish to have this platform's binaries in our builds
-		/// </summary>
-		public override bool IsBuildRequired()
-		{
-			return false;
-		}
-
-		/// <summary>
-		/// Return whether we wish to have this platform's binaries in our CIS tests
-		/// </summary>
-		public override bool IsCISRequired()
-		{
-			return false;
-		}
-
 		public override DirectoryReference GetBundleDirectory(ReadOnlyTargetRules Rules, List<FileReference> OutputFiles)
 		{
 			if(Rules.bIsBuildingConsoleApplication)
diff --git a/Engine/Source/Programs/UnrealBuildTool/Platform/Windows/UEBuildWindows.cs b/Engine/Source/Programs/UnrealBuildTool/Platform/Windows/UEBuildWindows.cs
--- a/Engine/Source/Programs/UnrealBuildTool/Platform/Windows/UEBuildWindows.cs
+++ b/Engine/Source/Programs/UnrealBuildTool/Platform/Windows/UEBuildWindows.cs
@@ -1329,12 +1329,24 @@ namespace UnrealBuildTool
 		}
 
 		/// <summary>
-		/// Return whether this platform has uniquely named binaries across multiple games
+		/// Gets the application icon for a given project
 		/// </summary>
-		public override bool HasUniqueBinaries()
+		/// <param name="ProjectFile">The project file</param>
+		/// <returns>The icon to use for this project</returns>
+		public static FileReference GetApplicationIcon(FileReference ProjectFile)
 		{
-			// Windows applications have many shared binaries between games
-			return false;
+			// Check if there's a custom icon
+			if(ProjectFile != null)
+			{
+				FileReference IconFile = FileReference.Combine(ProjectFile.Directory, "Build", "Windows", "Application.ico");
+				if(FileReference.Exists(IconFile))
+				{
+					return IconFile;
+				}
+			}
+
+			// Otherwise use the default
+			return FileReference.Combine(UnrealBuildTool.EngineDirectory, "Source", "Runtime", "Launch", "Resources", "Windows", "UE4.ico");
 		}
 
 		/// <summary>
@@ -1671,7 +1683,6 @@ namespace UnrealBuildTool
 		/// <param name="Target">Information about the target being deployed</param>
 		public override void Deploy(UEBuildDeployTarget Target)
 		{
-			new BaseWindowsDeploy().PrepTargetForDeployment(Target);
 		}
 	}
 
diff --git a/Engine/Source/Programs/UnrealBuildTool/Platform/Windows/WinResources.cs b/Engine/Source/Programs/UnrealBuildTool/Platform/Windows/WinResources.cs
deleted file mode 100644
diff --git a/Engine/Source/Programs/UnrealBuildTool/Platform/Windows/WindowsDeploy.cs b/Engine/Source/Programs/UnrealBuildTool/Platform/Windows/WindowsDeploy.cs
deleted file mode 100644
diff --git a/Engine/Source/Programs/UnrealBuildTool/Platform/Windows/WindowsExports.cs b/Engine/Source/Programs/UnrealBuildTool/Platform/Windows/WindowsExports.cs
--- a/Engine/Source/Programs/UnrealBuildTool/Platform/Windows/WindowsExports.cs
+++ b/Engine/Source/Programs/UnrealBuildTool/Platform/Windows/WindowsExports.cs
@@ -14,22 +14,6 @@ namespace UnrealBuildTool
 	/// </summary>
 	public static class WindowsExports
 	{
-		/// <summary>
-		/// 
-		/// </summary>
-		/// <param name="ProjectFile"></param>
-		/// <param name="ProjectName"></param>
-		/// <param name="ProjectDirectory"></param>
-		/// <param name="InTargetConfigurations"></param>
-		/// <param name="InExecutablePaths"></param>
-		/// <param name="EngineDirectory"></param>
-		/// <returns></returns>
-		public static bool SetApplicationIcon(FileReference ProjectFile, string ProjectName, DirectoryReference ProjectDirectory, List<UnrealTargetConfiguration> InTargetConfigurations, List<FileReference> InExecutablePaths, DirectoryReference EngineDirectory)
-		{
-	        BaseWindowsDeploy Deploy = new BaseWindowsDeploy();
-            return Deploy.SetApplicationIcon(ProjectFile, ProjectName, ProjectDirectory.FullName, InTargetConfigurations, InExecutablePaths.Select(x => x.FullName).ToList(), EngineDirectory.FullName);
-		}
-
 		/// <summary>
 		/// Tries to get the directory for an installed Visual Studio version
 		/// </summary>
diff --git a/Engine/Source/Programs/UnrealBuildTool/UnrealBuildTool.csproj b/Engine/Source/Programs/UnrealBuildTool/UnrealBuildTool.csproj
--- a/Engine/Source/Programs/UnrealBuildTool/UnrealBuildTool.csproj
+++ b/Engine/Source/Programs/UnrealBuildTool/UnrealBuildTool.csproj
@@ -251,8 +251,6 @@
     <Compile Include="ToolChain\RemoteMac.cs" />
     <Compile Include="UnrealBuildTool.cs" />
     <Compile Include="Platform\Windows\PVSToolChain.cs" />
-    <Compile Include="Platform\Windows\WindowsDeploy.cs" />
-    <Compile Include="Platform\Windows\WinResources.cs" />
   </ItemGroup>
   <Choose>
     <When Condition="Exists('Platform\IOS\UEBuildIOS.cs')">
diff --git a/Engine/Source/Runtime/Launch/Resources/Windows/PCLaunch.rc b/Engine/Source/Runtime/Launch/Resources/Windows/PCLaunch.rc
--- a/Engine/Source/Runtime/Launch/Resources/Windows/PCLaunch.rc
+++ b/Engine/Source/Runtime/Launch/Resources/Windows/PCLaunch.rc
@@ -118,14 +118,6 @@ END
 
 
 
-/////////////////////////////////////////////////////////////////////////////
-//
-// Icon
-//
-
-// Icon with lowest ID value placed first to ensure application icon
-// remains consistent on all systems.
-IDICON_UE4Game          ICON                    "UE4.ico"
 #endif    // English (United States) resources
 /////////////////////////////////////////////////////////////////////////////
 
diff --git a/Engine/Source/Runtime/Launch/Resources/Windows/PCLaunch.rc2 b/Engine/Source/Runtime/Launch/Resources/Windows/PCLaunch.rc2
--- a/Engine/Source/Runtime/Launch/Resources/Windows/PCLaunch.rc2
+++ b/Engine/Source/Runtime/Launch/Resources/Windows/PCLaunch.rc2
@@ -47,6 +47,18 @@ BEGIN
 	END
 END
 
+/////////////////////////////////////////////////////////////////////////////
+//
+// Icon
+//
+
+// Icon with lowest ID value placed first to ensure application icon
+// remains consistent on all systems.
+
+#ifdef BUILD_ICON_FILE_NAME
+IDICON_UE4Game          ICON                    BUILD_ICON_FILE_NAME
+#endif
+
 /////////////////////////////////////////////////////////////////////////////
 //
 // Application Manifest
