Fix inability to package Mac targets from Windows.

diff --git a/Engine/Source/Programs/AutomationTool/Mac/MacPlatform.Automation.cs b/Engine/Source/Programs/AutomationTool/Mac/MacPlatform.Automation.cs
--- a/Engine/Source/Programs/AutomationTool/Mac/MacPlatform.Automation.cs
+++ b/Engine/Source/Programs/AutomationTool/Mac/MacPlatform.Automation.cs
@@ -140,10 +140,12 @@ public class MacPlatform : Platform
 							BootstrapExeName = SC.ShortProjectName + ".app";
 						}
 
-						string AppPath = Executable.Path.FullName.Substring(0, Executable.Path.FullName.LastIndexOf(".app/") + 4);
+						string AppSuffix = ".app" + Path.DirectorySeparatorChar;
+
+						string AppPath = Executable.Path.FullName.Substring(0, Executable.Path.FullName.LastIndexOf(AppSuffix) + AppSuffix.Length);
 						foreach (var DestPath in StagedFiles)
 						{
-							string AppRelativePath = DestPath.Name.Substring(0, DestPath.Name.LastIndexOf(".app/") + 4);
+							string AppRelativePath = DestPath.Name.Substring(0, DestPath.Name.LastIndexOf(AppSuffix) + AppSuffix.Length);
 							StageBootstrapExecutable(SC, BootstrapExeName, AppPath, AppRelativePath, BootstrapArguments);
 						}
 					}
