UAT: Allow specifying the uebp_CLIENT Perforce environment variable without having to specify uebp_BuildRoot_P4 and uebp_CLIENT_ROOT variables.

diff --git a/Engine/Source/Programs/AutomationTool/AutomationUtils/P4Environment.cs b/Engine/Source/Programs/AutomationTool/AutomationUtils/P4Environment.cs
--- a/Engine/Source/Programs/AutomationTool/AutomationUtils/P4Environment.cs
+++ b/Engine/Source/Programs/AutomationTool/AutomationUtils/P4Environment.cs
@@ -167,33 +167,47 @@ namespace AutomationTool
 
 			// Get the Perforce client setting
 			Client = CommandUtils.GetEnvVar(EnvVarNames.Client);
-			if(String.IsNullOrEmpty(Client))
+			Branch = CommandUtils.GetEnvVar(EnvVarNames.BuildRootP4);
+			ClientRoot = CommandUtils.GetEnvVar(EnvVarNames.ClientRoot);
+			if(String.IsNullOrEmpty(Client) || String.IsNullOrEmpty(Branch) || String.IsNullOrEmpty(ClientRoot))
 			{
+				// Create a connection using the current setting
 				P4Connection DefaultConnection = new P4Connection(User: User, Client: null, ServerAndPort: ServerAndPort);
-				P4ClientInfo ThisClient = DetectClient(DefaultConnection, User, Environment.MachineName.ToLower(), CmdEnv.UATExe);
-				Log.TraceInformation("Using user {0} clientspec {1} {2}", User, ThisClient.Name, ThisClient.RootPath);
 
-				string BranchPath;
-				string ClientRootPath;
-				P4Connection ClientConnection = new P4Connection(User: User, Client: ThisClient.Name, ServerAndPort: ServerAndPort);
-				DetectRootPaths(ClientConnection, CmdEnv.LocalRoot, ThisClient, out BranchPath, out ClientRootPath);
-
-				Client = ThisClient.Name;
-				CommandUtils.SetEnvVar(EnvVarNames.Client, Client);
-
-				Branch = BranchPath;
-				CommandUtils.SetEnvVar(EnvVarNames.BuildRootP4, Branch);
+				// Get the client info
+				P4ClientInfo ThisClient;
+				if(String.IsNullOrEmpty(Client))
+				{
+					ThisClient = DetectClient(DefaultConnection, User, Environment.MachineName.ToLower(), CmdEnv.UATExe);
+					Log.TraceInformation("Using user {0} clientspec {1} {2}", User, ThisClient.Name, ThisClient.RootPath);
+					Client = ThisClient.Name;
+					CommandUtils.SetEnvVar(EnvVarNames.Client, Client);
+				}
+				else
+				{
+					ThisClient = DefaultConnection.GetClientInfo(Client);
+				}
 
-				ClientRoot = ClientRootPath;
-				CommandUtils.SetEnvVar(EnvVarNames.ClientRoot, ClientRootPath);
-			}
-			else
-			{
-				Branch = CommandUtils.GetEnvVar(EnvVarNames.BuildRootP4);
-				ClientRoot = CommandUtils.GetEnvVar(EnvVarNames.ClientRoot);
+				// Detect the root paths
 				if(String.IsNullOrEmpty(Branch) || String.IsNullOrEmpty(ClientRoot))
 				{
-					throw new AutomationException("{0} and {1} must also be set with {2}", EnvVarNames.ClientRoot, EnvVarNames.BuildRootP4, EnvVarNames.Client);
+					P4Connection ClientConnection = new P4Connection(User: User, Client: ThisClient.Name, ServerAndPort: ServerAndPort);
+
+					string BranchPath;
+					string ClientRootPath;
+					DetectRootPaths(ClientConnection, CmdEnv.LocalRoot, ThisClient, out BranchPath, out ClientRootPath);
+
+					if(String.IsNullOrEmpty(Branch))
+					{
+						Branch = BranchPath;
+						CommandUtils.SetEnvVar(EnvVarNames.BuildRootP4, Branch);
+					}
+
+					if(String.IsNullOrEmpty(ClientRoot))
+					{
+						ClientRoot = ClientRootPath;
+						CommandUtils.SetEnvVar(EnvVarNames.ClientRoot, ClientRootPath);
+					}
 				}
 			}
 
