Add -projectfileformat=XCode when generating project files for deploy.

diff --git a/Engine/Source/Programs/AutomationTool/IOS/IOSPlatform.Automation.cs b/Engine/Source/Programs/AutomationTool/IOS/IOSPlatform.Automation.cs
--- a/Engine/Source/Programs/AutomationTool/IOS/IOSPlatform.Automation.cs
+++ b/Engine/Source/Programs/AutomationTool/IOS/IOSPlatform.Automation.cs
@@ -530,7 +530,7 @@ public class IOSPlatform : Platform
 		{
 			// project.xcodeproj doesn't exist, so generate temp project
 			string Arguments = "-project=\"" + RawProjectPath + "\"";
-			Arguments += " -platforms=" + PlatformName + " -game -nointellisense -" + PlatformName + "deployonly -ignorejunk";
+			Arguments += " -platforms=" + PlatformName + " -game -nointellisense -" + PlatformName + "deployonly -ignorejunk -projectfileformat=XCode";
 
 			// If engine is installed then UBT doesn't need to be built
 			if (Automation.IsEngineInstalled())
