Fixes to plugin packaging script.

* Now much simpler to build a plugin against a target: UE4Game Win64 Development -Plugin=Foo.uplugin will now configure all appropriate settings and will include all files to be distributed in the manifest.
* BuildPlugin command in UAT now uses the manifest to determine which files need to be distributed.

diff --git a/Engine/Source/Programs/AutomationTool/Scripts/BuildPluginCommand.Automation.cs b/Engine/Source/Programs/AutomationTool/Scripts/BuildPluginCommand.Automation.cs
--- a/Engine/Source/Programs/AutomationTool/Scripts/BuildPluginCommand.Automation.cs
+++ b/Engine/Source/Programs/AutomationTool/Scripts/BuildPluginCommand.Automation.cs
@@ -120,7 +120,7 @@ class BuildPlugin : BuildCommand
 
 	FileReference[] CompilePlugin(FileReference HostProjectFile, FileReference HostProjectPluginFile, PluginDescriptor Plugin, List<UnrealTargetPlatform> HostPlatforms, List<UnrealTargetPlatform> TargetPlatforms, string AdditionalArgs)
 	{
-		List<FileReference> ReceiptFileNames = new List<FileReference>();
+		List<FileReference> ManifestFileNames = new List<FileReference>();
 
 		// Build the host platforms
 		if(HostPlatforms.Count > 0)
@@ -130,9 +130,9 @@ class BuildPlugin : BuildCommand
 			{
 				if (Plugin.SupportedPrograms != null && Plugin.SupportedPrograms.Contains("UnrealHeaderTool"))
 				{
-					CompilePluginWithUBT(null, HostProjectPluginFile, Plugin, "UnrealHeaderTool", TargetType.Program, HostPlatform, UnrealTargetConfiguration.Development, ReceiptFileNames, String.Format("{0} -plugin={1}", AdditionalArgs, CommandUtils.MakePathSafeToUseWithCommandLine(HostProjectPluginFile.FullName)));
+					CompilePluginWithUBT(null, HostProjectPluginFile, Plugin, "UnrealHeaderTool", TargetType.Program, HostPlatform, UnrealTargetConfiguration.Development, ManifestFileNames, String.Format("{0} -plugin={1}", AdditionalArgs, CommandUtils.MakePathSafeToUseWithCommandLine(HostProjectPluginFile.FullName)));
 				}
-				CompilePluginWithUBT(HostProjectFile, HostProjectPluginFile, Plugin, "UE4Editor", TargetType.Editor, HostPlatform, UnrealTargetConfiguration.Development, ReceiptFileNames, AdditionalArgs);
+				CompilePluginWithUBT(HostProjectFile, HostProjectPluginFile, Plugin, "UE4Editor", TargetType.Editor, HostPlatform, UnrealTargetConfiguration.Development, ManifestFileNames, AdditionalArgs);
 			}
 		}
 
@@ -144,27 +144,23 @@ class BuildPlugin : BuildCommand
 			{
 				if(Plugin.SupportsTargetPlatform(TargetPlatform))
 				{
-					CompilePluginWithUBT(HostProjectFile, HostProjectPluginFile, Plugin, "UE4Game", TargetType.Game, TargetPlatform, UnrealTargetConfiguration.Development, ReceiptFileNames, null);
-					CompilePluginWithUBT(HostProjectFile, HostProjectPluginFile, Plugin, "UE4Game", TargetType.Game, TargetPlatform, UnrealTargetConfiguration.Shipping, ReceiptFileNames, null);
+					CompilePluginWithUBT(HostProjectFile, HostProjectPluginFile, Plugin, "UE4Game", TargetType.Game, TargetPlatform, UnrealTargetConfiguration.Development, ManifestFileNames, null);
+					CompilePluginWithUBT(HostProjectFile, HostProjectPluginFile, Plugin, "UE4Game", TargetType.Game, TargetPlatform, UnrealTargetConfiguration.Shipping, ManifestFileNames, null);
 				}
 			}
 		}
 
 		// Package the plugin to the output folder
 		HashSet<FileReference> BuildProducts = new HashSet<FileReference>();
-		foreach(FileReference ReceiptFileName in ReceiptFileNames)
+		foreach(FileReference ManifestFileName in ManifestFileNames)
 		{
-			TargetReceipt Receipt;
-			if(!TargetReceipt.TryRead(ReceiptFileName, CommandUtils.EngineDirectory, HostProjectFile.Directory, out Receipt))
-			{
-				throw new AutomationException("Missing or invalid target receipt ({0})", ReceiptFileName);
-			}
-			BuildProducts.UnionWith(Receipt.BuildProducts.Select(x => x.Path).Where(x => x.IsUnderDirectory(HostProjectPluginFile.Directory)));
+			BuildManifest Manifest = CommandUtils.ReadManifest(ManifestFileName);
+			BuildProducts.UnionWith(Manifest.BuildProducts.Select(x => new FileReference(x)));
 		}
 		return BuildProducts.ToArray();
 	}
 
-	void CompilePluginWithUBT(FileReference HostProjectFile, FileReference HostProjectPluginFile, PluginDescriptor Plugin, string TargetName, TargetType TargetType, UnrealTargetPlatform Platform, UnrealTargetConfiguration Configuration, List<FileReference> ReceiptFileNames, string InAdditionalArgs)
+	void CompilePluginWithUBT(FileReference HostProjectFile, FileReference HostProjectPluginFile, PluginDescriptor Plugin, string TargetName, TargetType TargetType, UnrealTargetPlatform Platform, UnrealTargetConfiguration Configuration, List<FileReference> ManifestFileNames, string InAdditionalArgs)
 	{
 		// Find a list of modules that need to be built for this plugin
 		bool bCompilePlatform = false;
@@ -185,12 +181,10 @@ class BuildPlugin : BuildCommand
 		// Add these modules to the build agenda
 		if(bCompilePlatform)
 		{
-			string Architecture = PlatformExports.GetDefaultArchitecture(Platform, HostProjectFile);
-
-			FileReference ReceiptFileName = TargetReceipt.GetDefaultPath(HostProjectPluginFile.Directory, TargetName, Platform, Configuration, Architecture);
-			ReceiptFileNames.Add(ReceiptFileName);
+			FileReference ManifestFileName = FileReference.Combine(HostProjectFile.Directory, "Saved", String.Format("Manifest-{0}-{1}-{2}.xml", TargetName, Platform, Configuration));
+			ManifestFileNames.Add(ManifestFileName);
 			
-			string Arguments = String.Format("-plugin={0} -iwyu -precompile -nosharedpch -noubtmakefiles -receipt={1}", CommandUtils.MakePathSafeToUseWithCommandLine(HostProjectPluginFile.FullName), CommandUtils.MakePathSafeToUseWithCommandLine(ReceiptFileName.FullName));
+			string Arguments = String.Format("-plugin={0} -iwyu -noubtmakefiles -manifest={1}", CommandUtils.MakePathSafeToUseWithCommandLine(HostProjectPluginFile.FullName), CommandUtils.MakePathSafeToUseWithCommandLine(ManifestFileName.FullName));
 			if(!String.IsNullOrEmpty(InAdditionalArgs))
 			{
 				Arguments += InAdditionalArgs;
diff --git a/Engine/Source/Programs/UnrealBuildTool/Configuration/UEBuildBinary.cs b/Engine/Source/Programs/UnrealBuildTool/Configuration/UEBuildBinary.cs
--- a/Engine/Source/Programs/UnrealBuildTool/Configuration/UEBuildBinary.cs
+++ b/Engine/Source/Programs/UnrealBuildTool/Configuration/UEBuildBinary.cs
@@ -462,35 +462,32 @@ namespace UnrealBuildTool
 		public void GetBuildProducts(ReadOnlyTargetRules Target, UEToolChain ToolChain, Dictionary<FileReference, BuildProductType> BuildProducts, bool bCreateDebugInfo)
 		{
 			// Add all the precompiled outputs
-			if(Target.bPrecompile)
+			foreach(UEBuildModuleCPP Module in Modules.OfType<UEBuildModuleCPP>())
 			{
-				foreach(UEBuildModuleCPP Module in Modules.OfType<UEBuildModuleCPP>())
+				if(Module.Rules.bPrecompile)
 				{
-					if(Module.Rules.bPrecompile)
+					if(Module.GeneratedCodeDirectory != null && DirectoryReference.Exists(Module.GeneratedCodeDirectory))
 					{
-						if(Module.GeneratedCodeDirectory != null && DirectoryReference.Exists(Module.GeneratedCodeDirectory))
+						foreach(FileReference GeneratedCodeFile in DirectoryReference.EnumerateFiles(Module.GeneratedCodeDirectory))
 						{
-							foreach(FileReference GeneratedCodeFile in DirectoryReference.EnumerateFiles(Module.GeneratedCodeDirectory))
+							// Exclude timestamp files, since they're always updated and cause collisions between builds
+							if(!GeneratedCodeFile.GetFileName().Equals("Timestamp", StringComparison.OrdinalIgnoreCase) && !GeneratedCodeFile.HasExtension(".cpp"))
 							{
-								// Exclude timestamp files, since they're always updated and cause collisions between builds
-								if(!GeneratedCodeFile.GetFileName().Equals("Timestamp", StringComparison.OrdinalIgnoreCase))
-								{
-									BuildProducts.Add(GeneratedCodeFile, BuildProductType.BuildResource);
-								}
+								BuildProducts.Add(GeneratedCodeFile, BuildProductType.BuildResource);
 							}
 						}
-						if(Target.LinkType == TargetLinkType.Monolithic)
-						{
-							FileReference PrecompiledManifestLocation = Module.PrecompiledManifestLocation;
-							BuildProducts.Add(PrecompiledManifestLocation, BuildProductType.BuildResource);
+					}
+					if(Target.LinkType == TargetLinkType.Monolithic)
+					{
+						FileReference PrecompiledManifestLocation = Module.PrecompiledManifestLocation;
+						BuildProducts.Add(PrecompiledManifestLocation, BuildProductType.BuildResource);
 
-							PrecompiledManifest ModuleManifest = PrecompiledManifest.Read(PrecompiledManifestLocation);
-							foreach(FileReference OutputFile in ModuleManifest.OutputFiles)
+						PrecompiledManifest ModuleManifest = PrecompiledManifest.Read(PrecompiledManifestLocation);
+						foreach(FileReference OutputFile in ModuleManifest.OutputFiles)
+						{
+							if(!BuildProducts.ContainsKey(OutputFile))
 							{
-								if(!BuildProducts.ContainsKey(OutputFile))
-								{
-									BuildProducts.Add(OutputFile, BuildProductType.BuildResource);
-								}
+								BuildProducts.Add(OutputFile, BuildProductType.BuildResource);
 							}
 						}
 					}
diff --git a/Engine/Source/Programs/UnrealBuildTool/Configuration/UEBuildTarget.cs b/Engine/Source/Programs/UnrealBuildTool/Configuration/UEBuildTarget.cs
--- a/Engine/Source/Programs/UnrealBuildTool/Configuration/UEBuildTarget.cs
+++ b/Engine/Source/Programs/UnrealBuildTool/Configuration/UEBuildTarget.cs
@@ -468,6 +468,16 @@ namespace UnrealBuildTool
 				RulesObject.bDeployAfterCompile = false;
 			}
 
+			// If we're compiling a plugin, and this target is monolithic, just create the object files
+			if(Desc.ForeignPlugin != null && RulesObject.LinkType == TargetLinkType.Monolithic)
+			{
+				// Don't actually want an executable
+				RulesObject.bDisableLinking = true;
+
+				// Don't allow using shared PCHs; they won't be distributed with the plugin
+				RulesObject.bUseSharedPCHs = false;
+			}
+
 			// Include generated code plugin if not building an editor target and project is configured for nativization
 			if (RulesObject.ProjectFile != null
 				&& (RulesObject.Type == TargetType.Game || RulesObject.Type == TargetType.Client || RulesObject.Type == TargetType.Server)
@@ -1393,6 +1403,13 @@ namespace UnrealBuildTool
 				}
 			}
 
+			// Remove anything that's not part of the plugin
+			if(ForeignPlugin != null)
+			{
+				DirectoryReference ForeignPluginDir = ForeignPlugin.Directory;
+				Manifest.BuildProducts.RemoveAll(x => !new FileReference(x).IsUnderDirectory(ForeignPluginDir));
+			}
+
 			Manifest.BuildProducts.Sort();
 			Manifest.DeployTargetFiles.Sort();
 			Manifest.PostBuildScripts.Sort();
@@ -2214,7 +2231,7 @@ namespace UnrealBuildTool
 			// If we're just precompiling a plugin, only include output items which are under that directory
 			if(ForeignPlugin != null)
 			{
-				OutputItems.RemoveAll(x => x.Location.IsUnderDirectory(ForeignPlugin.Directory));
+				OutputItems.RemoveAll(x => !x.Location.IsUnderDirectory(ForeignPlugin.Directory));
 			}
 
 			// Allow the toolchain to modify the final output items
@@ -3230,12 +3247,6 @@ namespace UnrealBuildTool
 			// Find all the valid plugins
 			Dictionary<string, PluginInfo> NameToInfo = RulesAssembly.EnumeratePlugins().ToDictionary(x => x.Name, x => x, StringComparer.InvariantCultureIgnoreCase);
 
-			// Remove any foreign plugins; we will just precompile those
-			if(ForeignPlugin != null)
-			{
-				NameToInfo = NameToInfo.Where(x => !x.Value.File.IsUnderDirectory(ForeignPlugin.Directory)).ToDictionary(Pair => Pair.Key, Pair => Pair.Value, StringComparer.InvariantCultureIgnoreCase);
-			}
-
 			// Remove any plugins for platforms we don't have
 			List<UnrealTargetPlatform> MissingPlatforms = new List<UnrealTargetPlatform>();
 			foreach (UnrealTargetPlatform TargetPlatform in Enum.GetValues(typeof(UnrealTargetPlatform)))
@@ -3255,6 +3266,13 @@ namespace UnrealBuildTool
 			// Map of plugin names to instances of that plugin
 			Dictionary<string, UEBuildPlugin> NameToInstance = new Dictionary<string, UEBuildPlugin>(StringComparer.InvariantCultureIgnoreCase);
 
+			// Set up the foreign plugin
+			if(ForeignPlugin != null)
+			{
+				PluginReferenceDescriptor PluginReference = new PluginReferenceDescriptor(ForeignPlugin.GetFileNameWithoutExtension(), null, true);
+				AddPlugin(PluginReference, "command line", ExcludeFolders, NameToInstance, NameToInfo);
+			}
+
 			// Configure plugins explicitly enabled via target settings
 			foreach(string PluginName in Rules.EnablePlugins)
 			{
@@ -4011,6 +4029,7 @@ namespace UnrealBuildTool
 				// Clear the bUsePrecompiled flag if we're compiling a foreign plugin; since it's treated like an engine module, it will default to true in an installed build.
 				if(RulesObject.Plugin != null && RulesObject.Plugin.File == ForeignPlugin)
 				{
+					RulesObject.bPrecompile = true;
 					RulesObject.bUsePrecompiled = false;
 				}
 
