Create a separate installed DDC pak file for Enterprise.

diff --git a/Engine/Config/BaseEngine.ini b/Engine/Config/BaseEngine.ini
--- a/Engine/Config/BaseEngine.ini
+++ b/Engine/Config/BaseEngine.ini
@@ -1464,13 +1464,14 @@ EnginePak=(Type=ReadPak, Filename=../../../Engine/DerivedDataCache/DDC.ddp)
 MinimumDaysToKeepFile=7
 Root=(Type=KeyLength, Length=120, Inner=AsyncPut)
 AsyncPut=(Type=AsyncPut, Inner=Hierarchy)
-Hierarchy=(Type=Hierarchical, Inner=Boot, Inner=Pak, Inner=CompressedPak, Inner=EnginePak, Inner=Local, Inner=Shared)
+Hierarchy=(Type=Hierarchical, Inner=Boot, Inner=Pak, Inner=CompressedPak, Inner=EnginePak, Inner=EnterprisePak, Inner=Local, Inner=Shared)
 Boot=(Type=Boot, Filename="%ENGINEUSERDIR%DerivedDataCache/Boot.ddc", MaxCacheSize=512)
 Local=(Type=FileSystem, ReadOnly=false, Clean=false, Flush=false, PurgeTransient=true, DeleteUnused=true, UnusedFileAge=34, FoldersToClean=-1, Path="%ENGINEVERSIONAGNOSTICUSERDIR%DerivedDataCache", EditorOverrideSetting=LocalDerivedDataCache)
 Shared=(Type=FileSystem, ReadOnly=false, Clean=false, Flush=false, DeleteUnused=true, UnusedFileAge=10, FoldersToClean=10, MaxFileChecksPerSec=1, Path=?EpicDDC, EnvPathOverride=UE-SharedDataCachePath, EditorOverrideSetting=SharedDerivedDataCache)
 Pak=(Type=ReadPak, Filename="%GAMEDIR%DerivedDataCache/DDC.ddp")
 CompressedPak=(Type=ReadPak, Filename="%GAMEDIR%DerivedDataCache/Compressed.ddp", Compressed=true)
 EnginePak=(Type=ReadPak, Filename=../../../Engine/DerivedDataCache/Compressed.ddp, Compressed=true)
+EnterprisePak=(Type=ReadPak, Filename=../../../Enterprise/DerivedDataCache/Compressed.ddp, Compressed=true)
 
 [NoShared]
 MinimumDaysToKeepFile=7
diff --git a/Engine/Source/Editor/UnrealEd/Private/Commandlets/DerivedDataCacheCommandlet.cpp b/Engine/Source/Editor/UnrealEd/Private/Commandlets/DerivedDataCacheCommandlet.cpp
--- a/Engine/Source/Editor/UnrealEd/Private/Commandlets/DerivedDataCacheCommandlet.cpp
+++ b/Engine/Source/Editor/UnrealEd/Private/Commandlets/DerivedDataCacheCommandlet.cpp
@@ -152,7 +152,7 @@ int32 UDerivedDataCacheCommandlet::Main( const FString& Params )
 
 		if ( FilesInPath.Num() == 0 )
 		{
-			UE_LOG(LogDerivedDataCacheCommandlet, Warning, TEXT("No files found."));
+			UE_LOG(LogDerivedDataCacheCommandlet, Display, TEXT("No files found."));
 		}
 
 		ITargetPlatformManagerModule* TPM = GetTargetPlatformManager();
diff --git a/Engine/Source/Programs/AutomationTool/Scripts/BuildDerivedDataCache.Automation.cs b/Engine/Source/Programs/AutomationTool/Scripts/BuildDerivedDataCache.Automation.cs
--- a/Engine/Source/Programs/AutomationTool/Scripts/BuildDerivedDataCache.Automation.cs
+++ b/Engine/Source/Programs/AutomationTool/Scripts/BuildDerivedDataCache.Automation.cs
@@ -1,4 +1,4 @@
-﻿// Copyright 1998-2018 Epic Games, Inc. All Rights Reserved.
+// Copyright 1998-2018 Epic Games, Inc. All Rights Reserved.
 
 using AutomationTool;
 using UnrealBuildTool;
@@ -20,11 +20,13 @@ public class BuildDerivedDataCache : BuildCommand
 		UnrealTargetPlatform HostPlatform = BuildHostPlatform.Current.Platform;
 		string TargetPlatforms = ParseParamValue("TargetPlatforms");
 		string SavedDir = ParseParamValue("SavedDir");
+		string BackendName = ParseParamValue("BackendName", "CreateInstalledEnginePak");
+		string RelativePakPath = ParseParamValue("RelativePakPath", "Engine/DerivedDataCache/Compressed.ddp");
+		bool bSkipEngine = ParseParam("SkipEngine");
 
 
 		// Get paths to everything within the temporary directory
 		string EditorExe = CommandUtils.GetEditorCommandletExe(TempDir, HostPlatform);
-		string RelativePakPath = "Engine/DerivedDataCache/Compressed.ddp";
 		string OutputPakFile = CommandUtils.CombinePaths(TempDir, RelativePakPath);
 		string OutputCsvFile = Path.ChangeExtension(OutputPakFile, ".csv");
 
@@ -80,7 +82,7 @@ public class BuildDerivedDataCache : BuildCommand
 				ProjectSpecificPlatforms = CommandUtils.CombineCommandletParams(FilteredPlatforms.Distinct().ToArray());
 			}
 			CommandUtils.LogInformation("Generating DDC data for {0} on {1}", GameName, ProjectSpecificPlatforms);
-			CommandUtils.DDCCommandlet(FileRef, EditorExe, null, ProjectSpecificPlatforms, "-fill -DDC=CreateInstalledEnginePak -ProjectOnly");
+			CommandUtils.DDCCommandlet(FileRef, EditorExe, null, ProjectSpecificPlatforms, String.Format("-fill -DDC={0} -ProjectOnly", BackendName));
 
 			string ProjectPakFile = CommandUtils.CombinePaths(Path.GetDirectoryName(OutputPakFile), String.Format("Compressed-{0}.ddp", GameName));
 			CommandUtils.DeleteFile(ProjectPakFile);
@@ -96,7 +98,7 @@ public class BuildDerivedDataCache : BuildCommand
 
 		// Generate DDC for the editor, and merge all the other PAK files in
 		CommandUtils.LogInformation("Generating DDC data for engine content on {0}", TargetPlatforms);
-		CommandUtils.DDCCommandlet(null, EditorExe, null, TargetPlatforms, "-fill -DDC=CreateInstalledEnginePak " + CommandUtils.MakePathSafeToUseWithCommandLine("-MergePaks=" + String.Join("+", ProjectPakFiles)));
+		CommandUtils.DDCCommandlet(null, EditorExe, null, TargetPlatforms, String.Format("-fill -DDC={0} -MergePaks={1}{2}", BackendName, CommandUtils.MakePathSafeToUseWithCommandLine(String.Join("+", ProjectPakFiles)), bSkipEngine? " -projectonly" : ""));
 
 		string SavedPakFile = CommandUtils.CombinePaths(SavedDir, RelativePakPath);
 		CommandUtils.CopyFile(OutputPakFile, SavedPakFile);
