UE-64059 Failed to package for IOS on PC with BuildGraph
#4.21

BuildGraph is launched using AutomationTools, and the tasks (graph nodes) are also using AutomationTools, locking the Dlls.
Using AutomationToolLauncher instead, to allow shadow copying the Dlls.

diff --git a/Engine/Source/Programs/AutomationTool/AutomationUtils/CommandEnvironment.cs b/Engine/Source/Programs/AutomationTool/AutomationUtils/CommandEnvironment.cs
--- a/Engine/Source/Programs/AutomationTool/AutomationUtils/CommandEnvironment.cs
+++ b/Engine/Source/Programs/AutomationTool/AutomationUtils/CommandEnvironment.cs
@@ -76,6 +76,10 @@ namespace AutomationTool
 		{
 			// Get the path to the UAT executable
 			UATExe = Assembly.GetEntryAssembly().GetOriginalLocation();
+
+			// use the launcher to enable shadow copying the assemblies on Windows
+			UATExe = Regex.Replace(UATExe, "AutomationTool.exe", "AutomationToolLauncher.exe", RegexOptions.IgnoreCase);
+
 			if (!CommandUtils.FileExists(UATExe))
 			{
 				throw new AutomationException("Could not find AutomationTool.exe. Reflection indicated it was here: {0}", UATExe);
