Undo //UE4/Release-4.21/Engine/Source/Programs/AutomationTool/AutomationUtils/... changelist 4494623

Prevents UAT running on Mac.

diff --git a/Engine/Source/Programs/AutomationTool/AutomationUtils/CommandEnvironment.cs b/Engine/Source/Programs/AutomationTool/AutomationUtils/CommandEnvironment.cs
--- a/Engine/Source/Programs/AutomationTool/AutomationUtils/CommandEnvironment.cs
+++ b/Engine/Source/Programs/AutomationTool/AutomationUtils/CommandEnvironment.cs
@@ -76,10 +76,6 @@ namespace AutomationTool
 		{
 			// Get the path to the UAT executable
 			UATExe = Assembly.GetEntryAssembly().GetOriginalLocation();
-
-			// use the launcher to enable shadow copying the assemblies on Windows
-			UATExe = Regex.Replace(UATExe, "AutomationTool.exe", "AutomationToolLauncher.exe", RegexOptions.IgnoreCase);
-
 			if (!CommandUtils.FileExists(UATExe))
 			{
 				throw new AutomationException("Could not find AutomationTool.exe. Reflection indicated it was here: {0}", UATExe);
