BuildGraph: Allow specifying a target for badges, which marks all dependencies as required.

diff --git a/Engine/Source/Programs/AutomationTool/BuildGraph/Schema.cs b/Engine/Source/Programs/AutomationTool/BuildGraph/Schema.cs
--- a/Engine/Source/Programs/AutomationTool/BuildGraph/Schema.cs
+++ b/Engine/Source/Programs/AutomationTool/BuildGraph/Schema.cs
@@ -648,7 +648,8 @@ namespace AutomationTool
 			XmlSchemaComplexType BadgeType = new XmlSchemaComplexType();
 			BadgeType.Name = GetTypeName(ScriptSchemaStandardType.Badge);
 			BadgeType.Attributes.Add(CreateSchemaAttribute("Name", ScriptSchemaStandardType.Name, XmlSchemaUse.Required));
-			BadgeType.Attributes.Add(CreateSchemaAttribute("Requires", ScriptSchemaStandardType.NameOrTagList, XmlSchemaUse.Required));
+			BadgeType.Attributes.Add(CreateSchemaAttribute("Requires", ScriptSchemaStandardType.NameOrTagList, XmlSchemaUse.Optional));
+			BadgeType.Attributes.Add(CreateSchemaAttribute("Targets", ScriptSchemaStandardType.NameOrTagList, XmlSchemaUse.Optional));
 			BadgeType.Attributes.Add(CreateSchemaAttribute("Project", ScriptSchemaStandardType.BalancedString, XmlSchemaUse.Required));
 			BadgeType.Attributes.Add(CreateSchemaAttribute("Change", ScriptSchemaStandardType.BalancedString, XmlSchemaUse.Optional));
 			BadgeType.Attributes.Add(CreateSchemaAttribute("If", ScriptSchemaStandardType.BalancedString, XmlSchemaUse.Optional));
diff --git a/Engine/Source/Programs/AutomationTool/BuildGraph/Script.cs b/Engine/Source/Programs/AutomationTool/BuildGraph/Script.cs
--- a/Engine/Source/Programs/AutomationTool/BuildGraph/Script.cs
+++ b/Engine/Source/Programs/AutomationTool/BuildGraph/Script.cs
@@ -909,6 +909,7 @@ namespace AutomationTool
 			if (EvaluateCondition(Element) && TryReadObjectName(Element, out Name))
 			{
 				string[] RequiredNames = ReadListAttribute(Element, "Requires");
+				string[] TargetNames = ReadListAttribute(Element, "Targets");
 				string Project = ReadAttribute(Element, "Project");
 				int Change = ReadIntegerAttribute(Element, "Change", 0);
 
@@ -917,6 +918,11 @@ namespace AutomationTool
 				{
 					NewBadge.Nodes.Add(ReferencedNode);
 				}
+				foreach (Node ReferencedNode in ResolveReferences(Element, TargetNames))
+				{
+					NewBadge.Nodes.Add(ReferencedNode);
+					NewBadge.Nodes.UnionWith(ReferencedNode.OrderDependencies);
+				}
 				Graph.Badges.Add(NewBadge);
 			}
 		}
