UAT: restored Intermediate Commandline Dir

diff --git a/Engine/Source/Programs/AutomationTool/Scripts/CopyBuildToStagingDirectory.Automation.cs b/Engine/Source/Programs/AutomationTool/Scripts/CopyBuildToStagingDirectory.Automation.cs
--- a/Engine/Source/Programs/AutomationTool/Scripts/CopyBuildToStagingDirectory.Automation.cs
+++ b/Engine/Source/Programs/AutomationTool/Scripts/CopyBuildToStagingDirectory.Automation.cs
@@ -2289,7 +2289,7 @@ public partial class Project : CommandUtils
 
 	private static DirectoryReference GetIntermediateCommandlineDir(DeploymentContext SC)
 	{
-		return DirectoryReference.Combine(SC.LocalRoot, "Intermediate", "UAT", SC.FinalCookPlatform);
+		return DirectoryReference.Combine(SC.EngineRoot, "Intermediate", "UAT", SC.FinalCookPlatform);
 	}
 
 	public static void WriteStageCommandline(FileReference IntermediateCmdLineFile, ProjectParams Params, DeploymentContext SC)
