Add Dock_ to macos crash ignore list

diff --git a/Engine/Source/Programs/AutomationTool/AutomationUtils/CommandletUtils.cs b/Engine/Source/Programs/AutomationTool/AutomationUtils/CommandletUtils.cs
--- a/Engine/Source/Programs/AutomationTool/AutomationUtils/CommandletUtils.cs
+++ b/Engine/Source/Programs/AutomationTool/AutomationUtils/CommandletUtils.cs
@@ -359,7 +359,7 @@ namespace AutomationTool
 				{
 					// snmpd seems to often crash (suspect due to it being starved of CPU cycles during cooks)
 					// also ignore spotlight crash with the excel plugin
-					if(!CrashFileInfo.Name.StartsWith("snmpd_") && !CrashFileInfo.Name.StartsWith("mdworker32_"))
+					if(!CrashFileInfo.Name.StartsWith("snmpd_") && !CrashFileInfo.Name.StartsWith("mdworker32_") && !CrashFileInfo.Name.StartsWith("Dock_"))
 					{
 						CommandUtils.LogInformation("Found crash log - {0}", CrashFileInfo.FullName);
 						try
