UAT: Match the pattern "START XXX HACK" when looking for branch-specific hacks in the codebase.

diff --git a/Engine/Source/Programs/AutomationTool/Scripts/CheckForHacks.cs b/Engine/Source/Programs/AutomationTool/Scripts/CheckForHacks.cs
--- a/Engine/Source/Programs/AutomationTool/Scripts/CheckForHacks.cs
+++ b/Engine/Source/Programs/AutomationTool/Scripts/CheckForHacks.cs
@@ -48,7 +48,7 @@ namespace AutomationTool
 		/// <summary>
 		/// The pattern that should match hack comments (and captures a list of tags). Ignore anything with a semicolon between BEGIN and HACK to avoid matching statements with "HACK" as a comment (seen in LLVM).
 		/// </summary>
-		static readonly Regex CompiledRegex = new Regex("(?<!\\w)BEGIN\\s([^;]*)\\sHACK(?!\\w)", RegexOptions.IgnoreCase | RegexOptions.Compiled);
+		static readonly Regex CompiledRegex = new Regex("(?<!\\w)(?:BEGIN|START)\\s([^;]*)\\sHACK(?!\\w)", RegexOptions.IgnoreCase | RegexOptions.Compiled);
 
 		/// <summary>
 		/// Executes the command
