Merging UAT fixes from 4.21.1

diff --git a/Engine/Source/Programs/AutomationTool/Scripts/RebuildLightMapsCommand.Automation.cs b/Engine/Source/Programs/AutomationTool/Scripts/RebuildLightMapsCommand.Automation.cs
--- a/Engine/Source/Programs/AutomationTool/Scripts/RebuildLightMapsCommand.Automation.cs
+++ b/Engine/Source/Programs/AutomationTool/Scripts/RebuildLightMapsCommand.Automation.cs
@@ -260,43 +260,46 @@ namespace AutomationScripts.Automation
 		 */
 		void SendCompletionMessage(bool bWasSuccessful, String MessageBody)
 		{
-			MailMessage Message = new System.Net.Mail.MailMessage();
-			Message.Priority = MailPriority.High;
-			Message.From = new MailAddress("[REDACTED_EMAIL]");
+			if (StakeholdersEmailAddresses != null)
+			{
+				MailMessage Message = new System.Net.Mail.MailMessage();
+				Message.Priority = MailPriority.High;
+				Message.From = new MailAddress("[REDACTED_EMAIL]");
 
-            string Branch = "Unknown";
-            if ( P4Enabled )
-            {
-                Branch = P4Env.Branch;
-            }
+				string Branch = "Unknown";
+				if (P4Enabled)
+				{
+					Branch = P4Env.Branch;
+				}
 
-			foreach (String NextStakeHolder in StakeholdersEmailAddresses)
-			{
-				Message.To.Add(new MailAddress(NextStakeHolder));
-			}
+				foreach (String NextStakeHolder in StakeholdersEmailAddresses)
+				{
+					Message.To.Add(new MailAddress(NextStakeHolder));
+				}
 
-			ConfigHierarchy Ini = ConfigCache.ReadHierarchy(ConfigHierarchyType.EditorPerProjectUserSettings, DirectoryReference.FromFile(ProjectFullPath), UnrealTargetPlatform.Win64);
-			List<string> Emails;
-			Ini.GetArray("RebuildlightingSettings", "EmailNotification", out Emails);
-			foreach (var Email in Emails)
-			{
-				Message.CC.Add(new MailAddress(Email));
-			}
+				ConfigHierarchy Ini = ConfigCache.ReadHierarchy(ConfigHierarchyType.EditorPerProjectUserSettings, DirectoryReference.FromFile(ProjectFullPath), UnrealTargetPlatform.Win64);
+				List<string> Emails;
+				Ini.GetArray("RebuildlightingSettings", "EmailNotification", out Emails);
+				foreach (var Email in Emails)
+				{
+					Message.CC.Add(new MailAddress(Email));
+				}
 
-			Message.Subject = String.Format("Nightly lightmap rebuild {0} for {1}", bWasSuccessful ? "[SUCCESS]" : "[FAILED]", Branch);
-			Message.Body = MessageBody;
-            /*Attachment Attach = new Attachment();
-            Message.Attachments.Add()*/
-			try
-			{
-				#pragma warning disable CS0618 // Mono 4.6.x obsoletes this class
-				SmtpClient MailClient = new SmtpClient("smtp.epicgames.net");
-				MailClient.Send(Message);
-				#pragma warning restore CS0618
-			}
-			catch (Exception Ex)
-			{
-				LogError("Failed to send notify email to {0} ({1})", String.Join(", ", StakeholdersEmailAddresses.ToArray()), Ex.Message);
+				Message.Subject = String.Format("Nightly lightmap rebuild {0} for {1}", bWasSuccessful ? "[SUCCESS]" : "[FAILED]", Branch);
+				Message.Body = MessageBody;
+				/*Attachment Attach = new Attachment();
+				Message.Attachments.Add()*/
+				try
+				{
+#pragma warning disable CS0618 // Mono 4.6.x obsoletes this class
+					SmtpClient MailClient = new SmtpClient("smtp.epicgames.net");
+					MailClient.Send(Message);
+#pragma warning restore CS0618
+				}
+				catch (Exception Ex)
+				{
+					LogError("Failed to send notify email to {0} ({1})", String.Join(", ", StakeholdersEmailAddresses.ToArray()), Ex.Message);
+				}
 			}
 		}
 
diff --git a/Engine/Source/Programs/AutomationTool/Scripts/ResavePackagesCommand.Automation.cs b/Engine/Source/Programs/AutomationTool/Scripts/ResavePackagesCommand.Automation.cs
--- a/Engine/Source/Programs/AutomationTool/Scripts/ResavePackagesCommand.Automation.cs
+++ b/Engine/Source/Programs/AutomationTool/Scripts/ResavePackagesCommand.Automation.cs
@@ -263,38 +263,38 @@ namespace AutomationScripts.Automation
 		 */
 		void SendCompletionMessage(bool bWasSuccessful, String MessageBody)
 		{
-			MailMessage Message = new System.Net.Mail.MailMessage();
-			Message.Priority = MailPriority.High;
-			Message.From = new MailAddress("[REDACTED_EMAIL]");
+			if (StakeholdersEmailAddresses != null)
+			{
+				MailMessage Message = new System.Net.Mail.MailMessage();
+				Message.Priority = MailPriority.High;
+				Message.From = new MailAddress("[REDACTED_EMAIL]");
 
-            string Branch = "Unknown";
-            if ( P4Enabled )
-            {
-                Branch = P4Env.Branch;
-            }
+				string Branch = "Unknown";
+				if (P4Enabled)
+				{
+					Branch = P4Env.Branch;
+				}
 
-			foreach (String NextStakeHolder in StakeholdersEmailAddresses)
-			{
-				Message.To.Add(new MailAddress(NextStakeHolder));
-			}
+				foreach (String NextStakeHolder in StakeholdersEmailAddresses)
+				{
+					Message.To.Add(new MailAddress(NextStakeHolder));
+				}
 
-			Message.CC.Add(new MailAddress("[REDACTED_EMAIL]"));
-            Message.CC.Add(new MailAddress("[REDACTED_EMAIL]"));
-            Message.CC.Add(new MailAddress("[REDACTED_EMAIL]"));
-			Message.Subject = String.Format("Nightly lightmap rebuild {0} for {1}", bWasSuccessful ? "[SUCCESS]" : "[FAILED]", Branch);
-			Message.Body = MessageBody;
-            /*Attachment Attach = new Attachment();
-            Message.Attachments.Add()*/
-			try
-			{
-				#pragma warning disable CS0618 // Mono 4.6.x obsoletes this class
-				SmtpClient MailClient = new SmtpClient("smtp.epicgames.net");
-				MailClient.Send(Message);
-				#pragma warning restore CS0618
-			}
-			catch (Exception Ex)
-			{
-				LogError("Failed to send notify email to {0} ({1})", String.Join(", ", StakeholdersEmailAddresses.ToArray()), Ex.Message);
+				Message.Subject = String.Format("Nightly lightmap rebuild {0} for {1}", bWasSuccessful ? "[SUCCESS]" : "[FAILED]", Branch);
+				Message.Body = MessageBody;
+				/*Attachment Attach = new Attachment();
+				Message.Attachments.Add()*/
+				try
+				{
+#pragma warning disable CS0618 // Mono 4.6.x obsoletes this class
+					SmtpClient MailClient = new SmtpClient("smtp.epicgames.net");
+					MailClient.Send(Message);
+#pragma warning restore CS0618
+				}
+				catch (Exception Ex)
+				{
+					LogError("Failed to send notify email to {0} ({1})", String.Join(", ", StakeholdersEmailAddresses.ToArray()), Ex.Message);
+				}
 			}
 		}
 
