OneSky is no longer the default localization provider in UAT

diff --git a/Engine/Source/Programs/AutomationTool/Scripts/Localisation.Automation.cs b/Engine/Source/Programs/AutomationTool/Scripts/Localisation.Automation.cs
--- a/Engine/Source/Programs/AutomationTool/Scripts/Localisation.Automation.cs
+++ b/Engine/Source/Programs/AutomationTool/Scripts/Localisation.Automation.cs
@@ -17,7 +17,7 @@ using Tools.DotNETCommon;
 [Help("UEProjectName", "Optional name of the project we're gathering for (should match its .uproject file, eg QAGame).")]
 [Help("LocalizationProjectNames", "Comma separated list of the projects to gather text from.")]
 [Help("LocalizationBranch", "Optional suffix to use when uploading the new data to the localization provider.")]
-[Help("LocalizationProvider", "Optional localization provide override (default is OneSky).")]
+[Help("LocalizationProvider", "Optional localization provide override.")]
 [Help("LocalizationSteps", "Optional comma separated list of localization steps to perform [Download, Gather, Import, Export, Compile, GenerateReports, Upload] (default is all). Only valid for projects using a modular config.")]
 [Help("IncludePlugins", "Optional flag to include plugins from within the given UEProjectDirectory as part of the gather. This may optionally specify a comma separated list of the specific plugins to gather (otherwise all plugins will be gathered).")]
 [Help("ExcludePlugins", "Optional comma separated list of plugins to exclude from the gather.")]
@@ -75,7 +75,7 @@ class Localise : BuildCommand
 		var LocalizationProviderName = ParseParamValue("LocalizationProvider");
 		if (LocalizationProviderName == null)
 		{
-			LocalizationProviderName = "OneSky";
+			LocalizationProviderName = "";
 		}
 
 		var LocalizationSteps = new List<string>();
