Added support for staging DLC localization data, and consistently implemented target exclusion

diff --git a/Engine/Source/Programs/AutomationTool/Scripts/CopyBuildToStagingDirectory.Automation.cs b/Engine/Source/Programs/AutomationTool/Scripts/CopyBuildToStagingDirectory.Automation.cs
--- a/Engine/Source/Programs/AutomationTool/Scripts/CopyBuildToStagingDirectory.Automation.cs
+++ b/Engine/Source/Programs/AutomationTool/Scripts/CopyBuildToStagingDirectory.Automation.cs
@@ -191,6 +191,77 @@ public partial class Project : CommandUtils
 		LogLog("End Deployment Context **************");
 	}
 
+	private static string GetInternationalizationPreset(ProjectParams Params, ConfigHierarchy PlatformGameConfig, bool bMustExist = true)
+	{
+		// Initialize internationalization preset.
+		string InternationalizationPreset = Params.InternationalizationPreset;
+
+		// Use configuration if otherwise lacking an internationalization preset.
+		if (string.IsNullOrEmpty(InternationalizationPreset))
+		{
+			if (PlatformGameConfig != null)
+			{
+				PlatformGameConfig.GetString("/Script/UnrealEd.ProjectPackagingSettings", "InternationalizationPreset", out InternationalizationPreset);
+			}
+		}
+
+		// Error if no preset has been provided.
+		if (bMustExist && string.IsNullOrEmpty(InternationalizationPreset))
+		{
+			throw new AutomationException("No internationalization preset was specified for packaging. This will lead to fatal errors when launching. Specify preset via commandline (-I18NPreset=) or project packaging settings (InternationalizationPreset).");
+		}
+
+		return InternationalizationPreset;
+	}
+
+	private static List<string> GetCulturesToStage(ProjectParams Params, ConfigHierarchy PlatformGameConfig, bool bMustExist = true)
+	{
+		// Initialize cultures to stage.
+		List<string> CulturesToStage = null;
+
+		// Use parameters if provided.
+		if (Params.CulturesToCook != null && Params.CulturesToCook.Count > 0)
+		{
+			CulturesToStage = Params.CulturesToCook;
+		}
+
+		// Use configuration if otherwise lacking cultures to stage.
+		if (CulturesToStage == null || CulturesToStage.Count == 0)
+		{
+			if (PlatformGameConfig != null)
+			{
+				PlatformGameConfig.GetArray("/Script/UnrealEd.ProjectPackagingSettings", "CulturesToStage", out CulturesToStage);
+			}
+		}
+
+		// Error if no cultures have been provided.
+		if (bMustExist && (CulturesToStage == null || CulturesToStage.Count == 0))
+		{
+			throw new AutomationException("No cultures were specified for cooking and packaging. This will lead to fatal errors when launching. Specify culture codes via commandline (-CookCultures=) or using project packaging settings (+CulturesToStage).");
+		}
+
+		return CulturesToStage;
+	}
+
+	private static void StageLocalizationDataForPlugin(DeploymentContext SC, List<string> CulturesToStage, FileReference Plugin)
+	{
+		PluginDescriptor Descriptor = PluginDescriptor.FromFile(Plugin);
+		if (Descriptor.LocalizationTargets != null)
+		{
+			foreach (LocalizationTargetDescriptor LocalizationTarget in Descriptor.LocalizationTargets)
+			{
+				if (LocalizationTarget.LoadingPolicy == LocalizationTargetDescriptorLoadingPolicy.Always || LocalizationTarget.LoadingPolicy == LocalizationTargetDescriptorLoadingPolicy.Game)
+				{
+					DirectoryReference PluginLocTargetDirectory = DirectoryReference.Combine(Plugin.Directory, "Content", "Localization", LocalizationTarget.Name);
+					if (!SC.BlacklistLocalizationTargets.Contains(LocalizationTarget.Name) && DirectoryReference.Exists(PluginLocTargetDirectory))
+					{
+						StageLocalizationDataForTarget(SC, CulturesToStage, PluginLocTargetDirectory);
+					}
+				}
+			}
+		}
+	}
+
 	private static void StageLocalizationDataForTarget(DeploymentContext SC, List<string> CulturesToStage, DirectoryReference SourceDirectory)
 	{
 		foreach (FileReference SourceFile in DirectoryReference.EnumerateFiles(SourceDirectory, "*.locmeta"))
@@ -291,6 +362,8 @@ public partial class Project : CommandUtils
 
 		LogInformation("Creating Staging Manifest...");
 
+		ConfigHierarchy PlatformGameConfig = ConfigCache.ReadHierarchy(ConfigHierarchyType.Game, DirectoryReference.FromFile(Params.RawProjectPath), SC.StageTargetPlatform.IniPlatformType);
+
 		if (Params.HasIterateSharedCookedBuild)
 		{
 			// can't do shared cooked builds with DLC that's madness!!
@@ -316,6 +389,13 @@ public partial class Project : CommandUtils
 				{
 					SC.StageFile(StagedFileType.NonUFS, Params.DLCFile);
 				}
+
+				// Stage DLC localization targets
+				List<string> CulturesToStage = GetCulturesToStage(Params, PlatformGameConfig, false);
+				if (CulturesToStage != null)
+				{
+					StageLocalizationDataForPlugin(SC, CulturesToStage, Params.DLCFile);
+				}
 			}
 
 			// Put the binaries into the staged dir
@@ -439,54 +519,16 @@ public partial class Project : CommandUtils
 				SC.StageFile(StagedFileType.SystemNonUFS, CommandLineFile, StagedCommandLineFile);
 			}
 
-			ConfigHierarchy PlatformGameConfig = ConfigCache.ReadHierarchy(ConfigHierarchyType.Game, DirectoryReference.FromFile(Params.RawProjectPath), SC.StageTargetPlatform.IniPlatformType);
 			DirectoryReference ProjectContentRoot = DirectoryReference.Combine(SC.ProjectRoot, "Content");
 			StagedDirectoryReference StageContentRoot = StagedDirectoryReference.Combine(SC.RelativeProjectRootForStage, "Content");
 
 			if (!Params.CookOnTheFly && !Params.SkipCookOnTheFly) // only stage the UFS files if we are not using cook on the fly
 			{
-
 				// Initialize internationalization preset.
-				string InternationalizationPreset = Params.InternationalizationPreset;
-
-				// Use configuration if otherwise lacking an internationalization preset.
-				if (string.IsNullOrEmpty(InternationalizationPreset))
-				{
-					if (PlatformGameConfig != null)
-					{
-						PlatformGameConfig.GetString("/Script/UnrealEd.ProjectPackagingSettings", "InternationalizationPreset", out InternationalizationPreset);
-					}
-				}
-
-				// Error if no preset has been provided.
-				if (string.IsNullOrEmpty(InternationalizationPreset))
-				{
-					throw new AutomationException("No internationalization preset was specified for packaging. This will lead to fatal errors when launching. Specify preset via commandline (-I18NPreset=) or project packaging settings (InternationalizationPreset).");
-				}
+				string InternationalizationPreset = GetInternationalizationPreset(Params, PlatformGameConfig);
 
 				// Initialize cultures to stage.
-				List<string> CulturesToStage = null;
-
-				// Use parameters if provided.
-				if (Params.CulturesToCook != null && Params.CulturesToCook.Count > 0)
-				{
-					CulturesToStage = Params.CulturesToCook;
-				}
-
-				// Use configuration if otherwise lacking cultures to stage.
-				if (CulturesToStage == null || CulturesToStage.Count == 0)
-				{
-					if (PlatformGameConfig != null)
-					{
-						PlatformGameConfig.GetArray("/Script/UnrealEd.ProjectPackagingSettings", "CulturesToStage", out CulturesToStage);
-					}
-				}
-
-				// Error if no cultures have been provided.
-				if (CulturesToStage == null || CulturesToStage.Count == 0)
-				{
-					throw new AutomationException("No cultures were specified for cooking and packaging. This will lead to fatal errors when launching. Specify culture codes via commandline (-CookCultures=) or using project packaging settings (+CulturesToStage).");
-				}
+				List<string> CulturesToStage = GetCulturesToStage(Params, PlatformGameConfig);
 
 				// Stage ICU internationalization data from Engine.
 				SC.StageFiles(StagedFileType.UFS, DirectoryReference.Combine(SC.LocalRoot, "Engine", "Content", "Internationalization", InternationalizationPreset), StageFilesSearch.AllDirectories, new StagedDirectoryReference("Engine/Content/Internationalization"));
@@ -494,7 +536,10 @@ public partial class Project : CommandUtils
 				// Engine ufs (content)
 				StageConfigFiles(SC, DirectoryReference.Combine(SC.LocalRoot, "Engine", "Config"));
 
-				StageLocalizationDataForTarget(SC, CulturesToStage, DirectoryReference.Combine(SC.LocalRoot, "Engine", "Content", "Localization", "Engine"));
+				if (!SC.BlacklistLocalizationTargets.Contains("Engine"))
+				{
+					StageLocalizationDataForTarget(SC, CulturesToStage, DirectoryReference.Combine(SC.LocalRoot, "Engine", "Content", "Localization", "Engine"));
+				}
 
 				// Game ufs (content)
 				SC.StageFile(StagedFileType.UFS, SC.RawProjectPath);
@@ -523,30 +568,16 @@ public partial class Project : CommandUtils
 						{
 							if (!SC.BlacklistLocalizationTargets.Contains(ProjectLocTargetDirectory.MakeRelativeTo(ProjectLocRootDirectory)))
 							{
-							StageLocalizationDataForTarget(SC, CulturesToStage, ProjectLocTargetDirectory);
+								StageLocalizationDataForTarget(SC, CulturesToStage, ProjectLocTargetDirectory);
+							}
 						}
 					}
 				}
-				}
 
 				// Stage all plugin localization targets
 				foreach (KeyValuePair<StagedFileReference, FileReference> StagedPlugin in StagedPlugins)
 				{
-					PluginDescriptor Descriptor = PluginDescriptor.FromFile(StagedPlugin.Value);
-					if (Descriptor.LocalizationTargets != null)
-					{
-						foreach (LocalizationTargetDescriptor LocalizationTarget in Descriptor.LocalizationTargets)
-						{
-							if (LocalizationTarget.LoadingPolicy == LocalizationTargetDescriptorLoadingPolicy.Always || LocalizationTarget.LoadingPolicy == LocalizationTargetDescriptorLoadingPolicy.Game)
-							{
-								DirectoryReference PluginLocTargetDirectory = DirectoryReference.Combine(StagedPlugin.Value.Directory, "Content", "Localization", LocalizationTarget.Name);
-								if (DirectoryReference.Exists(PluginLocTargetDirectory))
-								{
-									StageLocalizationDataForTarget(SC, CulturesToStage, PluginLocTargetDirectory);
-								}
-							}
-						}
-					}
+					StageLocalizationDataForPlugin(SC, CulturesToStage, StagedPlugin.Value);
 				}
 
 				// Stage any additional UFS and NonUFS paths specified in the project ini files; these dirs are relative to the game content directory
