Add PGO heartbeat to reduce timeout state (30 minutes down from 4 hours in some cases)

diff --git a/Engine/Source/Programs/AutomationTool/Gauntlet/Gauntlet.Automation.csproj b/Engine/Source/Programs/AutomationTool/Gauntlet/Gauntlet.Automation.csproj
--- a/Engine/Source/Programs/AutomationTool/Gauntlet/Gauntlet.Automation.csproj
+++ b/Engine/Source/Programs/AutomationTool/Gauntlet/Gauntlet.Automation.csproj
@@ -253,6 +253,7 @@
       <ItemGroup>
         <Compile Include="Platform\Windows\Gauntlet.TargetDeviceWindows.cs" />
         <Compile Include="Platform\Windows\Gauntlet.WindowsBuildSources.cs" />
+        <Compile Include="Platform\Windows\Gauntlet.WindowsPGO.cs" />
         <Compile Include="Platform\Windows\Gauntlet.SelfTest.TestTargetDeviceWindows.cs" />
       </ItemGroup>
     </When>
diff --git a/Engine/Source/Programs/AutomationTool/Gauntlet/Platform/Windows/Gauntlet.WindowsPGO.cs b/Engine/Source/Programs/AutomationTool/Gauntlet/Platform/Windows/Gauntlet.WindowsPGO.cs
new file mode 100644
--- /dev/null
+++ b/Engine/Source/Programs/AutomationTool/Gauntlet/Platform/Windows/Gauntlet.WindowsPGO.cs
@@ -0,0 +1,37 @@
+// Copyright 1998-2018 Epic Games, Inc. All Rights Reserved.
+
+using UnrealBuildTool;
+
+namespace Gauntlet
+{
+	/// <summary>
+	/// Skeleton Windows PGO platform implementation (primarily used to test PGO locally with editor builds)
+	/// </summary>
+	internal class WindowsPGOPlatform : IPGOPlatform
+	{
+
+		public UnrealTargetPlatform GetPlatform()
+		{
+			return UnrealTargetPlatform.Win64;
+		}
+
+		public void GatherResults(string ArtifactPath)
+		{
+
+		}
+
+		public void ApplyConfiguration(PGOConfig Config)
+		{
+
+		}
+
+		public bool TakeScreenshot(ITargetDevice Device, string ScreenshotDirectory, out string ImageFilename)
+		{			
+			ImageFilename = string.Empty;
+			return false;
+		}
+	}
+}
+
+
+
