Copyright Notices.

diff --git a/Engine/Source/Editor/DataTableEditor/Private/SDataTableListViewRowName.cpp b/Engine/Source/Editor/DataTableEditor/Private/SDataTableListViewRowName.cpp
--- a/Engine/Source/Editor/DataTableEditor/Private/SDataTableListViewRowName.cpp
+++ b/Engine/Source/Editor/DataTableEditor/Private/SDataTableListViewRowName.cpp
@@ -1,4 +1,4 @@
-// Copyright 1998-2018 Epic Games, Inc. All Rights Reserved.
+// Copyright 1998-2019 Epic Games, Inc. All Rights Reserved.
 
 #include "SDataTableListViewRowName.h"
 
diff --git a/Engine/Source/Editor/DataTableEditor/Private/SDataTableListViewRowName.h b/Engine/Source/Editor/DataTableEditor/Private/SDataTableListViewRowName.h
--- a/Engine/Source/Editor/DataTableEditor/Private/SDataTableListViewRowName.h
+++ b/Engine/Source/Editor/DataTableEditor/Private/SDataTableListViewRowName.h
@@ -1,4 +1,4 @@
-// Copyright 1998-2018 Epic Games, Inc. All Rights Reserved.
+// Copyright 1998-2019 Epic Games, Inc. All Rights Reserved.
 
 #pragma once
 
diff --git a/Engine/Source/Editor/DataTableEditor/Public/DataTableRowUtlis.cpp b/Engine/Source/Editor/DataTableEditor/Public/DataTableRowUtlis.cpp
--- a/Engine/Source/Editor/DataTableEditor/Public/DataTableRowUtlis.cpp
+++ b/Engine/Source/Editor/DataTableEditor/Public/DataTableRowUtlis.cpp
@@ -1,4 +1,4 @@
-// Copyright 1998-2018 Epic Games, Inc. All Rights Reserved.
+// Copyright 1998-2019 Epic Games, Inc. All Rights Reserved.
 
 #include "DataTableRowUtlis.h"
 #include "DetailWidgetRow.h"
diff --git a/Engine/Source/Editor/DataTableEditor/Public/DataTableRowUtlis.h b/Engine/Source/Editor/DataTableEditor/Public/DataTableRowUtlis.h
--- a/Engine/Source/Editor/DataTableEditor/Public/DataTableRowUtlis.h
+++ b/Engine/Source/Editor/DataTableEditor/Public/DataTableRowUtlis.h
@@ -1,4 +1,4 @@
-// Copyright 1998-2018 Epic Games, Inc. All Rights Reserved.
+// Copyright 1998-2019 Epic Games, Inc. All Rights Reserved.
 
 #pragma once
 
diff --git a/Engine/Source/Editor/UMGEditor/Private/Details/SFunctionSelector.h b/Engine/Source/Editor/UMGEditor/Private/Details/SFunctionSelector.h
--- a/Engine/Source/Editor/UMGEditor/Private/Details/SFunctionSelector.h
+++ b/Engine/Source/Editor/UMGEditor/Private/Details/SFunctionSelector.h
@@ -1,4 +1,4 @@
-// Copyright 1998-2018 Epic Games, Inc. All Rights Reserved.
+// Copyright 1998-2019 Epic Games, Inc. All Rights Reserved.
 
 #pragma once
 
diff --git a/Engine/Source/Programs/AutomationTool/Gauntlet/Platform/Windows/Gauntlet.WindowsPGO.cs b/Engine/Source/Programs/AutomationTool/Gauntlet/Platform/Windows/Gauntlet.WindowsPGO.cs
--- a/Engine/Source/Programs/AutomationTool/Gauntlet/Platform/Windows/Gauntlet.WindowsPGO.cs
+++ b/Engine/Source/Programs/AutomationTool/Gauntlet/Platform/Windows/Gauntlet.WindowsPGO.cs
@@ -1,4 +1,4 @@
-// Copyright 1998-2018 Epic Games, Inc. All Rights Reserved.
+// Copyright 1998-2019 Epic Games, Inc. All Rights Reserved.
 
 using UnrealBuildTool;
 
diff --git a/Engine/Source/Runtime/Renderer/Private/FogRendering.h b/Engine/Source/Runtime/Renderer/Private/FogRendering.h
--- a/Engine/Source/Runtime/Renderer/Private/FogRendering.h
+++ b/Engine/Source/Runtime/Renderer/Private/FogRendering.h
@@ -1,4 +1,4 @@
-// Copyright 1998-2018 Epic Games, Inc. All Rights Reserved.
+// Copyright 1998-2019 Epic Games, Inc. All Rights Reserved.
 
 /*=============================================================================
 	FogRendering.h: 
diff --git a/Engine/Source/Runtime/WebBrowser/Private/Android/AndroidCookieManager.cpp b/Engine/Source/Runtime/WebBrowser/Private/Android/AndroidCookieManager.cpp
--- a/Engine/Source/Runtime/WebBrowser/Private/Android/AndroidCookieManager.cpp
+++ b/Engine/Source/Runtime/WebBrowser/Private/Android/AndroidCookieManager.cpp
@@ -1,4 +1,4 @@
-// Copyright 1998-2018 Epic Games, Inc. All Rights Reserved.
+// Copyright 1998-2019 Epic Games, Inc. All Rights Reserved.
 
 #include "AndroidCookieManager.h"
 
diff --git a/Engine/Source/Runtime/WebBrowser/Private/Android/AndroidCookieManager.h b/Engine/Source/Runtime/WebBrowser/Private/Android/AndroidCookieManager.h
--- a/Engine/Source/Runtime/WebBrowser/Private/Android/AndroidCookieManager.h
+++ b/Engine/Source/Runtime/WebBrowser/Private/Android/AndroidCookieManager.h
@@ -1,4 +1,4 @@
-// Copyright 1998-2018 Epic Games, Inc. All Rights Reserved.
+// Copyright 1998-2019 Epic Games, Inc. All Rights Reserved.
 
 #pragma once
 
