modifying BuildPhysX to automatically include a 'github ignore' tag so the automated Build PhysX can check in again.

diff --git a/Engine/Source/Programs/AutomationTool/Scripts/BuildPhysX.Automation.cs b/Engine/Source/Programs/AutomationTool/Scripts/BuildPhysX.Automation.cs
--- a/Engine/Source/Programs/AutomationTool/Scripts/BuildPhysX.Automation.cs
+++ b/Engine/Source/Programs/AutomationTool/Scripts/BuildPhysX.Automation.cs
@@ -1683,7 +1683,7 @@ class BuildPhysX : BuildCommand
 			{
 				RobomergeLine = Environment.NewLine + RobomergeCommand;
 			}
-            P4ChangeList = P4.CreateChange(P4Env.Client, String.Format("BuildPhysX.Automation: Deploying {0} libs.", LibDeploymentDesc) + Environment.NewLine + "#rb none" + Environment.NewLine + "#lockdown Nick.Penwarden" + Environment.NewLine + "#tests none" + Environment.NewLine + "#jira none" + RobomergeLine);
+            P4ChangeList = P4.CreateChange(P4Env.Client, String.Format("BuildPhysX.Automation: Deploying {0} libs.", LibDeploymentDesc) + Environment.NewLine + "#rb none" + Environment.NewLine + "#lockdown Nick.Penwarden" + Environment.NewLine + "#tests none" + Environment.NewLine + "#jira none" + Environment.NewLine + "#okforgithub ignore" + RobomergeLine);
 		}
 
 
