Gauntlet: Adding local executable support to iOS device (including updating symbols for symbolicated crash logs)

diff --git a/Engine/Source/Programs/AutomationTool/Gauntlet/Platform/IOS/Gauntlet.IOSBuildSource.cs b/Engine/Source/Programs/AutomationTool/Gauntlet/Platform/IOS/Gauntlet.IOSBuildSource.cs
--- a/Engine/Source/Programs/AutomationTool/Gauntlet/Platform/IOS/Gauntlet.IOSBuildSource.cs
+++ b/Engine/Source/Programs/AutomationTool/Gauntlet/Platform/IOS/Gauntlet.IOSBuildSource.cs
@@ -141,8 +141,8 @@ namespace Gauntlet
 
 				string AbsPath = Fi.Directory.FullName;
 
-				// IOS builds are always packaged, and we can always replace the command line
-				BuildFlags Flags = BuildFlags.Packaged | BuildFlags.CanReplaceCommandLine;
+				// IOS builds are always packaged, and can always replace the command line and executable as we cache the unzip'd IPA
+				BuildFlags Flags = BuildFlags.Packaged | BuildFlags.CanReplaceCommandLine | BuildFlags.CanReplaceExecutable;
 
 				if (AbsPath.Contains("Bulk"))
 				{
diff --git a/Engine/Source/Programs/AutomationTool/Gauntlet/Platform/IOS/Gauntlet.TargetDeviceIOS.cs b/Engine/Source/Programs/AutomationTool/Gauntlet/Platform/IOS/Gauntlet.TargetDeviceIOS.cs
--- a/Engine/Source/Programs/AutomationTool/Gauntlet/Platform/IOS/Gauntlet.TargetDeviceIOS.cs
+++ b/Engine/Source/Programs/AutomationTool/Gauntlet/Platform/IOS/Gauntlet.TargetDeviceIOS.cs
@@ -383,6 +383,91 @@
 
 		}
 
+		/// <summary>
+		/// Resign application using local executable and update debug symbols
+		/// </summary>
+		void ResignApplication(UnrealAppConfig AppConfig)
+		{
+			// check that we have the signing stuff we need
+			string SignProvision = Globals.Params.ParseValue("signprovision", String.Empty);
+			string SignEntitlements = Globals.Params.ParseValue("signentitlements", String.Empty);
+			string SigningIdentity =  Globals.Params.ParseValue("signidentity", String.Empty);
+			
+			// handle signing provision
+			if (string.IsNullOrEmpty(SignProvision) || !File.Exists(SignProvision))
+			{
+				throw new AutomationException("Absolute path to existing provision must be specified, example: -signprovision=/path/to/myapp.provision");
+			}
+			
+			// handle entitlements
+			// Note this extracts entitlements: which may be useful when using same provision/entitlements?: codesign -d --entitlements :entitlements.plist ~/.gauntletappcache/Payload/Example.app/
+
+			if (string.IsNullOrEmpty(SignEntitlements) || !File.Exists(SignEntitlements))
+			{
+				throw new AutomationException("Absolute path to existing entitlements must be specified, example: -signprovision=/path/to/entitlements.plist");
+			}
+
+			// signing identity
+			if (string.IsNullOrEmpty(SigningIdentity))
+			{
+				throw new AutomationException("Signing identity must be specified, example: -signidentity=\"iPhone Developer: John Smith\"");
+			}
+
+			string ProjectName = AppConfig.ProjectName;
+			string BundleName = Path.GetFileNameWithoutExtension(LocalAppBundle);
+			string ExecutableName = UnrealHelpers.GetExecutableName(ProjectName, UnrealTargetPlatform.IOS, AppConfig.Configuration, AppConfig.ProcessType, "");
+			string CachedAppPath = Path.Combine(GauntletAppCache, "Payload", string.Format("{0}.app", BundleName));			
+
+			string LocalExecutable = Path.Combine(Environment.CurrentDirectory, ProjectName, string.Format("Binaries/IOS/{0}", ExecutableName));
+			if (!File.Exists(LocalExecutable))
+			{
+				throw new AutomationException("Local executable not found for -dev argument: {0}", LocalExecutable);
+			}
+
+			File.WriteAllText(CacheResignedFilename, "The application has been resigned");
+
+			// copy local executable
+			FileInfo SrcInfo = new FileInfo(LocalExecutable);			
+			string DestPath = Path.Combine(CachedAppPath, BundleName);
+			SrcInfo.CopyTo(DestPath, true);
+			Log.Verbose("Copied local executable from {0} to {1}", LocalExecutable, DestPath);
+
+			// copy provision
+			SrcInfo = new FileInfo(SignProvision);
+			DestPath = Path.Combine(CachedAppPath, "embedded.mobileprovision");
+			SrcInfo.CopyTo(DestPath, true);
+			Log.Verbose("Copied provision from {0} to {1}", SignProvision, DestPath);
+
+			// handle symbols
+			string LocalSymbolsDir = Path.Combine(Environment.CurrentDirectory, ProjectName, string.Format("Binaries/IOS/{0}.dSYM", ExecutableName));
+			DestPath = Path.Combine(GauntletAppCache, string.Format("Symbols/{0}.dSYM", ExecutableName));
+
+			if (Directory.Exists(DestPath))
+			{
+				Directory.Delete(DestPath, true);
+			}
+
+			if (Directory.Exists(LocalSymbolsDir))
+			{				
+				CommandUtils.CopyDirectory_NoExceptions(LocalSymbolsDir, DestPath, true);
+			}
+			else
+			{
+				Log.Warning("No symbols found for local build at {0}, removing cached app symbols", LocalSymbolsDir);
+			}
+
+			// resign application
+			// @todo: this asks for password unless "Always Allow" is selected, also for builders, document how to permanently grant codesign access to keychain
+			string SignArgs = string.Format("-f -s \"{0}\" --entitlements \"{1}\" \"{2}\"", SigningIdentity, SignEntitlements, CachedAppPath);
+			Log.Info("\nResigning app, please enter keychain password if prompted:\n\ncodesign {0}", SignArgs);
+			var Result = IOSBuild.ExecuteCommand("codesign", SignArgs);
+			if (Result.ExitCode != 0)
+			{
+				throw new AutomationException("Failed to resign application");
+			}
+
+		}
+
 		public IAppInstall InstallApplication(UnrealAppConfig AppConfig)
 		{            
             IOSBuild Build = AppConfig.Build as IOSBuild;
@@ -398,9 +483,27 @@
 			// device artifact path
 			DeviceArtifactPath = string.Format("/Documents/{0}/Saved", AppConfig.ProjectName);
 
+			bool CacheResigned = File.Exists(CacheResignedFilename);
+			bool UseLocalExecutable = Globals.Params.ParseParam("dev");
+
+			if (CacheResigned && !UseLocalExecutable)
+			{
+				if (File.Exists(IPAHashFilename))
+				{
+					Log.Verbose("App was resigned, invalidating app cache");
+					File.Delete(IPAHashFilename);
+				}								
+			}
+
 			PrepareIPA(Build);
 
-			if (!CheckDeployedIPA(Build))
+			// local executable support			
+			if (UseLocalExecutable)
+			{					
+				ResignApplication(AppConfig);				
+			}
+
+			if (CacheResigned || UseLocalExecutable || !CheckDeployedIPA(Build))
 			{
 				// uninstall will clean all device artifacts
 				ExecuteIOSDeployCommand(String.Format("--uninstall -b \"{0}\"", LocalAppBundle), 10 * 60);
@@ -606,6 +709,10 @@
 		// the current IPA MD5 hash, which is tracked to avoid unneccessary deployments and unzip operations
 		string IPAHashFilename { get { return Path.Combine(GauntletAppCache, "IPAHash.txt"); } }
 
+		// file whose presence signals that cache was resigned
+		string CacheResignedFilename { get { return Path.Combine(GauntletAppCache, "Resigned.txt"); } }
+
+
 		/// <summary>
 		/// Generate MD5 and cache IPA bundle files
 		/// </summary>
@@ -650,6 +757,11 @@
 						Directory.Delete(SymbolsDir, true);
 					}
 
+					if (File.Exists(CacheResignedFilename))
+					{
+						File.Delete(CacheResignedFilename);
+					}
+
 					Log.Verbose("Unzipping IPA {0} to cache at: {1}", Build.SourceIPAPath, GauntletAppCache);
 					
 					IProcessResult Result = IOSBuild.ExecuteCommand("unzip", String.Format("{0} -d {1}", Build.SourceIPAPath, GauntletAppCache));
