UAT: Use Dns.GetHostName() rather than Environment.MachineName to determine the hostname for matching P4 clients. Environment.MachineName is the short NETBIOS name, and may not include the ".local" suffix typical on Mac.

diff --git a/Engine/Source/Programs/AutomationTool/AutomationUtils/P4Environment.cs b/Engine/Source/Programs/AutomationTool/AutomationUtils/P4Environment.cs
--- a/Engine/Source/Programs/AutomationTool/AutomationUtils/P4Environment.cs
+++ b/Engine/Source/Programs/AutomationTool/AutomationUtils/P4Environment.cs
@@ -178,7 +178,8 @@ namespace AutomationTool
 				P4ClientInfo ThisClient;
 				if(String.IsNullOrEmpty(Client))
 				{
-					ThisClient = DetectClient(DefaultConnection, User, Environment.MachineName.ToLower(), CmdEnv.UATExe);
+					string HostName = System.Net.Dns.GetHostName();
+					ThisClient = DetectClient(DefaultConnection, User, HostName, CmdEnv.UATExe);
 					Log.TraceInformation("Using user {0} clientspec {1} {2}", User, ThisClient.Name, ThisClient.RootPath);
 					Client = ThisClient.Name;
 					CommandUtils.SetEnvVar(EnvVarNames.Client, Client);
