When adding additional files to the UFS staging list, omit movie files on dedicated servers

diff --git a/Engine/Source/Programs/AutomationTool/Scripts/CopyBuildToStagingDirectory.Automation.cs b/Engine/Source/Programs/AutomationTool/Scripts/CopyBuildToStagingDirectory.Automation.cs
--- a/Engine/Source/Programs/AutomationTool/Scripts/CopyBuildToStagingDirectory.Automation.cs
+++ b/Engine/Source/Programs/AutomationTool/Scripts/CopyBuildToStagingDirectory.Automation.cs
@@ -343,8 +343,7 @@ public partial class Project : CommandUtils
 					if (bUFS)
 					{
 						List<FileReference> Files = SC.FindFilesToStage(InputDir, StageFilesSearch.AllDirectories);
-						Files.RemoveAll(x => x.HasExtension(".uasset") || x.HasExtension(".umap"));
-
+						Files.RemoveAll(x => x.HasExtension(".uasset") || x.HasExtension(".umap") || (SC.DedicatedServer && x.HasExtension(".mp4")));
 						SC.StageFiles(StagedFileType.UFS, InputDir, Files, OutputDir);
 					}
 					else
