Fix comment.

diff --git a/Engine/Source/Programs/AutomationTool/BuildGraph/Tasks/SubmitTask.cs b/Engine/Source/Programs/AutomationTool/BuildGraph/Tasks/SubmitTask.cs
--- a/Engine/Source/Programs/AutomationTool/BuildGraph/Tasks/SubmitTask.cs
+++ b/Engine/Source/Programs/AutomationTool/BuildGraph/Tasks/SubmitTask.cs
@@ -41,7 +41,7 @@ namespace AutomationTool.Tasks
 		public string Workspace;
 
 		/// <summary>
-		/// The stream for the workspace; defaults to the current stream. Ignored unless If the Workspace attribute is also specified.
+		/// The stream for the workspace; defaults to the current stream. Ignored unless the Workspace attribute is also specified.
 		/// </summary>
 		[TaskParameter(Optional = true)]
 		public string Stream;
