Fix error message when enumerating P4 changes.

diff --git a/Engine/Source/Programs/AutomationTool/AutomationUtils/P4Utils.cs b/Engine/Source/Programs/AutomationTool/AutomationUtils/P4Utils.cs
--- a/Engine/Source/Programs/AutomationTool/AutomationUtils/P4Utils.cs
+++ b/Engine/Source/Programs/AutomationTool/AutomationUtils/P4Utils.cs
@@ -3884,7 +3884,7 @@ namespace AutomationTool
 			P4IntegrateAction Action;
 			if(!DescriptionToIntegrationAction.Value.TryGetValue(ActionText, out Action))
 			{
-				throw new P4Exception("Invalid integration action '{0}'", Action);
+				throw new P4Exception("Invalid integration action '{0}'", ActionText);
 			}
 			return Action;
 		}
