Add missing p4 integration action.

#rnx

diff --git a/Engine/Source/Programs/AutomationTool/AutomationUtils/P4Utils.cs b/Engine/Source/Programs/AutomationTool/AutomationUtils/P4Utils.cs
--- a/Engine/Source/Programs/AutomationTool/AutomationUtils/P4Utils.cs
+++ b/Engine/Source/Programs/AutomationTool/AutomationUtils/P4Utils.cs
@@ -426,6 +426,12 @@ namespace AutomationTool
 		/// </summary>
 		[Description("add into")]
 		AddInto,
+
+		/// <summary>
+		/// file was reverted to a previous revision
+		/// </summary>
+		[Description("undid")]
+		Undid,
 	}
 
 	/// <summary>
