Change UAT to using the editor's UnrealPak commandlet rather than invoking the standalone UnrealPak executable. To improve performance when building several PAK files, also add a new -batch=<file> command which reads commands to execute in parallel from a text file.

diff --git a/Engine/Source/Developer/PakFileUtilities/Private/PakFileUtilities.cpp b/Engine/Source/Developer/PakFileUtilities/Private/PakFileUtilities.cpp
--- a/Engine/Source/Developer/PakFileUtilities/Private/PakFileUtilities.cpp
+++ b/Engine/Source/Developer/PakFileUtilities/Private/PakFileUtilities.cpp
@@ -19,6 +19,7 @@
 #include "Misc/FileHelper.h"
 #include "Misc/ConfigCacheIni.h"
 #include "HAL/PlatformFilemanager.h"
+#include "Async/ParallelFor.h"
 
 struct FPakCommandLineParameters
 {
@@ -2502,6 +2503,21 @@ bool ExecuteUnrealPak(const TCHAR* CmdLine)
 	FAES::FAESKey EncryptionKey;
 	PrepareEncryptionAndSigningKeys(CmdLine, SigningKey, EncryptionKey);
 
+	FString BatchFileName;
+	if (FParse::Value(CmdLine, TEXT("-Batch="), BatchFileName))
+	{
+		TArray<FString> Commands;
+		if (!FFileHelper::LoadFileToStringArray(Commands, *BatchFileName))
+		{
+			UE_LOG(LogPakFile, Error, TEXT("Unable to read '%s'"), *BatchFileName);
+			return false;
+		}
+
+		TAtomic<bool> Result(true);
+		ParallelFor(Commands.Num(), [&Commands, &Result](int32 Idx) { if (!ExecuteUnrealPak(*Commands[Idx])) { Result = false; } });
+		return Result;
+	}
+
 	FString KeyFilename;
 	if (FParse::Value(CmdLine, TEXT("GenerateKeys="), KeyFilename, false))
 	{
diff --git a/Engine/Source/Programs/AutomationTool/Scripts/CopyBuildToStagingDirectory.Automation.cs b/Engine/Source/Programs/AutomationTool/Scripts/CopyBuildToStagingDirectory.Automation.cs
--- a/Engine/Source/Programs/AutomationTool/Scripts/CopyBuildToStagingDirectory.Automation.cs
+++ b/Engine/Source/Programs/AutomationTool/Scripts/CopyBuildToStagingDirectory.Automation.cs
@@ -12,6 +12,7 @@ using System.Linq;
 using UnrealBuildTool;
 using System.Text;
 using Tools.DotNETCommon;
+using System.Diagnostics;
 
 /// <summary>
 /// Helper command used for cooking.
@@ -108,52 +109,45 @@ public partial class Project : CommandUtils
 		return Result;
 	}
 
-	static public void RunUnrealPak(Dictionary<string, string> UnrealPakResponseFile, FileReference OutputLocation, FileReference PakOrderFileLocation, string PlatformOptions, bool Compressed, EncryptionAndSigning.CryptoSettings CryptoSettings, FileReference CryptoKeysCacheFilename, String PatchSourceContentPath)
+	static public string GetUnrealPakArguments(Dictionary<string, string> UnrealPakResponseFile, FileReference OutputLocation, FileReference PakOrderFileLocation, string PlatformOptions, bool Compressed, EncryptionAndSigning.CryptoSettings CryptoSettings, FileReference CryptoKeysCacheFilename, String PatchSourceContentPath)
 	{
-		if (UnrealPakResponseFile.Count < 1)
-		{
-			return;
-		}
+		StringBuilder CmdLine = new StringBuilder(MakePathSafeToUseWithCommandLine(OutputLocation.FullName));
+
 		string PakName = Path.GetFileNameWithoutExtension(OutputLocation.FullName);
 		string UnrealPakResponseFileName = CombinePaths(CmdEnv.LogFolder, "PakList_" + PakName + ".txt");
 		WritePakResponseFile(UnrealPakResponseFileName, UnrealPakResponseFile, Compressed, CryptoSettings);
-
-		var UnrealPakExe = CombinePaths(CmdEnv.LocalRoot, "Engine/Binaries/Win64/UnrealPak.exe");
-		Log("Running UnrealPak *******");
-		string CmdLine = CommandUtils.MakePathSafeToUseWithCommandLine(OutputLocation.FullName) + " -create=" + CommandUtils.MakePathSafeToUseWithCommandLine(UnrealPakResponseFileName);
-		string LogFileName = CombinePaths(CmdEnv.LogFolder, "PakLog_" + PakName + ".log");
-
-		CmdLine += String.Format(" -abslog=\"{0}\"", LogFileName);
-
+		CmdLine.AppendFormat(" -create={0}", CommandUtils.MakePathSafeToUseWithCommandLine(UnrealPakResponseFileName));
+		
 		if(CryptoKeysCacheFilename != null)
 		{
-			CmdLine += String.Format(" -cryptokeys=\"{0}\"", CryptoKeysCacheFilename.ToString());
-		}
-		if (GlobalCommandLine.Installed)
-		{
-			CmdLine += " -installed";
+			CmdLine.AppendFormat(" -cryptokeys={0}", CommandUtils.MakePathSafeToUseWithCommandLine(CryptoKeysCacheFilename.FullName));
 		}
 		if(PakOrderFileLocation != null)
 		{
-			CmdLine += " -order=" + CommandUtils.MakePathSafeToUseWithCommandLine(PakOrderFileLocation.FullName);
-		}
-		if (GlobalCommandLine.UTF8Output)
-		{
-			CmdLine += " -UTF8Output";
+			CmdLine.AppendFormat(" -order={0}", CommandUtils.MakePathSafeToUseWithCommandLine(PakOrderFileLocation.FullName));
 		}
 		if (!String.IsNullOrEmpty(PatchSourceContentPath))
 		{
-			CmdLine += " -generatepatch=" + CommandUtils.MakePathSafeToUseWithCommandLine(PatchSourceContentPath) + " -tempfiles=" + CommandUtils.MakePathSafeToUseWithCommandLine(CommandUtils.CombinePaths(CommandUtils.CmdEnv.LocalRoot, "TempFiles" + PakName));
+			CmdLine.AppendFormat(" -generatepatch={0} -tempfiles={1}", CommandUtils.MakePathSafeToUseWithCommandLine(PatchSourceContentPath), CommandUtils.MakePathSafeToUseWithCommandLine(CommandUtils.CombinePaths(CommandUtils.CmdEnv.LocalRoot, "TempFiles" + Path.GetFileNameWithoutExtension(OutputLocation.FullName))));
 		}
 		if (CryptoSettings != null && CryptoSettings.bEnablePakIndexEncryption)
 		{
-			CmdLine += " -encryptindex";
+			CmdLine.AppendFormat(" -encryptindex");
 		}
-		CmdLine += " -multiprocess"; // Prevents warnings about being unable to write to config files
-		CmdLine += PlatformOptions;
-		string UnrealPakLogFileName = "UnrealPak_" + PakName;
-		RunAndLog(CmdEnv, UnrealPakExe, CmdLine, LogName: UnrealPakLogFileName, Options: ERunOptions.Default | ERunOptions.UTF8Output);
-		Log("UnrealPak Done *******");
+		CmdLine.Append(PlatformOptions);
+
+		return CmdLine.ToString();
+	}
+
+	static public void RunUnrealPak(ProjectParams Params, Dictionary<string, string> UnrealPakResponseFile, FileReference OutputLocation, FileReference PakOrderFileLocation, string PlatformOptions, bool Compressed, EncryptionAndSigning.CryptoSettings CryptoSettings, FileReference CryptoKeysCacheFilename, String PatchSourceContentPath)
+	{
+		if (UnrealPakResponseFile.Count < 1)
+		{
+			return;
+		}
+
+		string Arguments = GetUnrealPakArguments(UnrealPakResponseFile, OutputLocation, PakOrderFileLocation, PlatformOptions, Compressed, CryptoSettings, CryptoKeysCacheFilename, PatchSourceContentPath);
+		RunCommandlet(Params.RawProjectPath, Params.UE4Exe, "UnrealPak", Arguments);
 	}
 
 	static public void LogDeploymentContext(DeploymentContext SC)
@@ -1165,7 +1159,9 @@ public partial class Project : CommandUtils
 		FileReference CryptoKeysCacheFilename = FileReference.Combine(SC.MetadataDir, "Crypto.json");
 		PakCryptoSettings.Save(CryptoKeysCacheFilename);
 
-		CreatePak(Params, SC, UnrealPakResponseFile, SC.ShortProjectName, PakCryptoSettings, CryptoKeysCacheFilename);
+		List<CreatePakParams> PakInputs = new List<CreatePakParams>();
+		PakInputs.Add(new CreatePakParams(SC.ShortProjectName, UnrealPakResponseFile));
+		CreatePaks(Params, SC, PakInputs, PakCryptoSettings, CryptoKeysCacheFilename);
 	}
 
 	/// <summary>
@@ -1180,7 +1176,7 @@ public partial class Project : CommandUtils
 		Dictionary<string, string> PakResponseFile = CreatePakResponseFileFromStagingManifest(SC, SC.CrashReporterUFSFiles);
 		FileReference OutputLocation = FileReference.Combine(SC.RuntimeRootDir, "Engine", "Programs", "CrashReportClient", "Content", "Paks", "CrashReportClient.pak");
 
-		RunUnrealPak(PakResponseFile, OutputLocation, null, null, Params.Compressed, null, null, null);
+		RunUnrealPak(Params, PakResponseFile, OutputLocation, null, null, Params.Compressed, null, null, null);
 	}
 
 	/// <summary>
@@ -1275,6 +1271,33 @@ public partial class Project : CommandUtils
 		return UnrealPakResponseFile;
 	}
 
+	/// <summary>
+	/// Parameter class for CreatePaks(). Each instance stores information about a pak file to be created.
+	/// </summary>
+	[DebuggerDisplay("PakName")]
+	class CreatePakParams
+	{
+		/// <summary>
+		/// Path to the base output file for this pak file 
+		/// </summary>
+		public string PakName;
+
+		/// <summary>
+		/// Map of files within the pak file to their source file on disk
+		/// </summary>
+		public Dictionary<string, string> UnrealPakResponseFile;
+
+		/// <summary>
+		/// Constructor
+		/// </summary>
+		/// <param name="PakName">Path to the base output file for this pak file</param>
+		/// <param name="UnrealPakResponseFile">Map of files within the pak file to their source file on disk</param>
+		public CreatePakParams(string PakName, Dictionary<string, string> UnrealPakResponseFile)
+		{
+			this.PakName = PakName;
+			this.UnrealPakResponseFile = UnrealPakResponseFile;
+		}
+	}
 
 	/// <summary>
 	/// Creates a pak file using response file.
@@ -1283,7 +1306,7 @@ public partial class Project : CommandUtils
 	/// <param name="SC"></param>
 	/// <param name="UnrealPakResponseFile"></param>
 	/// <param name="PakName"></param>
-	private static void CreatePak(ProjectParams Params, DeploymentContext SC, Dictionary<string, string> UnrealPakResponseFile, string PakName, EncryptionAndSigning.CryptoSettings CryptoSettings, FileReference CryptoKeysCacheFilename)
+	private static void CreatePaks(ProjectParams Params, DeploymentContext SC, List<CreatePakParams> PakParamsList, EncryptionAndSigning.CryptoSettings CryptoSettings, FileReference CryptoKeysCacheFilename)
 	{
 		bool bShouldGeneratePatch = Params.IsGeneratingPatch && SC.StageTargetPlatform.GetPlatformPatchesWithDiffPak(Params, SC);
 
@@ -1293,264 +1316,294 @@ public partial class Project : CommandUtils
 		}
 
 		string PostFix = "";
-		string OutputFilename = PakName + "-" + SC.FinalCookPlatform;
-		string OutputFilenameExtension = ".pak";
 		if (bShouldGeneratePatch)
 		{
 			PostFix += "_P";
-			int TargetPatchIndex = 0;
-			string ExistingPatchSearchPath = SC.StageTargetPlatform.GetReleasePakFilePath(SC, Params, null);
-			if (Directory.Exists(ExistingPatchSearchPath))
-			{
-				IEnumerable<string> PakFileSet = Directory.EnumerateFiles(ExistingPatchSearchPath, OutputFilename + "*" + PostFix + OutputFilenameExtension);
-				foreach (string PakFilePath in PakFileSet)
-				{
-					string PakFileName = Path.GetFileName(PakFilePath);
-					int StartIndex = OutputFilename.Length + 1;
-					int LengthVar = PakFileName.Length - (OutputFilename.Length + 1 + PostFix.Length + OutputFilenameExtension.Length);
-					if (LengthVar > 0)
-					{
-						string PakFileIndex = PakFileName.Substring(StartIndex, LengthVar);
-						int ChunkIndex;
-						if (int.TryParse(PakFileIndex, out ChunkIndex))
-						{
-							if (ChunkIndex > TargetPatchIndex)
-							{
-								TargetPatchIndex = ChunkIndex;
-							}
-						}
-					}
-				}
-				if (Params.ShouldAddPatchLevel && PakFileSet.Count() > 0)
-				{
-					TargetPatchIndex++;
-				}
-			}
-			OutputFilename = OutputFilename + "_" + TargetPatchIndex + PostFix;
 		}
 		else if (Params.IterateSharedCookedBuild)
 		{
 			// shared cooked builds will produce a patch
 			// then be combined with the shared cooked build
 			PostFix += "_S_P";
-			OutputFilename = OutputFilename + PostFix;
 		}
 
-		StagedFileReference OutputRelativeLocation;
-		if (Params.HasDLCName)
-		{
-			OutputRelativeLocation = StagedFileReference.Combine(SC.RelativeProjectRootForStage, Params.DLCFile.Directory.MakeRelativeTo(SC.ProjectRoot), "Content", "Paks", SC.FinalCookPlatform, Params.DLCFile.GetFileNameWithoutExtension() + ".pak");
-		}
-		else
-		{
-			OutputRelativeLocation = StagedFileReference.Combine(SC.RelativeProjectRootForStage, "Content", "Paks", OutputFilename + OutputFilenameExtension);
-		}
-		if (SC.StageTargetPlatform.DeployLowerCaseFilenames())
-		{
-			OutputRelativeLocation = OutputRelativeLocation.ToLowerInvariant();
-		}
-		OutputRelativeLocation = SC.StageTargetPlatform.Remap(OutputRelativeLocation);
-
-		FileReference OutputLocation = FileReference.Combine(SC.RuntimeRootDir, OutputRelativeLocation.Name);
-		// Add input file to control order of file within the pak
-		DirectoryReference PakOrderFileLocationBase = DirectoryReference.Combine(SC.ProjectRoot, "Build", SC.FinalCookPlatform, "FileOpenOrder");
+		const string OutputFilenameExtension = ".pak";
 
-		FileReference PakOrderFileLocation = null;
+		List<string> Commands = new List<string>();
 
-		string[] OrderFileNames = new string[] { "GameOpenOrder.log", "CookerOpenOrder.log", "EditorOpenOrder.log" };
-		foreach (string OrderFileName in OrderFileNames)
+		List<Tuple<FileReference, StagedFileReference, string>> Outputs = new List<Tuple<FileReference, StagedFileReference, string>>();
+		foreach(CreatePakParams PakParams in PakParamsList)
 		{
-			PakOrderFileLocation = FileReference.Combine(PakOrderFileLocationBase, OrderFileName);
+			string OutputFilename = PakParams.PakName + "-" + SC.FinalCookPlatform;
+			if (bShouldGeneratePatch)
+			{
+				int TargetPatchIndex = 0;
+				string ExistingPatchSearchPath = SC.StageTargetPlatform.GetReleasePakFilePath(SC, Params, null);
+				if (Directory.Exists(ExistingPatchSearchPath))
+				{
+					IEnumerable<string> PakFileSet = Directory.EnumerateFiles(ExistingPatchSearchPath, OutputFilename + "*" + PostFix + OutputFilenameExtension);
+					foreach (string PakFilePath in PakFileSet)
+					{
+						string PakFileName = Path.GetFileName(PakFilePath);
+						int StartIndex = OutputFilename.Length + 1;
+						int LengthVar = PakFileName.Length - (OutputFilename.Length + 1 + PostFix.Length + OutputFilenameExtension.Length);
+						if (LengthVar > 0)
+						{
+							string PakFileIndex = PakFileName.Substring(StartIndex, LengthVar);
+							int ChunkIndex;
+							if (int.TryParse(PakFileIndex, out ChunkIndex))
+							{
+								if (ChunkIndex > TargetPatchIndex)
+								{
+									TargetPatchIndex = ChunkIndex;
+								}
+							}
+						}
+					}
+					if (Params.ShouldAddPatchLevel && PakFileSet.Count() > 0)
+					{
+						TargetPatchIndex++;
+					}
+				}
+				OutputFilename = OutputFilename + "_" + TargetPatchIndex;
+			}
+			OutputFilename = OutputFilename += PostFix;
 
-			if (FileExists_NoExceptions(PakOrderFileLocation.FullName))
+			StagedFileReference OutputRelativeLocation;
+			if (Params.HasDLCName)
 			{
-				break;
+				OutputRelativeLocation = StagedFileReference.Combine(SC.RelativeProjectRootForStage, Params.DLCFile.Directory.MakeRelativeTo(SC.ProjectRoot), "Content", "Paks", SC.FinalCookPlatform, Params.DLCFile.GetFileNameWithoutExtension() + ".pak");
 			}
-		}
+			else
+			{
+				OutputRelativeLocation = StagedFileReference.Combine(SC.RelativeProjectRootForStage, "Content", "Paks", OutputFilename + OutputFilenameExtension);
+			}
+			if (SC.StageTargetPlatform.DeployLowerCaseFilenames())
+			{
+				OutputRelativeLocation = OutputRelativeLocation.ToLowerInvariant();
+			}
+			OutputRelativeLocation = SC.StageTargetPlatform.Remap(OutputRelativeLocation);
 
-		bool bCopiedExistingPak = false;
+			FileReference OutputLocation = FileReference.Combine(SC.RuntimeRootDir, OutputRelativeLocation.Name);
+			// Add input file to control order of file within the pak
+			DirectoryReference PakOrderFileLocationBase = DirectoryReference.Combine(SC.ProjectRoot, "Build", SC.FinalCookPlatform, "FileOpenOrder");
 
-		if (SC.StageTargetPlatform != SC.CookSourcePlatform)
-		{
-			// Check to see if we have an existing pak file we can use
+			FileReference PakOrderFileLocation = null;
 
-			StagedFileReference SourceOutputRelativeLocation = StagedFileReference.Combine(SC.RelativeProjectRootForStage, "Content/Paks/", PakName + "-" + SC.CookPlatform + PostFix + ".pak");
-			if (SC.CookSourcePlatform.DeployLowerCaseFilenames())
+			string[] OrderFileNames = new string[] { "GameOpenOrder.log", "CookerOpenOrder.log", "EditorOpenOrder.log" };
+			foreach (string OrderFileName in OrderFileNames)
 			{
-				SourceOutputRelativeLocation = SourceOutputRelativeLocation.ToLowerInvariant();
+				PakOrderFileLocation = FileReference.Combine(PakOrderFileLocationBase, OrderFileName);
+
+				if (FileExists_NoExceptions(PakOrderFileLocation.FullName))
+				{
+					break;
+				}
 			}
-			SourceOutputRelativeLocation = SC.CookSourcePlatform.Remap(SourceOutputRelativeLocation);
 
-			FileReference SourceOutputLocation = FileReference.Combine(SC.CookSourceRuntimeRootDir, SourceOutputRelativeLocation.Name);
-			if (FileExists_NoExceptions(SourceOutputLocation.FullName))
+			bool bCopiedExistingPak = false;
+
+			if (SC.StageTargetPlatform != SC.CookSourcePlatform)
 			{
-				InternalUtils.SafeCreateDirectory(Path.GetDirectoryName(OutputLocation.FullName), true);
+				// Check to see if we have an existing pak file we can use
+
+				StagedFileReference SourceOutputRelativeLocation = StagedFileReference.Combine(SC.RelativeProjectRootForStage, "Content/Paks/", PakParams.PakName + "-" + SC.CookPlatform + PostFix + ".pak");
+				if (SC.CookSourcePlatform.DeployLowerCaseFilenames())
+				{
+					SourceOutputRelativeLocation = SourceOutputRelativeLocation.ToLowerInvariant();
+				}
+				SourceOutputRelativeLocation = SC.CookSourcePlatform.Remap(SourceOutputRelativeLocation);
 
-				if (InternalUtils.SafeCopyFile(SourceOutputLocation.FullName, OutputLocation.FullName))
+				FileReference SourceOutputLocation = FileReference.Combine(SC.CookSourceRuntimeRootDir, SourceOutputRelativeLocation.Name);
+				if (FileExists_NoExceptions(SourceOutputLocation.FullName))
 				{
-					Log("Copying source pak from {0} to {1} instead of creating new pak", SourceOutputLocation, OutputLocation);
-					bCopiedExistingPak = true;
+					InternalUtils.SafeCreateDirectory(Path.GetDirectoryName(OutputLocation.FullName), true);
 
-					FileReference InSigFile = SourceOutputLocation.ChangeExtension(".sig");
-					if (FileReference.Exists(InSigFile))
+					if (InternalUtils.SafeCopyFile(SourceOutputLocation.FullName, OutputLocation.FullName))
 					{
-						FileReference OutSigFile = OutputLocation.ChangeExtension(".sig");
-
-						Log("Copying pak sig from {0} to {1}", InSigFile, OutSigFile);
+						Log("Copying source pak from {0} to {1} instead of creating new pak", SourceOutputLocation, OutputLocation);
+						bCopiedExistingPak = true;
 
-						if (!InternalUtils.SafeCopyFile(InSigFile.FullName, OutSigFile.FullName))
+						FileReference InSigFile = SourceOutputLocation.ChangeExtension(".sig");
+						if (FileReference.Exists(InSigFile))
 						{
-							Log("Failed to copy pak sig {0} to {1}, creating new pak", InSigFile, InSigFile);
-							bCopiedExistingPak = false;
+							FileReference OutSigFile = OutputLocation.ChangeExtension(".sig");
+
+							Log("Copying pak sig from {0} to {1}", InSigFile, OutSigFile);
+
+							if (!InternalUtils.SafeCopyFile(InSigFile.FullName, OutSigFile.FullName))
+							{
+								Log("Failed to copy pak sig {0} to {1}, creating new pak", InSigFile, InSigFile);
+								bCopiedExistingPak = false;
+							}
 						}
-					}
 
+					}
+				}
+				if (!bCopiedExistingPak)
+				{
+					Log("Failed to copy source pak from {0} to {1}, creating new pak", SourceOutputLocation, OutputLocation);
 				}
 			}
-			if (!bCopiedExistingPak)
+
+			string PatchSourceContentPath = null;
+			if (bShouldGeneratePatch)
 			{
-				Log("Failed to copy source pak from {0} to {1}, creating new pak", SourceOutputLocation, OutputLocation);
+				// don't include the post fix in this filename because we are looking for the source pak path
+				string PakFilename = PakParams.PakName + "-" + SC.FinalCookPlatform + "*.pak";
+				PatchSourceContentPath = SC.StageTargetPlatform.GetReleasePakFilePath(SC, Params, PakFilename);
 			}
-		}
-
-		string PatchSourceContentPath = null;
-		if (bShouldGeneratePatch)
-		{
-			// don't include the post fix in this filename because we are looking for the source pak path
-			string PakFilename = PakName + "-" + SC.FinalCookPlatform + "*.pak";
-			PatchSourceContentPath = SC.StageTargetPlatform.GetReleasePakFilePath(SC, Params, PakFilename);
-		}
 		
-		if (!bCopiedExistingPak)
-		{
-			if (FileReference.Exists(OutputLocation))
+			if (!bCopiedExistingPak)
 			{
-				string UnrealPakResponseFileName = CombinePaths(CmdEnv.LogFolder, "PakList_" + OutputLocation.GetFileNameWithoutExtension() + ".txt");
-				if (File.Exists(UnrealPakResponseFileName) && FileReference.GetLastWriteTimeUtc(OutputLocation) > File.GetLastWriteTimeUtc(UnrealPakResponseFileName))
+				if (FileReference.Exists(OutputLocation))
+				{
+					string UnrealPakResponseFileName = CombinePaths(CmdEnv.LogFolder, "PakList_" + OutputLocation.GetFileNameWithoutExtension() + ".txt");
+					if (File.Exists(UnrealPakResponseFileName) && FileReference.GetLastWriteTimeUtc(OutputLocation) > File.GetLastWriteTimeUtc(UnrealPakResponseFileName))
+					{
+						bCopiedExistingPak = true;
+					}
+				}
+				if (!bCopiedExistingPak)
 				{
-					bCopiedExistingPak = true;
+					Commands.Add(GetUnrealPakArguments(PakParams.UnrealPakResponseFile, OutputLocation, PakOrderFileLocation, SC.StageTargetPlatform.GetPlatformPakCommandLine(), Params.Compressed, CryptoSettings, CryptoKeysCacheFilename, PatchSourceContentPath));
 				}
 			}
-			if (!bCopiedExistingPak)
-			{
-				RunUnrealPak(UnrealPakResponseFile, OutputLocation, PakOrderFileLocation, SC.StageTargetPlatform.GetPlatformPakCommandLine(), Params.Compressed, CryptoSettings, CryptoKeysCacheFilename, PatchSourceContentPath);
-			}
+
+			Outputs.Add(Tuple.Create(OutputLocation, OutputRelativeLocation, PatchSourceContentPath));
 		}
 
-		if (Params.HasCreateReleaseVersion)
+		// Actually execute UnrealPak
+		if(Commands.Count > 0)
 		{
-			// copy the created pak to the release version directory we might need this later if we want to generate patches
-			//string ReleaseVersionPath = CombinePaths( SC.ProjectRoot, "Releases", Params.CreateReleaseVersion, SC.StageTargetPlatform.GetCookPlatform(Params.DedicatedServer, false), Path.GetFileName(OutputLocation) );
-			string ReleaseVersionPath = SC.StageTargetPlatform.GetReleasePakFilePath(SC, Params, OutputLocation.GetFileName());
-
-			InternalUtils.SafeCreateDirectory(Path.GetDirectoryName(ReleaseVersionPath));
-			InternalUtils.SafeCopyFile(OutputLocation.FullName, ReleaseVersionPath);
+			string CommandsFile = LogUtils.GetUniqueLogName(CombinePaths(CmdEnv.EngineSavedFolder, "UnrealPak-Commands"));
+			File.WriteAllLines(CommandsFile, Commands);
+			RunCommandlet(Params.RawProjectPath, Params.UE4Exe, "UnrealPak", String.Format("-batch={0}", MakePathSafeToUseWithCommandLine(CommandsFile)));
 		}
 
-		if (Params.CreateChunkInstall)
+		// Do any additional processing on the command output
+		for(int Idx = 0; Idx < PakParamsList.Count; Idx++)
 		{
-			var RegEx = new Regex("pakchunk(\\d+)", RegexOptions.IgnoreCase | RegexOptions.Compiled);
-			var Matches = RegEx.Matches(PakName);
+			string PakName = PakParamsList[Idx].PakName;
+			FileReference OutputLocation = Outputs[Idx].Item1;
+			StagedFileReference OutputRelativeLocation = Outputs[Idx].Item2;
+			string PatchSourceContentPath = Outputs[Idx].Item3;
 
-			int ChunkID = 0;
-			if (Matches.Count != 0 && Matches[0].Groups.Count > 1)
+			if (Params.HasCreateReleaseVersion)
 			{
-				ChunkID = Convert.ToInt32(Matches[0].Groups[1].ToString());
-			}
-			else if (Params.HasDLCName)
-			{
-				// Assuming DLC is a single pack file
-				ChunkID = 1;
-			}
-			else
-			{
-				throw new AutomationException(String.Format("Failed Creating Chunk Install Data, Unable to parse chunk id from {0}", PakName));
+				// copy the created pak to the release version directory we might need this later if we want to generate patches
+				//string ReleaseVersionPath = CombinePaths( SC.ProjectRoot, "Releases", Params.CreateReleaseVersion, SC.StageTargetPlatform.GetCookPlatform(Params.DedicatedServer, false), Path.GetFileName(OutputLocation) );
+				string ReleaseVersionPath = SC.StageTargetPlatform.GetReleasePakFilePath(SC, Params, OutputLocation.GetFileName());
+
+				InternalUtils.SafeCreateDirectory(Path.GetDirectoryName(ReleaseVersionPath));
+				InternalUtils.SafeCopyFile(OutputLocation.FullName, ReleaseVersionPath);
 			}
 
-			if (ChunkID != 0)
+			if (Params.CreateChunkInstall)
 			{
-				var BPTExe = GetBuildPatchToolExecutable();
-				EnsureBuildPatchToolExists();
+				var RegEx = new Regex("pakchunk(\\d+)", RegexOptions.IgnoreCase | RegexOptions.Compiled);
+				var Matches = RegEx.Matches(PakName);
 
-				string VersionString = Params.ChunkInstallVersionString;
-				string ChunkInstallBasePath = CombinePaths(Params.ChunkInstallDirectory, SC.FinalCookPlatform);
-				string RawDataPath = CombinePaths(ChunkInstallBasePath, VersionString, PakName);
-				string RawDataPakPath = CombinePaths(RawDataPath, PakName + "-" + SC.FinalCookPlatform + PostFix + ".pak");
-				//copy the pak chunk to the raw data folder
-				if (InternalUtils.SafeFileExists(RawDataPakPath, true))
+				int ChunkID = 0;
+				if (Matches.Count != 0 && Matches[0].Groups.Count > 1)
+				{
+					ChunkID = Convert.ToInt32(Matches[0].Groups[1].ToString());
+				}
+				else if (Params.HasDLCName)
+				{
+					// Assuming DLC is a single pack file
+					ChunkID = 1;
+				}
+				else
 				{
-					InternalUtils.SafeDeleteFile(RawDataPakPath, true);
+					throw new AutomationException(String.Format("Failed Creating Chunk Install Data, Unable to parse chunk id from {0}", PakName));
 				}
-				InternalUtils.SafeCreateDirectory(RawDataPath, true);
-				InternalUtils.SafeCopyFile(OutputLocation.FullName, RawDataPakPath);
-				InternalUtils.SafeDeleteFile(OutputLocation.FullName, true);
 
-				if (Params.IsGeneratingPatch)
+				if (ChunkID != 0)
 				{
-					if (String.IsNullOrEmpty(PatchSourceContentPath))
+					var BPTExe = GetBuildPatchToolExecutable();
+					EnsureBuildPatchToolExists();
+
+					string VersionString = Params.ChunkInstallVersionString;
+					string ChunkInstallBasePath = CombinePaths(Params.ChunkInstallDirectory, SC.FinalCookPlatform);
+					string RawDataPath = CombinePaths(ChunkInstallBasePath, VersionString, PakName);
+					string RawDataPakPath = CombinePaths(RawDataPath, PakName + "-" + SC.FinalCookPlatform + PostFix + ".pak");
+					//copy the pak chunk to the raw data folder
+					if (InternalUtils.SafeFileExists(RawDataPakPath, true))
 					{
-						throw new AutomationException(String.Format("Failed Creating Chunk Install Data. No source pak for patch pak {0} given", OutputLocation));
+						InternalUtils.SafeDeleteFile(RawDataPakPath, true);
 					}
-					// If we are generating a patch, then we need to copy the original pak across
-					// for distribution.
-					string SourceRawDataPakPath = CombinePaths(RawDataPath, PakName + "-" + SC.FinalCookPlatform + ".pak");
-					InternalUtils.SafeCopyFile(PatchSourceContentPath, SourceRawDataPakPath);
-				}
-
-				string BuildRoot = MakePathSafeToUseWithCommandLine(RawDataPath);
-				string CloudDir = MakePathSafeToUseWithCommandLine(CombinePaths(ChunkInstallBasePath, "CloudDir"));
-				string ManifestDir = CombinePaths(ChunkInstallBasePath, "ManifestDir");
-				var AppID = 1; // For a chunk install this value doesn't seem to matter
-				string AppName = String.Format("{0}_{1}", SC.ShortProjectName, PakName);
-				string AppLaunch = ""; // For a chunk install this value doesn't seem to matter
-				string ManifestFilename = AppName + VersionString + ".manifest";
-				string SourceManifestPath = CombinePaths(CloudDir, ManifestFilename);
-				string DestManifestPath = CombinePaths(ManifestDir, ManifestFilename);
-				InternalUtils.SafeCreateDirectory(ManifestDir, true);
-
-				string CmdLine = String.Format("-BuildRoot=\"{0}\" -CloudDir=\"{1}\" -AppID={2} -AppName=\"{3}\" -BuildVersion=\"{4}\" -AppLaunch=\"{5}\"", BuildRoot, CloudDir, AppID, AppName, VersionString, AppLaunch);
-				CmdLine += " -AppArgs=\"\"";
-				CmdLine += " -custom=\"bIsPatch=false\"";
-				CmdLine += String.Format(" -customint=\"ChunkID={0}\"", ChunkID);
-				CmdLine += " -customint=\"PakReadOrdering=0\"";
-				CmdLine += " -stdout";
-
-				string UnrealPakLogFileName = "UnrealPak_" + PakName;
-				RunAndLog(CmdEnv, BPTExe, CmdLine, UnrealPakLogFileName, Options: ERunOptions.Default | ERunOptions.UTF8Output);
-
-				InternalUtils.SafeCopyFile(SourceManifestPath, DestManifestPath);
+					InternalUtils.SafeCreateDirectory(RawDataPath, true);
+					InternalUtils.SafeCopyFile(OutputLocation.FullName, RawDataPakPath);
+					InternalUtils.SafeDeleteFile(OutputLocation.FullName, true);
+
+					if (Params.IsGeneratingPatch)
+					{
+						if (String.IsNullOrEmpty(PatchSourceContentPath))
+						{
+							throw new AutomationException(String.Format("Failed Creating Chunk Install Data. No source pak for patch pak {0} given", OutputLocation));
+						}
+						// If we are generating a patch, then we need to copy the original pak across
+						// for distribution.
+						string SourceRawDataPakPath = CombinePaths(RawDataPath, PakName + "-" + SC.FinalCookPlatform + ".pak");
+						InternalUtils.SafeCopyFile(PatchSourceContentPath, SourceRawDataPakPath);
+					}
+
+					string BuildRoot = MakePathSafeToUseWithCommandLine(RawDataPath);
+					string CloudDir = MakePathSafeToUseWithCommandLine(CombinePaths(ChunkInstallBasePath, "CloudDir"));
+					string ManifestDir = CombinePaths(ChunkInstallBasePath, "ManifestDir");
+					var AppID = 1; // For a chunk install this value doesn't seem to matter
+					string AppName = String.Format("{0}_{1}", SC.ShortProjectName, PakName);
+					string AppLaunch = ""; // For a chunk install this value doesn't seem to matter
+					string ManifestFilename = AppName + VersionString + ".manifest";
+					string SourceManifestPath = CombinePaths(CloudDir, ManifestFilename);
+					string DestManifestPath = CombinePaths(ManifestDir, ManifestFilename);
+					InternalUtils.SafeCreateDirectory(ManifestDir, true);
+
+					string CmdLine = String.Format("-BuildRoot=\"{0}\" -CloudDir=\"{1}\" -AppID={2} -AppName=\"{3}\" -BuildVersion=\"{4}\" -AppLaunch=\"{5}\"", BuildRoot, CloudDir, AppID, AppName, VersionString, AppLaunch);
+					CmdLine += " -AppArgs=\"\"";
+					CmdLine += " -custom=\"bIsPatch=false\"";
+					CmdLine += String.Format(" -customint=\"ChunkID={0}\"", ChunkID);
+					CmdLine += " -customint=\"PakReadOrdering=0\"";
+					CmdLine += " -stdout";
+
+					string UnrealPakLogFileName = "UnrealPak_" + PakName;
+					RunAndLog(CmdEnv, BPTExe, CmdLine, UnrealPakLogFileName, Options: ERunOptions.Default | ERunOptions.UTF8Output);
+
+					InternalUtils.SafeCopyFile(SourceManifestPath, DestManifestPath);
+				}
+				else
+				{
+					// add the first pak file as needing deployment and convert to lower case again if needed
+					SC.FilesToStage.UFSFiles.Add(OutputRelativeLocation, OutputLocation);
+				}
 			}
 			else
 			{
-				// add the first pak file as needing deployment and convert to lower case again if needed
+				// add the pak file as needing deployment and convert to lower case again if needed
 				SC.FilesToStage.UFSFiles.Add(OutputRelativeLocation, OutputLocation);
-			}
-		}
-		else
-		{
-			// add the pak file as needing deployment and convert to lower case again if needed
-			SC.FilesToStage.UFSFiles.Add(OutputRelativeLocation, OutputLocation);
 
-			// add the base pak files to deployment as well
-			if (bShouldGeneratePatch)
-			{
-				string ExistingPatchSearchPath = SC.StageTargetPlatform.GetReleasePakFilePath(SC, Params, null);
-				if (Directory.Exists(ExistingPatchSearchPath))
+				// add the base pak files to deployment as well
+				if (bShouldGeneratePatch)
 				{
-					IEnumerable<string> PakFileSet = Directory.EnumerateFiles(ExistingPatchSearchPath, PakName + "-" + SC.FinalCookPlatform + "*" + OutputFilenameExtension);
-					foreach (string PakFilePath in PakFileSet)
+					string ExistingPatchSearchPath = SC.StageTargetPlatform.GetReleasePakFilePath(SC, Params, null);
+					if (Directory.Exists(ExistingPatchSearchPath))
 					{
-						FileReference OutputDestinationPath = FileReference.Combine(OutputLocation.Directory, Path.GetFileName(PakFilePath));
-						if (!File.Exists(OutputDestinationPath.FullName))
+						IEnumerable<string> PakFileSet = Directory.EnumerateFiles(ExistingPatchSearchPath, PakName + "-" + SC.FinalCookPlatform + "*" + OutputFilenameExtension);
+						foreach (string PakFilePath in PakFileSet)
 						{
-							InternalUtils.SafeCopyFile(PakFilePath, OutputDestinationPath.FullName);
-							StagedFileReference OutputDestinationRelativeLocation = StagedFileReference.Combine(SC.RelativeProjectRootForStage, "Content/Paks/", Path.GetFileName(PakFilePath));
-							if (SC.StageTargetPlatform.DeployLowerCaseFilenames())
+							FileReference OutputDestinationPath = FileReference.Combine(OutputLocation.Directory, Path.GetFileName(PakFilePath));
+							if (!File.Exists(OutputDestinationPath.FullName))
 							{
-								OutputDestinationRelativeLocation = OutputDestinationRelativeLocation.ToLowerInvariant();
+								InternalUtils.SafeCopyFile(PakFilePath, OutputDestinationPath.FullName);
+								StagedFileReference OutputDestinationRelativeLocation = StagedFileReference.Combine(SC.RelativeProjectRootForStage, "Content/Paks/", Path.GetFileName(PakFilePath));
+								if (SC.StageTargetPlatform.DeployLowerCaseFilenames())
+								{
+									OutputDestinationRelativeLocation = OutputDestinationRelativeLocation.ToLowerInvariant();
+								}
+								OutputDestinationRelativeLocation = SC.StageTargetPlatform.Remap(OutputDestinationRelativeLocation);
+								SC.FilesToStage.UFSFiles.Add(OutputDestinationRelativeLocation, OutputDestinationPath);
 							}
-							OutputDestinationRelativeLocation = SC.StageTargetPlatform.Remap(OutputDestinationRelativeLocation);
-							SC.FilesToStage.UFSFiles.Add(OutputDestinationRelativeLocation, OutputDestinationPath);
 						}
 					}
 				}
@@ -1738,21 +1791,17 @@ public partial class Project : CommandUtils
 		FileReference CryptoKeysCacheFilename = FileReference.Combine(SC.MetadataDir, "Crypto.json");
 		PakCryptoSettings.Save(CryptoKeysCacheFilename);
 
-		IEnumerable<Tuple<Dictionary<string,string>, string>> PakPairs = PakResponseFiles.Zip(ChunkList, (a, b) => Tuple.Create(a, b));
-	
-        System.Threading.Tasks.ParallelOptions Options = new System.Threading.Tasks.ParallelOptions();
-
-        Log("Creating Pak files utilizing {0} cores", Environment.ProcessorCount);
-        Options.MaxDegreeOfParallelism = Environment.ProcessorCount;        
-
-        System.Threading.Tasks.Parallel.ForEach(PakPairs, Options, (PakPair) =>
+		List<CreatePakParams> PakInputs = new List<CreatePakParams>();
+		for(int Idx = 0; Idx < PakResponseFiles.Length; Idx++)
 		{
-			var ChunkName = Path.GetFileNameWithoutExtension(PakPair.Item2);
-            if (PakPair.Item1.Count > 0)
+            if (PakResponseFiles[Idx].Count > 0)
             {
-			    CreatePak(Params, SC, PakPair.Item1, ChunkName, PakCryptoSettings, CryptoKeysCacheFilename);
+				string ChunkName = Path.GetFileNameWithoutExtension(ChunkList[Idx]);
+				PakInputs.Add(new CreatePakParams(ChunkName, PakResponseFiles[Idx]));
             }
-		});
+		}
+
+		CreatePaks(Params, SC, PakInputs, PakCryptoSettings, CryptoKeysCacheFilename);
 
 		if (Params.CreateChunkInstall)
 		{
