Fix incorrect assignment when deriving name for chunked pak file.

diff --git a/Engine/Source/Programs/AutomationTool/Scripts/CopyBuildToStagingDirectory.Automation.cs b/Engine/Source/Programs/AutomationTool/Scripts/CopyBuildToStagingDirectory.Automation.cs
--- a/Engine/Source/Programs/AutomationTool/Scripts/CopyBuildToStagingDirectory.Automation.cs
+++ b/Engine/Source/Programs/AutomationTool/Scripts/CopyBuildToStagingDirectory.Automation.cs
@@ -1367,7 +1367,7 @@ public partial class Project : CommandUtils
 				}
 				OutputFilename = OutputFilename + "_" + TargetPatchIndex;
 			}
-			OutputFilename = OutputFilename += PostFix;
+			OutputFilename = OutputFilename + PostFix;
 
 			StagedFileReference OutputRelativeLocation;
 			if (Params.HasDLCName)
