UAT: Remove the unused -SkipHeader argument to UE4Build.

#rnx

diff --git a/Engine/Source/Programs/AutomationTool/AutomationUtils/UE4Build.cs b/Engine/Source/Programs/AutomationTool/AutomationUtils/UE4Build.cs
--- a/Engine/Source/Programs/AutomationTool/AutomationUtils/UE4Build.cs
+++ b/Engine/Source/Programs/AutomationTool/AutomationUtils/UE4Build.cs
@@ -447,7 +447,7 @@ namespace AutomationTool
 		/// <summary>
 		/// Updates the engine version files
 		/// </summary>
-		public List<FileReference> UpdateVersionFiles(bool ActuallyUpdateVersionFiles = true, int? ChangelistNumberOverride = null, int? CompatibleChangelistNumberOverride = null, string Build = null, bool? IsPromotedOverride = null, bool bSkipHeader = false)
+		public List<FileReference> UpdateVersionFiles(bool ActuallyUpdateVersionFiles = true, int? ChangelistNumberOverride = null, int? CompatibleChangelistNumberOverride = null, string Build = null, bool? IsPromotedOverride = null)
 		{
 			bool bIsLicenseeVersion = ParseParam("Licensee") || !FileReference.Exists(FileReference.Combine(CommandUtils.EngineDirectory, "Build", "NotForLicensees", "EpicInternal.txt"));
 			bool bIsPromotedBuild = IsPromotedOverride.HasValue? IsPromotedOverride.Value : (ParseParamInt("Promoted", 1) != 0);
@@ -469,10 +469,10 @@ namespace AutomationTool
 				Branch = CommandUtils.P4Enabled ? CommandUtils.EscapePath(CommandUtils.P4Env.Branch) : "";
 			}
 
-			return StaticUpdateVersionFiles(ChangelistNumber, CompatibleChangelistNumber, Branch, Build, bIsLicenseeVersion, bIsPromotedBuild, bDoUpdateVersionFiles, bSkipHeader);
+			return StaticUpdateVersionFiles(ChangelistNumber, CompatibleChangelistNumber, Branch, Build, bIsLicenseeVersion, bIsPromotedBuild, bDoUpdateVersionFiles);
 		}
 
-		public static List<FileReference> StaticUpdateVersionFiles(int ChangelistNumber, int CompatibleChangelistNumber, string Branch, string Build, bool bIsLicenseeVersion, bool bIsPromotedBuild, bool bDoUpdateVersionFiles, bool bSkipHeader)
+		public static List<FileReference> StaticUpdateVersionFiles(int ChangelistNumber, int CompatibleChangelistNumber, string Branch, string Build, bool bIsLicenseeVersion, bool bIsPromotedBuild, bool bDoUpdateVersionFiles)
 		{
 			FileReference BuildVersionFile = BuildVersion.GetDefaultFileName();
 
diff --git a/Engine/Source/Programs/AutomationTool/BuildGraph/Tasks/SetVersionTask.cs b/Engine/Source/Programs/AutomationTool/BuildGraph/Tasks/SetVersionTask.cs
--- a/Engine/Source/Programs/AutomationTool/BuildGraph/Tasks/SetVersionTask.cs
+++ b/Engine/Source/Programs/AutomationTool/BuildGraph/Tasks/SetVersionTask.cs
@@ -52,12 +52,6 @@ namespace AutomationTool.Tasks
 		[TaskParameter(Optional = true)]
 		public bool Promoted = true;
 
-		/// <summary>
-		/// If set, don't update Version.h; just update data files (eg. Build.version, Metadata.cs). The editor can read the build changelist at runtime from build.version.
-		/// </summary>
-		[TaskParameter(Optional = true)]
-		public bool SkipHeader;
-
 		/// <summary>
 		/// If set, don't actually write to the files - just return the version files that would be updated. Useful for local builds.
 		/// </summary>
@@ -100,7 +94,7 @@ namespace AutomationTool.Tasks
 		public override void Execute(JobContext Job, HashSet<FileReference> BuildProducts, Dictionary<string, HashSet<FileReference>> TagNameToFileSet)
 		{
 			// Update the version files
-			List<FileReference> VersionFiles = UE4Build.StaticUpdateVersionFiles(Parameters.Change, Parameters.CompatibleChange, Parameters.Branch, Parameters.Build, Parameters.Licensee, Parameters.Promoted, !Parameters.SkipWrite, Parameters.SkipHeader);
+			List<FileReference> VersionFiles = UE4Build.StaticUpdateVersionFiles(Parameters.Change, Parameters.CompatibleChange, Parameters.Branch, Parameters.Build, Parameters.Licensee, Parameters.Promoted, !Parameters.SkipWrite);
 
 			// Apply the optional tag to them
 			foreach(string TagName in FindTagNamesFromList(Parameters.Tag))
diff --git a/Engine/Source/Programs/AutomationTool/Scripts/UpdateLocalVersion.Automation.cs b/Engine/Source/Programs/AutomationTool/Scripts/UpdateLocalVersion.Automation.cs
--- a/Engine/Source/Programs/AutomationTool/Scripts/UpdateLocalVersion.Automation.cs
+++ b/Engine/Source/Programs/AutomationTool/Scripts/UpdateLocalVersion.Automation.cs
@@ -22,7 +22,6 @@ public class UpdateLocalVersion : BuildCommand
 		int? ChangelistOverride = ParseParamNullableInt("cl");
 		int? CompatibleChangelistOverride = ParseParamNullableInt("compatiblecl");
 		string Build = ParseParamValue("Build", null);
-		bool bSkipHeader = ParseParam("SkipHeader");
-		UE4Build.UpdateVersionFiles(ChangelistNumberOverride: ChangelistOverride, CompatibleChangelistNumberOverride: CompatibleChangelistOverride, Build: Build, bSkipHeader: bSkipHeader);
+		UE4Build.UpdateVersionFiles(ChangelistNumberOverride: ChangelistOverride, CompatibleChangelistNumberOverride: CompatibleChangelistOverride, Build: Build);
 	}
 }
