Allow per-chunk settings for whether to enable compression in UnrealPak.

diff --git a/Engine/Source/Programs/AutomationTool/Scripts/CopyBuildToStagingDirectory.Automation.cs b/Engine/Source/Programs/AutomationTool/Scripts/CopyBuildToStagingDirectory.Automation.cs
--- a/Engine/Source/Programs/AutomationTool/Scripts/CopyBuildToStagingDirectory.Automation.cs
+++ b/Engine/Source/Programs/AutomationTool/Scripts/CopyBuildToStagingDirectory.Automation.cs
@@ -1211,8 +1211,8 @@ public partial class Project : CommandUtils
 		PakCryptoSettings.Save(CryptoKeysCacheFilename);
 
 		List<CreatePakParams> PakInputs = new List<CreatePakParams>();
-		PakInputs.Add(new CreatePakParams(SC.ShortProjectName, UnrealPakResponseFile));
-		CreatePaks(Params, SC, PakInputs, PakCryptoSettings, CryptoKeysCacheFilename, Params.Compressed);
+		PakInputs.Add(new CreatePakParams(SC.ShortProjectName, UnrealPakResponseFile, Params.Compressed));
+		CreatePaks(Params, SC, PakInputs, PakCryptoSettings, CryptoKeysCacheFilename);
 	}
 
 	/// <summary>
@@ -1338,15 +1338,22 @@ public partial class Project : CommandUtils
 		/// </summary>
 		public Dictionary<string, string> UnrealPakResponseFile;
 
+		/// <summary>
+		/// Whether to enable compression
+		/// </summary>
+		public bool bCompressed;
+
 		/// <summary>
 		/// Constructor
 		/// </summary>
 		/// <param name="PakName">Path to the base output file for this pak file</param>
 		/// <param name="UnrealPakResponseFile">Map of files within the pak file to their source file on disk</param>
-		public CreatePakParams(string PakName, Dictionary<string, string> UnrealPakResponseFile)
+		/// <param name="bCompressed">Whether to enable compression</param>
+		public CreatePakParams(string PakName, Dictionary<string, string> UnrealPakResponseFile, bool bCompressed)
 		{
 			this.PakName = PakName;
 			this.UnrealPakResponseFile = UnrealPakResponseFile;
+			this.bCompressed = bCompressed;
 		}
 	}
 
@@ -1357,7 +1364,7 @@ public partial class Project : CommandUtils
 	/// <param name="SC"></param>
 	/// <param name="UnrealPakResponseFile"></param>
 	/// <param name="PakName"></param>
-	private static void CreatePaks(ProjectParams Params, DeploymentContext SC, List<CreatePakParams> PakParamsList, EncryptionAndSigning.CryptoSettings CryptoSettings, FileReference CryptoKeysCacheFilename, bool bCompressed)
+	private static void CreatePaks(ProjectParams Params, DeploymentContext SC, List<CreatePakParams> PakParamsList, EncryptionAndSigning.CryptoSettings CryptoSettings, FileReference CryptoKeysCacheFilename)
 	{
 		bool bShouldGeneratePatch = Params.IsGeneratingPatch && SC.StageTargetPlatform.GetPlatformPatchesWithDiffPak(Params, SC);
 
@@ -1420,104 +1427,104 @@ public partial class Project : CommandUtils
 			}
 			OutputFilename = OutputFilename + PostFix;
 
-		StagedFileReference OutputRelativeLocation;
-		if (Params.HasDLCName)
-		{
-			OutputRelativeLocation = StagedFileReference.Combine(SC.RelativeProjectRootForStage, Params.DLCFile.Directory.MakeRelativeTo(SC.ProjectRoot), "Content", "Paks", SC.FinalCookPlatform, Params.DLCFile.GetFileNameWithoutExtension() + ".pak");
-		}
-		else
-		{
-			OutputRelativeLocation = StagedFileReference.Combine(SC.RelativeProjectRootForStage, "Content", "Paks", OutputFilename + OutputFilenameExtension);
-		}
-		if (SC.StageTargetPlatform.DeployLowerCaseFilenames())
-		{
-			OutputRelativeLocation = OutputRelativeLocation.ToLowerInvariant();
-		}
-		OutputRelativeLocation = SC.StageTargetPlatform.Remap(OutputRelativeLocation);
-
-		FileReference OutputLocation = FileReference.Combine(SC.RuntimeRootDir, OutputRelativeLocation.Name);
-		// Add input file to control order of file within the pak
-		DirectoryReference PakOrderFileLocationBase = DirectoryReference.Combine(SC.ProjectRoot, "Build", SC.FinalCookPlatform, "FileOpenOrder");
-
-		FileReference PakOrderFileLocation = null;
-
-		string[] OrderFileNames = new string[] { "GameOpenOrder.log", "CookerOpenOrder.log", "EditorOpenOrder.log" };
-		foreach (string OrderFileName in OrderFileNames)
-		{
-			PakOrderFileLocation = FileReference.Combine(PakOrderFileLocationBase, OrderFileName);
-
-			if (FileExists_NoExceptions(PakOrderFileLocation.FullName))
+			StagedFileReference OutputRelativeLocation;
+			if (Params.HasDLCName)
 			{
-				break;
+				OutputRelativeLocation = StagedFileReference.Combine(SC.RelativeProjectRootForStage, Params.DLCFile.Directory.MakeRelativeTo(SC.ProjectRoot), "Content", "Paks", SC.FinalCookPlatform, Params.DLCFile.GetFileNameWithoutExtension() + ".pak");
 			}
-		}
+			else
+			{
+				OutputRelativeLocation = StagedFileReference.Combine(SC.RelativeProjectRootForStage, "Content", "Paks", OutputFilename + OutputFilenameExtension);
+			}
+			if (SC.StageTargetPlatform.DeployLowerCaseFilenames())
+			{
+				OutputRelativeLocation = OutputRelativeLocation.ToLowerInvariant();
+			}
+			OutputRelativeLocation = SC.StageTargetPlatform.Remap(OutputRelativeLocation);
 
-		bool bCopiedExistingPak = false;
+			FileReference OutputLocation = FileReference.Combine(SC.RuntimeRootDir, OutputRelativeLocation.Name);
+			// Add input file to control order of file within the pak
+			DirectoryReference PakOrderFileLocationBase = DirectoryReference.Combine(SC.ProjectRoot, "Build", SC.FinalCookPlatform, "FileOpenOrder");
 
-		if (SC.StageTargetPlatform != SC.CookSourcePlatform)
-		{
-			// Check to see if we have an existing pak file we can use
+			FileReference PakOrderFileLocation = null;
 
-				StagedFileReference SourceOutputRelativeLocation = StagedFileReference.Combine(SC.RelativeProjectRootForStage, "Content/Paks/", PakParams.PakName + "-" + SC.CookPlatform + PostFix + ".pak");
-			if (SC.CookSourcePlatform.DeployLowerCaseFilenames())
+			string[] OrderFileNames = new string[] { "GameOpenOrder.log", "CookerOpenOrder.log", "EditorOpenOrder.log" };
+			foreach (string OrderFileName in OrderFileNames)
 			{
-				SourceOutputRelativeLocation = SourceOutputRelativeLocation.ToLowerInvariant();
+				PakOrderFileLocation = FileReference.Combine(PakOrderFileLocationBase, OrderFileName);
+
+				if (FileExists_NoExceptions(PakOrderFileLocation.FullName))
+				{
+					break;
+				}
 			}
-			SourceOutputRelativeLocation = SC.CookSourcePlatform.Remap(SourceOutputRelativeLocation);
 
-			FileReference SourceOutputLocation = FileReference.Combine(SC.CookSourceRuntimeRootDir, SourceOutputRelativeLocation.Name);
-			if (FileExists_NoExceptions(SourceOutputLocation.FullName))
+			bool bCopiedExistingPak = false;
+
+			if (SC.StageTargetPlatform != SC.CookSourcePlatform)
 			{
-				InternalUtils.SafeCreateDirectory(Path.GetDirectoryName(OutputLocation.FullName), true);
+				// Check to see if we have an existing pak file we can use
 
-				if (InternalUtils.SafeCopyFile(SourceOutputLocation.FullName, OutputLocation.FullName))
+					StagedFileReference SourceOutputRelativeLocation = StagedFileReference.Combine(SC.RelativeProjectRootForStage, "Content/Paks/", PakParams.PakName + "-" + SC.CookPlatform + PostFix + ".pak");
+				if (SC.CookSourcePlatform.DeployLowerCaseFilenames())
 				{
-					Log("Copying source pak from {0} to {1} instead of creating new pak", SourceOutputLocation, OutputLocation);
-					bCopiedExistingPak = true;
+					SourceOutputRelativeLocation = SourceOutputRelativeLocation.ToLowerInvariant();
+				}
+				SourceOutputRelativeLocation = SC.CookSourcePlatform.Remap(SourceOutputRelativeLocation);
 
-					FileReference InSigFile = SourceOutputLocation.ChangeExtension(".sig");
-					if (FileReference.Exists(InSigFile))
-					{
-						FileReference OutSigFile = OutputLocation.ChangeExtension(".sig");
+				FileReference SourceOutputLocation = FileReference.Combine(SC.CookSourceRuntimeRootDir, SourceOutputRelativeLocation.Name);
+				if (FileExists_NoExceptions(SourceOutputLocation.FullName))
+				{
+					InternalUtils.SafeCreateDirectory(Path.GetDirectoryName(OutputLocation.FullName), true);
 
-						Log("Copying pak sig from {0} to {1}", InSigFile, OutSigFile);
+					if (InternalUtils.SafeCopyFile(SourceOutputLocation.FullName, OutputLocation.FullName))
+					{
+						Log("Copying source pak from {0} to {1} instead of creating new pak", SourceOutputLocation, OutputLocation);
+						bCopiedExistingPak = true;
 
-						if (!InternalUtils.SafeCopyFile(InSigFile.FullName, OutSigFile.FullName))
+						FileReference InSigFile = SourceOutputLocation.ChangeExtension(".sig");
+						if (FileReference.Exists(InSigFile))
 						{
-							Log("Failed to copy pak sig {0} to {1}, creating new pak", InSigFile, InSigFile);
-							bCopiedExistingPak = false;
+							FileReference OutSigFile = OutputLocation.ChangeExtension(".sig");
+
+							Log("Copying pak sig from {0} to {1}", InSigFile, OutSigFile);
+
+							if (!InternalUtils.SafeCopyFile(InSigFile.FullName, OutSigFile.FullName))
+							{
+								Log("Failed to copy pak sig {0} to {1}, creating new pak", InSigFile, InSigFile);
+								bCopiedExistingPak = false;
+							}
 						}
-					}
 
+					}
+				}
+				if (!bCopiedExistingPak)
+				{
+					Log("Failed to copy source pak from {0} to {1}, creating new pak", SourceOutputLocation, OutputLocation);
 				}
 			}
-			if (!bCopiedExistingPak)
+
+			string PatchSourceContentPath = null;
+			if (bShouldGeneratePatch)
 			{
-				Log("Failed to copy source pak from {0} to {1}, creating new pak", SourceOutputLocation, OutputLocation);
+				// don't include the post fix in this filename because we are looking for the source pak path
+					string PakFilename = PakParams.PakName + "-" + SC.FinalCookPlatform + "*.pak";
+				PatchSourceContentPath = SC.StageTargetPlatform.GetReleasePakFilePath(SC, Params, PakFilename);
 			}
-		}
-
-		string PatchSourceContentPath = null;
-		if (bShouldGeneratePatch)
-		{
-			// don't include the post fix in this filename because we are looking for the source pak path
-				string PakFilename = PakParams.PakName + "-" + SC.FinalCookPlatform + "*.pak";
-			PatchSourceContentPath = SC.StageTargetPlatform.GetReleasePakFilePath(SC, Params, PakFilename);
-		}
 		
-		if (!bCopiedExistingPak)
-		{
-			if (FileReference.Exists(OutputLocation))
+			if (!bCopiedExistingPak)
 			{
-				string UnrealPakResponseFileName = CombinePaths(CmdEnv.LogFolder, "PakList_" + OutputLocation.GetFileNameWithoutExtension() + ".txt");
-				if (File.Exists(UnrealPakResponseFileName) && FileReference.GetLastWriteTimeUtc(OutputLocation) > File.GetLastWriteTimeUtc(UnrealPakResponseFileName))
+				if (FileReference.Exists(OutputLocation))
 				{
-					bCopiedExistingPak = true;
+					string UnrealPakResponseFileName = CombinePaths(CmdEnv.LogFolder, "PakList_" + OutputLocation.GetFileNameWithoutExtension() + ".txt");
+					if (File.Exists(UnrealPakResponseFileName) && FileReference.GetLastWriteTimeUtc(OutputLocation) > File.GetLastWriteTimeUtc(UnrealPakResponseFileName))
+					{
+						bCopiedExistingPak = true;
+					}
 				}
-			}
-			if (!bCopiedExistingPak)
-			{
-					Commands.Add(GetUnrealPakArguments(PakParams.UnrealPakResponseFile, OutputLocation, PakOrderFileLocation, SC.StageTargetPlatform.GetPlatformPakCommandLine(Params, SC) + " " + Params.AdditionalPakOptions, bCompressed, CryptoSettings, CryptoKeysCacheFilename, PatchSourceContentPath));
+				if (!bCopiedExistingPak)
+				{
+					Commands.Add(GetUnrealPakArguments(PakParams.UnrealPakResponseFile, OutputLocation, PakOrderFileLocation, SC.StageTargetPlatform.GetPlatformPakCommandLine(Params, SC) + " " + Params.AdditionalPakOptions, PakParams.bCompressed, CryptoSettings, CryptoKeysCacheFilename, PatchSourceContentPath));
 				}
 			}
 
@@ -1954,23 +1961,18 @@ public partial class Project : CommandUtils
 		Log("Creating Pak files utilizing {0} cores", Environment.ProcessorCount);
 		Options.MaxDegreeOfParallelism = Environment.ProcessorCount;
 
-		bool bCompression = false;
 		List<CreatePakParams> PakInputs = new List<CreatePakParams>();
-		//System.Threading.Tasks.Parallel.ForEach(ChunkDefinitions, Options, (Chunk) =>
-		foreach (var Chunk in ChunkDefinitions)
-		{			
-			bCompression |= Params.Compressed;
-			bCompression |= Chunk.bCompressed;
-
-			var ChunkName = Path.GetFileNameWithoutExtension(Chunk.ChunkName);
+		foreach (ChunkDefinition Chunk in ChunkDefinitions)
+		{
+			string ChunkName = Path.GetFileNameWithoutExtension(Chunk.ChunkName);
 			if (Chunk.ResponseFile.Count > 0)
 			{
-				PakInputs.Add(new CreatePakParams(ChunkName, Chunk.ResponseFile));
+				PakInputs.Add(new CreatePakParams(ChunkName, Chunk.ResponseFile, Params.Compressed || Chunk.bCompressed));
 				//CreatePak(Params, SC, Chunk.ResponseFile, ChunkName, PakCryptoSettings, CryptoKeysCacheFilename, bCompression);
 			}
 		}
 
-		CreatePaks(Params, SC, PakInputs, PakCryptoSettings, CryptoKeysCacheFilename, bCompression);
+		CreatePaks(Params, SC, PakInputs, PakCryptoSettings, CryptoKeysCacheFilename);
 
 		if (Params.CreateChunkInstall)
 		{
