All staging client executables for blueprint projects.

diff --git a/Engine/Source/Programs/AutomationTool/AutomationUtils/ProjectParams.cs b/Engine/Source/Programs/AutomationTool/AutomationUtils/ProjectParams.cs
--- a/Engine/Source/Programs/AutomationTool/AutomationUtils/ProjectParams.cs
+++ b/Engine/Source/Programs/AutomationTool/AutomationUtils/ProjectParams.cs
@@ -1956,14 +1956,28 @@ namespace AutomationTool
 
 			if (!bIsCodeBasedProject)
 			{
-				GameTarget = "UE4Game";
+				if(Client)
+				{
+					GameTarget = "UE4Client";
+				}
+				else
+				{
+					GameTarget = "UE4Game";
+				}
 				EditorTarget = "UE4Editor";
 				ServerTarget = "UE4Server";
 			}
 			else if (!CommandUtils.CmdEnv.HasCapabilityToCompile)
 			{
 				var ShortName = ProjectUtils.GetShortProjectName(RawProjectPath);
-				GameTarget = ShortName;
+				if(Client)
+				{
+					GameTarget = ShortName + "Client";
+				}
+				else
+				{
+					GameTarget = ShortName;
+				}
 				EditorTarget = ShortName + "Editor";
 				ServerTarget = ShortName + "Server";
 			}
@@ -2025,7 +2039,14 @@ namespace AutomationTool
 			else if (!this.Build)
 			{
 				var ShortName = ProjectUtils.GetShortProjectName(RawProjectPath);
-				GameTarget = ShortName;
+				if(Client)
+				{
+					GameTarget = ShortName + "Client";
+				}
+				else
+				{
+					GameTarget = ShortName;
+				}
 				EditorTarget = ShortName + "Editor";
 				ServerTarget = ShortName + "Server";
 			}
