UAT: Remote code to build UnrealPak when packaging; we use the editor now.

#rnx

diff --git a/Engine/Source/Programs/AutomationTool/Scripts/BuildProjectCommand.Automation.cs b/Engine/Source/Programs/AutomationTool/Scripts/BuildProjectCommand.Automation.cs
--- a/Engine/Source/Programs/AutomationTool/Scripts/BuildProjectCommand.Automation.cs
+++ b/Engine/Source/Programs/AutomationTool/Scripts/BuildProjectCommand.Automation.cs
@@ -88,10 +88,6 @@ public partial class Project : CommandUtils
 				{
 					Agenda.AddTargets(new string[] { "ShaderCompileWorker" }, EditorPlatform, EditorConfiguration);
 				}
-				if (Params.Pak && Params.EditorTargets.Contains("UnrealPak") == false)
-				{
-					Agenda.AddTargets(new string[] { "UnrealPak" }, EditorPlatform, EditorConfiguration);
-				}
 				if (Params.FileServer && Params.EditorTargets.Contains("UnrealFileServer") == false)
 				{
 					Agenda.AddTargets(new string[] { "UnrealFileServer" }, EditorPlatform, EditorConfiguration);
