UAT: Fix log output w/ correct count of non-code projects.

#rnx

diff --git a/Engine/Source/Programs/AutomationTool/AutomationUtils/ProjectUtils.cs b/Engine/Source/Programs/AutomationTool/AutomationUtils/ProjectUtils.cs
--- a/Engine/Source/Programs/AutomationTool/AutomationUtils/ProjectUtils.cs
+++ b/Engine/Source/Programs/AutomationTool/AutomationUtils/ProjectUtils.cs
@@ -761,7 +761,7 @@ namespace AutomationTool
 				{
 					Proj.Dump(InHostPlatforms);
 				}
-				CommandUtils.LogVerbose("  {0} Non-Code projects:", CodeProjects.Count);
+				CommandUtils.LogVerbose("  {0} Non-Code projects:", NonCodeProjects.Count);
 				foreach (var Proj in NonCodeProjects)
 				{
 					Proj.Dump(InHostPlatforms);
