UAT: fix -skipbuildclient param default

It was defaulting to skipbuildeditor's value, likely a copy-paste error.

#rnx

diff --git a/Engine/Source/Programs/AutomationTool/AutomationUtils/ProjectParams.cs b/Engine/Source/Programs/AutomationTool/AutomationUtils/ProjectParams.cs
--- a/Engine/Source/Programs/AutomationTool/AutomationUtils/ProjectParams.cs
+++ b/Engine/Source/Programs/AutomationTool/AutomationUtils/ProjectParams.cs
@@ -548,7 +548,7 @@ namespace AutomationTool
             this.ServerTargetPlatforms = SetupTargetPlatforms(ref this.ServerDependentPlatformMap, Command, ServerTargetPlatforms, this.ClientTargetPlatforms, false, "ServerTargetPlatform", "ServerPlatform");
 
 			this.Build = GetParamValueIfNotSpecified(Command, Build, this.Build, "build");
-			this.SkipBuildClient = GetParamValueIfNotSpecified(Command, SkipBuildClient, this.SkipBuildEditor, "skipbuildclient");
+			this.SkipBuildClient = GetParamValueIfNotSpecified(Command, SkipBuildClient, this.SkipBuildClient, "skipbuildclient");
 			this.SkipBuildEditor = GetParamValueIfNotSpecified(Command, SkipBuildEditor, this.SkipBuildEditor, "skipbuildeditor");
 			this.Run = GetParamValueIfNotSpecified(Command, Run, this.Run, "run");
 			this.Cook = GetParamValueIfNotSpecified(Command, Cook, this.Cook, "cook");
