BuildGraph: Allow setting the target to build when using the <CsCompile> task.

diff --git a/Engine/Source/Programs/AutomationTool/BuildGraph/Tasks/CsCompileTask.cs b/Engine/Source/Programs/AutomationTool/BuildGraph/Tasks/CsCompileTask.cs
--- a/Engine/Source/Programs/AutomationTool/BuildGraph/Tasks/CsCompileTask.cs
+++ b/Engine/Source/Programs/AutomationTool/BuildGraph/Tasks/CsCompileTask.cs
@@ -34,6 +34,12 @@ namespace AutomationTool.Tasks
 		[TaskParameter(Optional = true)]
 		public string Platform;
 
+		/// <summary>
+		/// The target to build
+		/// </summary>
+		[TaskParameter(Optional = true)]
+		public string Target;
+
 		/// <summary>
 		/// Additional options to pass to the compiler
 		/// </summary>
@@ -125,6 +131,10 @@ namespace AutomationTool.Tasks
 				{
 					Arguments.Add(Parameters.Arguments);
 				}
+				if(!String.IsNullOrEmpty(Parameters.Target))
+				{
+					Arguments.Add(String.Format("/target:{0}", CommandUtils.MakePathSafeToUseWithCommandLine(Parameters.Target)));
+				}
 				Arguments.Add("/verbosity:minimal");
 				Arguments.Add("/nologo");
 				foreach(FileReference ProjectFile in ProjectFiles)
