UAT: Remove StreamUtilis utility class; there is a simpler way to implement the one place it's used.

#rnx

diff --git a/Engine/Source/Programs/AutomationTool/AutomationUtils/DeviceReservation.cs b/Engine/Source/Programs/AutomationTool/AutomationUtils/DeviceReservation.cs
--- a/Engine/Source/Programs/AutomationTool/AutomationUtils/DeviceReservation.cs
+++ b/Engine/Source/Programs/AutomationTool/AutomationUtils/DeviceReservation.cs
@@ -131,7 +131,9 @@ namespace AutomationTool.DeviceReservation
 			using (var Response = (HttpWebResponse)Request.GetResponse())
 			using (var ResponseStream = Response.GetResponseStream())
 			{
-				string JsonString = Encoding.UTF8.GetString(ResponseStream.ReadIntoMemoryStream().ToArray());
+				MemoryStream MemoryStream = new MemoryStream();
+				ResponseStream.CopyTo(MemoryStream);
+				string JsonString = Encoding.UTF8.GetString(MemoryStream.ToArray());
 				return fastJSON.JSON.Instance.ToObject<T>(JsonString);
 			}
 		}
diff --git a/Engine/Source/Programs/UnrealBuildTool/System/StreamUtils.cs b/Engine/Source/Programs/UnrealBuildTool/System/StreamUtils.cs
deleted file mode 100644
diff --git a/Engine/Source/Programs/UnrealBuildTool/UnrealBuildTool.csproj b/Engine/Source/Programs/UnrealBuildTool/UnrealBuildTool.csproj
--- a/Engine/Source/Programs/UnrealBuildTool/UnrealBuildTool.csproj
+++ b/Engine/Source/Programs/UnrealBuildTool/UnrealBuildTool.csproj
@@ -204,7 +204,6 @@
     <Compile Include="System\ScopedTimer.cs" />
     <Compile Include="System\SourceFileWorkingSet.cs" />
     <Compile Include="System\StartupTraceListener.cs" />
-    <Compile Include="System\StreamUtils.cs" />
     <Compile Include="System\SupportedConfigurationsAttribute.cs" />
     <Compile Include="System\SupportedPlatformsAttribute.cs" />
     <Compile Include="System\TargetReceipt.cs" />
