UAT: Rename CommandUtils.Log to CommandUtils.LogInformation, to avoid conflicts with the underlying Tools.DotNETCommon.Log class.

diff --git a/Engine/Source/Programs/AutomationTool/Android/AndroidPlatform.Automation.cs b/Engine/Source/Programs/AutomationTool/Android/AndroidPlatform.Automation.cs
--- a/Engine/Source/Programs/AutomationTool/Android/AndroidPlatform.Automation.cs
+++ b/Engine/Source/Programs/AutomationTool/Android/AndroidPlatform.Automation.cs
@@ -252,7 +252,7 @@ public class AndroidPlatform : Platform
 				if (PluginExtras.FirstOrDefault(x => x == PluginPath) == null)
 				{
 					PluginExtras.Add(PluginPath);
-					Log("AndroidPlugin: {0}", PluginPath);
+					LogInformation("AndroidPlugin: {0}", PluginPath);
 				}
 			}
 		}
@@ -292,7 +292,7 @@ public class AndroidPlatform : Platform
 		bool bPackageDataInsideApk = Deploy.GetPackageDataInsideApk();
 
 		string BaseApkName = GetFinalApkName(Params, SC.StageExecutables[0], true, "", "");
-		Log("BaseApkName = {0}", BaseApkName);
+		LogInformation("BaseApkName = {0}", BaseApkName);
 
 		// Create main OBB with entire contents of staging dir. This
 		// includes any PAK files, movie files, etc.
@@ -306,7 +306,7 @@ public class AndroidPlatform : Platform
 		}
 
 		// Now create the OBB as a ZIP archive.
-		Log("Creating {0} from {1}", LocalObbName, SC.StageDirectory);
+		LogInformation("Creating {0} from {1}", LocalObbName, SC.StageDirectory);
 		using (ZipFile ObbFile = new ZipFile(LocalObbName))
 		{
 			ObbFile.CompressionMethod = CompressionMethod.None;
@@ -320,7 +320,7 @@ public class AndroidPlatform : Platform
 					if (e.EventType == ZipProgressEventType.Adding_AfterAddEntry)
 					{
 						ObbFileCount += 1;
-						Log("[{0}/{1}] Adding {2} to OBB",
+						LogInformation("[{0}/{1}] Adding {2} to OBB",
 							ObbFileCount, e.EntriesTotal,
 							e.CurrentEntry.FileName);
 					}
@@ -355,7 +355,7 @@ public class AndroidPlatform : Platform
 			}
 			catch (Exception)
 			{
-				Log("Failed to build OBB: " + LocalObbName);
+				LogInformation("Failed to build OBB: " + LocalObbName);
 				throw new AutomationException(ExitCode.Error_AndroidOBBError, "Stage Failed. Could not build OBB {0}. The file may be too big to fit in an OBB (2 GiB limit)", LocalObbName);
 			}
 		}
@@ -365,7 +365,7 @@ public class AndroidPlatform : Platform
 		Int64 ObbFileLength = OBBFileInfo.Length;
 		if (ObbFileLength > MaxOBBSizeAllowed)
 		{
-			Log("OBB exceeds 2 GiB limit: " + ObbFileLength + " bytes");
+			LogInformation("OBB exceeds 2 GiB limit: " + ObbFileLength + " bytes");
 			throw new AutomationException(ExitCode.Error_AndroidOBBError, "Stage Failed. OBB {0} exceeds 2 GiB limit)", LocalObbName);
 		}
 
@@ -377,7 +377,7 @@ public class AndroidPlatform : Platform
 		Ini.GetInt32("/Script/AndroidRuntimeSettings.AndroidRuntimeSettings", "MinSDKVersion", out MinSDKVersion);
 		int TargetSDKVersion = MinSDKVersion;
 		Ini.GetInt32("/Script/AndroidRuntimeSettings.AndroidRuntimeSettings", "TargetSDKVersion", out TargetSDKVersion);
-		Log("Target SDK Version" + TargetSDKVersion);
+		LogInformation("Target SDK Version" + TargetSDKVersion);
 
 		foreach (string Architecture in Architectures)
 		{
@@ -390,7 +390,7 @@ public class AndroidPlatform : Platform
                     UE4SOName = UE4SOName.Replace(".apk", ".so");
                     if (FileExists_NoExceptions(UE4SOName) == false)
 					{
-						Log("Failed to find game .so " + UE4SOName);
+						LogInformation("Failed to find game .so " + UE4SOName);
                         throw new AutomationException(ExitCode.Error_MissingExecutable, "Stage Failed. Could not find .so {0}. You may need to build the UE4 project with your target configuration and platform.", UE4SOName);
 					}
 				}
@@ -503,7 +503,7 @@ public class AndroidPlatform : Platform
 			// Note that $STORAGE/Android/obb will be the folder that contains the obb if you download the app from playstore.
 			string OBBInstallCommand = bNoObbInstall ? "shell 'rm -r $EXTERNAL_STORAGE/" + DeviceObbName + "'" : "push " + Path.GetFileName(ObbName) + " $STORAGE/" + (bIsDistribution ? "Download/" : "") + DeviceObbName;
 
-            Log("Writing shell script for install with {0}", bPackageDataInsideApk ? "data in APK" : "separate obb");
+            LogInformation("Writing shell script for install with {0}", bPackageDataInsideApk ? "data in APK" : "separate obb");
             BatchLines = new string[] {
 						"#!/bin/sh",
 						"cd \"`dirname \"$0\"`\"",
@@ -553,7 +553,7 @@ public class AndroidPlatform : Platform
         {
 			string OBBInstallCommand = bNoObbInstall ? "shell rm -r %STORAGE%/" + DeviceObbName : "push " + Path.GetFileName(ObbName) + " %STORAGE%/" + (bIsDistribution ? "Download/" : "") + DeviceObbName;
 
-			Log("Writing bat for install with {0}", bPackageDataInsideApk ? "data in APK" : "separate OBB");
+			LogInformation("Writing bat for install with {0}", bPackageDataInsideApk ? "data in APK" : "separate OBB");
             BatchLines = new string[] {
 						"setlocal",
                         "set ANDROIDHOME=%ANDROID_HOME%",
@@ -606,7 +606,7 @@ public class AndroidPlatform : Platform
 
 		if (!bIsPC)
 		{
-			Log("Writing shell script for uninstall with {0}", bPackageDataInsideApk ? "data in APK" : "separate obb");
+			LogInformation("Writing shell script for uninstall with {0}", bPackageDataInsideApk ? "data in APK" : "separate obb");
 			BatchLines = new string[] {
 						"#!/bin/sh",
 						"cd \"`dirname \"$0\"`\"",
@@ -630,7 +630,7 @@ public class AndroidPlatform : Platform
 		}
 		else
 		{
-			Log("Writing bat for uninstall with {0}", bPackageDataInsideApk ? "data in APK" : "separate OBB");
+			LogInformation("Writing bat for uninstall with {0}", bPackageDataInsideApk ? "data in APK" : "separate OBB");
 			BatchLines = new string[] {
 						"setlocal",
 						"set ANDROIDHOME=%ANDROID_HOME%",
@@ -661,7 +661,7 @@ public class AndroidPlatform : Platform
 
 		if (!bIsPC)
 		{
-			Log("Writing shell script for symbolize with {0}", "data in APK" );
+			LogInformation("Writing shell script for symbolize with {0}", "data in APK" );
 			BatchLines = new string[] {
 				"#!/bin/sh",
 				"if [ $? -ne 0]; then",
@@ -677,7 +677,7 @@ public class AndroidPlatform : Platform
 		}
 		else
 		{
-			Log("Writing bat for symbolize");
+			LogInformation("Writing bat for symbolize");
 			BatchLines = new string[] {
 						"@echo off",
 						"IF %1.==. GOTO NoArgs",
@@ -1495,7 +1495,7 @@ public class AndroidPlatform : Platform
 		CachedAaptPath = BestToolPath;
 		LastAndroidHomePath = HomePath;
 
-		Log("Using this aapt: {0}", CachedAaptPath);
+		LogInformation("Using this aapt: {0}", CachedAaptPath);
 
 		return CachedAaptPath;
 	}
@@ -1664,7 +1664,7 @@ public class AndroidPlatform : Platform
 					TimeSpan DeltaRunTime = DateTime.Now - StartTime;
 					if ((DeltaRunTime.TotalSeconds > TimeOutSeconds) && (TimeOutSeconds != 0))
 					{
-						Log("Device: " + DeviceName + " timed out while waiting for run to finish");
+						LogInformation("Device: " + DeviceName + " timed out while waiting for run to finish");
 						FinishedRunning = true;
 					}
 				}
diff --git a/Engine/Source/Programs/AutomationTool/AutomationUtils/Automation.cs b/Engine/Source/Programs/AutomationTool/AutomationUtils/Automation.cs
--- a/Engine/Source/Programs/AutomationTool/AutomationUtils/Automation.cs
+++ b/Engine/Source/Programs/AutomationTool/AutomationUtils/Automation.cs
@@ -565,7 +565,7 @@ AutomationTool.exe [-verbose] [-compileonly] [-p4] Command0 [-Arg0 -Arg1 -Arg2 .
 					{
 						return Result;
 					}
-					CommandUtils.Log("BUILD SUCCESSFUL");
+					CommandUtils.LogInformation("BUILD SUCCESSFUL");
 				}
 				finally
 				{
@@ -629,7 +629,7 @@ AutomationTool.exe [-verbose] [-compileonly] [-p4] Command0 [-Arg0 -Arg1 -Arg2 .
 			{
 				Message += String.Format("  {0}{1}", AvailableCommand.Key, Environment.NewLine);
 			}
-			CommandUtils.Log(Message);
+			CommandUtils.LogInformation(Message);
 		}
 
 		#endregion
diff --git a/Engine/Source/Programs/AutomationTool/AutomationUtils/CommandUtils.cs b/Engine/Source/Programs/AutomationTool/AutomationUtils/CommandUtils.cs
--- a/Engine/Source/Programs/AutomationTool/AutomationUtils/CommandUtils.cs
+++ b/Engine/Source/Programs/AutomationTool/AutomationUtils/CommandUtils.cs
@@ -114,6 +114,7 @@ namespace AutomationTool
 		/// <param name="Format">Format string</param>
 		/// <param name="Args">Parameters</param>
 		[MethodImplAttribute(MethodImplOptions.NoInlining)]
+		[Obsolete("CommandUtils.LogInformation() has been deprecated. Use CommandUtils.LogInformation() instead.", false)]
 		public static void Log(string Format, params object[] Args)
 		{
 			Tools.DotNETCommon.Log.WriteLine(1, Tools.DotNETCommon.LogEventType.Console, Format, Args);
@@ -124,11 +125,33 @@ namespace AutomationTool
 		/// </summary>
 		/// <param name="Message">Text</param>
 		[MethodImplAttribute(MethodImplOptions.NoInlining)]
+		[Obsolete("CommandUtils.LogInformation() has been deprecated. Use CommandUtils.LogInformation() instead.", false)]
 		public static void Log(string Message)
 		{
 			Tools.DotNETCommon.Log.WriteLine(1, Tools.DotNETCommon.LogEventType.Console, Message);
 		}
 
+		/// <summary>
+		/// Writes formatted text to log (with LogEventType.Console).
+		/// </summary>
+		/// <param name="Format">Format string</param>
+		/// <param name="Args">Parameters</param>
+		[MethodImplAttribute(MethodImplOptions.NoInlining)]
+		public static void LogInformation(string Format, params object[] Args)
+		{
+			Tools.DotNETCommon.Log.WriteLine(1, Tools.DotNETCommon.LogEventType.Console, Format, Args);
+		}
+
+		/// <summary>
+		/// Writes formatted text to log (with LogEventType.Console).
+		/// </summary>
+		/// <param name="Message">Text</param>
+		[MethodImplAttribute(MethodImplOptions.NoInlining)]
+		public static void LogInformation(string Message)
+		{
+			Tools.DotNETCommon.Log.WriteLine(1, Tools.DotNETCommon.LogEventType.Console, Message);
+		}
+
 		/// <summary>
 		/// Writes formatted text to log (with LogEventType.Error).
 		/// </summary>
@@ -276,7 +299,7 @@ namespace AutomationTool
 		{
 			if(bShowProgress)
 			{
-				Log("[@progress push {0}/{1} skipline]", Numerator, Denominator);
+				LogInformation("[@progress push {0}/{1} skipline]", Numerator, Denominator);
 			}
 		}
 
@@ -284,7 +307,7 @@ namespace AutomationTool
 		{
 			if(bShowProgress)
 			{
-				Log("[@progress pop skipline]");
+				LogInformation("[@progress pop skipline]");
 			}
 		}
 
@@ -292,7 +315,7 @@ namespace AutomationTool
 		{
 			if(bShowProgress)
 			{
-				Log("[@progress increment {0}/{1} skipline]", Numerator, Denominator);
+				LogInformation("[@progress increment {0}/{1} skipline]", Numerator, Denominator);
 			}
 		}
 
@@ -300,7 +323,7 @@ namespace AutomationTool
 		{
 			if(bShowProgress)
 			{
-				Log("[@progress '{0}' skipline]", String.Format(Format, Args));
+				LogInformation("[@progress '{0}' skipline]", String.Format(Format, Args));
 			}
 		}
 
@@ -308,7 +331,7 @@ namespace AutomationTool
 		{
 			if(bShowProgress)
 			{
-				Log("[@progress {0}/{1} '{2}' skipline]", Numerator, Denominator, String.Format(Format, Args));
+				LogInformation("[@progress {0}/{1} '{2}' skipline]", Numerator, Denominator, String.Format(Format, Args));
 			}
 		}
 
@@ -1580,7 +1603,7 @@ namespace AutomationTool
 		{
 			if(!bQuiet)
 			{
-				Log("Copying {0} file(s) using max {1} thread(s)", Source.Count, MaxThreads);
+				LogInformation("Copying {0} file(s) using max {1} thread(s)", Source.Count, MaxThreads);
 			}
 
             if (Source.Count != Dest.Count)
@@ -1621,7 +1644,7 @@ namespace AutomationTool
 		public static List<string> ThreadedCopyFiles(string SourceDir, string TargetDir, FileFilter Filter, bool bIgnoreSymlinks, int MaxThreads = 64)
 		{
 			// Filter all the relative paths
-			Log("Applying filter to {0}...", SourceDir);
+			LogInformation("Applying filter to {0}...", SourceDir);
 			DirectoryReference SourceDirRef = new DirectoryReference(SourceDir);
 			var RelativePaths = Filter.ApplyToDirectory(SourceDirRef, bIgnoreSymlinks).Select(x => x.MakeRelativeTo(SourceDirRef)).ToList();
 			return ThreadedCopyFiles(SourceDir, TargetDir, RelativePaths);
@@ -1977,13 +2000,13 @@ namespace AutomationTool
             try
             {
                 DirectoryInfo DirInfo = new DirectoryInfo(ParentDir);
-				Log("Looking for directories to delete in {0}", ParentDir);
+				LogInformation("Looking for directories to delete in {0}", ParentDir);
                 foreach (DirectoryInfo ThisDirInfo in DirInfo.EnumerateDirectories(SearchPattern))
                 {
 					double AgeDays = (DateTime.UtcNow - ThisDirInfo.CreationTimeUtc).TotalDays;
 					if (AgeDays > MaximumDaysToKeepTempStorage)
                     {
-                        Log("Deleting formal build directory {0}, because it is {1} days old (maximum {2}).", ThisDirInfo.FullName, (int)AgeDays, MaximumDaysToKeepTempStorage);
+                        LogInformation("Deleting formal build directory {0}, because it is {1} days old (maximum {2}).", ThisDirInfo.FullName, (int)AgeDays, MaximumDaysToKeepTempStorage);
                         DeleteDirectory_NoExceptions(true, ThisDirInfo.FullName);
                     }
                     else
@@ -2123,7 +2146,7 @@ namespace AutomationTool
 
 				foreach (string LogLine in Lines)
 				{
-					CommandUtils.Log(LogLine);
+					CommandUtils.LogInformation(LogLine);
 
 					// Split each line into two by whitespace
 					string[] SplitLine = LogLine.Split(new char[] { ' ', '\t' }, 2, StringSplitOptions.RemoveEmptyEntries);
@@ -2131,7 +2154,7 @@ namespace AutomationTool
 					{
 						// Second part of line should be a path
 						string FilePath = SplitLine[1].Trim();
-						CommandUtils.Log(FilePath);
+						CommandUtils.LogInformation(FilePath);
 						if (File.Exists(FilePath) && !OutputFileNames.Contains(FilePath) && FilePath != ZipFileName)
 						{
 							if (CommandUtils.IsProbablyAMacOrIOSExe(FilePath))
@@ -2790,7 +2813,7 @@ namespace AutomationTool
 		{
 			if (!Command.ParseParam("NoSign"))
 			{
-				CommandUtils.Log("Signing up to {0} files...", Files.Count());
+				CommandUtils.LogInformation("Signing up to {0} files...", Files.Count());
 				UnrealBuildTool.UnrealTargetPlatform TargetPlatform = UnrealBuildTool.BuildHostPlatform.Current.Platform;
 				if (TargetPlatform == UnrealBuildTool.UnrealTargetPlatform.Mac)
 				{
diff --git a/Engine/Source/Programs/AutomationTool/AutomationUtils/CommandletUtils.cs b/Engine/Source/Programs/AutomationTool/AutomationUtils/CommandletUtils.cs
--- a/Engine/Source/Programs/AutomationTool/AutomationUtils/CommandletUtils.cs
+++ b/Engine/Source/Programs/AutomationTool/AutomationUtils/CommandletUtils.cs
@@ -235,7 +235,7 @@ namespace AutomationTool
 		/// <param name="Parameters">Command line parameters (without -run=)</param>
 		public static void RunCommandlet(FileReference ProjectName, string UE4Exe, string Commandlet, string Parameters = null)
 		{
-			Log("Running UE4Editor {0} for project {1}", Commandlet, ProjectName);
+			LogInformation("Running UE4Editor {0} for project {1}", Commandlet, ProjectName);
 
             var CWD = Path.GetDirectoryName(UE4Exe);
 
@@ -252,7 +252,7 @@ namespace AutomationTool
 			DateTime StartTime = DateTime.UtcNow;
 
 			string LocalLogFile = LogUtils.GetUniqueLogName(CombinePaths(CmdEnv.EngineSavedFolder, Commandlet));
-			Log("Commandlet log file is {0}", LocalLogFile);
+			LogInformation("Commandlet log file is {0}", LocalLogFile);
 			string Args = String.Format(
 				"{0} -run={1} {2} -abslog={3} -stdout -CrashForUAT -unattended -NoLogTimes {5}{4}",
 				(ProjectName == null) ? "" : CommandUtils.MakePathSafeToUseWithCommandLine(ProjectName.FullName),
@@ -278,7 +278,7 @@ namespace AutomationTool
 				// If we exited normally, still check without waiting in case SCW or some other child process crashed.
 				if(RunResult.ExitCode > 128)
 				{
-					CommandUtils.Log("Pausing before checking for crash logs...");
+					CommandUtils.LogInformation("Pausing before checking for crash logs...");
 					Thread.Sleep(10 * 1000);
 				}
 				
@@ -317,13 +317,13 @@ namespace AutomationTool
 					// snmpd seems to often crash (suspect due to it being starved of CPU cycles during cooks)
 					if(!CrashFileInfo.Name.StartsWith("snmpd_"))
 					{
-						CommandUtils.Log("Found crash log - {0}", CrashFileInfo.FullName);
+						CommandUtils.LogInformation("Found crash log - {0}", CrashFileInfo.FullName);
 						try
 						{
 							string[] Lines = File.ReadAllLines(CrashFileInfo.FullName);
 							foreach(string Line in Lines)
 							{
-								CommandUtils.Log("Crash: {0}", Line);
+								CommandUtils.LogInformation("Crash: {0}", Line);
 							}
 						}
 						catch(Exception Ex)
diff --git a/Engine/Source/Programs/AutomationTool/AutomationUtils/HelpUtils.cs b/Engine/Source/Programs/AutomationTool/AutomationUtils/HelpUtils.cs
--- a/Engine/Source/Programs/AutomationTool/AutomationUtils/HelpUtils.cs
+++ b/Engine/Source/Programs/AutomationTool/AutomationUtils/HelpUtils.cs
@@ -149,7 +149,7 @@ namespace AutomationTool
 				}
 			}
 
-			Log(HelpMessage);
+			LogInformation(HelpMessage);
 		}
 
 		/// <summary>
diff --git a/Engine/Source/Programs/AutomationTool/AutomationUtils/MCPPublic.cs b/Engine/Source/Programs/AutomationTool/AutomationUtils/MCPPublic.cs
--- a/Engine/Source/Programs/AutomationTool/AutomationUtils/MCPPublic.cs
+++ b/Engine/Source/Programs/AutomationTool/AutomationUtils/MCPPublic.cs
@@ -197,7 +197,7 @@ namespace EpicGames.MCP.Automation
 					return _ManifestFilename;
 				}
 
-				CommandUtils.Log("Using legacy behavior of constructing manifest filename from appname, build version and platform. Update your code to specify manifest filename when constructing BuildPatchToolStagingInfo or call RetrieveManifestFilename to query it.");
+				CommandUtils.LogInformation("Using legacy behavior of constructing manifest filename from appname, build version and platform. Update your code to specify manifest filename when constructing BuildPatchToolStagingInfo or call RetrieveManifestFilename to query it.");
 				var BaseFilename = string.Format("{0}{1}-{2}.manifest",
 					AppName,
 					BuildVersion,
diff --git a/Engine/Source/Programs/AutomationTool/AutomationUtils/P4Blame.cs b/Engine/Source/Programs/AutomationTool/AutomationUtils/P4Blame.cs
--- a/Engine/Source/Programs/AutomationTool/AutomationUtils/P4Blame.cs
+++ b/Engine/Source/Programs/AutomationTool/AutomationUtils/P4Blame.cs
@@ -302,7 +302,7 @@ namespace AutomationTool
 		{
 			if (!bQuiet)
 			{
-				CommandUtils.Log("Filelog({0})", DepotFilename);
+				CommandUtils.LogInformation("Filelog({0})", DepotFilename);
 			}
 			var Result = new List<RevisionInfo>();
 			var FilelogResult = P4("filelog " + DepotFilename, AllowSpew: false);
@@ -335,7 +335,7 @@ namespace AutomationTool
 		{
 			if (!bQuiet)
 			{
-				CommandUtils.Log("P4Print({0})", CommandLine);
+				CommandUtils.LogInformation("P4Print({0})", CommandLine);
 			}
 			var PrintResult = P4("print " + CommandLine, AllowSpew: false);
 			var OutputLines = new List<string>(PrintResult.Output.Split(new string[] { Environment.NewLine }, StringSplitOptions.None));
@@ -354,7 +354,7 @@ namespace AutomationTool
 		{
 			if (!bQuiet)
 			{
-				CommandUtils.Log("Annotate({0})", CommandLine);
+				CommandUtils.LogInformation("Annotate({0})", CommandLine);
 			}
 			var PrintResult = P4("annotate " + CommandLine, AllowSpew: false);
 			var OutputLines = new List<string>(PrintResult.Output.Split(new string[] { Environment.NewLine }, StringSplitOptions.None));
@@ -374,7 +374,7 @@ namespace AutomationTool
 		{
 			if (!bQuiet)
 			{
-				CommandUtils.Log("Diff2({0}, {1}, {2})", DepotFilename, RevA, RevB);
+				CommandUtils.LogInformation("Diff2({0}, {1}, {2})", DepotFilename, RevA, RevB);
 			}
 			var Changes = new List<DiffChange>();
 			var Diff2Result = P4(String.Format("diff2 {0}#{1} {0}#{2}", DepotFilename, RevA, RevB), AllowSpew: false);
@@ -400,7 +400,7 @@ namespace AutomationTool
 		/// <returns>An array of lines with blame information.</returns>
 		public BlameLineInfo[] Timelapse(string DepotFilename, int MaxRevision = Int32.MaxValue)
 		{
-			CommandUtils.Log("Timelapse({0}#{1})", DepotFilename, (MaxRevision == Int32.MaxValue) ? "head" : MaxRevision.ToString());
+			CommandUtils.LogInformation("Timelapse({0}#{1})", DepotFilename, (MaxRevision == Int32.MaxValue) ? "head" : MaxRevision.ToString());
 			// Get the file history.
 			var FileHistory = Filelog(DepotFilename, bQuiet: true);
 			if (!CommandUtils.IsNullOrEmpty(FileHistory))
@@ -453,7 +453,7 @@ namespace AutomationTool
 		/// <returns>An array of lines with blame information.</returns>
 		public BlameLineInfo[] Blame(string DepotFilename)
 		{
-			CommandUtils.Log("Blame({0})", DepotFilename);
+			CommandUtils.LogInformation("Blame({0})", DepotFilename);
 			// Get the file history.
 			var FileHistory = Filelog(DepotFilename, bQuiet: true);
 			
diff --git a/Engine/Source/Programs/AutomationTool/AutomationUtils/P4Utils.cs b/Engine/Source/Programs/AutomationTool/AutomationUtils/P4Utils.cs
--- a/Engine/Source/Programs/AutomationTool/AutomationUtils/P4Utils.cs
+++ b/Engine/Source/Programs/AutomationTool/AutomationUtils/P4Utils.cs
@@ -2183,7 +2183,7 @@ namespace AutomationTool
                 bool isClPending = false;
                 if (ChangeFiles(CL, out isClPending, false).Count == 0)
                 {
-					CommandUtils.Log("No edits left to commit after brutal submit resolve. Assuming another build committed same changes already and exiting as success.");
+					CommandUtils.LogInformation("No edits left to commit after brutal submit resolve. Assuming another build committed same changes already and exiting as success.");
                     DeleteChange(CL);
                     // No changes to submit, no need to retry.
                     return;
@@ -2195,7 +2195,7 @@ namespace AutomationTool
 					{
 						throw new P4Exception("Change {0} failed to submit.\n{1}", CL, CmdOutput);
 					}
-					CommandUtils.Log("**** P4 Returned\n{0}\n*******", CmdOutput);
+					CommandUtils.LogInformation("**** P4 Returned\n{0}\n*******", CmdOutput);
 
 					LastCmdOutput = CmdOutput;
 					bool DidSomething = false;
@@ -2258,7 +2258,7 @@ namespace AutomationTool
                                 {
                                     continue;
                                 }
-								CommandUtils.Log("Brutal 'resolve' on {0} to force submit.\n", File);
+								CommandUtils.LogInformation("Brutal 'resolve' on {0} to force submit.\n", File);
 								Revert(CL, "-k " + CommandUtils.MakePathSafeToUseWithCommandLine(File));  // revert the file without overwriting the local one
 								Sync("-f -k " + CommandUtils.MakePathSafeToUseWithCommandLine(File + "#head"), false); // sync the file without overwriting local one
 								ReconcileNoDeletes(CL, CommandUtils.MakePathSafeToUseWithCommandLine(File));  // re-check out, if it changed, or add
@@ -2269,7 +2269,7 @@ namespace AutomationTool
                     }
 					if (!DidSomething)
 					{
-						CommandUtils.Log("Change {0} failed to submit for reasons we do not recognize.\n{1}\nWaiting and retrying.", CL, CmdOutput);
+						CommandUtils.LogInformation("Change {0} failed to submit for reasons we do not recognize.\n{1}\nWaiting and retrying.", CL, CmdOutput);
 					}
 					System.Threading.Thread.Sleep(30000);
 				}
@@ -2301,7 +2301,7 @@ namespace AutomationTool
 							}
 						}
 
-						CommandUtils.Log("Submitted CL {0} which became CL {1}\n", CL, SubmittedCL);
+						CommandUtils.LogInformation("Submitted CL {0} which became CL {1}\n", CL, SubmittedCL);
 					}
 
 					if (SubmittedCL < CL)
@@ -2345,7 +2345,7 @@ namespace AutomationTool
 			int CL = 0;
 			if(AllowSpew)
 			{
-				CommandUtils.Log("Creating Change\n {0}\n", ChangeSpec);
+				CommandUtils.LogInformation("Creating Change\n {0}\n", ChangeSpec);
 			}
 			if (LogP4Output(out CmdOutput, "change -i", Input: ChangeSpec, AllowSpew: AllowSpew))
 			{
@@ -2364,7 +2364,7 @@ namespace AutomationTool
 			}
 			else if(AllowSpew)
 			{
-				CommandUtils.Log("Returned CL {0}\n", CL);
+				CommandUtils.LogInformation("Returned CL {0}\n", CL);
 			}
 			return CL;
 		}
@@ -2488,7 +2488,7 @@ namespace AutomationTool
 				int EndOffset = CmdOutput.LastIndexOf(EndStr);
 				if (Offset == 0 && Offset < EndOffset)
 				{
-					CommandUtils.Log("Change {0} does not exist", CL);
+					CommandUtils.LogInformation("Change {0} does not exist", CL);
 					return false;
 				}
 
@@ -2502,7 +2502,7 @@ namespace AutomationTool
 				}
 
 				string Status = CmdOutput.Substring(StatusOffset + StatusStr.Length).TrimStart().Split('\n')[0].TrimEnd();
-				CommandUtils.Log("Change {0} exists ({1})", CL, Status);
+				CommandUtils.LogInformation("Change {0} exists ({1})", CL, Status);
 				Pending = (Status == "pending");
 				return true;
 			}
@@ -2591,7 +2591,7 @@ namespace AutomationTool
 			string Output;
 			if (!LogP4Output(out Output, CommandLine, null, AllowSpew))
 			{
-				CommandUtils.Log("Couldn't delete label '{0}'.  It may not have existed in the first place.", LabelName);
+				CommandUtils.LogInformation("Couldn't delete label '{0}'.  It may not have existed in the first place.", LabelName);
 			}
 		}
 
@@ -2622,7 +2622,7 @@ namespace AutomationTool
 			LabelSpec += "View: \n";
 			LabelSpec += " " + View;
 
-			CommandUtils.Log("Creating Label\n {0}\n", LabelSpec);
+			CommandUtils.LogInformation("Creating Label\n {0}\n", LabelSpec);
 			LogP4("label -i", Input: LabelSpec);
 		}
 
diff --git a/Engine/Source/Programs/AutomationTool/AutomationUtils/ParallelExecutor.cs b/Engine/Source/Programs/AutomationTool/AutomationUtils/ParallelExecutor.cs
--- a/Engine/Source/Programs/AutomationTool/AutomationUtils/ParallelExecutor.cs
+++ b/Engine/Source/Programs/AutomationTool/AutomationUtils/ParallelExecutor.cs
@@ -85,7 +85,7 @@ namespace AutomationTool
 		{
 			List<BuildAction> Actions = ReadActions(ActionsFileName);
 
-			CommandUtils.Log("Building {0} {1} with {2} {3}...", Actions.Count, (Actions.Count == 1) ? "action" : "actions", MaxProcesses, (MaxProcesses == 1)? "process" : "processes");
+			CommandUtils.LogInformation("Building {0} {1} with {2} {3}...", Actions.Count, (Actions.Count == 1) ? "action" : "actions", MaxProcesses, (MaxProcesses == 1)? "process" : "processes");
 
 			using (LogIndentScope Indent = new LogIndentScope("  "))
 			{
@@ -148,11 +148,11 @@ namespace AutomationTool
 										if(CurrentPrefix != CompletedAction.Action.GroupPrefix)
 										{
 											CurrentPrefix = CompletedAction.Action.GroupPrefix;
-											CommandUtils.Log(CurrentPrefix);
+											CommandUtils.LogInformation(CurrentPrefix);
 										}
 										foreach(string LogLine in CompletedAction.LogLines)
 										{
-											CommandUtils.Log(LogLine);
+											CommandUtils.LogInformation(LogLine);
 										}
 									}
 
diff --git a/Engine/Source/Programs/AutomationTool/AutomationUtils/ProcessUtils.cs b/Engine/Source/Programs/AutomationTool/AutomationUtils/ProcessUtils.cs
--- a/Engine/Source/Programs/AutomationTool/AutomationUtils/ProcessUtils.cs
+++ b/Engine/Source/Programs/AutomationTool/AutomationUtils/ProcessUtils.cs
@@ -157,7 +157,7 @@ namespace AutomationTool
 							if (ProcessResult.HasAnyDescendants(Process.GetCurrentProcess()))
 							{
 								AllDone = false;
-								CommandUtils.Log("Waiting for descendants of main process...");
+								CommandUtils.LogInformation("Waiting for descendants of main process...");
 							}
 						}
 						catch (Exception Ex)
@@ -1058,7 +1058,7 @@ namespace AutomationTool
 		{
 			while (!FileExists(LogFilename) && !LogProcess.HasExited)
 			{
-				Log("Waiting for logging process to start...");
+				LogInformation("Waiting for logging process to start...");
 				Thread.Sleep(2000);
 			}
 			Thread.Sleep(1000);
diff --git a/Engine/Source/Programs/AutomationTool/AutomationUtils/UE4Build.cs b/Engine/Source/Programs/AutomationTool/AutomationUtils/UE4Build.cs
--- a/Engine/Source/Programs/AutomationTool/AutomationUtils/UE4Build.cs
+++ b/Engine/Source/Programs/AutomationTool/AutomationUtils/UE4Build.cs
@@ -1455,7 +1455,7 @@ namespace AutomationTool
 				CommandUtils.LogLog("Build products *******");
 				if (BuildProductFiles.Count < 1)
 				{
-					CommandUtils.Log("No build products were made");
+					CommandUtils.LogInformation("No build products were made");
 				}
 				else
 				{
@@ -1479,7 +1479,7 @@ namespace AutomationTool
 		/// <param name="Files">List of files to check out</param>
 		public static void AddBuildProductsToChangelist(int WorkingCL, List<string> Files)
 		{
-			CommandUtils.Log("Adding {0} build products to changelist {1}...", Files.Count, WorkingCL);
+			CommandUtils.LogInformation("Adding {0} build products to changelist {1}...", Files.Count, WorkingCL);
 			foreach (var File in Files)
 			{
 				CommandUtils.P4.Sync("-f -k " + CommandUtils.MakePathSafeToUseWithCommandLine(File) + "#head"); // sync the file without overwriting local one
diff --git a/Engine/Source/Programs/AutomationTool/AutomationUtils/Utils.cs b/Engine/Source/Programs/AutomationTool/AutomationUtils/Utils.cs
--- a/Engine/Source/Programs/AutomationTool/AutomationUtils/Utils.cs
+++ b/Engine/Source/Programs/AutomationTool/AutomationUtils/Utils.cs
@@ -745,7 +745,7 @@ namespace AutomationTool
 	                if (Retry > 0)
 	                {
                         //@todo: These retries should be reported so we can track how often they are occurring.
-                        CommandUtils.Log("*** Mac temp storage retry {0}", OutputFileName);
+                        CommandUtils.LogInformation("*** Mac temp storage retry {0}", OutputFileName);
 	                    System.Threading.Thread.Sleep(1000);
 	                }
 	                bCopied = CommandUtils.CopyFile_NoExceptions(InputFileName, OutputFileName, true);
@@ -775,7 +775,7 @@ namespace AutomationTool
 	                while (!bFound && Retry < 60)
 	                {
                         //@todo: These retries should be reported so we can track how often they are occurring.
-                        CommandUtils.Log("*** Mac temp storage retry {0}", Filename);
+                        CommandUtils.LogInformation("*** Mac temp storage retry {0}", Filename);
 	                    System.Threading.Thread.Sleep(10000);
 	                    bFound = CommandUtils.FileExists_NoExceptions(bQuiet, Filename);
 	                    Retry++;
@@ -800,7 +800,7 @@ namespace AutomationTool
 	                while (!bFound && Retry < 60)
 	                {
                         //@todo: These retries should be reported so we can track how often they are occurring.
-                        CommandUtils.Log("*** Mac temp storage retry {0}", Directoryname);
+                        CommandUtils.LogInformation("*** Mac temp storage retry {0}", Directoryname);
 	                    System.Threading.Thread.Sleep(10000);
 	                    bFound = CommandUtils.DirectoryExists_NoExceptions(Directoryname);
 	                    Retry++;
@@ -831,7 +831,7 @@ namespace AutomationTool
 	                while (!bFound && Retry < NumRetries)
 	                {
                         //@todo: These retries should be reported so we can track how often they are occurring.
-                        CommandUtils.Log("*** Mac temp storage retry {0}", Directoryname);
+                        CommandUtils.LogInformation("*** Mac temp storage retry {0}", Directoryname);
 	                    System.Threading.Thread.Sleep(1000);
 	                    bFound = CommandUtils.DirectoryExistsAndIsWritable_NoExceptions(Directoryname);
 	                    Retry++;
@@ -933,7 +933,7 @@ namespace AutomationTool
             {
                 throw new AutomationException("Failed to find MAJOR, MINOR, and PATCH fields from version file {0}", Filename);
             }
-			CommandUtils.Log("Read {0}.{1}.{2} from {3}", foundElements["MAJOR"], foundElements["MINOR"], foundElements["PATCH"], Filename);
+			CommandUtils.LogInformation("Read {0}.{1}.{2} from {3}", foundElements["MAJOR"], foundElements["MINOR"], foundElements["PATCH"], Filename);
             return new Version(foundElements["MAJOR"], foundElements["MINOR"], foundElements["PATCH"]);
         }
 
diff --git a/Engine/Source/Programs/AutomationTool/BuildGraph/BuildGraph.cs b/Engine/Source/Programs/AutomationTool/BuildGraph/BuildGraph.cs
--- a/Engine/Source/Programs/AutomationTool/BuildGraph/BuildGraph.cs
+++ b/Engine/Source/Programs/AutomationTool/BuildGraph/BuildGraph.cs
@@ -176,7 +176,7 @@ namespace AutomationTool
 			if(SchemaFileName != null)
 			{
 				FileReference FullSchemaFileName = new FileReference(SchemaFileName);
-				Log("Writing schema to {0}...", FullSchemaFileName.FullName);
+				LogInformation("Writing schema to {0}...", FullSchemaFileName.FullName);
 				Schema.Export(FullSchemaFileName);
 				if(ScriptFileName == null)
 				{
@@ -244,12 +244,12 @@ namespace AutomationTool
 				// List out all the required tokens
 				if(SingleNodeName == null)
 				{
-					CommandUtils.Log("Required tokens:");
+					CommandUtils.LogInformation("Required tokens:");
 					foreach(Node Node in TargetNodes)
 					{
 						foreach(FileReference RequiredToken in Node.RequiredTokens)
 						{
-							CommandUtils.Log("  '{0}' requires {1}", Node, RequiredToken);
+							CommandUtils.LogInformation("  '{0}' requires {1}", Node, RequiredToken);
 						}
 					}
 				}
@@ -281,12 +281,12 @@ namespace AutomationTool
 						HashSet<Node> SkipNodes = new HashSet<Node>();
 						foreach(IGrouping<string, FileReference> MissingTokensForBuild in MissingTokens.GroupBy(x => x.Value, x => x.Key))
 						{
-							Log("Skipping the following nodes due to {0}:", MissingTokensForBuild.Key);
+							LogInformation("Skipping the following nodes due to {0}:", MissingTokensForBuild.Key);
 							foreach(FileReference MissingToken in MissingTokensForBuild)
 							{
 								foreach(Node SkipNode in TargetNodes.Where(x => x.RequiredTokens.Contains(MissingToken) && SkipNodes.Add(x)))
 								{
-									Log("    {0}", SkipNode);
+									LogInformation("    {0}", SkipNode);
 								}
 							}
 						}
@@ -295,14 +295,14 @@ namespace AutomationTool
 						if(SkipNodes.Count > 0)
 						{
 							TargetNodes.ExceptWith(SkipNodes);
-							Log("Remaining target nodes:");
+							LogInformation("Remaining target nodes:");
 							foreach(Node TargetNode in TargetNodes)
 							{
-								Log("    {0}", TargetNode);
+								LogInformation("    {0}", TargetNode);
 							}
 							if(TargetNodes.Count == 0)
 							{
-								Log("    None.");
+								LogInformation("    None.");
 							}
 						}
 					}
@@ -358,7 +358,7 @@ namespace AutomationTool
 			if (PreprocessedFileName != null)
 			{
 				FileReference PreprocessedFileLocation = new FileReference(PreprocessedFileName);
-				Log("Writing {0}...", PreprocessedFileLocation);
+				LogInformation("Writing {0}...", PreprocessedFileLocation);
 				Graph.Write(PreprocessedFileLocation, (SchemaFileName != null)? new FileReference(SchemaFileName) : null);
 				return ExitCode.Success;
 			}
@@ -394,7 +394,7 @@ namespace AutomationTool
 				{
 					if(Diagnostic.EventType == LogEventType.Console)
 					{
-						CommandUtils.Log(Diagnostic.Message);
+						CommandUtils.LogInformation(Diagnostic.Message);
 					}
 					else if(Diagnostic.EventType == LogEventType.Warning)
 					{
@@ -622,15 +622,15 @@ namespace AutomationTool
 			int NodeIdx = 0;
 			foreach(Node NodeToExecute in NodesToExecute)
 			{
-				Log("****** [{0}/{1}] {2}", ++NodeIdx, NodesToExecute.Length, NodeToExecute.Name);
+				LogInformation("****** [{0}/{1}] {2}", ++NodeIdx, NodesToExecute.Length, NodeToExecute.Name);
 				if(!Storage.IsComplete(NodeToExecute.Name))
 				{
-					Log("");
+					LogInformation("");
 					if(!BuildNode(Job, Graph, NodeToExecute, Storage, false))
 					{
 						return false;
 					} 
-					Log("");
+					LogInformation("");
 				}
 			}
 			return true;
@@ -695,7 +695,7 @@ namespace AutomationTool
 			if(bWithBanner)
 			{
 				Console.WriteLine();
-				CommandUtils.Log("========== Starting: {0} ==========", Node.Name);
+				CommandUtils.LogInformation("========== Starting: {0} ==========", Node.Name);
 			}
 			if(!Node.Build(Job, TagNameToFileSet))
 			{
@@ -703,7 +703,7 @@ namespace AutomationTool
 			}
 			if(bWithBanner)
 			{
-				CommandUtils.Log("========== Finished: {0} ==========", Node.Name);
+				CommandUtils.LogInformation("========== Finished: {0} ==========", Node.Name);
 				Console.WriteLine();
 			}
 
@@ -859,7 +859,7 @@ namespace AutomationTool
 			}
 
 			// Write the output file
-			Log("Writing {0}...", OutputFile);
+			LogInformation("Writing {0}...", OutputFile);
 			using (StreamWriter Writer = new StreamWriter(OutputFile.FullName))
 			{
 				Writer.WriteLine("<html>");
diff --git a/Engine/Source/Programs/AutomationTool/BuildGraph/Graph.cs b/Engine/Source/Programs/AutomationTool/BuildGraph/Graph.cs
--- a/Engine/Source/Programs/AutomationTool/BuildGraph/Graph.cs
+++ b/Engine/Source/Programs/AutomationTool/BuildGraph/Graph.cs
@@ -604,12 +604,12 @@ namespace AutomationTool
 				// Format them to the log
 				if(Messages.Count > 0)
 				{
-					CommandUtils.Log("");
-					CommandUtils.Log("Options:");
-					CommandUtils.Log("");
+					CommandUtils.LogInformation("");
+					CommandUtils.LogInformation("Options:");
+					CommandUtils.LogInformation("");
 					foreach(string Line in CommandUtils.FormatParams(Messages, 4, 24))
 					{
-						CommandUtils.Log(Line);
+						CommandUtils.LogInformation(Line);
 					}
 				}
 			}
@@ -620,8 +620,8 @@ namespace AutomationTool
 			AllTriggers.AddRange(NameToTrigger.Values.OrderBy(x => x.QualifiedName));
 
 			// Output all the triggers in order
-			CommandUtils.Log("");
-			CommandUtils.Log("Graph:");
+			CommandUtils.LogInformation("");
+			CommandUtils.LogInformation("Graph:");
 			foreach(ManualTrigger Trigger in AllTriggers)
 			{
 				// Filter everything by this trigger
@@ -642,12 +642,12 @@ namespace AutomationTool
 				}
 
 				// Print the trigger name
-				CommandUtils.Log("    Trigger: {0}", (Trigger == null)? "None" : Trigger.QualifiedName);
+				CommandUtils.LogInformation("    Trigger: {0}", (Trigger == null)? "None" : Trigger.QualifiedName);
 				if(Trigger != null && PrintOptions.HasFlag(GraphPrintOptions.ShowNotifications))
 				{
 					foreach(string User in Trigger.NotifyUsers)
 					{
-						CommandUtils.Log("            notify> {0}", User);
+						CommandUtils.LogInformation("            notify> {0}", User);
 					}
 				}
 
@@ -657,21 +657,21 @@ namespace AutomationTool
 					Node[] Nodes;
 					if(FilteredAgentToNodes.TryGetValue(Agent, out Nodes))
 					{
-						CommandUtils.Log("        Agent: {0} ({1})", Agent.Name, String.Join(";", Agent.PossibleTypes));
+						CommandUtils.LogInformation("        Agent: {0} ({1})", Agent.Name, String.Join(";", Agent.PossibleTypes));
 						foreach(Node Node in Nodes)
 						{
-							CommandUtils.Log("            Node: {0}{1}", Node.Name, CompletedNodes.Contains(Node)? " (completed)" : Node.bRunEarly? " (early)" : "");
+							CommandUtils.LogInformation("            Node: {0}{1}", Node.Name, CompletedNodes.Contains(Node)? " (completed)" : Node.bRunEarly? " (early)" : "");
 							if(PrintOptions.HasFlag(GraphPrintOptions.ShowDependencies))
 							{
 								HashSet<Node> InputDependencies = new HashSet<Node>(Node.GetDirectInputDependencies());
 								foreach(Node InputDependency in InputDependencies)
 								{
-									CommandUtils.Log("                input> {0}", InputDependency.Name);
+									CommandUtils.LogInformation("                input> {0}", InputDependency.Name);
 								}
 								HashSet<Node> OrderDependencies = new HashSet<Node>(Node.GetDirectOrderDependencies());
 								foreach(Node OrderDependency in OrderDependencies.Except(InputDependencies))
 								{
-									CommandUtils.Log("                after> {0}", OrderDependency.Name);
+									CommandUtils.LogInformation("                after> {0}", OrderDependency.Name);
 								}
 							}
 							if(PrintOptions.HasFlag(GraphPrintOptions.ShowNotifications))
@@ -679,29 +679,29 @@ namespace AutomationTool
 								string Label = Node.bNotifyOnWarnings? "warnings" : "errors";
 								foreach(string User in Node.NotifyUsers)
 								{
-									CommandUtils.Log("                {0}> {1}", Label, User);
+									CommandUtils.LogInformation("                {0}> {1}", Label, User);
 								}
 								foreach(string Submitter in Node.NotifySubmitters)
 								{
-									CommandUtils.Log("                {0}> submitters to {1}", Label, Submitter);
+									CommandUtils.LogInformation("                {0}> submitters to {1}", Label, Submitter);
 								}
 							}
 						}
 					}
 				}
 			}
-			CommandUtils.Log("");
+			CommandUtils.LogInformation("");
 
 			// Print out all the aggregates
 			string[] AggregateNames = AggregateNameToNodes.Keys.OrderBy(x => x).ToArray();
 			if(AggregateNames.Length > 0)
 			{
-				CommandUtils.Log("Aggregates:");
+				CommandUtils.LogInformation("Aggregates:");
 				foreach(string AggregateName in AggregateNames)
 				{
-					CommandUtils.Log("    {0}", AggregateName);
+					CommandUtils.LogInformation("    {0}", AggregateName);
 				}
-				CommandUtils.Log("");
+				CommandUtils.LogInformation("");
 			}
 		}
 	}
diff --git a/Engine/Source/Programs/AutomationTool/BuildGraph/LockFile.cs b/Engine/Source/Programs/AutomationTool/BuildGraph/LockFile.cs
--- a/Engine/Source/Programs/AutomationTool/BuildGraph/LockFile.cs
+++ b/Engine/Source/Programs/AutomationTool/BuildGraph/LockFile.cs
@@ -58,11 +58,11 @@ namespace AutomationTool
 
                     if (Iterations == 0)
                     {
-                        CommandUtils.Log("Waiting for lock file '{0}' to be removed...", LockFilePath);
+                        CommandUtils.LogInformation("Waiting for lock file '{0}' to be removed...", LockFilePath);
                     }
                     else if ((Iterations % 30) == 0)
                     {
-                        CommandUtils.Log("Still waiting for lock file '{0}' after {1} seconds.", LockFilePath, CurrentTime.Subtract(StartTime).TotalSeconds);
+                        CommandUtils.LogInformation("Still waiting for lock file '{0}' after {1} seconds.", LockFilePath, CurrentTime.Subtract(StartTime).TotalSeconds);
                     }
 
                     // Wait for a while before retrying.
diff --git a/Engine/Source/Programs/AutomationTool/BuildGraph/Tasks/AgeStoreTask.cs b/Engine/Source/Programs/AutomationTool/BuildGraph/Tasks/AgeStoreTask.cs
--- a/Engine/Source/Programs/AutomationTool/BuildGraph/Tasks/AgeStoreTask.cs
+++ b/Engine/Source/Programs/AutomationTool/BuildGraph/Tasks/AgeStoreTask.cs
@@ -69,7 +69,7 @@ namespace Win.Automation
             try
             {
                 Directory.Delete(true);
-                CommandUtils.Log("Removed '{0}'", Directory.FullName);
+                CommandUtils.LogInformation("Removed '{0}'", Directory.FullName);
             }
             catch
             {
@@ -127,7 +127,7 @@ namespace Win.Automation
 
             // Get the time at which to expire files
             DateTime ExpireTimeUtc = DateTime.UtcNow - TimeSpan.FromDays(Parameters.Days);
-            CommandUtils.Log("Expiring all files before {0}...", ExpireTimeUtc);
+            CommandUtils.LogInformation("Expiring all files before {0}...", ExpireTimeUtc);
             
             // Scan the store directory and delete old symbol files
             DirectoryReference SymbolServerDirectory = ResolveDirectory(Parameters.StoreDir);
diff --git a/Engine/Source/Programs/AutomationTool/BuildGraph/Tasks/CopyTask.cs b/Engine/Source/Programs/AutomationTool/BuildGraph/Tasks/CopyTask.cs
--- a/Engine/Source/Programs/AutomationTool/BuildGraph/Tasks/CopyTask.cs
+++ b/Engine/Source/Programs/AutomationTool/BuildGraph/Tasks/CopyTask.cs
@@ -101,7 +101,7 @@ namespace BuildGraph.Tasks
 				{
 					if (FileReference.Exists(File.Key))
 					{
-						CommandUtils.Log("Not copying existing file {0}", File.Key);
+						CommandUtils.LogInformation("Not copying existing file {0}", File.Key);
 						return false;
 					}
 					return true;
@@ -111,7 +111,7 @@ namespace BuildGraph.Tasks
 			// Check we got some files
 			if(TargetFileToSourceFile.Count == 0)
 			{
-				CommandUtils.Log("No files found matching '{0}'", SourcePattern);
+				CommandUtils.LogInformation("No files found matching '{0}'", SourcePattern);
 				return;
 			}
 
@@ -156,7 +156,7 @@ namespace BuildGraph.Tasks
 
 			// Copy them all
 			KeyValuePair<FileReference, FileReference>[] FilePairs = TargetFileToSourceFile.ToArray();
-			CommandUtils.Log("Copying {0} file{1} from {2} to {3}...", FilePairs.Length, (FilePairs.Length == 1)? "" : "s", SourcePattern.BaseDirectory, TargetPattern.BaseDirectory);
+			CommandUtils.LogInformation("Copying {0} file{1} from {2} to {3}...", FilePairs.Length, (FilePairs.Length == 1)? "" : "s", SourcePattern.BaseDirectory, TargetPattern.BaseDirectory);
 			CommandUtils.ThreadedCopyFiles(FilePairs.Select(x => x.Value.FullName).ToList(), FilePairs.Select(x => x.Key.FullName).ToList(), bQuiet: true);
 
 			// Update the list of build products
diff --git a/Engine/Source/Programs/AutomationTool/BuildGraph/Tasks/LogTask.cs b/Engine/Source/Programs/AutomationTool/BuildGraph/Tasks/LogTask.cs
--- a/Engine/Source/Programs/AutomationTool/BuildGraph/Tasks/LogTask.cs
+++ b/Engine/Source/Programs/AutomationTool/BuildGraph/Tasks/LogTask.cs
@@ -68,7 +68,7 @@ namespace BuildGraph.Tasks
 			// Print the message
 			if(!String.IsNullOrEmpty(Parameters.Message))
 			{
-				CommandUtils.Log(Parameters.Message);
+				CommandUtils.LogInformation(Parameters.Message);
 			}
 
 			// Print the contents of the given tag, if specified
@@ -77,12 +77,12 @@ namespace BuildGraph.Tasks
 				HashSet<FileReference> Files = ResolveFilespec(CommandUtils.RootDirectory, Parameters.Files, TagNameToFileSet);
 				foreach(FileReference File in Files.OrderBy(x => x.FullName))
 				{
-					CommandUtils.Log("  {0}", File.FullName);
+					CommandUtils.LogInformation("  {0}", File.FullName);
 					if(Parameters.IncludeContents)
 					{
 						foreach(string Line in System.IO.File.ReadAllLines(File.FullName))
 						{
-							CommandUtils.Log("    {0}", Line);
+							CommandUtils.LogInformation("    {0}", Line);
 						}
 					}
 				}
diff --git a/Engine/Source/Programs/AutomationTool/BuildGraph/Tasks/MoveTask.cs b/Engine/Source/Programs/AutomationTool/BuildGraph/Tasks/MoveTask.cs
--- a/Engine/Source/Programs/AutomationTool/BuildGraph/Tasks/MoveTask.cs
+++ b/Engine/Source/Programs/AutomationTool/BuildGraph/Tasks/MoveTask.cs
@@ -90,12 +90,12 @@ namespace BuildGraph.Tasks
 			// Check we got some files
 			if(TargetFileToSourceFile.Count == 0)
 			{
-				CommandUtils.Log("No files found matching '{0}'", SourcePattern);
+				CommandUtils.LogInformation("No files found matching '{0}'", SourcePattern);
 				return;
 			}
 
 			// Copy them all
-			CommandUtils.Log("Moving {0} file{1} from {2} to {3}...", TargetFileToSourceFile.Count, (TargetFileToSourceFile.Count == 1)? "" : "s", SourcePattern.BaseDirectory, TargetPattern.BaseDirectory);
+			CommandUtils.LogInformation("Moving {0} file{1} from {2} to {3}...", TargetFileToSourceFile.Count, (TargetFileToSourceFile.Count == 1)? "" : "s", SourcePattern.BaseDirectory, TargetPattern.BaseDirectory);
 			CommandUtils.ParallelMoveFiles(TargetFileToSourceFile.Select(x => new KeyValuePair<FileReference, FileReference>(x.Value, x.Key)));
 
 			// Update the list of build products
diff --git a/Engine/Source/Programs/AutomationTool/BuildGraph/Tasks/PakFileTask.cs b/Engine/Source/Programs/AutomationTool/BuildGraph/Tasks/PakFileTask.cs
--- a/Engine/Source/Programs/AutomationTool/BuildGraph/Tasks/PakFileTask.cs
+++ b/Engine/Source/Programs/AutomationTool/BuildGraph/Tasks/PakFileTask.cs
@@ -175,7 +175,7 @@ namespace BuildGraph.Tasks
 			}
 
 			// Run it
-			CommandUtils.Log("Running '{0} {1}'", CommandUtils.MakePathSafeToUseWithCommandLine(UnrealPakExe.FullName), CommandLine.ToString());
+			CommandUtils.LogInformation("Running '{0} {1}'", CommandUtils.MakePathSafeToUseWithCommandLine(UnrealPakExe.FullName), CommandLine.ToString());
 			CommandUtils.RunAndLog(CommandUtils.CmdEnv, UnrealPakExe.FullName, CommandLine.ToString(), Options: CommandUtils.ERunOptions.Default | CommandUtils.ERunOptions.UTF8Output);
 			BuildProducts.Add(OutputFile);
 
diff --git a/Engine/Source/Programs/AutomationTool/BuildGraph/Tasks/StripTask.cs b/Engine/Source/Programs/AutomationTool/BuildGraph/Tasks/StripTask.cs
--- a/Engine/Source/Programs/AutomationTool/BuildGraph/Tasks/StripTask.cs
+++ b/Engine/Source/Programs/AutomationTool/BuildGraph/Tasks/StripTask.cs
@@ -95,11 +95,11 @@ namespace BuildGraph.Tasks
 				DirectoryReference.CreateDirectory(TargetFiles[Idx].Directory);
 				if (SourceFiles[Idx] == TargetFiles[Idx])
 				{
-					CommandUtils.Log("Stripping symbols: {0}", SourceFiles[Idx].FullName);
+					CommandUtils.LogInformation("Stripping symbols: {0}", SourceFiles[Idx].FullName);
 				}
 				else
 				{
-					CommandUtils.Log("Stripping symbols: {0} -> {1}", SourceFiles[Idx].FullName, TargetFiles[Idx].FullName);
+					CommandUtils.LogInformation("Stripping symbols: {0} -> {1}", SourceFiles[Idx].FullName, TargetFiles[Idx].FullName);
 				}
 				TargetPlatform.StripSymbols(SourceFiles[Idx], TargetFiles[Idx]);
 			}
diff --git a/Engine/Source/Programs/AutomationTool/BuildGraph/Tasks/SubmitTask.cs b/Engine/Source/Programs/AutomationTool/BuildGraph/Tasks/SubmitTask.cs
--- a/Engine/Source/Programs/AutomationTool/BuildGraph/Tasks/SubmitTask.cs
+++ b/Engine/Source/Programs/AutomationTool/BuildGraph/Tasks/SubmitTask.cs
@@ -137,7 +137,7 @@ namespace AutomationTool.Tasks
 					SubmitP4.RevertUnchanged(NewCL);
 					if(SubmitP4.TryDeleteEmptyChange(NewCL))
 					{
-						CommandUtils.Log("No files to submit; ignored.");
+						CommandUtils.LogInformation("No files to submit; ignored.");
 						return;
 					}
 				}
@@ -149,7 +149,7 @@ namespace AutomationTool.Tasks
 				{
 					throw new AutomationException("Submit failed.");
 				}
-				CommandUtils.Log("Submitted in changelist {0}", SubmittedCL);
+				CommandUtils.LogInformation("Submitted in changelist {0}", SubmittedCL);
 			}
 		}
 
diff --git a/Engine/Source/Programs/AutomationTool/BuildGraph/TempStorage.cs b/Engine/Source/Programs/AutomationTool/BuildGraph/TempStorage.cs
--- a/Engine/Source/Programs/AutomationTool/BuildGraph/TempStorage.cs
+++ b/Engine/Source/Programs/AutomationTool/BuildGraph/TempStorage.cs
@@ -148,7 +148,7 @@ namespace AutomationTool
 			{
 				if(Message != null)
 				{
-					CommandUtils.Log(Message);
+					CommandUtils.LogInformation(Message);
 				}
 				return true;
 			}
@@ -672,7 +672,7 @@ namespace AutomationTool
 			FileReference LocalFileListLocation = GetTaggedFileListLocation(LocalDir, NodeName, TagName);
 			if(FileReference.Exists(LocalFileListLocation))
 			{
-				CommandUtils.Log("Reading local file list from {0}", LocalFileListLocation.FullName);
+				CommandUtils.LogInformation("Reading local file list from {0}", LocalFileListLocation.FullName);
 				FileList = TempStorageFileList.Load(LocalFileListLocation);
 			}
 			else
@@ -691,7 +691,7 @@ namespace AutomationTool
 				}
 
 				// Read the shared manifest
-				CommandUtils.Log("Copying shared tag set from {0} to {1}", SharedFileListLocation.FullName, LocalFileListLocation.FullName);
+				CommandUtils.LogInformation("Copying shared tag set from {0} to {1}", SharedFileListLocation.FullName, LocalFileListLocation.FullName);
 				FileList = TempStorageFileList.Load(SharedFileListLocation);
 
 				// Save the manifest locally
@@ -719,14 +719,14 @@ namespace AutomationTool
 			if(SharedDir != null && bWriteToSharedStorage)
 			{
 				FileReference SharedFileListLocation = GetTaggedFileListLocation(SharedDir, NodeName, TagName);
-				CommandUtils.Log("Saving file list to {0} and {1}", LocalFileListLocation.FullName, SharedFileListLocation.FullName);
+				CommandUtils.LogInformation("Saving file list to {0} and {1}", LocalFileListLocation.FullName, SharedFileListLocation.FullName);
 
 				DirectoryReference.CreateDirectory(SharedFileListLocation.Directory);
 				FileList.Save(SharedFileListLocation);
 			}
 			else
 			{
-				CommandUtils.Log("Saving file list to {0}", LocalFileListLocation.FullName);
+				CommandUtils.LogInformation("Saving file list to {0}", LocalFileListLocation.FullName);
 			}
 
 			// Save the local file list
@@ -767,13 +767,13 @@ namespace AutomationTool
 					Manifest.ZipFiles = ZipFiles.Select(x => new TempStorageZipFile(x)).ToArray();
 
 					// Save the shared manifest
-					CommandUtils.Log("Saving shared manifest to {0}", SharedManifestFile.FullName);
+					CommandUtils.LogInformation("Saving shared manifest to {0}", SharedManifestFile.FullName);
 					Manifest.Save(SharedManifestFile);
 				}
 
 				// Save the local manifest
 				FileReference LocalManifestFile = GetManifestLocation(LocalDir, NodeName, BlockName);
-				CommandUtils.Log("Saving local manifest to {0}", LocalManifestFile.FullName);
+				CommandUtils.LogInformation("Saving local manifest to {0}", LocalManifestFile.FullName);
 				Manifest.Save(LocalManifestFile);
 
 				// Update the stats
@@ -801,7 +801,7 @@ namespace AutomationTool
 				TempStorageManifest Manifest;
 				if(bLocal)
 				{
-					CommandUtils.Log("Reading shared manifest from {0}", LocalManifestFile.FullName);
+					CommandUtils.LogInformation("Reading shared manifest from {0}", LocalManifestFile.FullName);
 					Manifest = TempStorageManifest.Load(LocalManifestFile);
 				}
 				else
@@ -822,7 +822,7 @@ namespace AutomationTool
 					}
 
 					// Read the shared manifest
-					CommandUtils.Log("Copying shared manifest from {0} to {1}", SharedManifestFile.FullName, LocalManifestFile.FullName);
+					CommandUtils.LogInformation("Copying shared manifest from {0} to {1}", SharedManifestFile.FullName, LocalManifestFile.FullName);
 					Manifest = TempStorageManifest.Load(SharedManifestFile);
 
 					// Unzip all the build products
@@ -1142,12 +1142,12 @@ namespace AutomationTool
 			CheckManifest(WorkingDir, NamedManifest, NamedOutput);
 
 			// Delete local temp storage and the working directory and try again
-			CommandUtils.Log("Clearing local folders...");
+			CommandUtils.LogInformation("Clearing local folders...");
 			CommandUtils.DeleteDirectoryContents(WorkingDir.FullName);
 			CommandUtils.DeleteDirectoryContents(LocalDir.FullName);
 
 			// First output should fail
-			CommandUtils.Log("Checking default manifest is now unavailable...");
+			CommandUtils.LogInformation("Checking default manifest is now unavailable...");
 			bool bGotManifest = false;
 			try
 			{
@@ -1257,12 +1257,12 @@ namespace AutomationTool
 			DateTime RetainTime = DateTime.UtcNow - TimeSpan.FromDays(DaysValue);
 
 			// Enumerate all the build directories
-			CommandUtils.Log("Scanning {0}...", TempStorageDir);
+			CommandUtils.LogInformation("Scanning {0}...", TempStorageDir);
 			int NumBuilds = 0;
 			List<DirectoryInfo> BuildsToDelete = new List<DirectoryInfo>();
 			foreach (DirectoryInfo StreamDirectory in new DirectoryInfo(TempStorageDir).EnumerateDirectories().OrderBy(x => x.Name))
 			{
-				CommandUtils.Log("Scanning {0}...", StreamDirectory.FullName);
+				CommandUtils.LogInformation("Scanning {0}...", StreamDirectory.FullName);
 				foreach (DirectoryInfo BuildDirectory in StreamDirectory.EnumerateDirectories())
 				{
 					if(!BuildDirectory.EnumerateFiles("*", SearchOption.AllDirectories).Any(x => x.LastWriteTimeUtc > RetainTime))
@@ -1272,14 +1272,14 @@ namespace AutomationTool
 					NumBuilds++;
 				}
 			}
-			CommandUtils.Log("Found {0} builds; {1} to delete.", NumBuilds, BuildsToDelete.Count);
+			CommandUtils.LogInformation("Found {0} builds; {1} to delete.", NumBuilds, BuildsToDelete.Count);
 
 			// Loop through them all, checking for files older than the delete time
 			for (int Idx = 0; Idx < BuildsToDelete.Count; Idx++)
 			{
 				try
 				{
-					CommandUtils.Log("[{0}/{1}] Deleting {2}...", Idx + 1, BuildsToDelete.Count, BuildsToDelete[Idx].FullName);
+					CommandUtils.LogInformation("[{0}/{1}] Deleting {2}...", Idx + 1, BuildsToDelete.Count, BuildsToDelete[Idx].FullName);
 					BuildsToDelete[Idx].Delete(true);
 				}
 				catch (Exception Ex)
diff --git a/Engine/Source/Programs/AutomationTool/HTML5/HTML5Platform.Automation.cs b/Engine/Source/Programs/AutomationTool/HTML5/HTML5Platform.Automation.cs
--- a/Engine/Source/Programs/AutomationTool/HTML5/HTML5Platform.Automation.cs
+++ b/Engine/Source/Programs/AutomationTool/HTML5/HTML5Platform.Automation.cs
@@ -25,9 +25,9 @@ public class HTML5Platform : Platform
 
 	public override void Package(ProjectParams Params, DeploymentContext SC, int WorkingCL)
 	{
-		Log("Package {0}", Params.RawProjectPath);
+		LogInformation("Package {0}", Params.RawProjectPath);
 
-		Log("Setting Emscripten SDK for packaging..");
+		LogInformation("Setting Emscripten SDK for packaging..");
 		HTML5SDKInfo.SetupEmscriptenTemp();
 		HTML5SDKInfo.SetUpEmscriptenConfigFile();
 
@@ -56,9 +56,9 @@ public class HTML5Platform : Platform
 			ConfigCache.GetBool("/Script/HTML5PlatformEditor.HTML5TargetSettings", "Compressed", out Compressed);
 			ConfigCache.GetBool("/Script/HTML5PlatformEditor.HTML5TargetSettings", "EnableIndexedDB", out enableIndexedDB);
 		}
-		Log("HTML5Platform.Automation: TargetWebGL2 = "     + targetWebGL2    );
-		Log("HTML5Platform.Automation: Compressed = "       + Compressed      );
-		Log("HTML5Platform.Automation: EnableIndexedDB = "  + enableIndexedDB );
+		LogInformation("HTML5Platform.Automation: TargetWebGL2 = "     + targetWebGL2    );
+		LogInformation("HTML5Platform.Automation: Compressed = "       + Compressed      );
+		LogInformation("HTML5Platform.Automation: EnableIndexedDB = "  + enableIndexedDB );
 
 		string FinalDataLocation = Path.Combine(PackagePath, Params.ShortProjectName) + ".data";
 
@@ -128,7 +128,7 @@ public class HTML5Platform : Platform
 		// ensure the ue4game binary exists, if applicable
 		if (!FileExists_NoExceptions(FullGameExePath))
 		{
-			Log("Failed to find game application " + FullGameExePath);
+			LogInformation("Failed to find game application " + FullGameExePath);
 			throw new AutomationException(ExitCode.Error_MissingExecutable, "Stage Failed. Could not find application {0}. You may need to build the UE4 project with your target configuration and platform.", FullGameExePath);
 		}
 
@@ -196,7 +196,7 @@ public class HTML5Platform : Platform
 
 		if (Compressed)
 		{
-			Log("Build configuration is " + Params.ClientConfigsToBuild[0].ToString() + ", so (gzip) compressing files for web servers.");
+			LogInformation("Build configuration is " + Params.ClientConfigsToBuild[0].ToString() + ", so (gzip) compressing files for web servers.");
 
 			// Compress all files. These are independent tasks which can be threaded.
 			List<Task> CompressionTasks = new List<Task>();
@@ -222,7 +222,7 @@ public class HTML5Platform : Platform
 		}
 		else
 		{
-			Log("Build configuration is " + Params.ClientConfigsToBuild[0].ToString() + ", so not compressing. Build Shipping configuration to compress files to save space.");
+			LogInformation("Build configuration is " + Params.ClientConfigsToBuild[0].ToString() + ", so not compressing. Build Shipping configuration to compress files to save space.");
 
 			// nuke old compressed files to prevent using stale files
 			File.Delete(FinalDataLocation + "gz");
@@ -240,7 +240,7 @@ public class HTML5Platform : Platform
 
 	void CompressFile(string Source, string Destination)
 	{
-		Log(" Compressing " + Source);
+		LogInformation(" Compressing " + Source);
 		bool DeleteSource = false;
 
 		if(  Source == Destination )
@@ -651,7 +651,7 @@ public class HTML5Platform : Platform
 
 		if ( !AmazonIdentity )
 		{
-			Log("Amazon S3 Incorrectly configured");
+			LogInformation("Amazon S3 Incorrectly configured");
 			return;
 		}
 
@@ -682,13 +682,13 @@ public class HTML5Platform : Platform
 		{
 			msTimeOut = 100*1000;
 		}
-		Log("Upload Timeout set to: " + (msTimeOut/1000) + "secs");
+		LogInformation("Upload Timeout set to: " + (msTimeOut/1000) + "secs");
 		Task.WaitAll(UploadTasks.ToArray(), (int)msTimeOut); // set timeout [miliseconds]
 
 		string URL = "https://" + BucketName + ".s3.amazonaws.com/" + FolderName + "/" + OutputFilename;
-		Log("Your project's shareable link is: " + URL);
+		LogInformation("Your project's shareable link is: " + URL);
 
-		Log("Upload Tasks finished.");
+		LogInformation("Upload Tasks finished.");
 	}
 
 	private static IDictionary<string, string> MimeTypeMapping = new Dictionary<string, string>(StringComparer.InvariantCultureIgnoreCase)
@@ -713,7 +713,7 @@ public class HTML5Platform : Platform
 		// "AWS Signature Version 4"
 		// http://docs.aws.amazon.com/AmazonS3/latest/API/sig-v4-header-based-auth.html
 		// --------------------------------------------------
-		Log(" Uploading " + Info.Name);
+		LogInformation(" Uploading " + Info.Name);
 
 		// --------------------------------------------------
 		// http://docs.aws.amazon.com/AmazonS3/latest/API/sigv4-post-example.html
@@ -805,7 +805,7 @@ public class HTML5Platform : Platform
 		var response = await httpClient.PostAsync(URL, formData);
 		if (response.IsSuccessStatusCode)
 		{
-			Log("Upload done: " + Info.Name);
+			LogInformation("Upload done: " + Info.Name);
 		}
 		else
 		{
diff --git a/Engine/Source/Programs/AutomationTool/IOS/IOSPlatform.Automation.cs b/Engine/Source/Programs/AutomationTool/IOS/IOSPlatform.Automation.cs
--- a/Engine/Source/Programs/AutomationTool/IOS/IOSPlatform.Automation.cs
+++ b/Engine/Source/Programs/AutomationTool/IOS/IOSPlatform.Automation.cs
@@ -246,7 +246,7 @@ public class IOSPlatform : Platform
 
 	public override void Package(ProjectParams Params, DeploymentContext SC, int WorkingCL)
 	{
-		Log("Package {0}", Params.RawProjectPath);
+		LogInformation("Package {0}", Params.RawProjectPath);
 
 		// ensure the ue4game binary exists, if applicable
 #if !PLATFORM_MAC
diff --git a/Engine/Source/Programs/AutomationTool/Lumin/LuminPlatform.Automation.cs b/Engine/Source/Programs/AutomationTool/Lumin/LuminPlatform.Automation.cs
--- a/Engine/Source/Programs/AutomationTool/Lumin/LuminPlatform.Automation.cs
+++ b/Engine/Source/Programs/AutomationTool/Lumin/LuminPlatform.Automation.cs
@@ -316,7 +316,7 @@ public class LuminPlatform : Platform
 				if (PluginExtras.FirstOrDefault(x => x == PluginPath) == null)
 				{
 					PluginExtras.Add(PluginPath);
-					Log("LuminPlugin: {0}", PluginPath);
+					LogInformation("LuminPlugin: {0}", PluginPath);
 				}
 			}
 		}
@@ -328,7 +328,7 @@ public class LuminPlatform : Platform
 		string[] BatchLines = null;
 		if (HostPlatform.Current.HostEditorPlatform == UnrealTargetPlatform.Win64)
 		{
-			Log("Writing bat for install");
+			LogInformation("Writing bat for install");
 			BatchLines = new string[] {
 				"@echo off",
 				"setlocal",
@@ -353,7 +353,7 @@ public class LuminPlatform : Platform
 		}
 		else
 		{
-			Log("Writing shell for install");
+			LogInformation("Writing shell for install");
 			BatchLines = new string[] {
 				"#!/bin/sh",
 				"MLSDK_ROOT=$MLSDK",
@@ -386,7 +386,7 @@ public class LuminPlatform : Platform
 		string[] BatchLines = null;
 		if (HostPlatform.Current.HostEditorPlatform == UnrealTargetPlatform.Win64)
 		{
-			Log("Writing bat for uninstall");
+			LogInformation("Writing bat for uninstall");
 			BatchLines = new string[] {
 				"@echo off",
 				"setlocal",
@@ -403,7 +403,7 @@ public class LuminPlatform : Platform
 		}
 		else
 		{
-			Log("Writing shell for uninstall");
+			LogInformation("Writing shell for uninstall");
 			BatchLines = new string[] {
 				"#!/bin/sh",
 				"MLSDK_ROOT=$MLSDK",
@@ -490,13 +490,13 @@ public class LuminPlatform : Platform
 				// if we don't need to deploy any NonUFS files, and the exe is up to date, then there's no need to waste time packaging!
 				if (Lines.Length > 0 && !string.IsNullOrWhiteSpace(Lines[0]))
 				{
-					Log("Need minimal package because {0} NonUFS files needed to be staged", Lines.Length);
+					LogInformation("Need minimal package because {0} NonUFS files needed to be staged", Lines.Length);
 					return false;
 				}
 			}
 			else
 			{
-				Log("Need minimal package because delta staging file {0} wasn't pulled from device", NonUFSManifestPath);
+				LogInformation("Need minimal package because delta staging file {0} wasn't pulled from device", NonUFSManifestPath);
 				return false;
 			}
 
@@ -509,17 +509,17 @@ public class LuminPlatform : Platform
 
 				if (ExeTimestamp > DeployedTimestamp)
 				{
-					Log("Need minimal package because exe timestamp {0:O} is newer than deployed timestamp {1:O}", ExeTimestamp, DeployedTimestamp);
+					LogInformation("Need minimal package because exe timestamp {0:O} is newer than deployed timestamp {1:O}", ExeTimestamp, DeployedTimestamp);
 					return false;
 				}
 				else
 				{
-					Log("Exe up to date: Exe is {0:O} / {1}, Deployed is {2:O}", ExeTimestamp, SC.StageExecutables[0], DeployedTimestamp);
+					LogInformation("Exe up to date: Exe is {0:O} / {1}, Deployed is {2:O}", ExeTimestamp, SC.StageExecutables[0], DeployedTimestamp);
 				}
 			}
 			else
 			{
-				Log("Need minimal package because exe timestamp file {0} wasn't pulled from device", ExeTimestampFileName);
+				LogInformation("Need minimal package because exe timestamp file {0} wasn't pulled from device", ExeTimestampFileName);
 				return false;
 			}
 
@@ -767,7 +767,7 @@ public class LuminPlatform : Platform
 			// @todo Lumin support iterative deploy! and packaging for iterative deploy
 			if (Params.IterativeDeploy)
 			{
-				Log("ITERATIVE DEPLOY..");
+				LogInformation("ITERATIVE DEPLOY..");
 				// always send UE4CommandLine.txt (it was written above after delta checks applied)
 				EntriesToDeploy.Add(IntermediateCmdLineFile);
 
@@ -787,7 +787,7 @@ public class LuminPlatform : Platform
 				}
 				else
 				{
-					Log("Unable to read delta file {0}", NonUFSManifestPath);
+					LogInformation("Unable to read delta file {0}", NonUFSManifestPath);
 				}
 
 				// Add UFS files if any to deploy
@@ -806,7 +806,7 @@ public class LuminPlatform : Platform
 				}
 				else
 				{
-					Log("Unable to read delta file {0}", UFSManifestPath);
+					LogInformation("Unable to read delta file {0}", UFSManifestPath);
 				}
 
 				//// For now, if too many files may be better to just push them all
@@ -874,7 +874,7 @@ public class LuminPlatform : Platform
 			}
 			else
 			{
-				Log("CLEAN DEPLOY..");
+				LogInformation("CLEAN DEPLOY..");
 
 				CleanInstall(Params, SC, DeviceName);
 			}
@@ -1024,7 +1024,7 @@ public class LuminPlatform : Platform
 			TimeSpan DeltaRunTime = DateTime.Now - StartTime;
 			if ((DeltaRunTime.TotalSeconds > TimeOutSeconds) && (TimeOutSeconds != 0))
 			{
-				Log("Device: " + DeviceName + " timed out while waiting for run to finish");
+				LogInformation("Device: " + DeviceName + " timed out while waiting for run to finish");
 				break;
 			}
 		}
diff --git a/Engine/Source/Programs/AutomationTool/Mac/MacPlatform.Automation.cs b/Engine/Source/Programs/AutomationTool/Mac/MacPlatform.Automation.cs
--- a/Engine/Source/Programs/AutomationTool/Mac/MacPlatform.Automation.cs
+++ b/Engine/Source/Programs/AutomationTool/Mac/MacPlatform.Automation.cs
@@ -443,22 +443,22 @@ public class MacPlatform : Platform
 		{
 			// Sign everything we built
 			List<FileReference> FilesToSign = GetExecutableNames(SC);
-			Log("RuntimeProjectRootDir: " + SC.RuntimeProjectRootDir);
+			LogInformation("RuntimeProjectRootDir: " + SC.RuntimeProjectRootDir);
 			foreach (var Exe in FilesToSign)
 			{
-				Log("Signing: " + Exe);
+				LogInformation("Signing: " + Exe);
 				string AppBundlePath = "";
 				if (Exe.IsUnderDirectory(DirectoryReference.Combine(SC.RuntimeProjectRootDir, "Binaries", SC.PlatformDir)))
 				{
-					Log("Starts with Binaries");
+					LogInformation("Starts with Binaries");
 					AppBundlePath = CombinePaths(SC.RuntimeProjectRootDir.FullName, "Binaries", SC.PlatformDir, Path.GetFileNameWithoutExtension(Exe.FullName) + ".app");
 				}
 				else if (Exe.IsUnderDirectory(DirectoryReference.Combine(SC.RuntimeRootDir, "Engine/Binaries", SC.PlatformDir)))
 				{
-					Log("Starts with Engine/Binaries");
+					LogInformation("Starts with Engine/Binaries");
 					AppBundlePath = CombinePaths("Engine/Binaries", SC.PlatformDir, Path.GetFileNameWithoutExtension(Exe.FullName) + ".app");
 				}
-				Log("Signing: " + AppBundlePath);
+				LogInformation("Signing: " + AppBundlePath);
 				CodeSign.SignMacFileOrFolder(AppBundlePath);
 			}
 
diff --git a/Engine/Source/Programs/AutomationTool/OneSkyLocalization/OneSkyConfigHelper.cs b/Engine/Source/Programs/AutomationTool/OneSkyLocalization/OneSkyConfigHelper.cs
--- a/Engine/Source/Programs/AutomationTool/OneSkyLocalization/OneSkyConfigHelper.cs
+++ b/Engine/Source/Programs/AutomationTool/OneSkyLocalization/OneSkyConfigHelper.cs
@@ -75,8 +75,8 @@ namespace EpicGames.OneSkyLocalization.Config
 
 		public void SpewValues()
 		{
-			CommandUtils.Log("Name : {0}", Name);
-			CommandUtils.Log("ApiKey : {0}", ApiKey);
+			CommandUtils.LogInformation("Name : {0}", Name);
+			CommandUtils.LogInformation("ApiKey : {0}", ApiKey);
 			//CommandUtils.LogConsole("ApiSecret : {0}", ApiSecret);  // This should probably never be revealed.
 		}
 	}
diff --git a/Engine/Source/Programs/AutomationTool/Scripts/AnalyzeThirdPartyLibs.Automation.cs b/Engine/Source/Programs/AutomationTool/Scripts/AnalyzeThirdPartyLibs.Automation.cs
--- a/Engine/Source/Programs/AutomationTool/Scripts/AnalyzeThirdPartyLibs.Automation.cs
+++ b/Engine/Source/Programs/AutomationTool/Scripts/AnalyzeThirdPartyLibs.Automation.cs
@@ -127,7 +127,7 @@ class AnalyzeThirdPartyLibs : BuildCommand
 
 	public override void ExecuteBuild()
 	{
-		Log("************************* Analyze Third Party Libs");
+		LogInformation("************************* Analyze Third Party Libs");
 
 		// figure out what batch/script to run
 		switch (UnrealBuildTool.BuildHostPlatform.Current.Platform)
@@ -194,19 +194,19 @@ class AnalyzeThirdPartyLibs : BuildCommand
 
 			long Size = Info.GetSize(Platforms);
 
-			Log("Library {0} is {1}", Lib, ToMegabytes(Size));
+			LogInformation("Library {0} is {1}", Lib, ToMegabytes(Size));
 
 			long Total = 0;
 			for (int Index = 0; Index < Platforms.Count; ++Index)
 			{
 				PlatformLibraryInfo Platform = Platforms[Index];
 				long Growth = Platform.TotalSize - LastSizes[Index];
-				Log("  {0} is {1}", Platform.PlatformName, ToMegabytes(Growth));
+				LogInformation("  {0} is {1}", Platform.PlatformName, ToMegabytes(Growth));
 
 				LastSizes[Index] = Platform.TotalSize;
 				Total += Growth;
 			}
-			Log("  Platform neutral is probably {0} (specific sum {1})", ToMegabytes(Size - Total), ToMegabytes(Total));
+			LogInformation("  Platform neutral is probably {0} (specific sum {1})", ToMegabytes(Size - Total), ToMegabytes(Total));
 
 			TotalSize += Size;
 		}
@@ -217,19 +217,19 @@ class AnalyzeThirdPartyLibs : BuildCommand
 		LargeFileExtensions.AddRange(new string[] { ".pdb", ".a", ".lib", ".dll", ".dylib", ".bc", ".so" });
 
 		// Hackery, look for big files (re-traverses everything)
-		Log("----");
+		LogInformation("----");
 		foreach (string Lib in LibsToEvaluate)
 		{
 			ThirdPartyLibraryInfo Info = new ThirdPartyLibraryInfo(Lib);
 			Info.FindLargeFiles(LargeFileExtensions, 1024 * 1024);
 		}
 
-		Log("----");
+		LogInformation("----");
 		foreach (var Platform in Platforms)
 		{
-			Log("  {0} is {1} (estimate)", Platform.PlatformName, ToMegabytes(Platform.TotalSize));
+			LogInformation("  {0} is {1} (estimate)", Platform.PlatformName, ToMegabytes(Platform.TotalSize));
 		}
-		Log("  OVERALL is {0} (accurate)", ToMegabytes(TotalSize));
+		LogInformation("  OVERALL is {0} (accurate)", ToMegabytes(TotalSize));
 
 		// undo the LibDir push
 		CommandUtils.PopDir();
diff --git a/Engine/Source/Programs/AutomationTool/Scripts/ArchiveCommand.Automation.cs b/Engine/Source/Programs/AutomationTool/Scripts/ArchiveCommand.Automation.cs
--- a/Engine/Source/Programs/AutomationTool/Scripts/ArchiveCommand.Automation.cs
+++ b/Engine/Source/Programs/AutomationTool/Scripts/ArchiveCommand.Automation.cs
@@ -58,9 +58,9 @@ public partial class Project : CommandUtils
 			return;
 		}
 
-		Log("********** ARCHIVE COMMAND STARTED **********");
+		LogInformation("********** ARCHIVE COMMAND STARTED **********");
 
-		Log("Archiving to {0}", Params.ArchiveDirectoryParam);
+		LogInformation("Archiving to {0}", Params.ArchiveDirectoryParam);
 
 		if (!Params.NoClient)
 		{
@@ -84,7 +84,7 @@ public partial class Project : CommandUtils
 				SC.StageTargetPlatform.ProcessArchivedProject(Params, SC);
 			}
 		}
-		Log("********** ARCHIVE COMMAND COMPLETED **********");
+		LogInformation("********** ARCHIVE COMMAND COMPLETED **********");
     }
 
 	#endregion
diff --git a/Engine/Source/Programs/AutomationTool/Scripts/BlameKeyword.Automation.cs b/Engine/Source/Programs/AutomationTool/Scripts/BlameKeyword.Automation.cs
--- a/Engine/Source/Programs/AutomationTool/Scripts/BlameKeyword.Automation.cs
+++ b/Engine/Source/Programs/AutomationTool/Scripts/BlameKeyword.Automation.cs
@@ -45,7 +45,7 @@ class BlameKeyword : BuildCommand
 
 	public override void ExecuteBuild()
 	{
-		Log("************************** BlameKeyword");
+		LogInformation("************************** BlameKeyword");
 
 		var Extensions = new HashSet<string>(KnownExtensions, StringComparer.InvariantCultureIgnoreCase);
 		var PathToSearch = ParseParamValue("Path", CombinePaths(CmdEnv.LocalRoot, "Engine", "Source", "Runtime"));
@@ -58,10 +58,10 @@ class BlameKeyword : BuildCommand
 			throw new AutomationException("No keywords to look for");
 		}
 
-		Log("Path={0}", PathToSearch);
-		Log("Keywords={0}", String.Join(",", Keywords));
-		Log("Timelapse={0}", bUseTimelapse);
-		Log("Out={0}", OutFilename);
+		LogInformation("Path={0}", PathToSearch);
+		LogInformation("Keywords={0}", String.Join(",", Keywords));
+		LogInformation("Timelapse={0}", bUseTimelapse);
+		LogInformation("Out={0}", OutFilename);
 
 		// Find files to search
 		var FilenamesToSearch = new List<string>();
@@ -131,7 +131,7 @@ class BlameKeyword : BuildCommand
 		}
 
 		// Write out the csv
-		Log("Creating blame csv: {0}", OutFilename);
+		LogInformation("Creating blame csv: {0}", OutFilename);
 		WriteAllText(OutFilename, BlameCSV.ToString());
 	}
 }
diff --git a/Engine/Source/Programs/AutomationTool/Scripts/BuildCommonTools.Automation.cs b/Engine/Source/Programs/AutomationTool/Scripts/BuildCommonTools.Automation.cs
--- a/Engine/Source/Programs/AutomationTool/Scripts/BuildCommonTools.Automation.cs
+++ b/Engine/Source/Programs/AutomationTool/Scripts/BuildCommonTools.Automation.cs
@@ -16,7 +16,7 @@ public class BuildCommonTools : BuildCommand
 {
 	public override void ExecuteBuild()
 	{
-		Log("************************* BuildCommonTools");
+		LogInformation("************************* BuildCommonTools");
 
 		// Get the list of platform names
 		string[] PlatformNames = ParseParamValue("platforms", BuildHostPlatform.Current.Platform.ToString()).Split('+');
@@ -179,10 +179,10 @@ public class ZipProjectUp : BuildCommand
 		string InstallDirectory = ParseParamValue("install", "");
 		ProjectDirectory = Path.GetDirectoryName(ProjectDirectory);
 
-		Log("Started zipping project up");
-		Log("Project directory: {0}", ProjectDirectory);
-		Log("Install directory: {0}", InstallDirectory);
-		Log("Packaging up the project...");
+		LogInformation("Started zipping project up");
+		LogInformation("Project directory: {0}", ProjectDirectory);
+		LogInformation("Install directory: {0}", InstallDirectory);
+		LogInformation("Packaging up the project...");
 
 		// Setup filters
 		FileFilter Filter = new FileFilter();
@@ -193,6 +193,6 @@ public class ZipProjectUp : BuildCommand
 
 		ZipFiles(new FileReference(InstallDirectory), new DirectoryReference(ProjectDirectory), Filter);
 
-		Log("Completed zipping project up");
+		LogInformation("Completed zipping project up");
 	}
 }
diff --git a/Engine/Source/Programs/AutomationTool/Scripts/BuildCookRun.Automation.cs b/Engine/Source/Programs/AutomationTool/Scripts/BuildCookRun.Automation.cs
--- a/Engine/Source/Programs/AutomationTool/Scripts/BuildCookRun.Automation.cs
+++ b/Engine/Source/Programs/AutomationTool/Scripts/BuildCookRun.Automation.cs
@@ -47,7 +47,7 @@ public class BuildCookRun : BuildCommand
 
 	protected ProjectParams SetupParams()
 	{
-		Log("Setting up ProjectParams for {0}", ProjectPath);
+		LogInformation("Setting up ProjectParams for {0}", ProjectPath);
 
 		var Params = new ProjectParams
 		(
@@ -126,13 +126,13 @@ public class BuildCookRun : BuildCommand
 	private string GetDefaultMap(ProjectParams Params)
 	{
 		const string EngineEntryMap = "/Engine/Maps/Entry";
-		Log("Trying to find DefaultMap in ini files");
+		LogInformation("Trying to find DefaultMap in ini files");
 		string DefaultMap = null;
 		var ProjectFolder = GetDirectoryName(Params.RawProjectPath.FullName);
 		var DefaultGameEngineConfig = CombinePaths(ProjectFolder, "Config", "DefaultEngine.ini");
 		if (FileExists(DefaultGameEngineConfig))
 		{
-			Log("Looking for DefaultMap in {0}", DefaultGameEngineConfig);
+			LogInformation("Looking for DefaultMap in {0}", DefaultGameEngineConfig);
 			DefaultMap = GetDefaultMapFromIni(DefaultGameEngineConfig, Params.DedicatedServer);
 			if (DefaultMap == null && Params.DedicatedServer)
 			{
@@ -144,7 +144,7 @@ public class BuildCookRun : BuildCommand
 			var BaseEngineConfig = CombinePaths(CmdEnv.LocalRoot, "Config", "BaseEngine.ini");
 			if (FileExists(BaseEngineConfig))
 			{
-				Log("Looking for DefaultMap in {0}", BaseEngineConfig);
+				LogInformation("Looking for DefaultMap in {0}", BaseEngineConfig);
 				DefaultMap = GetDefaultMapFromIni(BaseEngineConfig, Params.DedicatedServer);
 				if (DefaultMap == null && Params.DedicatedServer)
 				{
@@ -155,12 +155,12 @@ public class BuildCookRun : BuildCommand
 		// We check for null here becase null == not found
 		if (DefaultMap == null)
 		{
-			Log("No DefaultMap found, assuming: {0}", EngineEntryMap);
+			LogInformation("No DefaultMap found, assuming: {0}", EngineEntryMap);
 			DefaultMap = EngineEntryMap;
 		}
 		else
 		{
-			Log("Found DefaultMap={0}", DefaultMap);
+			LogInformation("Found DefaultMap={0}", DefaultMap);
 		}
 		return DefaultMap;
 	}
@@ -243,7 +243,7 @@ public class BuildCookRun : BuildCommand
 		}
 
 		var Dest = ParseParamValue("ForeignDest", CombinePaths(@"C:\testue4\foreign\", DestSample + "_ _Dir"));
-		Log("Make a foreign sample {0} -> {1}", Src, Dest);
+		LogInformation("Make a foreign sample {0} -> {1}", Src, Dest);
 
 		CloneDirectory(Src, Dest);
 
@@ -269,7 +269,7 @@ public class BuildCookRun : BuildCommand
 		}
 
 		var Dest = ParseParamValue("ForeignDest", CombinePaths(@"C:\testue4\foreign\", DestSample + "_ _Dir"));
-		Log("Make a foreign sample {0} -> {1}", Src, Dest);
+		LogInformation("Make a foreign sample {0} -> {1}", Src, Dest);
 
 		CloneDirectory(Src, Dest);
 		DeleteDirectory_NoExceptions(CombinePaths(Dest, "Intermediate"));
diff --git a/Engine/Source/Programs/AutomationTool/Scripts/BuildDerivedDataCache.Automation.cs b/Engine/Source/Programs/AutomationTool/Scripts/BuildDerivedDataCache.Automation.cs
--- a/Engine/Source/Programs/AutomationTool/Scripts/BuildDerivedDataCache.Automation.cs
+++ b/Engine/Source/Programs/AutomationTool/Scripts/BuildDerivedDataCache.Automation.cs
@@ -76,7 +76,7 @@ public class BuildDerivedDataCache : BuildCommand
 				ProjectSpecificPlatforms = CommandUtils.CombineCommandletParams(FilteredPlatforms.Distinct().ToArray());
 			}
 
-			CommandUtils.Log("Generating DDC data for {0} on {1}", GameName, ProjectSpecificPlatforms);
+			CommandUtils.LogInformation("Generating DDC data for {0} on {1}", GameName, ProjectSpecificPlatforms);
 			CommandUtils.DDCCommandlet(FileRef, EditorExe, null, ProjectSpecificPlatforms, "-fill -DDC=CreateInstalledEnginePak -ProjectOnly");
 
 			string ProjectPakFile = CommandUtils.CombinePaths(Path.GetDirectoryName(OutputPakFile), String.Format("Compressed-{0}.ddp", GameName));
@@ -92,7 +92,7 @@ public class BuildDerivedDataCache : BuildCommand
 		}
 
 		// Generate DDC for the editor, and merge all the other PAK files in
-		CommandUtils.Log("Generating DDC data for engine content on {0}", TargetPlatforms);
+		CommandUtils.LogInformation("Generating DDC data for engine content on {0}", TargetPlatforms);
 		CommandUtils.DDCCommandlet(null, EditorExe, null, TargetPlatforms, "-fill -DDC=CreateInstalledEnginePak " + CommandUtils.MakePathSafeToUseWithCommandLine("-MergePaks=" + String.Join("+", ProjectPakFiles)));
 
 		string SavedPakFile = CommandUtils.CombinePaths(SavedDir, RelativePakPath);
diff --git a/Engine/Source/Programs/AutomationTool/Scripts/BuildPhysX.Automation.cs b/Engine/Source/Programs/AutomationTool/Scripts/BuildPhysX.Automation.cs
--- a/Engine/Source/Programs/AutomationTool/Scripts/BuildPhysX.Automation.cs
+++ b/Engine/Source/Programs/AutomationTool/Scripts/BuildPhysX.Automation.cs
@@ -98,7 +98,7 @@ class BuildPhysX : BuildCommand
 		}
 
 		Environment.SetEnvironmentVariable("CMAKE_ROOT", DirectoryReference.Combine(CMakeRootDirectory, "share").ToString());
-		Log("set {0}={1}", "CMAKE_ROOT", Environment.GetEnvironmentVariable("CMAKE_ROOT"));
+		LogInformation("set {0}={1}", "CMAKE_ROOT", Environment.GetEnvironmentVariable("CMAKE_ROOT"));
 
 		switch (TargetData.Platform)
 		{
@@ -590,9 +590,9 @@ class BuildPhysX : BuildCommand
 	{
 		// make sure we set up the environment variable specifying where the root of the PhysX SDK is
 		Environment.SetEnvironmentVariable("GW_DEPS_ROOT", PhysXSourceRootDirectory.ToString());
-		Log("set {0}={1}", "GW_DEPS_ROOT", Environment.GetEnvironmentVariable("GW_DEPS_ROOT"));
+		LogInformation("set {0}={1}", "GW_DEPS_ROOT", Environment.GetEnvironmentVariable("GW_DEPS_ROOT"));
 		Environment.SetEnvironmentVariable("CMAKE_MODULE_PATH", DirectoryReference.Combine(PhysXSourceRootDirectory, "Externals", "CMakeModules").ToString());
-		Log("set {0}={1}", "CMAKE_MODULE_PATH", Environment.GetEnvironmentVariable("CMAKE_MODULE_PATH"));
+		LogInformation("set {0}={1}", "CMAKE_MODULE_PATH", Environment.GetEnvironmentVariable("CMAKE_MODULE_PATH"));
 
 		string CMakeName = GetCMakeNameAndSetupEnv(TargetData);
 
@@ -608,7 +608,7 @@ class BuildPhysX : BuildCommand
 
 					if(!bCleanOnly)
 					{
-						Log("Generating projects for lib " + TargetLib.ToString() + ", " + TargetData.ToString());
+						LogInformation("Generating projects for lib " + TargetLib.ToString() + ", " + TargetData.ToString());
 
 						ProcessStartInfo StartInfo = new ProcessStartInfo();
 						StartInfo.FileName = CMakeName;
@@ -626,7 +626,7 @@ class BuildPhysX : BuildCommand
 
 					if (!bCleanOnly)
 					{
-						Log("Generating projects for lib " + TargetLib.ToString() + ", " + TargetData.ToString());
+						LogInformation("Generating projects for lib " + TargetLib.ToString() + ", " + TargetData.ToString());
 
 						ProcessStartInfo StartInfo = new ProcessStartInfo();
 						StartInfo.FileName = CMakeName;
@@ -648,7 +648,7 @@ class BuildPhysX : BuildCommand
 
 					if (!bCleanOnly)
 					{
-						Log("Generating projects for lib " + TargetLib.ToString() + ", " + TargetData.ToString());
+						LogInformation("Generating projects for lib " + TargetLib.ToString() + ", " + TargetData.ToString());
 
 						ProcessStartInfo StartInfo = new ProcessStartInfo();
 						StartInfo.FileName = CMakeName;
@@ -656,10 +656,10 @@ class BuildPhysX : BuildCommand
 						StartInfo.Arguments = GetCMakeArguments(TargetLib, TargetData, BuildConfig);
 
 						System.Console.WriteLine("Working in '{0}'", StartInfo.WorkingDirectory);
-						Log("Working in '{0}'", StartInfo.WorkingDirectory);
+						LogInformation("Working in '{0}'", StartInfo.WorkingDirectory);
 
 						System.Console.WriteLine("{0} {1}", StartInfo.FileName, StartInfo.Arguments);
-						Log("{0} {1}", StartInfo.FileName, StartInfo.Arguments);
+						LogInformation("{0} {1}", StartInfo.FileName, StartInfo.Arguments);
 						
 						if (RunLocalProcessAndLogOutput(StartInfo) != 0)
 						{
@@ -677,7 +677,7 @@ class BuildPhysX : BuildCommand
 
 					if (!bCleanOnly)
 					{
-						Log("Generating projects for lib " + TargetLib.ToString() + ", " + TargetData.ToString());
+						LogInformation("Generating projects for lib " + TargetLib.ToString() + ", " + TargetData.ToString());
 
 						ProcessStartInfo StartInfo = new ProcessStartInfo();
 						StartInfo.FileName = CMakeName;
@@ -710,7 +710,7 @@ class BuildPhysX : BuildCommand
 
 					if (!bCleanOnly)
 					{
-						Log("Generating projects for lib " + TargetLib.ToString() + ", " + TargetData.ToString());
+						LogInformation("Generating projects for lib " + TargetLib.ToString() + ", " + TargetData.ToString());
 
 						// CMAKE_TOOLCHAIN_FILE
 						Environment.SetEnvironmentVariable("LIB_SUFFIX", GetConfigurationSuffix(BuildConfig, TargetData)); // only used in HTML5's CMakefiles
@@ -728,8 +728,8 @@ class BuildPhysX : BuildCommand
 						StartInfo.WorkingDirectory = CMakeTargetDirectory.ToString();
 						StartInfo.Arguments = "\"" + HTML5SDKInfo.EMSCRIPTEN_ROOT + "\\emcmake\" cmake " + GetCMakeArguments(TargetLib, TargetData, BuildConfig);
 
-						Log("Working in: {0}", StartInfo.WorkingDirectory);
-						Log("{0} {1}", StartInfo.FileName, StartInfo.Arguments);
+						LogInformation("Working in: {0}", StartInfo.WorkingDirectory);
+						LogInformation("{0} {1}", StartInfo.FileName, StartInfo.Arguments);
 						
 						if (RunLocalProcessAndLogOutput(StartInfo) != 0)
 						{
@@ -745,7 +745,7 @@ class BuildPhysX : BuildCommand
 
 					if (!bCleanOnly)
 					{
-						Log("Generating projects for lib " + TargetLib.ToString() + ", " + TargetData.ToString());
+						LogInformation("Generating projects for lib " + TargetLib.ToString() + ", " + TargetData.ToString());
 
 						ProcessStartInfo StartInfo = new ProcessStartInfo();
 						StartInfo.FileName = CMakeName;
@@ -800,7 +800,7 @@ class BuildPhysX : BuildCommand
 				if (File.Exists(PathComponent + "/make.exe") || File.Exists(PathComponent + "make.exe"))
 				{
 					// gotcha!
-					Log("Removing {0} from PATH since it contains possibly colliding make.exe", PathComponent);
+					LogInformation("Removing {0} from PATH since it contains possibly colliding make.exe", PathComponent);
 					continue;
 				}
 			}
@@ -835,7 +835,7 @@ class BuildPhysX : BuildCommand
 				string PathWithoutCygwin = RemoveOtherMakeFromPath(PrevPath);
 				Environment.SetEnvironmentVariable("PATH", CMakePath + ";" + MakePath + ";" + PathWithoutCygwin);
 				Environment.SetEnvironmentVariable("PATH", CMakePath + ";" + MakePath + ";" + Environment.GetEnvironmentVariable("PATH"));
-				Log("set {0}={1}", "PATH", Environment.GetEnvironmentVariable("PATH"));
+				LogInformation("set {0}={1}", "PATH", Environment.GetEnvironmentVariable("PATH"));
 			}
 		}
 	}
@@ -1044,8 +1044,8 @@ class BuildPhysX : BuildCommand
 			StartInfo.WorkingDirectory = ConfigDirectory.ToString();
 			StartInfo.Arguments = MakeOptions;
 
-			Log("Working in: {0}", StartInfo.WorkingDirectory);
-			Log("{0} {1}", StartInfo.FileName, StartInfo.Arguments);
+			LogInformation("Working in: {0}", StartInfo.WorkingDirectory);
+			LogInformation("{0} {1}", StartInfo.FileName, StartInfo.Arguments);
 
 			if (RunLocalProcessAndLogOutput(StartInfo) != 0)
 			{
@@ -1403,7 +1403,7 @@ class BuildPhysX : BuildCommand
 				StartInfo.Arguments = SOFile.FullName + " " + PSymbolFile.ToString();
 				StartInfo.RedirectStandardError = true;
 
-				Log("Running: '{0} {1}'", StartInfo.FileName, StartInfo.Arguments);
+				LogInformation("Running: '{0} {1}'", StartInfo.FileName, StartInfo.Arguments);
 
 				RunLocalProcessAndLogOutput(StartInfo);
 
@@ -1411,7 +1411,7 @@ class BuildPhysX : BuildCommand
 				StartInfo.FileName = BreakpadSymbolEncoderPath.ToString() + ExeSuffix;
 				StartInfo.Arguments = PSymbolFile.ToString() + " " + SymbolFile.ToString();
 
-				Log("Running: '{0} {1}'", StartInfo.FileName, StartInfo.Arguments);
+				LogInformation("Running: '{0} {1}'", StartInfo.FileName, StartInfo.Arguments);
 
 				RunLocalProcessAndLogOutput(StartInfo);
 
@@ -1729,13 +1729,13 @@ class BuildPhysX : BuildCommand
 		{
 			if (!P4.TryDeleteEmptyChange(P4ChangeList))
 			{
-				Log("Submitting changelist " + P4ChangeList.ToString());
+				LogInformation("Submitting changelist " + P4ChangeList.ToString());
 				int SubmittedChangeList = InvalidChangeList;
 				P4.Submit(P4ChangeList, out SubmittedChangeList);
 			}
 			else
 			{
-				Log("Nothing to submit!");
+				LogInformation("Nothing to submit!");
 			}
 		}
 	}
diff --git a/Engine/Source/Programs/AutomationTool/Scripts/BuildPluginCommand.Automation.cs b/Engine/Source/Programs/AutomationTool/Scripts/BuildPluginCommand.Automation.cs
--- a/Engine/Source/Programs/AutomationTool/Scripts/BuildPluginCommand.Automation.cs
+++ b/Engine/Source/Programs/AutomationTool/Scripts/BuildPluginCommand.Automation.cs
@@ -83,7 +83,7 @@ class BuildPlugin : BuildCommand
 		FileReference HostProjectPluginFile = CreateHostProject(HostProjectFile, PluginFile);
 
 		// Read the plugin
-		CommandUtils.Log("Reading plugin from {0}...", HostProjectPluginFile);
+		CommandUtils.LogInformation("Reading plugin from {0}...", HostProjectPluginFile);
 		PluginDescriptor Plugin = PluginDescriptor.FromFile(HostProjectPluginFile);
 
 		// Compile the plugin for all the target platforms
@@ -125,7 +125,7 @@ class BuildPlugin : BuildCommand
 		// Build the host platforms
 		if(HostPlatforms.Count > 0)
 		{
-			CommandUtils.Log("Building plugin for host platforms: {0}", String.Join(", ", HostPlatforms));
+			CommandUtils.LogInformation("Building plugin for host platforms: {0}", String.Join(", ", HostPlatforms));
 			foreach (UnrealTargetPlatform HostPlatform in HostPlatforms)
 			{
 				if (Plugin.bCanBeUsedWithUnrealHeaderTool)
@@ -139,7 +139,7 @@ class BuildPlugin : BuildCommand
 		// Add the game targets
 		if(TargetPlatforms.Count > 0)
 		{
-			CommandUtils.Log("Building plugin for target platforms: {0}", String.Join(", ", TargetPlatforms));
+			CommandUtils.LogInformation("Building plugin for target platforms: {0}", String.Join(", ", TargetPlatforms));
 			foreach (UnrealTargetPlatform TargetPlatform in TargetPlatforms)
 			{
 				if(Plugin.SupportsTargetPlatform(TargetPlatform))
@@ -245,7 +245,7 @@ class BuildPlugin : BuildCommand
 		FileReference FilterFile = FileReference.Combine(PluginFile.Directory, "Config", "FilterPlugin.ini");
 		if(FileReference.Exists(FilterFile))
 		{
-			CommandUtils.Log("Reading filter rules from {0}", FilterFile);
+			CommandUtils.LogInformation("Reading filter rules from {0}", FilterFile);
 			Filter.ReadRulesFromFile(FilterFile, "FilterPlugin");
 		}
 
diff --git a/Engine/Source/Programs/AutomationTool/Scripts/BuildProjectCommand.Automation.cs b/Engine/Source/Programs/AutomationTool/Scripts/BuildProjectCommand.Automation.cs
--- a/Engine/Source/Programs/AutomationTool/Scripts/BuildProjectCommand.Automation.cs
+++ b/Engine/Source/Programs/AutomationTool/Scripts/BuildProjectCommand.Automation.cs
@@ -67,7 +67,7 @@ public partial class Project : CommandUtils
 			return;
 		}
 
-		Log("********** BUILD COMMAND STARTED **********");
+		LogInformation("********** BUILD COMMAND STARTED **********");
 
 		var UE4Build = new UE4Build(Command);
 		var Agenda = new UE4Build.BuildAgenda();
@@ -220,7 +220,7 @@ public partial class Project : CommandUtils
 			UE4Build.AddBuildProductsToChangelist(WorkingCL, UE4Build.BuildProductFiles);
 		}
 
-		Log("********** BUILD COMMAND COMPLETED **********");
+		LogInformation("********** BUILD COMMAND COMPLETED **********");
 	}
 
 	#endregion
diff --git a/Engine/Source/Programs/AutomationTool/Scripts/BuildThirdPartyLibs.Automation.cs b/Engine/Source/Programs/AutomationTool/Scripts/BuildThirdPartyLibs.Automation.cs
--- a/Engine/Source/Programs/AutomationTool/Scripts/BuildThirdPartyLibs.Automation.cs
+++ b/Engine/Source/Programs/AutomationTool/Scripts/BuildThirdPartyLibs.Automation.cs
@@ -63,7 +63,7 @@ class BuildThirdPartyLibs : BuildCommand
 
 	public override void ExecuteBuild()
 	{
-		Log("************************* Build Third Party Libs");
+		LogInformation("************************* Build Third Party Libs");
 
 		// figure out what batch/script to run
 		string CompileScriptName;
@@ -90,7 +90,7 @@ class BuildThirdPartyLibs : BuildCommand
 		{
 			WorkingCL = P4.CreateChange(P4Env.Client, String.Format("Third party libs built from changelist {0}", P4Env.Changelist));
 		}
-		Log("Build from {0}    Working in {1}", P4Env.Changelist, WorkingCL);
+		LogInformation("Build from {0}    Working in {1}", P4Env.Changelist, WorkingCL);
 
 		// go to the third party lib dir
 		string SearchLibraryDir = ParseParamValue("SearchDir", DefaultLibraryDir);
@@ -142,7 +142,7 @@ class BuildThirdPartyLibs : BuildCommand
 		// now go through and run each batch file, 
 		foreach (string Lib in LibsToCompile)
 		{
-			Log("Building {0}", Lib);
+			LogInformation("Building {0}", Lib);
 
 			// go into the lib dir
 			CommandUtils.PushDir(Lib);
@@ -166,7 +166,7 @@ class BuildThirdPartyLibs : BuildCommand
 		{
 			int SubmittedCL;
 			P4.Submit(WorkingCL, out SubmittedCL, true, true);
-			Log("Submitted changelist {0}", SubmittedCL);
+			LogInformation("Submitted changelist {0}", SubmittedCL);
 		}
 	}
 }
diff --git a/Engine/Source/Programs/AutomationTool/Scripts/CheckCsprojDotNetVersion.cs b/Engine/Source/Programs/AutomationTool/Scripts/CheckCsprojDotNetVersion.cs
--- a/Engine/Source/Programs/AutomationTool/Scripts/CheckCsprojDotNetVersion.cs
+++ b/Engine/Source/Programs/AutomationTool/Scripts/CheckCsprojDotNetVersion.cs
@@ -18,7 +18,7 @@ class CheckCsprojDotNetVersion : BuildCommand
 		{
 			throw new AutomationException("-TargetVersion was not specified.");
 		}
-		CommandUtils.Log("Scanning for all csproj's...");
+		CommandUtils.LogInformation("Scanning for all csproj's...");
 		// Check for all csproj's in the engine dir
 		DirectoryReference EngineDir = CommandUtils.EngineDirectory;
 
@@ -64,7 +64,7 @@ class CheckCsprojDotNetVersion : BuildCommand
                 Match = PossibleAppConfigRegex.Match(Contents);
 				if (!Match.Success)
 				{
-					CommandUtils.Log("No TargetFrameworkVersion or TargetFrameworkProfile found for project {0}, is it a mono project? If not, does it compile properly?", CsProj);
+					CommandUtils.LogInformation("No TargetFrameworkVersion or TargetFrameworkProfile found for project {0}, is it a mono project? If not, does it compile properly?", CsProj);
 					continue;
 				}
 
@@ -73,7 +73,7 @@ class CheckCsprojDotNetVersion : BuildCommand
 				string Profile = Match.Groups[1].Value;
 				if (!FileReference.Exists(AppConfigFile))
 				{
-					CommandUtils.Log("Found TargetFrameworkProfile but no associated app.config containing the version for project {0}.", CsProj);
+					CommandUtils.LogInformation("Found TargetFrameworkProfile but no associated app.config containing the version for project {0}.", CsProj);
 					continue;
 				}
 
@@ -82,7 +82,7 @@ class CheckCsprojDotNetVersion : BuildCommand
                 Match = AppConfigRegex.Match(Contents);
 				if (!Match.Success)
 				{
-					CommandUtils.Log("Couldn't find a supportedRuntime match for the version in the app.config for project {0}.", CsProj);
+					CommandUtils.LogInformation("Couldn't find a supportedRuntime match for the version in the app.config for project {0}.", CsProj);
 					continue;
 				}
 
diff --git a/Engine/Source/Programs/AutomationTool/Scripts/CheckForHacks.cs b/Engine/Source/Programs/AutomationTool/Scripts/CheckForHacks.cs
--- a/Engine/Source/Programs/AutomationTool/Scripts/CheckForHacks.cs
+++ b/Engine/Source/Programs/AutomationTool/Scripts/CheckForHacks.cs
@@ -71,7 +71,7 @@ namespace AutomationTool
 			}
 
 			// Enumerate all the files in the workspace
-			Log("Finding files in workspace...");
+			LogInformation("Finding files in workspace...");
 
 			List<FileInfo> FilesToCheck = new List<FileInfo>();
 			using (ThreadPoolWorkQueue Queue = new ThreadPoolWorkQueue())
@@ -82,7 +82,7 @@ namespace AutomationTool
 			}
 
 			// Scan all of the files for invalid comments
-			Log("Scanning files...", FilesToCheck.Count);
+			LogInformation("Scanning files...", FilesToCheck.Count);
 			using (ThreadPoolWorkQueue Queue = new ThreadPoolWorkQueue())
 			{
 				foreach(FileInfo File in FilesToCheck)
@@ -94,7 +94,7 @@ namespace AutomationTool
 				{
 					lock(this)
 					{
-						Log("Scanning files... [{0}/{1}]", FilesToCheck.Count - Queue.NumRemaining, FilesToCheck.Count);
+						LogInformation("Scanning files... [{0}/{1}]", FilesToCheck.Count - Queue.NumRemaining, FilesToCheck.Count);
 					}
 				}
 			}
diff --git a/Engine/Source/Programs/AutomationTool/Scripts/CheckRestrictedFolders.cs b/Engine/Source/Programs/AutomationTool/Scripts/CheckRestrictedFolders.cs
--- a/Engine/Source/Programs/AutomationTool/Scripts/CheckRestrictedFolders.cs
+++ b/Engine/Source/Programs/AutomationTool/Scripts/CheckRestrictedFolders.cs
@@ -33,7 +33,7 @@ namespace AutomationTool
 			}
 
 			// Find all the folders which are problematic
-			CommandUtils.Log("Searching for folders under {0} named {1}...", BaseDir, String.Join(", ", RestrictedNames));
+			CommandUtils.LogInformation("Searching for folders under {0} named {1}...", BaseDir, String.Join(", ", RestrictedNames));
 			List<DirectoryInfo> ProblemFolders = new List<DirectoryInfo>();
 			FindRestrictedFolders(new DirectoryInfo(BaseDir.FullName), RestrictedNames, ProblemFolders);
 
@@ -49,7 +49,7 @@ namespace AutomationTool
 			}
 
 			// Otherwise return success
-			CommandUtils.Log("No restricted folders found under {0}", BaseDir);
+			CommandUtils.LogInformation("No restricted folders found under {0}", BaseDir);
 			return ExitCode.Success;
 		}
 
diff --git a/Engine/Source/Programs/AutomationTool/Scripts/CleanAutomationReports.cs b/Engine/Source/Programs/AutomationTool/Scripts/CleanAutomationReports.cs
--- a/Engine/Source/Programs/AutomationTool/Scripts/CleanAutomationReports.cs
+++ b/Engine/Source/Programs/AutomationTool/Scripts/CleanAutomationReports.cs
@@ -45,7 +45,7 @@ namespace AutomationTool
 			DateTime RetainTime = DateTime.UtcNow - TimeSpan.FromDays(DaysValue);
 
 			// Enumerate all the build directories
-			CommandUtils.Log("Scanning {0}...", ReportDir);
+			CommandUtils.LogInformation("Scanning {0}...", ReportDir);
 			int NumFolders = 0;
 			List<DirectoryInfo> FoldersToDelete = new List<DirectoryInfo>();
 			foreach (DirectoryInfo BuildDirectory in new DirectoryInfo(ReportDir).EnumerateDirectories())
@@ -56,14 +56,14 @@ namespace AutomationTool
 				}
 				NumFolders++;
 			}
-			CommandUtils.Log("Found {0} builds; {1} to delete.", NumFolders, FoldersToDelete.Count);
+			CommandUtils.LogInformation("Found {0} builds; {1} to delete.", NumFolders, FoldersToDelete.Count);
 
 			// Delete them all
 			for (int Idx = 0; Idx < FoldersToDelete.Count; Idx++)
 			{
 				try
 				{
-					CommandUtils.Log("[{0}/{1}] Deleting {2}...", Idx + 1, FoldersToDelete.Count, FoldersToDelete[Idx].FullName);
+					CommandUtils.LogInformation("[{0}/{1}] Deleting {2}...", Idx + 1, FoldersToDelete.Count, FoldersToDelete[Idx].FullName);
 					FoldersToDelete[Idx].Delete(true);
 				}
 				catch (Exception Ex)
diff --git a/Engine/Source/Programs/AutomationTool/Scripts/CodeSurgery.Automation.cs b/Engine/Source/Programs/AutomationTool/Scripts/CodeSurgery.Automation.cs
--- a/Engine/Source/Programs/AutomationTool/Scripts/CodeSurgery.Automation.cs
+++ b/Engine/Source/Programs/AutomationTool/Scripts/CodeSurgery.Automation.cs
@@ -11,7 +11,7 @@ class CodeSurgery : BuildCommand
 {
 	public override void ExecuteBuild()
 	{
-		Log("************************* Reworking stats code");
+		LogInformation("************************* Reworking stats code");
 
 		var Wildcards = new List<string>() { "*.h", "*.cpp", "*.inl" };
 
@@ -44,7 +44,7 @@ class CodeSurgery : BuildCommand
 						!ThisFile.Contains(@"\ThirdParty\")
 						)
 					{
-						Log("Source File: {0}", ThisFile);
+						LogInformation("Source File: {0}", ThisFile);
 						AllFiles.Add(ThisFile);
 					}
 				}
@@ -124,7 +124,7 @@ class CodeSurgery : BuildCommand
 
 						if (bEnumLine || bDeclareLine)
 						{
-							Log("{0} {1} Line: {2} : {3}", bEnumLine ? "Enum" : "Declare", TypeString, StatName, Line);
+							LogInformation("{0} {1} Line: {2} : {3}", bEnumLine ? "Enum" : "Declare", TypeString, StatName, Line);
 
 							if (bEnumLine)
 							{
@@ -189,32 +189,32 @@ class CodeSurgery : BuildCommand
 		}
 
 		var ToCheckOuts = new HashSet<string>();
-		Log("Stats *************************");
+		LogInformation("Stats *************************");
 		foreach (var AllGoodStat in AllGoodStats)
 		{
-			Log("{0}", AllGoodStat);
+			LogInformation("{0}", AllGoodStat);
 			ToCheckOuts.Add(DeclareFiles[AllGoodStat]);
 			ToCheckOuts.Add(EnumFiles[AllGoodStat]);
 		}
-		Log("Groups *************************");
+		LogInformation("Groups *************************");
 		foreach (var AllGoodGroup in AllGoodGroups)
 		{
-			Log("{0}", AllGoodGroup);
+			LogInformation("{0}", AllGoodGroup);
 			ToCheckOuts.Add(DeclareFiles[AllGoodGroup]);
 			ToCheckOuts.Add(EnumFiles[AllGoodGroup]);
 		}
-		Log("Broken *************************");
+		LogInformation("Broken *************************");
 		foreach (var BrokenItem in Broken)
 		{
-			Log("{0}", BrokenItem.Key);
+			LogInformation("{0}", BrokenItem.Key);
 		}
-		Log("*************************");
+		LogInformation("*************************");
 
 		int WorkingCL = -1;
 		if (P4Enabled)
 		{
 			WorkingCL = P4.CreateChange(P4Env.Client, "Stat code surgery");
-			Log("Working in {0}", WorkingCL);
+			LogInformation("Working in {0}", WorkingCL);
 		}
 		else
 		{
@@ -229,20 +229,20 @@ class CodeSurgery : BuildCommand
 				CheckedOuts.Add(ToCheckOut);
 			}
 		}
-		Log("Checked Out *************************");
+		LogInformation("Checked Out *************************");
 		foreach (var CheckedOut in CheckedOuts)
 		{
-			Log("{0}", CheckedOut);
+			LogInformation("{0}", CheckedOut);
 		}
-		Log("Failed to check out *************************");
+		LogInformation("Failed to check out *************************");
 		foreach (var ToCheckOut in ToCheckOuts)
 		{
 			if (!CheckedOuts.Contains(ToCheckOut))
 			{
-				Log("{0}", ToCheckOut);
+				LogInformation("{0}", ToCheckOut);
 			}
 		}
-		Log("*************************");
+		LogInformation("*************************");
 		foreach (var AllGoodStat in AllGoodStats)
 		{
 			if (EnumFiles[AllGoodStat].EndsWith(".cpp", StringComparison.InvariantCultureIgnoreCase))
@@ -286,7 +286,7 @@ class CodeSurgery : BuildCommand
 			WriteAllText(EnumFiles[AllGoodGroup], EnumFileText);
 		}
 
-		Log("*************************");
+		LogInformation("*************************");
 
 	}
 }
diff --git a/Engine/Source/Programs/AutomationTool/Scripts/CookCommand.Automation.cs b/Engine/Source/Programs/AutomationTool/Scripts/CookCommand.Automation.cs
--- a/Engine/Source/Programs/AutomationTool/Scripts/CookCommand.Automation.cs
+++ b/Engine/Source/Programs/AutomationTool/Scripts/CookCommand.Automation.cs
@@ -31,7 +31,7 @@ public partial class Project : CommandUtils
 		}
 		Params.ValidateAndLog();
 
-		Log("********** COOK COMMAND STARTED **********");
+		LogInformation("********** COOK COMMAND STARTED **********");
 
 		string UE4EditorExe = HostPlatform.Current.GetUE4ExePath(Params.UE4Exe);
 		if (!FileExists(UE4EditorExe))
@@ -68,7 +68,7 @@ public partial class Project : CommandUtils
 
 				if (ServerProcess != null)
 				{
-					Log("Waiting a few seconds for the server to start...");
+					LogInformation("Waiting a few seconds for the server to start...");
 					Thread.Sleep(5000);
 				}
 			}
@@ -103,7 +103,7 @@ public partial class Project : CommandUtils
 
 			if (Params.Clean.HasValue && Params.Clean.Value && !Params.IterativeCooking)
 			{
-				Log("Cleaning cooked data.");
+				LogInformation("Cleaning cooked data.");
 				CleanupCookedData(PlatformsToCook.ToList(), Params);
 			}
 
@@ -114,11 +114,11 @@ public partial class Project : CommandUtils
 				Maps = Params.MapsToCook.ToArray();
                 foreach (var M in Maps)
                 {
-					Log("HasMapsToCook " + M.ToString());
+					LogInformation("HasMapsToCook " + M.ToString());
                 }
                 foreach (var M in Params.MapsToCook)
                 {
-					Log("Params.HasMapsToCook " + M.ToString());
+					LogInformation("Params.HasMapsToCook " + M.ToString());
                 }
 			}
 			
@@ -277,7 +277,7 @@ public partial class Project : CommandUtils
 		}
 
 
-		Log("********** COOK COMMAND COMPLETED **********");
+		LogInformation("********** COOK COMMAND COMPLETED **********");
 	}
 
     public struct FileInfo
@@ -329,7 +329,7 @@ public partial class Project : CommandUtils
             {
                 if (!(Ex is System.IO.DirectoryNotFoundException))
                 {
-                    Log("Failed deleting temporary directories " + TemporaryPakPath + " continuing. " + Ex.GetType().ToString());
+                    LogInformation("Failed deleting temporary directories " + TemporaryPakPath + " continuing. " + Ex.GetType().ToString());
                 }
             }
             try
@@ -340,7 +340,7 @@ public partial class Project : CommandUtils
             {
                 if (!(Ex is System.IO.DirectoryNotFoundException))
                 {
-                    Log("Failed deleting temporary directories " + TemporaryFilesPath + " continuing. " + Ex.GetType().ToString());
+                    LogInformation("Failed deleting temporary directories " + TemporaryFilesPath + " continuing. " + Ex.GetType().ToString());
                 }
             }
 
@@ -388,7 +388,7 @@ public partial class Project : CommandUtils
 
                     if ( PakFiles.Count <= 0 )
                     {
-                        Log("No Pak files found in " + SourceCookedContentPlatformPath +" :(");
+                        LogInformation("No Pak files found in " + SourceCookedContentPlatformPath +" :(");
                     }
                 }
                 else if (SourceCookedContentPath.EndsWith(".pak"))
@@ -406,7 +406,7 @@ public partial class Project : CommandUtils
 
                 foreach (var Name in PakFiles)
                 {
-                    Log("Extracting pak " + Name + " for comparision to location " + TemporaryFilesPath);
+                    LogInformation("Extracting pak " + Name + " for comparision to location " + TemporaryFilesPath);
 
                     string UnrealPakParams = Name + " -Extract " + " " + TemporaryFilesPath + " -ExtractToMountPoint";
                     try
@@ -415,7 +415,7 @@ public partial class Project : CommandUtils
                     }
                     catch(Exception Ex)
                     {
-                        Log("Pak failed to extract because of " + Ex.GetType().ToString());
+                        LogInformation("Pak failed to extract because of " + Ex.GetType().ToString());
                     }
                 }
 
@@ -466,7 +466,7 @@ public partial class Project : CommandUtils
 
                     if (SourceFile == null || DestFile == null)
                     {
-                        Log(LogStringBuilder.ToString());
+                        LogInformation(LogStringBuilder.ToString());
                         LogError("Diff cooked content failed on file " + SourceFilename + " when comparing against " + DestFilename + " " + (SourceFile == null ? SourceFilename : DestFilename) + " file is missing");
 			return;
                     }
@@ -520,7 +520,7 @@ public partial class Project : CommandUtils
                             bool bFailedToSaveDestFile = !Directory.Exists(Path.GetDirectoryName(SavedDestFilename));
                             if (bFailedToSaveSourceFile || bFailedToSaveDestFile)
                             {
-                                Log(LogStringBuilder.ToString());
+                                LogInformation(LogStringBuilder.ToString());
 
                                 if(bFailedToSaveSourceFile)
                                 {
@@ -556,26 +556,26 @@ public partial class Project : CommandUtils
                         FileReport.Add(DiffFileInfo);
                     }
                     
-                    Log(LogStringBuilder.ToString());
+                    LogInformation(LogStringBuilder.ToString());
                 });
 
-                Log("Mismatching files:");
+                LogInformation("Mismatching files:");
                 foreach (var Report in FileReport)
                 {
                     if ( Report.FirstByteFailed == -1)
                     {
-                        Log("File " + Report.Filename + " size mismatch: " + Report.File1Size + " VS " +Report.File2Size);
+                        LogInformation("File " + Report.Filename + " size mismatch: " + Report.File1Size + " VS " +Report.File2Size);
                     }
                     else
                     {
-                        Log("File " + Report.Filename + " bytes mismatch: " + Report.BytesMismatch + " first byte failed at: " + Report.FirstByteFailed + " file size: " + Report.File1Size);
+                        LogInformation("File " + Report.Filename + " bytes mismatch: " + Report.BytesMismatch + " first byte failed at: " + Report.FirstByteFailed + " file size: " + Report.File1Size);
                     }
                 }
 
             }
             catch ( Exception Ex )
             {
-                Log("Exception " + Ex.ToString());
+                LogInformation("Exception " + Ex.ToString());
                 continue;
             }
         }
diff --git a/Engine/Source/Programs/AutomationTool/Scripts/CopyBuildToStagingDirectory.Automation.cs b/Engine/Source/Programs/AutomationTool/Scripts/CopyBuildToStagingDirectory.Automation.cs
--- a/Engine/Source/Programs/AutomationTool/Scripts/CopyBuildToStagingDirectory.Automation.cs
+++ b/Engine/Source/Programs/AutomationTool/Scripts/CopyBuildToStagingDirectory.Automation.cs
@@ -282,7 +282,7 @@ public partial class Project : CommandUtils
 		}
 		var ThisPlatform = SC.StageTargetPlatform;
 
-		Log("Creating Staging Manifest...");
+		LogInformation("Creating Staging Manifest...");
 
 		if (Params.HasIterateSharedCookedBuild)
 		{
@@ -897,7 +897,7 @@ public partial class Project : CommandUtils
 			}
 			else
 			{
-				CommandUtils.Log("Excluding config file {0}", ConfigFile);
+				CommandUtils.LogInformation("Excluding config file {0}", ConfigFile);
 			}
 		}
 	}
@@ -1140,7 +1140,7 @@ public partial class Project : CommandUtils
 
 	public static void CopyManifestFilesToStageDir(Dictionary<StagedFileReference, FileReference> Mapping, DirectoryReference StageDir, DirectoryReference ManifestDir, string ManifestName, HashSet<StagedFileReference> CRCFiles, string PlatformName)
 	{
-		Log("Copying {0} to staging directory: {1}", ManifestName, StageDir);
+		LogInformation("Copying {0} to staging directory: {1}", ManifestName, StageDir);
 		FileReference ManifestPath = null;
 		string ManifestFile = "";
 		if (!String.IsNullOrEmpty(ManifestName))
@@ -1228,7 +1228,7 @@ public partial class Project : CommandUtils
 	/// <param name="SC"></param>
 	private static void CreatePakUsingStagingManifest(ProjectParams Params, DeploymentContext SC)
 	{
-		Log("Creating pak using staging manifest.");
+		LogInformation("Creating pak using staging manifest.");
 
 		DumpManifest(SC, CombinePaths(CmdEnv.LogFolder, "PrePak" + (SC.DedicatedServer ? "_Server" : "")));
 
@@ -1250,7 +1250,7 @@ public partial class Project : CommandUtils
 	/// <param name="SC">Staging context</param>
 	private static void CreatePakForCrashReporter(ProjectParams Params, DeploymentContext SC)
 	{
-		Log("Creating pak for crash reporter.");
+		LogInformation("Creating pak for crash reporter.");
 
 		Dictionary<string, string> PakResponseFile = CreatePakResponseFileFromStagingManifest(SC, SC.CrashReporterUFSFiles);
 		FileReference OutputLocation = FileReference.Combine(SC.RuntimeRootDir, "Engine", "Programs", "CrashReportClient", "Content", "Paks", "CrashReportClient.pak");
@@ -1280,7 +1280,7 @@ public partial class Project : CommandUtils
 			FileReference PakBlacklistFilename = FileReference.Combine(SC.ProjectRoot, "Build", SC.PlatformDir, string.Format("PakBlacklist-{0}.txt", SC.StageTargetConfigurations[0].ToString()));
 			if (FileReference.Exists(PakBlacklistFilename))
 			{
-				Log("Applying PAK blacklist file {0}", PakBlacklistFilename);
+				LogInformation("Applying PAK blacklist file {0}", PakBlacklistFilename);
 				string[] BlacklistContents = FileReference.ReadAllLines(PakBlacklistFilename);
 				foreach (string Candidate in BlacklistContents)
 				{
@@ -1318,7 +1318,7 @@ public partial class Project : CommandUtils
 
 				if (bExcludeFile)
 				{
-					Log("Excluding {0}", Src);
+					LogInformation("Excluding {0}", Src);
 					continue;
 				}
 			}
@@ -1398,7 +1398,7 @@ public partial class Project : CommandUtils
 
 		if (bShouldGeneratePatch && !Params.HasBasedOnReleaseVersion)
 		{
-			Log("Generating patch required a based on release version flag");
+			LogInformation("Generating patch required a based on release version flag");
 		}
 
 		string PostFix = "";
@@ -1507,7 +1507,7 @@ public partial class Project : CommandUtils
 
 					if (InternalUtils.SafeCopyFile(SourceOutputLocation.FullName, OutputLocation.FullName))
 					{
-						Log("Copying source pak from {0} to {1} instead of creating new pak", SourceOutputLocation, OutputLocation);
+						LogInformation("Copying source pak from {0} to {1} instead of creating new pak", SourceOutputLocation, OutputLocation);
 						bCopiedExistingPak = true;
 
 						FileReference InSigFile = SourceOutputLocation.ChangeExtension(".sig");
@@ -1515,11 +1515,11 @@ public partial class Project : CommandUtils
 						{
 							FileReference OutSigFile = OutputLocation.ChangeExtension(".sig");
 
-							Log("Copying pak sig from {0} to {1}", InSigFile, OutSigFile);
+							LogInformation("Copying pak sig from {0} to {1}", InSigFile, OutSigFile);
 
 							if (!InternalUtils.SafeCopyFile(InSigFile.FullName, OutSigFile.FullName))
 							{
-								Log("Failed to copy pak sig {0} to {1}, creating new pak", InSigFile, InSigFile);
+								LogInformation("Failed to copy pak sig {0} to {1}, creating new pak", InSigFile, InSigFile);
 								bCopiedExistingPak = false;
 							}
 						}
@@ -1528,7 +1528,7 @@ public partial class Project : CommandUtils
 				}
 				if (!bCopiedExistingPak)
 				{
-					Log("Failed to copy source pak from {0} to {1}, creating new pak", SourceOutputLocation, OutputLocation);
+					LogInformation("Failed to copy source pak from {0} to {1}, creating new pak", SourceOutputLocation, OutputLocation);
 				}
 			}
 
@@ -1706,7 +1706,7 @@ public partial class Project : CommandUtils
     /// <param name="SC"></param>
     private static void CopyPaksFromNetwork(ProjectParams Params, DeploymentContext SC)
     {
-        Log("Copying paks from network.");
+        LogInformation("Copying paks from network.");
 
         if (!CommandUtils.P4Enabled)
         {
@@ -1785,7 +1785,7 @@ public partial class Project : CommandUtils
                             var InfoCache = new System.IO.FileInfo(CacheSrcFile);
                             if (Info.Exists && InfoCache.Exists && Info.Length == InfoCache.Length)
                             {
-                                Log("Copying from cache {0} -> {1}", CacheSrcFile, DestFileName);
+                                LogInformation("Copying from cache {0} -> {1}", CacheSrcFile, DestFileName);
                                 CopyFileIncremental(new FileReference(CacheSrcFile), new FileReference(DestFileName));
                                 continue;
                             }
@@ -1796,7 +1796,7 @@ public partial class Project : CommandUtils
 
                     }
                 }
-                Log("Copying {0} -> {1}", SrcFile, DestFileName);
+                LogInformation("Copying {0} -> {1}", SrcFile, DestFileName);
                 CopyFileIncremental(new FileReference(SrcFile), new FileReference(DestFileName));
             }
         }
@@ -1825,7 +1825,7 @@ public partial class Project : CommandUtils
     /// <param name="SC"></param>
     private static void CreatePaksUsingChunkManifests(ProjectParams Params, DeploymentContext SC)
 	{
-		Log("Creating pak using streaming install manifests.");
+		LogInformation("Creating pak using streaming install manifests.");
 		DumpManifest(SC, CombinePaths(CmdEnv.LogFolder, "PrePak" + (SC.DedicatedServer ? "_Server" : "")));
 
 
@@ -1988,7 +1988,7 @@ public partial class Project : CommandUtils
 
 		System.Threading.Tasks.ParallelOptions Options = new System.Threading.Tasks.ParallelOptions();
 
-		Log("Creating Pak files utilizing {0} cores", Environment.ProcessorCount);
+		LogInformation("Creating Pak files utilizing {0} cores", Environment.ProcessorCount);
 		Options.MaxDegreeOfParallelism = Environment.ProcessorCount;
 
 		List<CreatePakParams> PakInputs = new List<CreatePakParams>();
@@ -2140,7 +2140,7 @@ public partial class Project : CommandUtils
 
 	public static void CleanStagingDirectory(ProjectParams Params, DeploymentContext SC)
 	{
-		Log("Cleaning Stage Directory: {0}", SC.StageDirectory.FullName);
+		LogInformation("Cleaning Stage Directory: {0}", SC.StageDirectory.FullName);
 		if (SC.Stage && !Params.NoCleanStage && !Params.SkipStage && !Params.IterativeDeploy)
 		{
 			try
@@ -2224,7 +2224,7 @@ public partial class Project : CommandUtils
 			return;
 		}
 
-		Log("Creating UE4CommandLine.txt");
+		LogInformation("Creating UE4CommandLine.txt");
 		if (!string.IsNullOrEmpty(Params.StageCommandline) || !string.IsNullOrEmpty(Params.RunCommandline))
 		{
 			string FileHostParams = " ";
@@ -2703,7 +2703,7 @@ public partial class Project : CommandUtils
 		{
 			Params.ValidateAndLog();
 
-			Log("********** STAGE COMMAND STARTED **********");
+			LogInformation("********** STAGE COMMAND STARTED **********");
 
 			if (!Params.NoClient)
 			{
@@ -2780,7 +2780,7 @@ public partial class Project : CommandUtils
 					ApplyStagingManifest(Params, SC);
 				}
 			}
-			Log("********** STAGE COMMAND COMPLETED **********");
+			LogInformation("********** STAGE COMMAND COMPLETED **********");
 		}
 	}
 
diff --git a/Engine/Source/Programs/AutomationTool/Scripts/CopySharedCookedBuild.Automation.cs b/Engine/Source/Programs/AutomationTool/Scripts/CopySharedCookedBuild.Automation.cs
--- a/Engine/Source/Programs/AutomationTool/Scripts/CopySharedCookedBuild.Automation.cs
+++ b/Engine/Source/Programs/AutomationTool/Scripts/CopySharedCookedBuild.Automation.cs
@@ -59,7 +59,7 @@ class CopySharedCookedBuild : BuildCommand
 
 	public override void ExecuteBuild()
 	{
-		Log("************************* CopySharedCookedBuild");
+		LogInformation("************************* CopySharedCookedBuild");
 
 		// Parse the project filename (as a local path)
 		
diff --git a/Engine/Source/Programs/AutomationTool/Scripts/CopyUAT.Automation.cs b/Engine/Source/Programs/AutomationTool/Scripts/CopyUAT.Automation.cs
--- a/Engine/Source/Programs/AutomationTool/Scripts/CopyUAT.Automation.cs
+++ b/Engine/Source/Programs/AutomationTool/Scripts/CopyUAT.Automation.cs
@@ -61,6 +61,6 @@ public class CopyUAT : BuildCommand
 			CommandUtils.CopyFile(SourceFile.FullName, TargetFile.FullName);
 		}
 
-		Log("Copied {0} files to {1}", SourceFiles.Count, TargetDir);
+		LogInformation("Copied {0} files to {1}", SourceFiles.Count, TargetDir);
 	}
 }
diff --git a/Engine/Source/Programs/AutomationTool/Scripts/CryptoKeys.Automation.cs b/Engine/Source/Programs/AutomationTool/Scripts/CryptoKeys.Automation.cs
--- a/Engine/Source/Programs/AutomationTool/Scripts/CryptoKeys.Automation.cs
+++ b/Engine/Source/Programs/AutomationTool/Scripts/CryptoKeys.Automation.cs
@@ -22,7 +22,7 @@ namespace AutomationTool
 				RawProjectPath: ProjectPath
 			);
 
-			Log("********** CRYPTOKEYS COMMAND STARTED **********");
+			LogInformation("********** CRYPTOKEYS COMMAND STARTED **********");
 
 			string UE4EditorExe = HostPlatform.Current.GetUE4ExePath(Params.UE4Exe);
 			if (!FileExists(UE4EditorExe))
@@ -87,7 +87,7 @@ namespace AutomationTool
 			}
 			else
 			{
-				Log(ChangeDescription);
+				LogInformation(ChangeDescription);
 				FileReference.MakeWriteable(OutputFile);
 			}
 
diff --git a/Engine/Source/Programs/AutomationTool/Scripts/DeployCommand.Automation.cs b/Engine/Source/Programs/AutomationTool/Scripts/DeployCommand.Automation.cs
--- a/Engine/Source/Programs/AutomationTool/Scripts/DeployCommand.Automation.cs
+++ b/Engine/Source/Programs/AutomationTool/Scripts/DeployCommand.Automation.cs
@@ -19,7 +19,7 @@ public partial class Project : CommandUtils
 			return;
 		}
 
-		Log("********** DEPLOY COMMAND STARTED **********");
+		LogInformation("********** DEPLOY COMMAND STARTED **********");
 		
 		if (!Params.NoClient)
 		{
@@ -34,8 +34,8 @@ public partial class Project : CommandUtils
 					ClientCmdLine += "-SourceDir=\"" + CombinePaths(Params.BaseStageDirectory, SC.StageTargetPlatform.PlatformType.ToString()) + "\" ";
 					string ClientApp = CombinePaths(CmdEnv.LocalRoot, "Engine/Binaries/Win64/UnrealFrontend.exe");
 
-					Log("Deploying via UFE:");
-					Log("\tClientCmdLine: " + ClientCmdLine + "");
+					LogInformation("Deploying via UFE:");
+					LogInformation("\tClientCmdLine: " + ClientCmdLine + "");
 
 					//@todo UAT: Consolidate running of external applications like UFE (See 'RunProjectCommand' for other instances)
 					PushDir(Path.GetDirectoryName(ClientApp));
@@ -68,7 +68,7 @@ public partial class Project : CommandUtils
 			}
 		}
 
-		Log("********** DEPLOY COMMAND COMPLETED **********");
+		LogInformation("********** DEPLOY COMMAND COMPLETED **********");
 	}
 
 	#endregion
diff --git a/Engine/Source/Programs/AutomationTool/Scripts/FixupRedirects.Automation.cs b/Engine/Source/Programs/AutomationTool/Scripts/FixupRedirects.Automation.cs
--- a/Engine/Source/Programs/AutomationTool/Scripts/FixupRedirects.Automation.cs
+++ b/Engine/Source/Programs/AutomationTool/Scripts/FixupRedirects.Automation.cs
@@ -13,7 +13,7 @@ class FixupRedirects : BuildCommand
 		var ProjectName = ParseParamValue("project", "");
 
 		var EditorExe = CombinePaths(CmdEnv.LocalRoot, @"Engine/Binaries/Win64/UE4Editor-Cmd.exe");
-		Log("********** Running FixupRedirects: {0} -run=ResavePackages -unattended -nopause -buildmachine -fixupredirects -autocheckout -autocheckin -projectonly", EditorExe);
+		LogInformation("********** Running FixupRedirects: {0} -run=ResavePackages -unattended -nopause -buildmachine -fixupredirects -autocheckout -autocheckin -projectonly", EditorExe);
 		RunCommandlet(GetCommandletProjectFile(ProjectName), EditorExe, "ResavePackages", "-unattended -nopause -buildmachine -fixupredirects -autocheckout -autocheckin -projectonly");
 	}
 }
diff --git a/Engine/Source/Programs/AutomationTool/Scripts/GetFileCommand.Automation.cs b/Engine/Source/Programs/AutomationTool/Scripts/GetFileCommand.Automation.cs
--- a/Engine/Source/Programs/AutomationTool/Scripts/GetFileCommand.Automation.cs
+++ b/Engine/Source/Programs/AutomationTool/Scripts/GetFileCommand.Automation.cs
@@ -21,7 +21,7 @@ public partial class Project : CommandUtils
 			return;
 		}
 
-		Log("********** GETFILE COMMAND STARTED **********");
+		LogInformation("********** GETFILE COMMAND STARTED **********");
 
 		var FileName = Path.GetFileName(Params.GetFile);
 		var LocalFile = CombinePaths(CmdEnv.EngineSavedFolder, FileName);
@@ -34,7 +34,7 @@ public partial class Project : CommandUtils
 
 		SC[0].StageTargetPlatform.GetTargetFile(Params.GetFile, LocalFile, Params);
 
-		Log("********** GETFILE COMMAND COMPLETED **********");
+		LogInformation("********** GETFILE COMMAND COMPLETED **********");
 	}
 
 	#endregion
diff --git a/Engine/Source/Programs/AutomationTool/Scripts/IPhonePackager.Automation.cs b/Engine/Source/Programs/AutomationTool/Scripts/IPhonePackager.Automation.cs
--- a/Engine/Source/Programs/AutomationTool/Scripts/IPhonePackager.Automation.cs
+++ b/Engine/Source/Programs/AutomationTool/Scripts/IPhonePackager.Automation.cs
@@ -11,7 +11,7 @@ class IPhonePackager : BuildCommand
 {
 	public override void ExecuteBuild()
 	{
-		Log("************************* Calling IPP");
+		LogInformation("************************* Calling IPP");
 
 		Platform IOS = Platform.GetPlatform(UnrealTargetPlatform.IOS);
 		string Command = ParseParamValue("cmd", "");
diff --git a/Engine/Source/Programs/AutomationTool/Scripts/LauncherLocalization.Automation.cs b/Engine/Source/Programs/AutomationTool/Scripts/LauncherLocalization.Automation.cs
--- a/Engine/Source/Programs/AutomationTool/Scripts/LauncherLocalization.Automation.cs
+++ b/Engine/Source/Programs/AutomationTool/Scripts/LauncherLocalization.Automation.cs
@@ -28,7 +28,7 @@ class LauncherLocalization : BuildCommand
 
 		if (P4Enabled)
 		{
-			Log("Sync necessary content to head revision");
+			LogInformation("Sync necessary content to head revision");
 			P4.Sync(P4Env.Branch + "/Engine/Config/...");
 			P4.Sync(P4Env.Branch + "/Engine/Content/...");
 			P4.Sync(P4Env.Branch + "/Engine/Source/...");
@@ -79,9 +79,9 @@ class LauncherLocalization : BuildCommand
 		// Execute commandlet for each set of arguments.
 		foreach (var CommandletArguments in CommandletArgumentSets)
 		{
-			Log("Localization for {0} {1}", EditorArguments, CommandletArguments);
+			LogInformation("Localization for {0} {1}", EditorArguments, CommandletArguments);
 
-			Log("Running UE4Editor to generate Localization data");
+			LogInformation("Running UE4Editor to generate Localization data");
 
 			string Arguments = String.Format("-run=GatherText {0} {1}", EditorArguments, CommandletArguments);
 			var RunResult = Run(EditorExe, Arguments);
diff --git a/Engine/Source/Programs/AutomationTool/Scripts/ListMobileDevices.Automation.cs b/Engine/Source/Programs/AutomationTool/Scripts/ListMobileDevices.Automation.cs
--- a/Engine/Source/Programs/AutomationTool/Scripts/ListMobileDevices.Automation.cs
+++ b/Engine/Source/Programs/AutomationTool/Scripts/ListMobileDevices.Automation.cs
@@ -11,7 +11,7 @@ namespace AutomationScripts.Automation
 	{
 		public override void ExecuteBuild()
 		{
-			Log("======= ListMobileDevices - Start =======");
+			LogInformation("======= ListMobileDevices - Start =======");
 
 			var GlobalParams = new ProjectParams(
 				Command: this,
@@ -28,7 +28,7 @@ namespace AutomationScripts.Automation
 				throw new AutomationException("iOS is not yet implemented.");
 			}
 
-			Log("======= ListMobileDevices - Done ========");
+			LogInformation("======= ListMobileDevices - Done ========");
 		}
 
 		private static void GetConnectedDevices(ProjectParams Params, Platform TargetPlatform)
@@ -41,7 +41,7 @@ namespace AutomationScripts.Automation
 			{
 				foreach (var DeviceName in ConnectedDevices)
 				{
-					Log("Device:{0}:{1}", PlatformName, DeviceName);
+					LogInformation("Device:{0}:{1}", PlatformName, DeviceName);
 				}
 			}
 			catch
diff --git a/Engine/Source/Programs/AutomationTool/Scripts/ListThirdPartySoftware.Automation.cs b/Engine/Source/Programs/AutomationTool/Scripts/ListThirdPartySoftware.Automation.cs
--- a/Engine/Source/Programs/AutomationTool/Scripts/ListThirdPartySoftware.Automation.cs
+++ b/Engine/Source/Programs/AutomationTool/Scripts/ListThirdPartySoftware.Automation.cs
@@ -15,7 +15,7 @@ class ListThirdPartySoftware : BuildCommand
 {
 	public override void ExecuteBuild()
 	{
-		CommandUtils.Log("************************* List Third Party Software");
+		CommandUtils.LogInformation("************************* List Third Party Software");
 
 		string ProjectPath = ParseParamValue("Project", String.Empty);
 
@@ -136,7 +136,7 @@ class ListThirdPartySoftware : BuildCommand
 		// Print them all out
 		foreach(string OutputMessage in OutputMessages)
 		{
-			CommandUtils.Log(OutputMessage);
+			CommandUtils.LogInformation(OutputMessage);
 		}
 	}
 }
diff --git a/Engine/Source/Programs/AutomationTool/Scripts/Localisation.Automation.cs b/Engine/Source/Programs/AutomationTool/Scripts/Localisation.Automation.cs
--- a/Engine/Source/Programs/AutomationTool/Scripts/Localisation.Automation.cs
+++ b/Engine/Source/Programs/AutomationTool/Scripts/Localisation.Automation.cs
@@ -241,7 +241,7 @@ class Localise : BuildCommand
 		// Make sure the Localization configs and content is up-to-date to ensure we don't get errors later on
 		if (P4Enabled)
 		{
-			Log("Sync necessary content to head revision");
+			LogInformation("Sync necessary content to head revision");
 			P4.Sync(P4Env.Branch + "/" + LocalizationBatch.LocalizationTargetDirectory + "/Config/Localization/...");
 			P4.Sync(P4Env.Branch + "/" + LocalizationBatch.LocalizationTargetDirectory + "/Content/Localization/...");
 		}
@@ -301,11 +301,11 @@ class Localise : BuildCommand
 				}
 
 				string Arguments = String.Format("{0} -run=GatherText {1} {2}", UEProjectName, EditorArguments, CommandletArguments);
-				Log("Running localization commandlet: {0}", Arguments);
+				LogInformation("Running localization commandlet: {0}", Arguments);
 				var StartTime = DateTime.UtcNow;
 				var RunResult = Run(EditorExe, Arguments, null, ERunOptions.Default | ERunOptions.NoLoggingOfRunCommand); // Disable logging of the run command as it will print the exit code which GUBP can pick up as an error (we do that ourselves below)
 				var RunDuration = (DateTime.UtcNow - StartTime).TotalMilliseconds;
-				Log("Localization commandlet finished in {0}s", RunDuration / 1000);
+				LogInformation("Localization commandlet finished in {0}s", RunDuration / 1000);
 
 				if (RunResult.ExitCode != 0)
 				{
diff --git a/Engine/Source/Programs/AutomationTool/Scripts/MegaXGE.Automation.cs b/Engine/Source/Programs/AutomationTool/Scripts/MegaXGE.Automation.cs
--- a/Engine/Source/Programs/AutomationTool/Scripts/MegaXGE.Automation.cs
+++ b/Engine/Source/Programs/AutomationTool/Scripts/MegaXGE.Automation.cs
@@ -27,7 +27,7 @@ class MegaXGE : BuildCommand
 			WorkingCL = P4.CreateChange(P4Env.Client, String.Format("MegaXGE build from changelist {0} - Params: {1}", P4Env.Changelist, CmdLine));
 		}
 
-		Log("************************* MegaXGE");
+		LogInformation("************************* MegaXGE");
 
 		bool Clean = ParseParam("Clean");
 		string CleanToolLocation = CombinePaths(CmdEnv.LocalRoot, "Engine", "Build", "Batchfiles", "Clean.bat");
@@ -54,7 +54,7 @@ class MegaXGE : BuildCommand
 			}
 		}
 
-		Log("*************************");
+		LogInformation("*************************");
 		for (int Arg = 1; Arg < 100; Arg++)
 		{
 			string Parm = String.Format("Target{0}", Arg);
@@ -147,7 +147,7 @@ class MegaXGE : BuildCommand
 					foreach (var Configuration in Configurations)
 					{
 						Agenda.AddTargets(new string[] { CurTarget }, Platform, Configuration, ProjectFile);
-						Log("Target {0} {1} {2}", CurTarget, Platform.ToString(), Configuration.ToString());
+						LogInformation("Target {0} {1} {2}", CurTarget, Platform.ToString(), Configuration.ToString());
 						if (Clean)
 						{
 							string Args = String.Format("{0} {1} {2}", CurTarget, Platform.ToString(), Configuration.ToString());
@@ -157,7 +157,7 @@ class MegaXGE : BuildCommand
 				}
 			}
 		}
-		Log("*************************");
+		LogInformation("*************************");
 
 		Agenda.DoRetries = ParseParam("Retry");
 		UE4Build.Build(Agenda, InUpdateVersionFiles: IsBuildMachine, InUseParallelExecutor: ParseParam("useparallelexecutor"), InShowProgress: ShowProgress);
diff --git a/Engine/Source/Programs/AutomationTool/Scripts/PackageCommand.Automation.cs b/Engine/Source/Programs/AutomationTool/Scripts/PackageCommand.Automation.cs
--- a/Engine/Source/Programs/AutomationTool/Scripts/PackageCommand.Automation.cs
+++ b/Engine/Source/Programs/AutomationTool/Scripts/PackageCommand.Automation.cs
@@ -28,7 +28,7 @@ public partial class Project : CommandUtils
 
 			if (DeployContextList.Count > 0 )
 			{
-				Log("********** PACKAGE COMMAND STARTED **********");
+				LogInformation("********** PACKAGE COMMAND STARTED **********");
 
 				foreach (var SC in DeployContextList)
 				{
@@ -38,7 +38,7 @@ public partial class Project : CommandUtils
 					}
 				}
 
-				Log("********** PACKAGE COMMAND COMPLETED **********");
+				LogInformation("********** PACKAGE COMMAND COMPLETED **********");
 			}
 		}
 	}
diff --git a/Engine/Source/Programs/AutomationTool/Scripts/RebasePublicIncludePaths.cs b/Engine/Source/Programs/AutomationTool/Scripts/RebasePublicIncludePaths.cs
--- a/Engine/Source/Programs/AutomationTool/Scripts/RebasePublicIncludePaths.cs
+++ b/Engine/Source/Programs/AutomationTool/Scripts/RebasePublicIncludePaths.cs
@@ -160,7 +160,7 @@ namespace AutomationTool
 			// Output them all to disk
 			if(bWrite && ModifiedFiles.Count > 0)
 			{
-				Log("Updating {0} files...", ModifiedFiles.Count);
+				LogInformation("Updating {0} files...", ModifiedFiles.Count);
 
 				List<FileReference> FilesToCheckOut = new List<FileReference>();
 				foreach(FileReference ModifiedFile in ModifiedFiles.Keys)
@@ -178,7 +178,7 @@ namespace AutomationTool
 						throw new AutomationException("{0} files have been modified, but are read only. Run with -P4 to enable Perforce checkout.\n{1}", FilesToCheckOut.Count, String.Join("\n", FilesToCheckOut.Select(x => "  " + x)));
 					}
 
-					Log("Checking out files from Perforce");
+					LogInformation("Checking out files from Perforce");
 
 					int ChangeNumber = P4.CreateChange(Description: "Updating source files");
 					P4.Edit(ChangeNumber, FilesToCheckOut.Select(x => x.FullName).ToList(), false);
@@ -186,7 +186,7 @@ namespace AutomationTool
 
 				foreach(KeyValuePair<FileReference, string[]> FileToWrite in ModifiedFiles)
 				{
-					Log("Writing {0}", FileToWrite.Key);
+					LogInformation("Writing {0}", FileToWrite.Key);
 					FileReference.WriteAllLines(FileToWrite.Key, FileToWrite.Value);
 				}
 			}
diff --git a/Engine/Source/Programs/AutomationTool/Scripts/RebuildHLODCommand.Automation.cs b/Engine/Source/Programs/AutomationTool/Scripts/RebuildHLODCommand.Automation.cs
--- a/Engine/Source/Programs/AutomationTool/Scripts/RebuildHLODCommand.Automation.cs
+++ b/Engine/Source/Programs/AutomationTool/Scripts/RebuildHLODCommand.Automation.cs
@@ -31,7 +31,7 @@ namespace AutomationScripts.Automation
 
 		public override void ExecuteBuild()
 		{
-			Log("********** REBUILD HLODS COMMAND STARTED **********");
+			LogInformation("********** REBUILD HLODS COMMAND STARTED **********");
             bool DelayCheckin = ParseParam("DelaySubmission");
             int SubmittedCL = 0;
             try
@@ -64,8 +64,8 @@ namespace AutomationScripts.Automation
 
             catch (Exception ProcessEx)
             {
-                Log("********** REBUILD HLODS COMMAND FAILED **********");
-                Log("Error message: {0}", ProcessEx.Message);
+                LogInformation("********** REBUILD HLODS COMMAND FAILED **********");
+                LogInformation("Error message: {0}", ProcessEx.Message);
                 HandleFailure(ProcessEx.Message);
                 throw ProcessEx;
             }
@@ -80,7 +80,7 @@ namespace AutomationScripts.Automation
                 HandleSuccessNoCheckin(WorkingCL);
             }			
 
-			Log("********** REBUILD HLODS COMMAND COMPLETED **********");
+			LogInformation("********** REBUILD HLODS COMMAND COMPLETED **********");
 		}
 
 		#endregion
@@ -90,7 +90,7 @@ namespace AutomationScripts.Automation
 
 		private void BuildNecessaryTargets()
 		{
-			Log("Running Step:- RebuildHLOD::BuildNecessaryTargets");
+			LogInformation("Running Step:- RebuildHLOD::BuildNecessaryTargets");
 			UE4Build.BuildAgenda Agenda = new UE4Build.BuildAgenda();
             Agenda.AddTarget("UnrealHeaderTool", UnrealBuildTool.UnrealTargetPlatform.Win64, UnrealBuildTool.UnrealTargetConfiguration.Development);
 			Agenda.AddTarget("ShaderCompileWorker", UnrealBuildTool.UnrealTargetPlatform.Win64, UnrealBuildTool.UnrealTargetConfiguration.Development);			
@@ -121,16 +121,16 @@ namespace AutomationScripts.Automation
 
 		private void CreateChangelist(ProjectParams Params)
 		{
-			Log("Running Step:- RebuildHLOD::CheckOutMaps");
+			LogInformation("Running Step:- RebuildHLOD::CheckOutMaps");
 			// Setup a P4 Cl we will use to submit the new HLOD data
 			WorkingCL = P4.CreateChange(P4Env.Client, String.Format("{0} rebuilding HLODs from changelist {1}\n#rb None\n#tests None\n#jira none\n#robomerge #DisregardExcludedAuthors", Params.ShortProjectName, P4Env.Changelist));
-			Log("Working in {0}", WorkingCL);
+			LogInformation("Working in {0}", WorkingCL);
 
 		}
 
 		private void RunRebuildHLODCommandlet(ProjectParams Params)
 		{
-			Log("Running Step:- RebuildHLOD::RunRebuildHLODCommandlet");
+			LogInformation("Running Step:- RebuildHLOD::RunRebuildHLODCommandlet");
 
 			// Find the commandlet binary
 			string UE4EditorExe = HostPlatform.Current.GetUE4ExePath(Params.UE4Exe);
@@ -202,14 +202,14 @@ namespace AutomationScripts.Automation
 
 		private void SubmitRebuiltMaps(ref int SubmittedCL)
 		{
-			Log("Running Step:- RebuildHLOD::SubmitRebuiltMaps");
+			LogInformation("Running Step:- RebuildHLOD::SubmitRebuiltMaps");
 
 			// Check everything in!
 			if (WorkingCL != -1)
 			{
-                Log("Running Step:- Submitting CL " + WorkingCL);
+                LogInformation("Running Step:- Submitting CL " + WorkingCL);
 				P4.Submit(WorkingCL, out SubmittedCL, true, true);
-				Log("INFO: HLODs successfully submitted in cl " + SubmittedCL.ToString());
+				LogInformation("INFO: HLODs successfully submitted in cl " + SubmittedCL.ToString());
 			}
 		}
 
diff --git a/Engine/Source/Programs/AutomationTool/Scripts/RebuildLightMapsCommand.Automation.cs b/Engine/Source/Programs/AutomationTool/Scripts/RebuildLightMapsCommand.Automation.cs
--- a/Engine/Source/Programs/AutomationTool/Scripts/RebuildLightMapsCommand.Automation.cs
+++ b/Engine/Source/Programs/AutomationTool/Scripts/RebuildLightMapsCommand.Automation.cs
@@ -28,7 +28,7 @@ namespace AutomationScripts.Automation
 
 		public override void ExecuteBuild()
 		{
-			Log("********** REBUILD LIGHT MAPS COMMAND STARTED **********");
+			LogInformation("********** REBUILD LIGHT MAPS COMMAND STARTED **********");
 			int SubmittedCL = 0;
 			try
 			{
@@ -54,8 +54,8 @@ namespace AutomationScripts.Automation
 			}
 			catch (Exception ProcessEx)
 			{
-				Log("********** REBUILD LIGHT MAPS COMMAND FAILED **********");
-                Log("Error message: {0}", ProcessEx.Message);
+				LogInformation("********** REBUILD LIGHT MAPS COMMAND FAILED **********");
+                LogInformation("Error message: {0}", ProcessEx.Message);
 				HandleFailure(ProcessEx.Message);
 				throw ProcessEx;
 			}
@@ -63,7 +63,7 @@ namespace AutomationScripts.Automation
 			// The processes steps have completed successfully.
 			HandleSuccess(SubmittedCL);
 
-			Log("********** REBUILD LIGHT MAPS COMMAND COMPLETED **********");
+			LogInformation("********** REBUILD LIGHT MAPS COMMAND COMPLETED **********");
 		}
 
 		#endregion
@@ -73,7 +73,7 @@ namespace AutomationScripts.Automation
 
 		private void BuildNecessaryTargets()
 		{
-			Log("Running Step:- RebuildLightMaps::BuildNecessaryTargets");
+			LogInformation("Running Step:- RebuildLightMaps::BuildNecessaryTargets");
 			UE4Build.BuildAgenda Agenda = new UE4Build.BuildAgenda();
             Agenda.AddTarget("UnrealHeaderTool", UnrealBuildTool.UnrealTargetPlatform.Win64, UnrealBuildTool.UnrealTargetConfiguration.Development);
 			Agenda.AddTarget("ShaderCompileWorker", UnrealBuildTool.UnrealTargetPlatform.Win64, UnrealBuildTool.UnrealTargetConfiguration.Development);
@@ -105,16 +105,16 @@ namespace AutomationScripts.Automation
 
 		private void CreateChangelist(ProjectParams Params)
 		{
-			Log("Running Step:- RebuildLightMaps::CheckOutMaps");
+			LogInformation("Running Step:- RebuildLightMaps::CheckOutMaps");
 			// Setup a P4 Cl we will use to submit the new lightmaps
 			WorkingCL = P4.CreateChange(P4Env.Client, String.Format("{0} rebuilding lightmaps from changelist {1}\n#rb None\n#tests None", Params.ShortProjectName, P4Env.Changelist));
-			Log("Working in {0}", WorkingCL);
+			LogInformation("Working in {0}", WorkingCL);
 
 		}
 
 		private void RunRebuildLightmapsCommandlet(ProjectParams Params)
 		{
-			Log("Running Step:- RebuildLightMaps::RunRebuildLightmapsCommandlet");
+			LogInformation("Running Step:- RebuildLightMaps::RunRebuildLightmapsCommandlet");
 
 			// Find the commandlet binary
 			string UE4EditorExe = HostPlatform.Current.GetUE4ExePath(Params.UE4Exe);
@@ -176,14 +176,14 @@ namespace AutomationScripts.Automation
 
 		private void SubmitRebuiltMaps(ref int SubmittedCL)
 		{
-			Log("Running Step:- RebuildLightMaps::SubmitRebuiltMaps");
+			LogInformation("Running Step:- RebuildLightMaps::SubmitRebuiltMaps");
 
 			// Check everything in!
 			if (WorkingCL != -1)
 			{
-                Log("Running Step:- Submitting CL " + WorkingCL);
+                LogInformation("Running Step:- Submitting CL " + WorkingCL);
 				P4.Submit(WorkingCL, out SubmittedCL, true, true);
-				Log("INFO: Lightmaps successfully submitted in cl "+ SubmittedCL.ToString());
+				LogInformation("INFO: Lightmaps successfully submitted in cl "+ SubmittedCL.ToString());
 			}
 		}
 
diff --git a/Engine/Source/Programs/AutomationTool/Scripts/RecordPerformance.Automation.cs b/Engine/Source/Programs/AutomationTool/Scripts/RecordPerformance.Automation.cs
--- a/Engine/Source/Programs/AutomationTool/Scripts/RecordPerformance.Automation.cs
+++ b/Engine/Source/Programs/AutomationTool/Scripts/RecordPerformance.Automation.cs
@@ -59,7 +59,7 @@ class RecordPerformance : BuildCommand
 				var Platform = GetPlatformByName(PlatformName);
 				if (Platform == null)
 				{
-					Log("Cannot find platform '{0}'. Skipping.", PlatformName);
+					LogInformation("Cannot find platform '{0}'. Skipping.", PlatformName);
 					continue;
 				}
 
@@ -98,7 +98,7 @@ class RecordPerformance : BuildCommand
 			CreateChart(Stat);
 		}
 
-		Log("Performance charts stored in {0}", OutputDir);
+		LogInformation("Performance charts stored in {0}", OutputDir);
 	}
 
 	private IEnumerable<string> ParseMultipleParams(string ParamName)
@@ -182,7 +182,7 @@ class RecordPerformance : BuildCommand
 		var AdditionalCommandLineParameters = "-savevulkanpsocacheonexit";
 		var StatFileName = Path.Combine(ProjectToRun, "Saved", "FXPerformance", PerformanceMonitorConfig + ".csv");
 
-		Log("Running {0} on {1}...", Rhi, Platform.PlatformType);
+		LogInformation("Running {0} on {1}...", Rhi, Platform.PlatformType);
 
 		var RhiParam = (Rhi != "default") ? "-" + Rhi : "";
 
@@ -193,7 +193,7 @@ class RecordPerformance : BuildCommand
 		}
 		catch
 		{
-			Log("Running game failed. Stats might be wrong / incomplete.");
+			LogInformation("Running game failed. Stats might be wrong / incomplete.");
 		}
 
 		if (UseStats)
@@ -215,7 +215,7 @@ class RecordPerformance : BuildCommand
 
 	private void BuildAndCookPlatform(UnrealTargetPlatform Platform)
 	{
-		Log("Building {0}...", Platform.ToString());
+		LogInformation("Building {0}...", Platform.ToString());
 
 		RunBCR(string.Format("-project={0} -platform={1} -build -cook -stage -pak -deploy", ProjectToRun, Platform));
 	}
diff --git a/Engine/Source/Programs/AutomationTool/Scripts/ResavePackagesCommand.Automation.cs b/Engine/Source/Programs/AutomationTool/Scripts/ResavePackagesCommand.Automation.cs
--- a/Engine/Source/Programs/AutomationTool/Scripts/ResavePackagesCommand.Automation.cs
+++ b/Engine/Source/Programs/AutomationTool/Scripts/ResavePackagesCommand.Automation.cs
@@ -30,7 +30,7 @@ namespace AutomationScripts.Automation
 
         public override void ExecuteBuild()
 		{
-			Log("********** RESAVE PACKAGES COMMAND STARTED **********");
+			LogInformation("********** RESAVE PACKAGES COMMAND STARTED **********");
 
 			try
 			{
@@ -56,8 +56,8 @@ namespace AutomationScripts.Automation
 			}
 			catch (Exception ProcessEx)
 			{
-				Log("********** RESAVE PACKAGES COMMAND FAILED **********");
-                Log("Error message: {0}", ProcessEx.Message);
+				LogInformation("********** RESAVE PACKAGES COMMAND FAILED **********");
+                LogInformation("Error message: {0}", ProcessEx.Message);
 				HandleFailure(ProcessEx.Message);
 				throw ProcessEx;
 			}
@@ -65,7 +65,7 @@ namespace AutomationScripts.Automation
 			// The processes steps have completed successfully.
 			HandleSuccess();
 
-			Log("********** RESAVE PACKAGES COMMAND COMPLETED **********");
+			LogInformation("********** RESAVE PACKAGES COMMAND COMPLETED **********");
 		}
 
 		#endregion
@@ -75,7 +75,7 @@ namespace AutomationScripts.Automation
 
 		private void BuildNecessaryTargets()
 		{
-			Log("Running Step:- RebuildLightMaps::BuildNecessaryTargets");
+			LogInformation("Running Step:- RebuildLightMaps::BuildNecessaryTargets");
 			UE4Build.BuildAgenda Agenda = new UE4Build.BuildAgenda();
             Agenda.AddTarget("UnrealHeaderTool", UnrealBuildTool.UnrealTargetPlatform.Win64, UnrealBuildTool.UnrealTargetConfiguration.Development);
 			Agenda.AddTarget("ShaderCompileWorker", UnrealBuildTool.UnrealTargetPlatform.Win64, UnrealBuildTool.UnrealTargetConfiguration.Development);
@@ -107,16 +107,16 @@ namespace AutomationScripts.Automation
 
 		private void CreateChangelist(ProjectParams Params)
 		{
-			Log("Running Step:- RebuildLightMaps::CheckOutMaps");
+			LogInformation("Running Step:- RebuildLightMaps::CheckOutMaps");
 			// Setup a P4 Cl we will use to submit the new lightmaps
 			WorkingCL = P4.CreateChange(P4Env.Client, String.Format("{0} rebuilding lightmaps from changelist {1}\n#rb None\n#tests None", Params.ShortProjectName, P4Env.Changelist));
-			Log("Working in {0}", WorkingCL);
+			LogInformation("Working in {0}", WorkingCL);
 
 		}
 
 		private void RunResavePackagesCommandlet(ProjectParams Params)
 		{
-			Log("Running Step:- ResavePackages::RunResavePackagesCommandlet");
+			LogInformation("Running Step:- ResavePackages::RunResavePackagesCommandlet");
 
 			// Find the commandlet binary
 			string UE4EditorExe = HostPlatform.Current.GetUE4ExePath(Params.UE4Exe);
@@ -178,15 +178,15 @@ namespace AutomationScripts.Automation
 
 		private void SubmitResavedPackages()
 		{
-			Log("Running Step:- ResavePackages::SubmitResavedPackages");
+			LogInformation("Running Step:- ResavePackages::SubmitResavedPackages");
 
 			// Check everything in!
 			if (WorkingCL != -1)
 			{
-                Log("Running Step:- Submitting CL " + WorkingCL);
+                LogInformation("Running Step:- Submitting CL " + WorkingCL);
 				int SubmittedCL;
 				P4.Submit(WorkingCL, out SubmittedCL, true, true);
-				Log("INFO: Packages successfully submitted in cl "+ SubmittedCL.ToString());
+				LogInformation("INFO: Packages successfully submitted in cl "+ SubmittedCL.ToString());
 			}
 		}
 
diff --git a/Engine/Source/Programs/AutomationTool/Scripts/ResavePluginDescriptors.cs b/Engine/Source/Programs/AutomationTool/Scripts/ResavePluginDescriptors.cs
--- a/Engine/Source/Programs/AutomationTool/Scripts/ResavePluginDescriptors.cs
+++ b/Engine/Source/Programs/AutomationTool/Scripts/ResavePluginDescriptors.cs
@@ -32,7 +32,7 @@ namespace AutomationTool
 
 			foreach(FileReference PluginFile in DirectoryReference.EnumerateFiles(new DirectoryReference(RootDirParam), "*.uplugin", System.IO.SearchOption.AllDirectories))
 			{
-				Log("Reading {0}", PluginFile);
+				LogInformation("Reading {0}", PluginFile);
 				string InputText = File.ReadAllText(PluginFile.FullName);
 
 				// Parse the descriptor
@@ -50,12 +50,12 @@ namespace AutomationTool
 				// Update the fields
 				if(CreatedBy != null && Descriptor.CreatedBy != CreatedBy)
 				{
-					Log("  Updating 'CreatedBy' field from '{0}' to '{1}'", Descriptor.CreatedBy ?? "<empty>", CreatedBy);
+					LogInformation("  Updating 'CreatedBy' field from '{0}' to '{1}'", Descriptor.CreatedBy ?? "<empty>", CreatedBy);
 					Descriptor.CreatedBy = CreatedBy;
 				}
 				if(CreatedByUrl != null)
 				{
-					Log("  Updating 'CreatedByURL' field from '{0}' to '{1}'", Descriptor.CreatedByURL ?? "<empty>", CreatedByUrl);
+					LogInformation("  Updating 'CreatedByURL' field from '{0}' to '{1}'", Descriptor.CreatedByURL ?? "<empty>", CreatedByUrl);
 					Descriptor.CreatedByURL = CreatedByUrl;
 				}
 
@@ -81,7 +81,7 @@ namespace AutomationTool
 						}
 						CommandUtils.SetFileAttributes(PluginFile.FullName, ReadOnly: false);
 					}
-					Log("  Writing updated file.", PluginFile);
+					LogInformation("  Writing updated file.", PluginFile);
 					FileReference.WriteAllText(PluginFile, OutputText);
 				}
 			}
diff --git a/Engine/Source/Programs/AutomationTool/Scripts/ResaveProjectDescriptors.cs b/Engine/Source/Programs/AutomationTool/Scripts/ResaveProjectDescriptors.cs
--- a/Engine/Source/Programs/AutomationTool/Scripts/ResaveProjectDescriptors.cs
+++ b/Engine/Source/Programs/AutomationTool/Scripts/ResaveProjectDescriptors.cs
@@ -51,7 +51,7 @@ namespace AutomationTool
 
 			foreach (FileReference ProjectFile in ProjectFiles)
 			{
-				Log("Reading {0}", ProjectFile);
+				LogInformation("Reading {0}", ProjectFile);
 				string InputText = File.ReadAllText(ProjectFile.FullName);
 
 				// Parse the descriptor
@@ -123,7 +123,7 @@ namespace AutomationTool
 						}
 						CommandUtils.SetFileAttributes(ProjectFile.FullName, ReadOnly: false);
 					}
-					Log("  Writing updated file.", ProjectFile);
+					LogInformation("  Writing updated file.", ProjectFile);
 					FileReference.WriteAllText(ProjectFile, OutputText);
 				}
 			}
diff --git a/Engine/Source/Programs/AutomationTool/Scripts/RunProjectCommand.Automation.cs b/Engine/Source/Programs/AutomationTool/Scripts/RunProjectCommand.Automation.cs
--- a/Engine/Source/Programs/AutomationTool/Scripts/RunProjectCommand.Automation.cs
+++ b/Engine/Source/Programs/AutomationTool/Scripts/RunProjectCommand.Automation.cs
@@ -94,7 +94,7 @@ public partial class Project : CommandUtils
 			return;
 		}
 
-		Log("********** RUN COMMAND STARTED **********");
+		LogInformation("********** RUN COMMAND STARTED **********");
 
 		var LogFolderOutsideOfSandbox = GetLogFolderOutsideOfSandbox();
 		if (!GlobalCommandLine.Installed && ServerProcess == null)
@@ -119,7 +119,7 @@ public partial class Project : CommandUtils
 			CopyLogsBackToLogFolder();
 		}
 
-		Log("********** RUN COMMAND COMPLETED **********");
+		LogInformation("********** RUN COMMAND COMPLETED **********");
 	}
 
 	private static void RunInternal(ProjectParams Params, string ServerLogFile, string ClientLogFile)
@@ -157,13 +157,13 @@ public partial class Project : CommandUtils
 
 		if (ServerProcess != null)
 		{
-			Log("Waiting a few seconds for the server to start...");
+			LogInformation("Waiting a few seconds for the server to start...");
 			Thread.Sleep(5000);
 		}
 
 		if (!Params.NoClient)
 		{
-			Log("Starting Client....");
+			LogInformation("Starting Client....");
 
 			var SC = CreateDeploymentContext(Params, false);
 
@@ -221,17 +221,17 @@ public partial class Project : CommandUtils
 				IProcessResult ClientProcess = null;
 				FileStream ClientProcessLog = null;
 				StreamReader ClientLogReader = null;
-				Log("Starting Client for unattended test....");
+				LogInformation("Starting Client for unattended test....");
 				ClientProcess = Run(ClientApp, ClientCmdLine + " -testexit=\"" + LookFor + "\"", null, ClientRunFlags | ERunOptions.NoWaitForExit);
 				while (!FileExists(ClientLogFile) && !ClientProcess.HasExited)
 				{
-					Log("Waiting for client logging process to start...{0}", ClientLogFile);
+					LogInformation("Waiting for client logging process to start...{0}", ClientLogFile);
 					Thread.Sleep(2000);
 				}
 				if (FileExists(ClientLogFile))
 				{
 					Thread.Sleep(2000);
-					Log("Client logging process started...{0}", ClientLogFile);
+					LogInformation("Client logging process started...{0}", ClientLogFile);
 					ClientProcessLog = File.Open(ClientLogFile, FileMode.Open, FileAccess.Read, FileShare.ReadWrite);
 					ClientLogReader = new StreamReader(ClientProcessLog);
 				}
@@ -264,7 +264,7 @@ public partial class Project : CommandUtils
 						if (AllClientOutput.LastIndexOf(LookFor) > AllClientOutput.IndexOf(LookFor))
 						{
 							WelcomedCorrectly = true;
-							Log("Test complete...");
+							LogInformation("Test complete...");
 							bKeepReading = false;
 						}
 						else if (Params.RunAutomationTests)
@@ -320,13 +320,13 @@ public partial class Project : CommandUtils
 					}
 					else if (bClientExited && (DateTime.UtcNow - ExitTime).TotalSeconds > 30)
 					{
-						Log("Client exited and has been quiet for 30 seconds...exiting");
+						LogInformation("Client exited and has been quiet for 30 seconds...exiting");
 						bKeepReading = false;
 					}
 				}
 				if (ClientProcess != null && !ClientProcess.HasExited)
 				{
-					Log("Client is supposed to exit, lets wait a while for it to exit naturally...");
+					LogInformation("Client is supposed to exit, lets wait a while for it to exit naturally...");
 					for (int i = 0; i < 120 && !ClientProcess.HasExited; i++)
 					{
 						Thread.Sleep(1000);
@@ -334,7 +334,7 @@ public partial class Project : CommandUtils
 				}
 				if (ClientProcess != null && !ClientProcess.HasExited)
 				{
-					Log("Stopping client...");
+					LogInformation("Stopping client...");
 					ClientProcess.StopProcess();
 					Thread.Sleep(10000);
 				}
@@ -413,7 +413,7 @@ public partial class Project : CommandUtils
 			{
 				while (!FileExists(ServerLogFile) && !ServerProcess.HasExited)
 				{
-					Log("Waiting for logging process to start...");
+					LogInformation("Waiting for logging process to start...");
 					Thread.Sleep(2000);
 				}
 				Thread.Sleep(1000);
@@ -439,33 +439,33 @@ public partial class Project : CommandUtils
 								if (ClientProcess == null &&
 										(AllServerOutput.Contains("Game Engine Initialized") || AllServerOutput.Contains("Unreal Network File Server is ready")))
 								{
-									Log("Starting Client for unattended test....");
+									LogInformation("Starting Client for unattended test....");
 									ClientProcess = Run(ClientApp, ClientCmdLine + " -testexit=\"" + LookFor + "\"", null, ClientRunFlags | ERunOptions.NoWaitForExit);
 									//@todo no testing is done on these
 									if (NumClients > 1 && NumClients < 9)
 									{
 										for (int i = 1; i < NumClients; i++)
 										{
-											Log("Starting Extra Client....");
+											LogInformation("Starting Extra Client....");
 											OtherClients.Add(Run(ClientApp, ClientCmdLine, null, ClientRunFlags | ERunOptions.NoWaitForExit));
 										}
 									}
 									while (!FileExists(ClientLogFile) && !ClientProcess.HasExited)
 									{
-										Log("Waiting for client logging process to start...{0}", ClientLogFile);
+										LogInformation("Waiting for client logging process to start...{0}", ClientLogFile);
 										Thread.Sleep(2000);
 									}
 									if (!ClientProcess.HasExited)
 									{
 										Thread.Sleep(2000);
-										Log("Client logging process started...{0}", ClientLogFile);
+										LogInformation("Client logging process started...{0}", ClientLogFile);
 										ClientProcessLog = File.Open(ClientLogFile, FileMode.Open, FileAccess.Read, FileShare.ReadWrite);
 										ClientLogReader = new StreamReader(ClientProcessLog);
 									}
 								}
 								else if (ClientProcess == null && !ServerProcess.HasExited)
 								{
-									Log("Waiting for server to start....");
+									LogInformation("Waiting for server to start....");
 									Thread.Sleep(2000);
 								}
 								if (ClientProcess != null && ClientProcess.HasExited)
@@ -496,12 +496,12 @@ public partial class Project : CommandUtils
 										{
 											if (Params.FakeClient)
 											{
-												Log("Welcomed by server or client loaded, lets wait ten minutes...");
+												LogInformation("Welcomed by server or client loaded, lets wait ten minutes...");
 												Thread.Sleep(60000 * 10);
 											}
 											else
 											{
-												Log("Welcomed by server or client loaded, lets wait 30 seconds...");
+												LogInformation("Welcomed by server or client loaded, lets wait 30 seconds...");
 												Thread.Sleep(30000);
 											}
 											WelcomedCorrectly = true;
@@ -544,7 +544,7 @@ public partial class Project : CommandUtils
 					AllClientOutput += Output;
 					if (AllClientOutput.Contains("Game Engine Initialized") || AllClientOutput.Contains("Unreal Network File Server is ready"))
 					{
-						Log("Starting Client....");
+						LogInformation("Starting Client....");
 						var SC = DeployContextList[0];
 						ClientProcess = SC.StageTargetPlatform.RunClient(ClientRunFlags | ERunOptions.NoWaitForExit, ClientApp, ClientCmdLine, Params);
 //						ClientProcess = Run(ClientApp, ClientCmdLine, null, ClientRunFlags | ERunOptions.NoWaitForExit);
@@ -552,7 +552,7 @@ public partial class Project : CommandUtils
 						{
 							for (int i = 1; i < NumClients; i++)
 							{
-								Log("Starting Extra Client....");
+								LogInformation("Starting Extra Client....");
 								IProcessResult NewClient = SC.StageTargetPlatform.RunClient(ClientRunFlags | ERunOptions.NoWaitForExit, ClientApp, ClientCmdLine, Params);
 								OtherClients.Add(NewClient);
 							}
@@ -561,7 +561,7 @@ public partial class Project : CommandUtils
 				}
 				else if (ClientProcess == null && !ServerProcess.HasExited)
 				{
-					Log("Waiting for server to start....");
+					LogInformation("Waiting for server to start....");
 					Thread.Sleep(2000);
 				}
 
@@ -573,7 +573,7 @@ public partial class Project : CommandUtils
 				if (ClientProcess != null && ClientProcess.HasExited)
 				{
 
-					Log("Client exited, stopping server....");
+					LogInformation("Client exited, stopping server....");
 					if (!GlobalCommandLine.NoKill)
 					{
 						ServerProcess.StopProcess();
@@ -584,7 +584,7 @@ public partial class Project : CommandUtils
 				return bKeepReading; // Keep reading
 			});
 		}
-		Log("Server exited....");
+		LogInformation("Server exited....");
 		if (ClientProcess != null && !ClientProcess.HasExited)
 		{
 			ClientProcess.StopProcess();
@@ -904,8 +904,8 @@ public partial class Project : CommandUtils
 			ClientCmdLine += "-Params=\"" + TempCmdLine + "\"";
 			ClientApp = CombinePaths(CmdEnv.LocalRoot, "Engine/Binaries/Win64/UnrealFrontend.exe");
 
-			Log("Launching via UFE:");
-			Log("\tClientCmdLine: " + ClientCmdLine + "");
+			LogInformation("Launching via UFE:");
+			LogInformation("\tClientCmdLine: " + ClientCmdLine + "");
 		}
 		else
 		{
@@ -926,7 +926,7 @@ public partial class Project : CommandUtils
 		{
 			if (String.IsNullOrEmpty(Output) == false)
 			{
-				Log(Output);
+				LogInformation(Output);
 			}
 			return true;
 		});
@@ -1078,7 +1078,7 @@ public partial class Project : CommandUtils
 #endif
 		var UnrealFileServerExe = HostPlatform.Current.GetUE4ExePath("UnrealFileServer.exe");
 
-		Log("Running UnrealFileServer *******");
+		LogInformation("Running UnrealFileServer *******");
 		var Args = String.Format("{0} -abslog={1} -unattended -CrashForUAT -log {2}",
 						CommandUtils.MakePathSafeToUseWithCommandLine(Params.RawProjectPath.FullName),
 						CommandUtils.MakePathSafeToUseWithCommandLine(ServerLogFile),
diff --git a/Engine/Source/Programs/AutomationTool/Scripts/SharedCookedBuild.cs b/Engine/Source/Programs/AutomationTool/Scripts/SharedCookedBuild.cs
--- a/Engine/Source/Programs/AutomationTool/Scripts/SharedCookedBuild.cs
+++ b/Engine/Source/Programs/AutomationTool/Scripts/SharedCookedBuild.cs
@@ -36,7 +36,7 @@ public class SharedCookedBuild
 		List<string> CookedBuildPaths;
 		if (Hierarchy.GetArray("SharedCookedBuildSettings", "SharedCookedBuildPath", out CookedBuildPaths) == false)
 		{
-			CommandUtils.Log("Unable to copy shared cooked build: SharedCookedBuildPath not set in Engine.ini SharedCookedBuildSettings");
+			CommandUtils.LogInformation("Unable to copy shared cooked build: SharedCookedBuildPath not set in Engine.ini SharedCookedBuildSettings");
 			return false;
 		}
 
@@ -49,13 +49,13 @@ public class SharedCookedBuild
 			{
 				if (CurrentCL == "" && FinalCookedBuildPath.Contains("[CL]"))
 				{
-					CommandUtils.Log("Unable to copy shared cooked build: Unable to determine CL number from P4 or UGS, and is required by SharedCookedBuildPath");
+					CommandUtils.LogInformation("Unable to copy shared cooked build: Unable to determine CL number from P4 or UGS, and is required by SharedCookedBuildPath");
 					return false;
 				}
 
 				if (CurrentCL == "" && FinalCookedBuildPath.Contains("[BRANCHNAME]"))
 				{
-					CommandUtils.Log("Unable to copy shared cooked build: Unable to determine BRANCHNAME number from P4 or UGS, and is required by SharedCookedBuildPath");
+					CommandUtils.LogInformation("Unable to copy shared cooked build: Unable to determine BRANCHNAME number from P4 or UGS, and is required by SharedCookedBuildPath");
 					return false;
 				}
 
@@ -168,24 +168,24 @@ public class SharedCookedBuild
 			BuildCL = File.ReadAllText(SyncedBuildFile);
 		}
 
-		CommandUtils.Log("Attempting download of latest shared build from {0}", CookedBuildPath);
+		CommandUtils.LogInformation("Attempting download of latest shared build from {0}", CookedBuildPath);
 
 		string SavedBuildCL = string.Format( "{0} {1}", CookedBuildPath, bOnlyCopyAssetRegistry ? "RegistryOnly" : "" );
 
 		if (BuildCL == SavedBuildCL)
 		{
-			CommandUtils.Log("Already downloaded latest shared build at CL {0}", SavedBuildCL);
+			CommandUtils.LogInformation("Already downloaded latest shared build at CL {0}", SavedBuildCL);
 			return false;
 		}
 
 		if (Directory.Exists(CookedBuildPath) == false)
 		{
-			CommandUtils.Log("Unable to copy shared cooked build: Unable to find shared build at location {0} check SharedCookedBuildPath in Engine.ini SharedCookedBuildSettings is correct", CookedBuildPath);
+			CommandUtils.LogInformation("Unable to copy shared cooked build: Unable to find shared build at location {0} check SharedCookedBuildPath in Engine.ini SharedCookedBuildSettings is correct", CookedBuildPath);
 			return false;
 		}
 		
 		// delete all the stuff
-		CommandUtils.Log("Deleting previous shared build because it was out of date");
+		CommandUtils.LogInformation("Deleting previous shared build because it was out of date");
 		CommandUtils.DeleteDirectory(LocalPath);
 		Directory.CreateDirectory(LocalPath);
 
diff --git a/Engine/Source/Programs/AutomationTool/Scripts/SubmitUtililzationReportToEC.Automation.cs b/Engine/Source/Programs/AutomationTool/Scripts/SubmitUtililzationReportToEC.Automation.cs
--- a/Engine/Source/Programs/AutomationTool/Scripts/SubmitUtililzationReportToEC.Automation.cs
+++ b/Engine/Source/Programs/AutomationTool/Scripts/SubmitUtililzationReportToEC.Automation.cs
@@ -25,7 +25,7 @@ class SubmitUtilizationReportToEC : BuildCommand
 		if (string.IsNullOrEmpty(DayParam))
 		{
 			Day = DateTime.Today.Add(new TimeSpan(1,0,0,0));
-			Log("Day parameter not specified, defaulting to today's report: " + Day.Date.ToString("MM-dd-yyyy"));
+			LogInformation("Day parameter not specified, defaulting to today's report: " + Day.Date.ToString("MM-dd-yyyy"));
 		}
 		else
 		{
@@ -68,13 +68,13 @@ class SubmitUtilizationReportToEC : BuildCommand
 			// delete anything older than two weeks
 			if(ExistingDate < TwoWeeksAgo)
 			{
-				Log("Deleting out-of-date property: " + prop.InnerText);
+				LogInformation("Deleting out-of-date property: " + prop.InnerText);
 				CommandUtils.RunAndLog("ectool", string.Format("--timeout 900 deleteProperty --propertyName \"/projects/GUBP_V5/Generated/Utilization2/{0}\"", prop.InnerText), Options: Opts);
 			}
 
 		}
 		// add today's
-		Log("Adding report: " + Day.Subtract(new TimeSpan(1,0,0,0)).ToString("yyyy-MM-dd") + "->" + Day.ToString("yyyy-MM-dd"));
+		LogInformation("Adding report: " + Day.Subtract(new TimeSpan(1,0,0,0)).ToString("yyyy-MM-dd") + "->" + Day.ToString("yyyy-MM-dd"));
 		CommandUtils.RunAndLog("ectool", string.Format("--timeout 900 setProperty \"/projects/GUBP_V5/Generated/Utilization2/Report_{0}\" --valueFile {1}", Day.Subtract(new TimeSpan(1,0,0,0)).ToString("yyyy-MM-dd"), FileName), Options: Opts);
 	}
 }
diff --git a/Engine/Source/Programs/AutomationTool/Scripts/TestSnippet.Automation.cs b/Engine/Source/Programs/AutomationTool/Scripts/TestSnippet.Automation.cs
--- a/Engine/Source/Programs/AutomationTool/Scripts/TestSnippet.Automation.cs
+++ b/Engine/Source/Programs/AutomationTool/Scripts/TestSnippet.Automation.cs
@@ -8,6 +8,6 @@ class TestSnippetUtils : CommandUtils
 {
 	public static void Test(params object[] Args)
 	{
-		Log("Args[0]=" + Args[0].ToString());
+		LogInformation("Args[0]=" + Args[0].ToString());
 	}
 }
\ No newline at end of file
diff --git a/Engine/Source/Programs/AutomationTool/Scripts/Tests.Automation.cs b/Engine/Source/Programs/AutomationTool/Scripts/Tests.Automation.cs
--- a/Engine/Source/Programs/AutomationTool/Scripts/Tests.Automation.cs
+++ b/Engine/Source/Programs/AutomationTool/Scripts/Tests.Automation.cs
@@ -20,8 +20,8 @@ class TestP4_Info : BuildCommand
 {
 	public override void ExecuteBuild()
 	{
-		Log("P4CLIENT={0}", GetEnvVar("P4CLIENT"));
-		Log("P4PORT={0}", GetEnvVar("P4PORT"));
+		LogInformation("P4CLIENT={0}", GetEnvVar("P4CLIENT"));
+		LogInformation("P4PORT={0}", GetEnvVar("P4PORT"));
 
 		var Result = P4.P4("info");
 		if (Result.ExitCode != 0)
@@ -74,7 +74,7 @@ class GitPullRequest : BuildCommand
 			throw new AutomationException("Unable to find Git.exe in the system path under a 'PortableGit' subdirectory.  Make sure the GitHub client is installed, and you are running this script from within a GitHub Command Shell.  We want to make sure we're using the correct version of Git, in case multiple versions are on the computer, which is why we check for a PortableGit folder that the GitHub Shell uses.");
 		}
 
-		Log("Running {0} {1}", GitExePath, GitCommandLine);
+		LogInformation("Running {0} {1}", GitExePath, GitCommandLine);
 
 		IProcessResult Result = Run(App: GitExePath, CommandLine: GitCommandLine, Options: (ERunOptions.NoLoggingOfRunCommand | ERunOptions.AllowSpew | ERunOptions.AppMustExist));
 		if (Result.ExitCode != 0)
@@ -116,7 +116,7 @@ class GitPullRequest : BuildCommand
 				{
 					throw new AutomationException("Unrecognized branch.");
 				}
-				Log("Depot {0}", Depot);
+				LogInformation("Depot {0}", Depot);
 
 				Base = Base.Substring(0, Base.IndexOf(BaseEnd));
 				if (!Base.Contains(" "))
@@ -124,10 +124,10 @@ class GitPullRequest : BuildCommand
 					throw new AutomationException("Unrecognized commit3.");
 				}
 				Base = Base.Substring(Base.LastIndexOf(" "));
-				Log("CL String {0}", Base);
+				LogInformation("CL String {0}", Base);
 				CL = int.Parse(Base);
 			}
-			Log("CL int {0}", CL);
+			LogInformation("CL int {0}", CL);
 			if (CL < 2000000 || String.IsNullOrWhiteSpace(Depot))
 			{
 				throw new AutomationException("Unrecognized commit3.");
@@ -161,7 +161,7 @@ class GitPullRequest : BuildCommand
 					throw new AutomationException("Unrecognized commit5.");
 				}
 				CLStuff = CLStuff.Substring(0, CLStuff.IndexOf(" "));
-				Log("CL String {0}", CLStuff);
+				LogInformation("CL String {0}", CLStuff);
 				CL = int.Parse(CLStuff);
 
 				var BranchStuff = Base.Substring(Base.IndexOf(LiveStart) + LiveStart.Length, Base.IndexOf(LiveEnd) - Base.IndexOf(LiveStart) - LiveStart.Length);
@@ -170,7 +170,7 @@ class GitPullRequest : BuildCommand
 					throw new AutomationException("Unrecognized commit6.");
 				}
 				BranchStuff = BranchStuff.Substring(BranchStuff.LastIndexOf(" ") + 1);
-				Log("Branch String {0}", BranchStuff);
+				LogInformation("Branch String {0}", BranchStuff);
 				if (BranchStuff.StartsWith("4."))
 				{
 					Depot = "//depot/UE4-Releases/4." + BranchStuff.Substring(2, 1);
@@ -188,7 +188,7 @@ class GitPullRequest : BuildCommand
 					throw new AutomationException("Unrecognized branch2.");
 				}
 			}
-			Log("CL int {0}", CL);
+			LogInformation("CL int {0}", CL);
 			if (CL < 2000000 || String.IsNullOrWhiteSpace(Depot))
 			{
 				throw new AutomationException("Unrecognized commit3.");
@@ -256,7 +256,7 @@ class GitPullRequest : BuildCommand
 
 		foreach (var p in NewClient.View)
 		{
-			Log("{0} = {1}", p.Key, p.Value);
+			LogInformation("{0} = {1}", p.Key, p.Value);
 		}
 
 		string TestClient = P4Env.User + "_" + Environment.MachineName + "_PullRequests_" + Depot.Replace("/", "_");
@@ -387,7 +387,7 @@ class GitPullRequest : BuildCommand
 			}
 			catch (Exception Ex)
 			{
-				Log(" Exception was {0}", LogUtils.FormatException(Ex));
+				LogInformation(" Exception was {0}", LogUtils.FormatException(Ex));
 				Failures.Add(String.Format("PR {0} Failed with {1}", PR, LogUtils.FormatException(Ex)));
 			}
 		}
@@ -395,7 +395,7 @@ class GitPullRequest : BuildCommand
 
 		foreach (var Failed in Failures)
 		{
-			Log("{0}", Failed);
+			LogInformation("{0}", Failed);
 		}
 	}
 }
@@ -414,7 +414,7 @@ class TestSuccess : BuildCommand
 {
 	public override void ExecuteBuild()
 	{
-		Log("TestSuccess message.");
+		LogInformation("TestSuccess message.");
 	}
 }
 
@@ -423,11 +423,11 @@ class TestMessage : BuildCommand
 {
 	public override void ExecuteBuild()
 	{
-		Log("TestMessage message.");
-		Log("you must be in error");
-		Log("Behold the Error of you ways");
-		Log("ERROR, you are silly");
-		Log("ERROR: Something must be broken");
+		LogInformation("TestMessage message.");
+		LogInformation("you must be in error");
+		LogInformation("Behold the Error of you ways");
+		LogInformation("ERROR, you are silly");
+		LogInformation("ERROR: Something must be broken");
 	}
 }
 
@@ -436,9 +436,9 @@ class TestRecursion : BuildCommand
 {
 	public override void ExecuteBuild()
 	{
-		Log("TestRecursion *********************");
+		LogInformation("TestRecursion *********************");
 		string Params = ParseParamValue("Cmd");
-		Log("Recursive Command: {0}", Params);
+		LogInformation("Recursive Command: {0}", Params);
 		RunUAT(CmdEnv, Params, "Recur");
 
 	}
@@ -448,7 +448,7 @@ class TestRecursionAuto : BuildCommand
 {
 	public override void ExecuteBuild()
 	{
-		Log("TestRecursionAuto *********************");
+		LogInformation("TestRecursionAuto *********************");
 		RunUAT(CmdEnv, "TestMessage", "TestMessage");
 		RunUAT(CmdEnv, "TestMessage", "TestMessage");
 		RunUAT(CmdEnv, "TestRecursion -Cmd=TestMessage", "TestRecursion");
@@ -462,7 +462,7 @@ class TestMacZip : BuildCommand
 {
 	public override void ExecuteBuild()
 	{
-		Log("TestMacZip *********************");
+		LogInformation("TestMacZip *********************");
 
 		if (UnrealBuildTool.Utils.IsRunningOnMono)
 		{
@@ -483,14 +483,14 @@ class TestP4_CreateChangelist : BuildCommand
 {
 	public override void ExecuteBuild()
 	{
-		Log("P4CLIENT={0}", GetEnvVar("P4CLIENT"));
-		Log("P4PORT={0}", GetEnvVar("P4PORT"));
+		LogInformation("P4CLIENT={0}", GetEnvVar("P4CLIENT"));
+		LogInformation("P4PORT={0}", GetEnvVar("P4PORT"));
 
 		var CLDescription = "AutomationTool TestP4";
 
-		Log("Creating new changelist \"{0}\" using client \"{1}\"", CLDescription, GetEnvVar("P4CLIENT"));
+		LogInformation("Creating new changelist \"{0}\" using client \"{1}\"", CLDescription, GetEnvVar("P4CLIENT"));
 		var ChangelistNumber = P4.CreateChange(Description: CLDescription);
-		Log("Created changelist {0}", ChangelistNumber);
+		LogInformation("Created changelist {0}", ChangelistNumber);
 	}
 }
 
@@ -501,12 +501,12 @@ class TestP4_StrandCheckout : BuildCommand
 	{
 		int WorkingCL = P4.CreateChange(P4Env.Client, String.Format("TestP4_StrandCheckout, head={0}", P4Env.Changelist));
 
-		Log("Build from {0}    Working in {1}", P4Env.Changelist, WorkingCL);
+		LogInformation("Build from {0}    Working in {1}", P4Env.Changelist, WorkingCL);
 
 		List<string> Sign = new List<string>();
 		Sign.Add(CombinePaths(CmdEnv.LocalRoot, @"\Engine\Binaries\DotNET\AgentInterface.dll"));
 
-		Log("Signing and adding {0} build products to changelist {1}...", Sign.Count, WorkingCL);
+		LogInformation("Signing and adding {0} build products to changelist {1}...", Sign.Count, WorkingCL);
 
 		CodeSign.SignMultipleIfEXEOrDLL(this, Sign);
 		foreach (var File in Sign)
@@ -529,13 +529,13 @@ class TestP4_LabelDescription : BuildCommand
 	{
 		string Label = GetEnvVar("SBF_LabelFromUser");
 		string Desc;
-		Log("LabelDescription {0}", Label);
+		LogInformation("LabelDescription {0}", Label);
 		var Result = P4.LabelDescription(Label, out Desc);
 		if (!Result)
 		{
 			throw new AutomationException("Could not get label description");
 		}
-		Log("Result***\n{0}\n***\n", Desc);
+		LogInformation("Result***\n{0}\n***\n", Desc);
 	}
 }
 
@@ -608,8 +608,8 @@ class TestP4_ClientOps : BuildCommand
 					throw new AutomationException("{0} does not contain {1} and it is supposed to.", depot, P4Env.Branch);
 				}
 				string local = CombinePaths(depot.Replace(SlashRoot, LocalRoot));
-				Log("{0}", depot);
-				Log("    {0}", local);
+				LogInformation("{0}", depot);
+				LogInformation("    {0}", local);
 			}
 		}
 
@@ -634,7 +634,7 @@ class CleanDDC : BuildCommand
 {
 	public override void ExecuteBuild()
 	{
-		Log("*********************** Clean DDC");
+		LogInformation("*********************** Clean DDC");
 
 		bool DoIt = ParseParam("DoIt");
 
@@ -655,10 +655,10 @@ class CleanDDC : BuildCommand
 						if (FileExists_NoExceptions(file))
 						{
 							FileInfo Info = new FileInfo(file);
-							Log("{0}", file);
+							LogInformation("{0}", file);
 							if ((DateTime.UtcNow - Info.LastWriteTimeUtc).TotalDays > 20)
 							{
-								Log("  is old.");
+								LogInformation("  is old.");
 								if (DoIt)
 								{
 									DeleteFile_NoExceptions(file);
@@ -666,7 +666,7 @@ class CleanDDC : BuildCommand
 							}
 							else
 							{
-								Log("  is NOT old.");
+								LogInformation("  is NOT old.");
 							}
 						}
 					}
@@ -681,7 +681,7 @@ class TestTestFarm : BuildCommand
 {
 	public override void ExecuteBuild()
 	{
-		Log("*********************** TestTestFarm");
+		LogInformation("*********************** TestTestFarm");
 
 		string Exe = CombinePaths(CmdEnv.LocalRoot, "Engine", "Binaries", "Win64", "UE4Editor.exe");
 		string ClientLogFile = CombinePaths(CmdEnv.LogFolder, "HoverGameRun");
@@ -694,7 +694,7 @@ class TestTestFarm : BuildCommand
 			if (!String.IsNullOrEmpty(Output) &&
 					Output.Contains("Game Engine Initialized"))
 			{
-				Log("It looks started, let me kill it....");
+				LogInformation("It looks started, let me kill it....");
 				App.StopProcess();
 			}
 			return true; //keep reading
@@ -707,10 +707,10 @@ class TestArguments : BuildCommand
 {
 	public override void ExecuteBuild()
 	{
-		Log("List of provided arguments: ");
+		LogInformation("List of provided arguments: ");
 		for (int Index = 0; Index < Params.Length; ++Index)
 		{
-			Log("{0}={1}", Index, Params[Index].ToString());
+			LogInformation("{0}={1}", Index, Params[Index].ToString());
 		}
 
 		string[] TestArgs = new string[] { "NoXGE", "SomeArg", "Map=AwesomeMap", "run=whatever", "Stuff" };
@@ -763,7 +763,7 @@ public class TestCombinePaths : BuildCommand
 		for (int Index = 0; Index < Results.Length; ++Index)
 		{
 			var CombinedPath = Results[Index];
-			Log(CombinedPath);
+			LogInformation(CombinedPath);
 			if (CombinedPath != Results[0])
 			{
 				throw new AutomationException("Path {0} does not match path 0: {1} (expected: {2})", Index, CombinedPath, Results[0]);
@@ -787,7 +787,7 @@ class TestLog : BuildCommand
 {
 	public override void ExecuteBuild()
 	{
-		Log("************************* ShooterGameRun");
+		LogInformation("************************* ShooterGameRun");
 		string MapToRun = ParseParamValue("map", "/game/maps/sanctuary");
 
 		PushDir(CombinePaths(CmdEnv.LocalRoot, @"Engine/Binaries/Win64/"));
@@ -838,7 +838,7 @@ class TestGamePerf : BuildCommand
 
 	public override void ExecuteBuild()
 	{
-		Log("*********************** TestGamePerf");
+		LogInformation("*********************** TestGamePerf");
 
 		string UserName = "test.automation";
 		string Password = "JJGfh9CX";
@@ -1038,13 +1038,13 @@ class TestGamePerf : BuildCommand
 						{
 							if (!File.Exists(ClientLogFile))
 							{
-								Log("Creating log file");
+								LogInformation("Creating log file");
 
 								File.Create(ClientLogFile).Close();
 
-								Log(ClientLogFile);
+								LogInformation(ClientLogFile);
 
-								Log("Log file created");
+								LogInformation("Log file created");
 							}
 
 							RunAndLog(Exe, CmdLine, ClientLogFile);
@@ -1082,7 +1082,7 @@ class TestGamePerf : BuildCommand
 
 						catch (Exception Err)
 						{
-							Log(Err.Message);
+							LogInformation(Err.Message);
 						}
 					}
 				}
@@ -1091,7 +1091,7 @@ class TestGamePerf : BuildCommand
 
 		catch (Exception Err)
 		{
-			Log(Err.Message);
+			LogInformation(Err.Message);
 		}
 	}
 }
@@ -1124,10 +1124,10 @@ public class TestUATBuildProducts : BuildCommand
 		UE4Build Build = new UE4Build(this);
 		Build.AddUATFilesToBuildProducts();
 
-		Log("Build products:");
+		LogInformation("Build products:");
 		foreach (var Product in Build.BuildProductFiles)
 		{
-			Log("  " + Product);
+			LogInformation("  " + Product);
 		}
 	}
 }
@@ -1154,7 +1154,7 @@ class TestOSSCommands : BuildCommand
 		}
 		catch (Exception Err)
 		{
-			Log(Err.Message);
+			LogInformation(Err.Message);
 		}
 
 		FileInfo[] LogFiles = new DirectoryInfo(LogDirectory).GetFiles();
@@ -1262,15 +1262,15 @@ public class UBT : BuildCommand
 
 		Agenda.AddTargets(Targets.ToArray(), Platform, Configuration);
 
-		Log("UBT Buid");
-		Log("Targets={0}", String.Join(",", Targets));
-		Log("Platform={0}", Platform);
-		Log("Configuration={0}", Configuration);
-		Log("Clean={0}", Clean);
+		LogInformation("UBT Buid");
+		LogInformation("Targets={0}", String.Join(",", Targets));
+		LogInformation("Platform={0}", Platform);
+		LogInformation("Configuration={0}", Configuration);
+		LogInformation("Clean={0}", Clean);
 
 		Build.Build(Agenda, InUpdateVersionFiles: false);
 
-		Log("UBT Completed");
+		LogInformation("UBT Completed");
 	}
 }
 
@@ -1339,7 +1339,7 @@ public class SyncSource : BuildCommand
 		foreach (var AlwaysSyncSubFolder in AlwaysSync)
 		{
 			var SyncCmdLine = CombinePaths(PathSeparator.Depot, DepotPath, AlwaysSyncSubFolder, "...") + Changelist;
-			Log("Syncing {0}", SyncCmdLine, Changelist);
+			LogInformation("Syncing {0}", SyncCmdLine, Changelist);
 			SafeSync(SyncCmdLine);
 		}
 
@@ -1355,7 +1355,7 @@ public class SyncSource : BuildCommand
 				if (!ExclusionList.Contains(SubDirName))
 				{
 					var SyncCmdLine = CombinePaths(PathSeparator.Depot, SubDir, "...") + Changelist;
-					Log("Syncing {0}", SyncCmdLine);
+					LogInformation("Syncing {0}", SyncCmdLine);
 					SafeSync(SyncCmdLine);
 				}
 			}
@@ -1432,7 +1432,7 @@ class DumpBranch : BuildCommand
 
 	public override void ExecuteBuild()
 	{
-		Log("************************* DumpBranch");
+		LogInformation("************************* DumpBranch");
 
 		var HostPlatforms = new List<UnrealTargetPlatform>();
 		HostPlatforms.Add(UnrealTargetPlatform.Win64);
@@ -1473,7 +1473,7 @@ class TestBlame : BuildCommand
 		var Filename = ParseParamValue("File", "//depot/UE4/Engine/Source/Runtime/PakFile/Public/IPlatformFilePak.h");
 		var OutFilename = ParseParamValue("Out");
 
-		Log("Creating blame file for {0}", Filename);
+		LogInformation("Creating blame file for {0}", Filename);
 		P4Connection.BlameLineInfo[] Result = null;
 		if (ParseParam("Timelapse"))
 		{
@@ -1487,7 +1487,7 @@ class TestBlame : BuildCommand
 		foreach (var BlameLine in Result)
 		{
 			var ResultLine = String.Format("#{0} in {1} by {2}: {3}", BlameLine.Revision.Revision, BlameLine.Revision.Changelist, BlameLine.Revision.User, BlameLine.Line);
-			Log(ResultLine);
+			LogInformation(ResultLine);
 			BlameResult.AppendLine(ResultLine);
 		}
 		if (String.IsNullOrEmpty(OutFilename) == false)
@@ -1514,7 +1514,7 @@ class TestChanges : BuildCommand
 			}
 			foreach (var Record in ChangeRecords)
 			{
-				Log("{0} {1} {2}", Record.CL, Record.UserEmail, Record.Summary);
+				LogInformation("{0} {1} {2}", Record.CL, Record.UserEmail, Record.Summary);
 			}
 		}
 		{
@@ -1525,7 +1525,7 @@ class TestChanges : BuildCommand
 			}
 			foreach (var Record in ChangeRecords)
 			{
-				Log("{0} {1} {2}", Record.CL, Record.UserEmail, Record.Summary);
+				LogInformation("{0} {1} {2}", Record.CL, Record.UserEmail, Record.Summary);
 			}
 		}
 	}
@@ -1536,7 +1536,7 @@ class TestKillAll : BuildCommand
 {
 	public override void ExecuteBuild()
 	{
-		Log("*********************** TestKillAll");
+		LogInformation("*********************** TestKillAll");
 
 		string Exe = CombinePaths(CmdEnv.LocalRoot, "Engine", "Binaries", "Win64", "UE4Editor.exe");
 		string ClientLogFile = CombinePaths(CmdEnv.LogFolder, "HoverGameRun");
@@ -1552,7 +1552,7 @@ class TestCleanFormalBuilds : BuildCommand
 {
 	public override void ExecuteBuild()
 	{
-		Log("*********************** TestCleanFormalBuilds");
+		LogInformation("*********************** TestCleanFormalBuilds");
 		var Dir = ParseParamValue("Dir", @"P:\Builds\UE4\PackagedBuilds\++UE4+Release-4.11");
 		CleanFormalBuilds(Dir, "CL-*");
 	}
@@ -1564,7 +1564,7 @@ class TestStopProcess : BuildCommand
 {
 	public override void ExecuteBuild()
 	{
-		Log("*********************** TestStopProcess");
+		LogInformation("*********************** TestStopProcess");
 
 		string Exe = CombinePaths(CmdEnv.LocalRoot, "Engine", "Binaries", "Win64", "UE4Editor.exe");
 		string ClientLogFile = CombinePaths(CmdEnv.LogFolder, "HoverGameRun");
@@ -1572,13 +1572,13 @@ class TestStopProcess : BuildCommand
 
 		for (int Attempt = 0; Attempt < 5; ++Attempt)
 		{
-			Log("Attempt: {0}", Attempt);
+			LogInformation("Attempt: {0}", Attempt);
 			var Proc = Run(Exe, CmdLine, null, ERunOptions.AllowSpew | ERunOptions.NoWaitForExit | ERunOptions.AppMustExist | ERunOptions.NoStdOutRedirect);
 			Thread.Sleep(10000);
 			Proc.StopProcess();
 		}
 
-		Log("One final attempt to test KillAll");
+		LogInformation("One final attempt to test KillAll");
 		Run(Exe, CmdLine, null, ERunOptions.AllowSpew | ERunOptions.NoWaitForExit | ERunOptions.AppMustExist | ERunOptions.NoStdOutRedirect);
 		Thread.Sleep(10000);
 	}
diff --git a/Engine/Source/Programs/AutomationTool/Scripts/UE4BuildUtils.cs b/Engine/Source/Programs/AutomationTool/Scripts/UE4BuildUtils.cs
--- a/Engine/Source/Programs/AutomationTool/Scripts/UE4BuildUtils.cs
+++ b/Engine/Source/Programs/AutomationTool/Scripts/UE4BuildUtils.cs
@@ -51,7 +51,7 @@ public class UE4BuildUtils : CommandUtils
 			throw new AutomationException("Target is required when calling UE4BuildUtils.BuildProduct");
 		}
 
-		Log("Building {0}", Target.TargetName);
+		LogInformation("Building {0}", Target.TargetName);
 
 		if (Command == null)
 		{
diff --git a/Engine/Source/Programs/AutomationTool/Scripts/ZipUtils.Automation.cs b/Engine/Source/Programs/AutomationTool/Scripts/ZipUtils.Automation.cs
--- a/Engine/Source/Programs/AutomationTool/Scripts/ZipUtils.Automation.cs
+++ b/Engine/Source/Programs/AutomationTool/Scripts/ZipUtils.Automation.cs
@@ -20,7 +20,7 @@ public class ZipUtils : BuildCommand
 	public override ExitCode Execute()
 	{
 
-		Log("************************ STARTING ZIPUTILS ************************");
+		LogInformation("************************ STARTING ZIPUTILS ************************");
 
 		if (Params.Length < 2)
 		{
@@ -57,7 +57,7 @@ public class ZipUtils : BuildCommand
 		{
 			string OutputFolderPath = ExtractArgs[0];
 
-			Log("Running unzip : {0} OuputFolder:{1}", ZipFilePath.ToString(), OutputFolderPath);
+			LogInformation("Running unzip : {0} OuputFolder:{1}", ZipFilePath.ToString(), OutputFolderPath);
 
 			if (!System.IO.File.Exists(ZipFilePath))
 			{
@@ -74,10 +74,10 @@ public class ZipUtils : BuildCommand
 			}
 			else
 			{
-				Log("List of files extracted:");
+				LogInformation("List of files extracted:");
 				foreach (var file in ExtractedFiles)
 				{
-					Log("\t{0}", file);
+					LogInformation("\t{0}", file);
 				}
 			}
 		}
@@ -85,8 +85,8 @@ public class ZipUtils : BuildCommand
 		{
 
 			string ArchiveFolder = AddArgs[0];
-			Log("Compress level : {0}", CompressionLevel);
-			Log("Running zip : {0}", ArchiveFolder.ToString());
+			LogInformation("Compress level : {0}", CompressionLevel);
+			LogInformation("Running zip : {0}", ArchiveFolder.ToString());
 
 			FileAttributes attr = File.GetAttributes(ArchiveFolder);
 
@@ -99,7 +99,7 @@ public class ZipUtils : BuildCommand
 			InternalZipFiles(new FileReference(ZipFilePath), new DirectoryReference(ArchiveFolder), Filter, CompressionLevel);
 		}
 
-		Log("************************ ZIPUTIL WORK COMPLETED ************************");
+		LogInformation("************************ ZIPUTIL WORK COMPLETED ************************");
 
 		return ExitCode.Success;
 	}
diff --git a/Engine/Source/Programs/AutomationTool/Win/Tasks/SrcSrvTask.cs b/Engine/Source/Programs/AutomationTool/Win/Tasks/SrcSrvTask.cs
--- a/Engine/Source/Programs/AutomationTool/Win/Tasks/SrcSrvTask.cs
+++ b/Engine/Source/Programs/AutomationTool/Win/Tasks/SrcSrvTask.cs
@@ -156,7 +156,7 @@ namespace Win.Automation
                 {
                     foreach (string Message in Messages)
                     {
-                        CommandUtils.Log(Message);
+                        CommandUtils.LogInformation(Message);
                     }
                 }
 
diff --git a/Engine/Source/Programs/AutomationTool/Win/WinPlatform.Automation.cs b/Engine/Source/Programs/AutomationTool/Win/WinPlatform.Automation.cs
--- a/Engine/Source/Programs/AutomationTool/Win/WinPlatform.Automation.cs
+++ b/Engine/Source/Programs/AutomationTool/Win/WinPlatform.Automation.cs
@@ -263,7 +263,7 @@ public abstract class BaseWinPlatform : Platform
 		List<StagedFileReference> FilesInTargetDir = SC.FilesToStage.NonUFSFiles.Keys.Where(x => x.IsUnderDirectory(StagedBinariesDir) && (x.HasExtension(".exe") || x.HasExtension(".dll"))).ToList();
 		if(FilesInTargetDir.Count > 0)
 		{
-			Log("Copying AppLocal dependencies from {0} to {1}", BaseAppLocalDependenciesPath, StagedBinariesDir);
+			LogInformation("Copying AppLocal dependencies from {0} to {1}", BaseAppLocalDependenciesPath, StagedBinariesDir);
 
 			// Stage files in subdirs
 			foreach (DirectoryReference DependencyDirectory in DirectoryReference.EnumerateDirectories(BaseAppLocalDependenciesPath))
