Force-sync files for building documentation, to fix issues with files not being updated.

diff --git a/Engine/Source/Programs/AutomationTool/Scripts/SyncDepotPath.cs b/Engine/Source/Programs/AutomationTool/Scripts/SyncDepotPath.cs
--- a/Engine/Source/Programs/AutomationTool/Scripts/SyncDepotPath.cs
+++ b/Engine/Source/Programs/AutomationTool/Scripts/SyncDepotPath.cs
@@ -55,7 +55,7 @@ namespace AutomationTool
 			try
 			{
 				P4Connection Perforce = new P4Connection(P4Env.User, ClientName);
-				Perforce.Sync("//...");
+				Perforce.Sync("-f //...");
 			}
 			finally
 			{
