Always build UnrealPak when building client targets. The ProjectParams.Pak option is not reliable, because it can be forced on later by the target platform.

diff --git a/Engine/Source/Programs/AutomationTool/Scripts/BuildProjectCommand.Automation.cs b/Engine/Source/Programs/AutomationTool/Scripts/BuildProjectCommand.Automation.cs
--- a/Engine/Source/Programs/AutomationTool/Scripts/BuildProjectCommand.Automation.cs
+++ b/Engine/Source/Programs/AutomationTool/Scripts/BuildProjectCommand.Automation.cs
@@ -112,7 +112,7 @@ public partial class Project : CommandUtils
 		// Build any tools we need to stage
 		if ((TargetMask & ProjectBuildTargets.UnrealPak) == ProjectBuildTargets.UnrealPak && !Automation.IsEngineInstalled())
 		{
-			if (Params.Pak && Params.EditorTargets.Contains("UnrealPak") == false)
+			if (Params.EditorTargets.Contains("UnrealPak") == false)
 			{
 				Agenda.AddTargets(new string[] { "UnrealPak" }, HostPlatform.Current.HostEditorPlatform, UnrealTargetConfiguration.Development, Params.CodeBasedUprojectPath);
 			}
