Add the -VersionCookedContent argument to disable the -unversioned parameter on the cooker command line.

diff --git a/Engine/Source/Programs/AutomationTool/AutomationUtils/ProjectParams.cs b/Engine/Source/Programs/AutomationTool/AutomationUtils/ProjectParams.cs
--- a/Engine/Source/Programs/AutomationTool/AutomationUtils/ProjectParams.cs
+++ b/Engine/Source/Programs/AutomationTool/AutomationUtils/ProjectParams.cs
@@ -44,6 +44,21 @@ namespace AutomationTool
 			return Default;
 		}
 
+		/// <summary>
+		/// Gets optional parameter from the command line if it hasn't been specified in the constructor. 
+		/// If the command line is not available or the command has not been specified in the command line, default value will be used.
+		/// </summary>
+		/// <param name="Command">Command to parse the command line for. Can be null.</param>
+		/// <param name="SpecifiedValue">Value specified in the constructor (or not)</param>
+		/// <param name="Default">Default value.</param>
+		/// <param name="TrueParam">Name of a parameter that sets the value to 'true', for example: -clean</param>
+		/// <param name="FalseParam">Name of a parameter that sets the value to 'false', for example: -noclean</param>
+		/// <returns>Parameter value or default value if the paramater has not been specified</returns>
+		bool GetOptionalParamValueIfNotSpecified(BuildCommand Command, bool? SpecifiedValue, bool Default, string TrueParam, string FalseParam)
+		{
+			return GetOptionalParamValueIfNotSpecified(Command, SpecifiedValue, (bool?)Default, TrueParam, FalseParam).Value;
+		}
+
 		/// <summary>
 		/// Gets optional parameter from the command line if it hasn't been specified in the constructor. 
 		/// If the command line is not available or the command has not been specified in the command line, default value will be used.
@@ -619,7 +634,7 @@ namespace AutomationTool
                 this.Cook = false;
             }
             this.CookOnTheFlyStreaming = GetParamValueIfNotSpecified(Command, CookOnTheFlyStreaming, this.CookOnTheFlyStreaming, "cookontheflystreaming");
-            this.UnversionedCookedContent = GetParamValueIfNotSpecified(Command, UnversionedCookedContent, this.UnversionedCookedContent, "UnversionedCookedContent");
+            this.UnversionedCookedContent = GetOptionalParamValueIfNotSpecified(Command, UnversionedCookedContent, this.UnversionedCookedContent, "UnversionedCookedContent", "VersionCookedContent");
 			this.SkipCookingEditorContent = GetParamValueIfNotSpecified(Command, SkipCookingEditorContent, this.SkipCookingEditorContent, "SkipCookingEditorContent");
             if (NumCookersToSpawn.HasValue)
             {
@@ -658,7 +673,7 @@ namespace AutomationTool
 				this.Stage = true;
 			}
 			this.StageDirectoryParam = ParseParamValueIfNotSpecified(Command, StageDirectoryParam, "stagingdirectory", String.Empty, true);
-			this.bCodeSign = GetOptionalParamValueIfNotSpecified(Command, CodeSign, CommandUtils.IsBuildMachine, "CodeSign", "NoCodeSign").GetValueOrDefault();
+			this.bCodeSign = GetOptionalParamValueIfNotSpecified(Command, CodeSign, CommandUtils.IsBuildMachine, "CodeSign", "NoCodeSign");
 			this.bTreatNonShippingBinariesAsDebugFiles = GetParamValueIfNotSpecified(Command, TreatNonShippingBinariesAsDebugFiles, false, "TreatNonShippingBinariesAsDebugFiles");
 			this.bUseExtraFlavor = GetParamValueIfNotSpecified(Command, UseExtraFlavor, false, "UseExtraFlavor");
 			this.Manifests = GetParamValueIfNotSpecified(Command, Manifests, this.Manifests, "manifests");
