Added GetCachedConfig accessor to prevent multiple configs being constructed.

diff --git a/Engine/Source/Programs/AutomationTool/Gauntlet/Unreal/Base/Gauntlet.UnrealTestNode.cs b/Engine/Source/Programs/AutomationTool/Gauntlet/Unreal/Base/Gauntlet.UnrealTestNode.cs
--- a/Engine/Source/Programs/AutomationTool/Gauntlet/Unreal/Base/Gauntlet.UnrealTestNode.cs
+++ b/Engine/Source/Programs/AutomationTool/Gauntlet/Unreal/Base/Gauntlet.UnrealTestNode.cs
@@ -179,6 +179,20 @@ namespace Gauntlet
 			return CachedConfig;
 		}
 
+		/// <summary>
+		/// Returns the cached version of our config. Avoids repeatedly calling GetConfiguration() on derived nodes
+		/// </summary>
+		/// <returns></returns>
+		protected TConfigClass GetCachedConfiguration()
+		{
+			if (CachedConfig == null)
+			{
+				return GetConfiguration();
+			}
+
+			return CachedConfig;
+		}
+
 		/// <summary>
 		/// Returns a priority value for this test
 		/// </summary>
@@ -225,7 +239,7 @@ namespace Gauntlet
 			}
 
 			// This test only ends when all roles are gone
-			if (GetConfiguration().AllRolesExit)
+			if (GetCachedConfiguration().AllRolesExit)
 			{
 				return true;
 			}
@@ -251,7 +265,7 @@ namespace Gauntlet
 		protected bool PrepareUnrealApp()
 		{
 			// Get our configuration
-			TConfigClass Config = GetConfiguration();
+			TConfigClass Config = GetCachedConfiguration();
 
 			if (Config == null)
 			{
@@ -303,13 +317,16 @@ namespace Gauntlet
 
 				foreach (UnrealTestRole TestRole in TypesToRoles.Value)
 				{
+					// If a config has overriden a platform then we can't use the context constraints from the commandline
+					bool UseContextConstraint = TestRole.Type == UnrealTargetRole.Client && TestRole.PlatformOverride == UnrealTargetPlatform.Unknown;
+
 					// important, use the type from the ContextRolke because Server may have been mapped to EditorServer etc
 					UnrealTargetPlatform SessionPlatform = TestRole.PlatformOverride != UnrealTargetPlatform.Unknown ? TestRole.PlatformOverride : RoleContext.Platform;
 
 					UnrealSessionRole SessionRole = new UnrealSessionRole(RoleContext.Type, SessionPlatform, RoleContext.Configuration, TestRole.CommandLine);
 
 					SessionRole.RoleModifier = TestRole.RoleType;
-					SessionRole.Constraint = TestRole.Type == UnrealTargetRole.Client ? Context.Constraint : new UnrealTargetConstraint(SessionPlatform);
+					SessionRole.Constraint = UseContextConstraint ? Context.Constraint : new UnrealTargetConstraint(SessionPlatform);
 					
 					Log.Verbose("Created SessionRole {0} from RoleContext {1} (RoleType={2})", SessionRole, RoleContext, TypesToRoles.Key);
 
@@ -386,7 +403,7 @@ namespace Gauntlet
 				throw new AutomationException("Node already has a null UnrealApp, was PrepareUnrealSession or IsReadyToStart called?");
 			}
 
-			TConfigClass Config = GetConfiguration();
+			TConfigClass Config = GetCachedConfiguration();
 
 			CurrentPass = Pass;
 			NumPasses = InNumPasses;			
@@ -780,7 +797,7 @@ namespace Gauntlet
 
 			MB.Paragraph(string.Format("ResultHash: {0}", GetRoleResultHash(InArtifacts)));
 
-			if (GetConfiguration().ShowErrorsInSummary && InArtifacts.LogSummary.Errors.Count() > 0)
+			if (GetCachedConfiguration().ShowErrorsInSummary && InArtifacts.LogSummary.Errors.Count() > 0)
 			{
 				MB.H4("Errors");
 				MB.UnorderedList(LogSummary.Errors.Take(MaxLogLines));
@@ -791,7 +808,7 @@ namespace Gauntlet
 				}
 			}
 
-			if (GetConfiguration().ShowWarningsInSummary && InArtifacts.LogSummary.Warnings.Count() > 0)
+			if (GetCachedConfiguration().ShowWarningsInSummary && InArtifacts.LogSummary.Warnings.Count() > 0)
 			{
 				MB.H4("Warnings");
 				MB.UnorderedList(LogSummary.Warnings.Take(MaxLogLines));
