UE-50603 Cook on the fly process does not finish when closing the app on Android
UE-64633 UFE continually relaunches closed Android app on device

#Android
#4.22

AndroidPlatform.Automation.cs - adding a return value for RunClient, because the calling method uses the HasExited field

diff --git a/Engine/Source/Programs/AutomationTool/Android/AndroidPlatform.Automation.cs b/Engine/Source/Programs/AutomationTool/Android/AndroidPlatform.Automation.cs
--- a/Engine/Source/Programs/AutomationTool/Android/AndroidPlatform.Automation.cs
+++ b/Engine/Source/Programs/AutomationTool/Android/AndroidPlatform.Automation.cs
@@ -1579,6 +1579,7 @@ public class AndroidPlatform : Platform
 
 	public override IProcessResult RunClient(ERunOptions ClientRunFlags, string ClientApp, string ClientCmdLine, ProjectParams Params)
 	{
+		IProcessResult Result = null;
 		//make a copy of the device names, we'll be working through them
 		List<string> DeviceNames = new List<string>();
 		//same with the package names
@@ -1689,6 +1690,10 @@ public class AndroidPlatform : Platform
 					File.WriteAllText(LogFilename, LogFileProcess.Output);
 					File.WriteAllText(ServerLogFilename, LogFileProcess.Output);
 
+					if (Result == null)
+					{
+						Result = LogFileProcess;
+					}
 					DeviceNames.RemoveAt(DeviceIndex);
 					PackageNames.RemoveAt(DeviceIndex);
 
@@ -1697,7 +1702,7 @@ public class AndroidPlatform : Platform
 			}
 		}
 
-		return null;
+		return Result;
 	}
 
 	public override void GetFilesToDeployOrStage(ProjectParams Params, DeploymentContext SC)
