UE-64059 Failed to package for IOS on PC with BuildGraph
#4.22

Don't compile IPP, because it's checked into Perforce and expected to be built anyway. Note: we ship iPhonePackager binaries on GitHub. It's downloaded when running Setup.bat - it's required for the editor to query installed provisions and certificates.

diff --git a/Engine/Source/Programs/AutomationTool/IOS/IOSPlatform.Automation.cs b/Engine/Source/Programs/AutomationTool/IOS/IOSPlatform.Automation.cs
--- a/Engine/Source/Programs/AutomationTool/IOS/IOSPlatform.Automation.cs
+++ b/Engine/Source/Programs/AutomationTool/IOS/IOSPlatform.Automation.cs
@@ -1653,7 +1653,6 @@ public class IOSPlatform : Platform
 		if (UnrealBuildTool.BuildHostPlatform.Current.Platform != UnrealTargetPlatform.Mac && !Automation.IsEngineInstalled())
 		{
 			Agenda.DotNetProjects.Add(@"Engine\Source\Programs\IOS\MobileDeviceInterface\MobileDeviceInterface.csproj");
-			Agenda.DotNetProjects.Add(@"Engine\Source\Programs\IOS\iPhonePackager\iPhonePackager.csproj");
 			Agenda.DotNetProjects.Add(@"Engine\Source\Programs\IOS\DeploymentServer\DeploymentServer.csproj");
 		}
 	}
