Gauntlet: Remove workaround for missing log, add changes in UnrealSession missed in previous commit

diff --git a/Engine/Source/Programs/AutomationTool/Gauntlet/Unreal/Base/Gauntlet.UnrealSession.cs b/Engine/Source/Programs/AutomationTool/Gauntlet/Unreal/Base/Gauntlet.UnrealSession.cs
--- a/Engine/Source/Programs/AutomationTool/Gauntlet/Unreal/Base/Gauntlet.UnrealSession.cs
+++ b/Engine/Source/Programs/AutomationTool/Gauntlet/Unreal/Base/Gauntlet.UnrealSession.cs
@@ -133,12 +133,20 @@ namespace Gauntlet
 				RequiredBuildFlags |= BuildFlags.CanReplaceExecutable;
 			}
 
-            if (Globals.Params.ParseParam("bulk") && (InPlatform == UnrealTargetPlatform.Android || InPlatform == UnrealTargetPlatform.IOS))
+			//@todo: for bulk/packaged builds, we should mark the platform as capable of these as we're catching global and not test specific flags
+			// where we may be running parallel tests on multiple platforms
+			if (Globals.Params.ParseParam("bulk") && (InPlatform == UnrealTargetPlatform.Android || InPlatform == UnrealTargetPlatform.IOS))
             {
                 RequiredBuildFlags |= BuildFlags.Bulk;
             }
 
-            Options = InOptions;
+			if (Globals.Params.ParseParam("packaged") && (InPlatform == UnrealTargetPlatform.Switch))
+			{
+				RequiredBuildFlags |= BuildFlags.Packaged;
+			}
+
+
+			Options = InOptions;
             FilesToCopy = new List<UnrealFileToCopy>();
 			RoleModifier = ERoleModifier.None;
         }
