Enabled compression on first 4 chunks for windows.
Enabled custom compression on windows.
Prevented ChunkDBs from being uploaded to network drive now are uploaded straight to cloud.

diff --git a/Engine/Source/Editor/UnrealEd/Private/Commandlets/AssetRegistryGenerator.cpp b/Engine/Source/Editor/UnrealEd/Private/Commandlets/AssetRegistryGenerator.cpp
--- a/Engine/Source/Editor/UnrealEd/Private/Commandlets/AssetRegistryGenerator.cpp
+++ b/Engine/Source/Editor/UnrealEd/Private/Commandlets/AssetRegistryGenerator.cpp
@@ -278,6 +278,7 @@ bool FAssetRegistryGenerator::GenerateStreamingInstallManifest(int64 InExtraFlav
 				if ( PakChunkFilename.MatchesWildcard(CompressedChunkWildcard) )
 				{
 					PakChunkOptions += " compressed";
+					break;
 				}
 			}
 
diff --git a/Engine/Source/Programs/AutomationTool/Win/WinPlatform.Automation.cs b/Engine/Source/Programs/AutomationTool/Win/WinPlatform.Automation.cs
--- a/Engine/Source/Programs/AutomationTool/Win/WinPlatform.Automation.cs
+++ b/Engine/Source/Programs/AutomationTool/Win/WinPlatform.Automation.cs
@@ -174,11 +174,19 @@ public abstract class BaseWinPlatform : Platform
 	{
 		return "Windows";
 	}
+	
+	public override string GetPlatformPakCommandLine(ProjectParams Params, DeploymentContext SC)
+	{
+		string PakParams = " -patchpaddingalign=2048";
 
-    public override string GetPlatformPakCommandLine(ProjectParams Params, DeploymentContext SC)
-    {
-        return " -patchpaddingalign=2048";
-    }
+		string OodleDllPath = DirectoryReference.Combine(SC.ProjectRoot, "Binaries/ThirdParty/Oodle/Win64/UnrealPakPlugin.dll").FullName;
+		if (File.Exists(OodleDllPath))
+		{
+			PakParams += String.Format(" -customcompressor=\"{0}\"", OodleDllPath);
+		}
+
+		return PakParams;
+	}
 
 	public override void Package(ProjectParams Params, DeploymentContext SC, int WorkingCL)
 	{
