Check whether Gauntlet process has exited before log file created

diff --git a/Engine/Source/Programs/AutomationTool/Gauntlet/Platform/Windows/Gauntlet.TargetDeviceWindows.cs b/Engine/Source/Programs/AutomationTool/Gauntlet/Platform/Windows/Gauntlet.TargetDeviceWindows.cs
--- a/Engine/Source/Programs/AutomationTool/Gauntlet/Platform/Windows/Gauntlet.TargetDeviceWindows.cs
+++ b/Engine/Source/Programs/AutomationTool/Gauntlet/Platform/Windows/Gauntlet.TargetDeviceWindows.cs
@@ -69,6 +69,13 @@ namespace Gauntlet
 				Thread.Sleep(2000);
 			}
 
+			// Check whether the process exited before log file was created (this can happen for example if a server role exits and forces client to shutdown)
+			if (!File.Exists(ProcessLogFile))
+			{
+				ProcessLogOutput += "Process exited before log file created";
+				return;
+			}
+
 			Thread.Sleep(1000);
 
 			using (FileStream ProcessLog = File.Open(ProcessLogFile, FileMode.Open, FileAccess.Read, FileShare.ReadWrite))
