Disable compression on windows as it could be crashing servers.

diff --git a/Engine/Source/Programs/AutomationTool/Win/WinPlatform.Automation.cs b/Engine/Source/Programs/AutomationTool/Win/WinPlatform.Automation.cs
--- a/Engine/Source/Programs/AutomationTool/Win/WinPlatform.Automation.cs
+++ b/Engine/Source/Programs/AutomationTool/Win/WinPlatform.Automation.cs
@@ -178,13 +178,13 @@ public abstract class BaseWinPlatform : Platform
 	public override string GetPlatformPakCommandLine(ProjectParams Params, DeploymentContext SC)
 	{
 		string PakParams = " -patchpaddingalign=2048";
-
+		/*
 		string OodleDllPath = DirectoryReference.Combine(SC.ProjectRoot, "Binaries/ThirdParty/Oodle/Win64/UnrealPakPlugin.dll").FullName;
 		if (File.Exists(OodleDllPath))
 		{
 			PakParams += String.Format(" -customcompressor=\"{0}\"", OodleDllPath);
 		}
-
+		*/
 		return PakParams;
 	}
 
