Gauntlet - When number of desktop clients > 1 tile the windows into a 2xN configuration.

diff --git a/Engine/Source/Programs/AutomationTool/Gauntlet/Unreal/Base/Gauntlet.UnrealSession.cs b/Engine/Source/Programs/AutomationTool/Gauntlet/Unreal/Base/Gauntlet.UnrealSession.cs
--- a/Engine/Source/Programs/AutomationTool/Gauntlet/Unreal/Base/Gauntlet.UnrealSession.cs
+++ b/Engine/Source/Programs/AutomationTool/Gauntlet/Unreal/Base/Gauntlet.UnrealSession.cs
@@ -767,6 +767,53 @@ namespace Gauntlet
 
 				bool InstallSuccess = true;
 
+
+#if !__MonoCS__
+				// count how many desktop clients
+				IEnumerable<UnrealSessionRole> DesktopClients = SessionRoles.Where(R => R.Platform == BuildHostPlatform.Current.Platform).Where(R => R.RoleType.IsClient());
+
+				if (DesktopClients.Count() > 1)
+				{
+					int NumClientRows = (int)Math.Ceiling(DesktopClients.Count() / 2.0);
+
+					var ScreenRect = System.Windows.Forms.Screen.PrimaryScreen.Bounds;
+
+					double Ratio = 16.0 / 9.0;
+
+					// pick width
+					int DesiredWidth = ScreenRect.Width / 2;
+
+					// height at 16:9 aspect
+					int DesiredHeight = (int)(DesiredWidth / Ratio);
+
+					// constrain width if the screen can't have that many rows
+					if (DesiredHeight * NumClientRows > ScreenRect.Height)
+					{
+						DesiredHeight = ScreenRect.Height / NumClientRows;
+						DesiredWidth = (int)(DesiredHeight * Ratio);
+					}
+
+					// now set all these params and disable the regular params that pick windowed stuff
+					for (int i = 0; i < DesktopClients.Count(); i++)
+					{
+						UnrealSessionRole Role = DesktopClients.ElementAt(i);
+						int Row = i / 2;
+						int Column = i - (Row * 2);
+
+						// this is hacky, but not sure if a better way to do it right now without changing interfaces in a TBD way
+						UnrealTestConfiguration ConfigOptions = Role.Options as UnrealTestConfiguration;
+
+						if (ConfigOptions != null)
+						{
+							ConfigOptions.IgnoreDefaultResolutionAndWindowMode = true;
+						}
+
+						Role.CommandLine += string.Format(" -WinX={0} -WinY={1} -ResX={2} -ResY={3} -windowed",
+							Column * DesiredWidth, Row * DesiredHeight, DesiredWidth, DesiredHeight);
+					}
+				}
+#endif
+
 				// sort by constraints, so that we pick constrained devices first
 				List<UnrealSessionRole> SortedRoles = SessionRoles.OrderBy(R => R.Constraint.IsIdentity() ? 1 : 0).ToList();				
 
