Modified AutoParam so that default values can be specified as regular initializers.
Created custom config class for "DefaultTest" so that it's easier to see which options are supported.

diff --git a/Engine/Source/Programs/AutomationTool/Gauntlet/Framework/Gauntlet.AutoParam.cs b/Engine/Source/Programs/AutomationTool/Gauntlet/Framework/Gauntlet.AutoParam.cs
--- a/Engine/Source/Programs/AutomationTool/Gauntlet/Framework/Gauntlet.AutoParam.cs
+++ b/Engine/Source/Programs/AutomationTool/Gauntlet/Framework/Gauntlet.AutoParam.cs
@@ -40,23 +40,54 @@ namespace Gauntlet
 		protected string[] OptionNames;
 
 		/// <summary>
-		/// Constructor taking a name and default value
+		/// Constructor that takes nothing. Param option should be -MemberName or -MemberName=value.
+		/// Members with no matching param will be left as-is.
+		/// </summary>
+		public AutoParam()
+		{
+			this.OptionNames = null;
+			this.Default = null;
+		}
+
+		/// <summary>
+		/// Constructor that takes an array of of potential argument names, e.g. {"build","builds"}
+		/// Members with no matching param will be left as-is.
 		/// </summary>
 		/// <param name="Name"></param>
 		/// <param name="Default"></param>
-		protected AutoParam(object Default, params string[] OptionNames)
+		protected AutoParam(params string[] OptionNames)
 		{
 			this.OptionNames = OptionNames;
-			this.Default = Default;
+			this.Default = null;
 		}
 
+		/// <summary>
 
+		/// <summary>
+		/// Constructor that takes a default argument to use if no param is specified. Param option should be -MemberName or -MemberName=value.
+		/// Members with no matching param will be set to 'Default'
+		/// </summary>
+		/// <param name="Default"></param>
 		public AutoParam(object Default)
 		{
 			this.OptionNames = null;
 			this.Default = Default;
 		}
 
+		/// <summary>
+		/// Constructor that takes an array of of potential argument names, e.g. {"build","builds"}
+		/// Members with no matching param will be set to 'Default'
+		/// </summary>
+		/// <param name="Name"></param>
+		/// <param name="Default"></param>
+		protected AutoParam(object Default, params string[] OptionNames)
+		{
+			this.OptionNames = OptionNames;
+			this.Default = Default;
+		}
+
+		
+
 		/// <summary>
 		/// Checks whether Args contains a -Param statement, if so returns true else
 		/// returns Default
@@ -232,7 +263,7 @@ namespace Gauntlet
 							object DefaultValue = Opt.Default;
 							object NewValue = null;
 
-							if (DefaultValue.GetType() != MemberType)
+							if (DefaultValue != null && DefaultValue.GetType() != MemberType)
 							{
 								Log.Warning("AutoParam default value for member {0} is type {1}, not {2}", Member.Name, DefaultValue.GetType(), MemberType);
 							}
@@ -377,5 +408,10 @@ namespace Gauntlet
 			base(Default, OptionNames)
 		{
 		}
+
+		public AutoParamWithNames(params string[] OptionNames) :
+			base(OptionNames)
+		{
+		}
 	}
 }
diff --git a/Engine/Source/Programs/AutomationTool/Gauntlet/Unreal/Base/Gauntlet.UnrealTestContext.cs b/Engine/Source/Programs/AutomationTool/Gauntlet/Unreal/Base/Gauntlet.UnrealTestContext.cs
--- a/Engine/Source/Programs/AutomationTool/Gauntlet/Unreal/Base/Gauntlet.UnrealTestContext.cs
+++ b/Engine/Source/Programs/AutomationTool/Gauntlet/Unreal/Base/Gauntlet.UnrealTestContext.cs
@@ -47,20 +47,20 @@ namespace Gauntlet
 		/// <summary>
 		/// Name of this project
 		/// </summary>
-		[AutoParam("")]
-		public string Project;
+		[AutoParam]
+		public string Project = "";
 
 		/// <summary>
 		/// Reference to the build that is being tested
 		/// </summary>
-		[AutoParamWithNames("", "Build", "Builds")]
-		public string Build;
+		[AutoParamWithNames("Build", "Builds")]
+		public string Build = "";
 
 		/// <summary>
 		/// Does this project use 'Game' or Client/Server?
 		/// </summary>
 		[AutoParam(true)]
-		public bool UsesSharedBuildType;
+		public bool UsesSharedBuildType = true;
 
 		// todo - remove this and pass in BuildSource
 		public Type BuildSourceType;
diff --git a/Engine/Source/Programs/AutomationTool/Gauntlet/Unreal/Game/EpicGame.TestConfig.cs b/Engine/Source/Programs/AutomationTool/Gauntlet/Unreal/Game/EpicGame.TestConfig.cs
--- a/Engine/Source/Programs/AutomationTool/Gauntlet/Unreal/Game/EpicGame.TestConfig.cs
+++ b/Engine/Source/Programs/AutomationTool/Gauntlet/Unreal/Game/EpicGame.TestConfig.cs
@@ -19,22 +19,22 @@ namespace EpicGame
 		/// <summary>
 		/// Should this test skip mcp?
 		/// </summary>
-		[AutoParam(false)]
+		[AutoParam]
 		public bool NoMCP = false;
 
-		[AutoParam(false)]
+		[AutoParam]
 		public bool FastCook = false;
 
 		/// <summary>
 		/// Which backend to use for matchmaking
 		/// </summary>
-		[AutoParam("DevLatest")]
-		public string EpicApp = "";
+		[AutoParam]
+		public string EpicApp = "DevLatest";
 
 		/// <summary>
 		/// Unique buildid to avoid matchmaking collisions
 		/// </summary>
-		[AutoParam("")]
+		[AutoParam]
 		protected string BuildIDOverride = "";
 
 		/// <summary>
@@ -54,13 +54,13 @@ namespace EpicGame
 		/// <summary>
 		/// Make sure the client gets -logpso when we are collecting them
 		/// </summary>
-		[AutoParam(false)]
+		[AutoParam]
         public bool LogPSO = false;
 
         /// <summary>
         /// Should this test assign a random test account?
         /// </summary>
-        [AutoParam(true)]
+        [AutoParam]
         public bool PreAssignAccount = true;
 
 
diff --git a/Engine/Source/Programs/AutomationTool/Gauntlet/Unreal/Game/UE4Game.TestConfig.cs b/Engine/Source/Programs/AutomationTool/Gauntlet/Unreal/Game/UE4Game.TestConfig.cs
--- a/Engine/Source/Programs/AutomationTool/Gauntlet/Unreal/Game/UE4Game.TestConfig.cs
+++ b/Engine/Source/Programs/AutomationTool/Gauntlet/Unreal/Game/UE4Game.TestConfig.cs
@@ -18,7 +18,8 @@ namespace UE4Game
 		/// <summary>
 		/// Map to use
 		/// </summary>
-		public string Map { get; set; }
+		[AutoParam]
+		public string Map = "";
 
 		/// <summary>
 		/// Applies these options to the provided app config
