Modified build launcher to skip specification of test node when not specified
Removed redundant sbmm option

diff --git a/Engine/Source/Programs/AutomationTool/Gauntlet/Unreal/Game/EpicGame.TestConfig.cs b/Engine/Source/Programs/AutomationTool/Gauntlet/Unreal/Game/EpicGame.TestConfig.cs
--- a/Engine/Source/Programs/AutomationTool/Gauntlet/Unreal/Game/EpicGame.TestConfig.cs
+++ b/Engine/Source/Programs/AutomationTool/Gauntlet/Unreal/Game/EpicGame.TestConfig.cs
@@ -194,9 +194,6 @@ namespace EpicGame
 
 			if (AppConfig.ProcessType.IsClient())
 			{
-				// turn off skill-based matchmaking, turn off porta;
-				AppConfig.CommandLine += " -nosbmm";
-
                 if (LogPSO)
                 {
                     AppConfig.CommandLine += " -logpso";
