Linux: fix iterative deployment

diff --git a/Engine/Source/Programs/AutomationTool/Linux/LinuxPlatform.Automation.cs b/Engine/Source/Programs/AutomationTool/Linux/LinuxPlatform.Automation.cs
--- a/Engine/Source/Programs/AutomationTool/Linux/LinuxPlatform.Automation.cs
+++ b/Engine/Source/Programs/AutomationTool/Linux/LinuxPlatform.Automation.cs
@@ -236,6 +236,11 @@ chmod +x {0}
 					// non-null input is essential, since RedirectStandardInput=true is needed for PLINK, see http://stackoverflow.com/questions/1910592/process-waitforexit-on-console-vs-windows-forms
 					RunAndLog(CmdEnv, PlinkPath, String.Format("-batch -l {0} -pw {1} {2} rm -rf {3} && mkdir -p {3}", Params.DeviceUsername, Params.DevicePassword, Params.DeviceNames[0], DestPath), Input: "");
 				}
+				else
+				{
+					// still make sure that the path exists
+					RunAndLog(CmdEnv, PlinkPath, String.Format("-batch -l {0} -pw {1} {2} mkdir -p {3}", Params.DeviceUsername, Params.DevicePassword, Params.DeviceNames[0], DestPath), Input: "");
+				}
 
 				// copy the contents
 				RunAndLog(CmdEnv, PScpPath, String.Format("-batch -pw {0} -r \"{1}\" {2}", Params.DevicePassword, SourcePath, Params.DeviceUsername + "@" + DeviceAddress + ":" + DestPath));
