Changed some of the localization automation to use classes with read-only members, rather than use structs

diff --git a/Engine/Source/Programs/AutomationTool/Localization/LocalizationProvider.cs b/Engine/Source/Programs/AutomationTool/Localization/LocalizationProvider.cs
--- a/Engine/Source/Programs/AutomationTool/Localization/LocalizationProvider.cs
+++ b/Engine/Source/Programs/AutomationTool/Localization/LocalizationProvider.cs
@@ -8,193 +8,211 @@ using Tools.DotNETCommon;
 
 namespace EpicGames.Localization
 {
+	public class ProjectImportExportInfo
+	{
+		public ProjectImportExportInfo(string InDestinationPath, string InManifestName, string InArchiveName, string InPortableObjectName, string InNativeCulture, IReadOnlyList<string> InCulturesToGenerate, bool InUseCultureDirectory)
+		{
+			DestinationPath = InDestinationPath;
+			ManifestName = InManifestName;
+			ArchiveName = InArchiveName;
+			PortableObjectName = InPortableObjectName;
+			NativeCulture = InNativeCulture;
+			CulturesToGenerate = InCulturesToGenerate;
+			bUseCultureDirectory = InUseCultureDirectory;
+		}
 
-public struct ProjectImportExportInfo
-{
-	/** The destination path for the files containing the gathered data for this project (extracted from its config file - relative to the root working directory for the commandlet) */
-	public string DestinationPath;
+		/** The destination path for the files containing the gathered data for this project (extracted from its config file - relative to the root working directory for the commandlet) */
+		public string DestinationPath { get; private set; }
 
-	/** The name to use for this projects manifest file (extracted from its config file) */
-	public string ManifestName;
+		/** The name to use for this projects manifest file (extracted from its config file) */
+		public string ManifestName { get; private set; }
 
-	/** The name to use for this projects archive file (extracted from its config file) */
-	public string ArchiveName;
+		/** The name to use for this projects archive file (extracted from its config file) */
+		public string ArchiveName { get; private set; }
 
-	/** The name to use for this projects portable object file (extracted from its config file) */
-	public string PortableObjectName;
+		/** The name to use for this projects portable object file (extracted from its config file) */
+		public string PortableObjectName { get; private set; }
 
-	/** The native culture for this project (extracted from its config file) */
-	public string NativeCulture;
+		/** The native culture for this project (extracted from its config file) */
+		public string NativeCulture { get; private set; }
 
-	/** The cultures to generate for this project (extracted from its config file) */
-	public List<string> CulturesToGenerate;
+		/** The cultures to generate for this project (extracted from its config file) */
+		public IReadOnlyList<string> CulturesToGenerate { get; private set; }
 
-	/** True if we should use a per-culture directly when importing/exporting */
-	public bool bUseCultureDirectory;
+		/** True if we should use a per-culture directly when importing/exporting */
+		public bool bUseCultureDirectory { get; private set; }
 
-	/** The platform names that we have split into sub-folders */
-	public List<string> SplitPlatformNames;
+		/** The platform names that we have split into sub-folders */
+		public IReadOnlyList<string> SplitPlatformNames { get; private set; }
 
-	/** The platforms sub-folder */
-	public const string PlatformLocalizationFolderName = "Platforms";
+		/** The platforms sub-folder */
+		public const string PlatformLocalizationFolderName = "Platforms";
 
-	public void CalculateSplitPlatformNames(string RootWorkingDirectory)
-	{
-		// Is this isn't a single culture import/export, then also check for split platform sub-folders
-		SplitPlatformNames = new List<string>();
-		if (bUseCultureDirectory)
+		public void CalculateSplitPlatformNames(string RootWorkingDirectory)
 		{
-			var PlatformSourceDirectory = new DirectoryReference(CommandUtils.CombinePaths(RootWorkingDirectory, DestinationPath, PlatformLocalizationFolderName));
-			if (DirectoryReference.Exists(PlatformSourceDirectory))
+			// Is this isn't a single culture import/export, then also check for split platform sub-folders
+			var NewSplitPlatformNames = new List<string>();
+			if (bUseCultureDirectory)
 			{
-				foreach (DirectoryReference FoundDirectory in DirectoryReference.EnumerateDirectories(PlatformSourceDirectory, "*", SearchOption.TopDirectoryOnly))
+				var PlatformSourceDirectory = new DirectoryReference(CommandUtils.CombinePaths(RootWorkingDirectory, DestinationPath, PlatformLocalizationFolderName));
+				if (DirectoryReference.Exists(PlatformSourceDirectory))
 				{
-					string SplitPlatformName = CommandUtils.GetLastDirectoryName(FoundDirectory.FullName);
-					SplitPlatformNames.Add(SplitPlatformName);
+					foreach (DirectoryReference FoundDirectory in DirectoryReference.EnumerateDirectories(PlatformSourceDirectory, "*", SearchOption.TopDirectoryOnly))
+					{
+						string SplitPlatformName = CommandUtils.GetLastDirectoryName(FoundDirectory.FullName);
+						NewSplitPlatformNames.Add(SplitPlatformName);
+					}
 				}
 			}
+			SplitPlatformNames = NewSplitPlatformNames;
 		}
-	}
-};
+	};
 
-public struct ProjectStepInfo
-{
-	public ProjectStepInfo(string InName, string InLocalizationConfigFile)
+	public class ProjectStepInfo
 	{
-		Name = InName;
-		LocalizationConfigFile = InLocalizationConfigFile;
-	}
-	
-	/** The name of this localization step */
-	public string Name;
-
-	/** Absolute path to this steps localization config file */
-	public string LocalizationConfigFile;
-};
-
-public struct ProjectInfo
-{
-	/** The name of this project */
-	public string ProjectName;
-
-	/** Path to this projects localization step data - ordered so that iterating them runs in the correct order */
-	public List<ProjectStepInfo> LocalizationSteps;
-
-	/** Config data used by the PO file import process */
-	public ProjectImportExportInfo ImportInfo;
+		public ProjectStepInfo(string InName, string InLocalizationConfigFile)
+		{
+			Name = InName;
+			LocalizationConfigFile = InLocalizationConfigFile;
+		}
 
-	/** Config data used by the PO file export process */
-	public ProjectImportExportInfo ExportInfo;
-};
+		/** The name of this localization step */
+		public string Name { get; private set; }
 
-public abstract class LocalizationProvider
-{
-	public struct LocalizationProviderArgs
-	{
-		public string RootWorkingDirectory;
-		public string RootLocalizationTargetDirectory;
-		public string RemoteFilenamePrefix;
-		public BuildCommand Command;
-		public int PendingChangeList;
+		/** Absolute path to this steps localization config file */
+		public string LocalizationConfigFile { get; private set; }
 	};
 
-	public LocalizationProvider(LocalizationProviderArgs InArgs)
+	public class ProjectInfo
 	{
-		RootWorkingDirectory = InArgs.RootWorkingDirectory;
-		RemoteFilenamePrefix = InArgs.RemoteFilenamePrefix;
-		Command = InArgs.Command;
-		PendingChangeList = InArgs.PendingChangeList;
+		public ProjectInfo(string InProjectName, IReadOnlyList<ProjectStepInfo> InLocalizationSteps, ProjectImportExportInfo InImportInfo, ProjectImportExportInfo InExportInfo)
+		{
+			ProjectName = InProjectName;
+			LocalizationSteps = InLocalizationSteps;
+			ImportInfo = InImportInfo;
+			ExportInfo = InExportInfo;
+		}
 
-		LocalizationBranchName = Command.ParseParamValue("LocalizationBranch");
-		bUploadAllCultures = Command.ParseParam("UploadAllCultures");
-	}
+		/** The name of this project */
+		public string ProjectName { get; private set; }
 
-	public virtual string GetLocalizationProviderId()
-	{
-		throw new AutomationException("Unimplemented GetLocalizationProviderId.");
-	}
+		/** Path to this projects localization step data - ordered so that iterating them runs in the correct order */
+		public IReadOnlyList<ProjectStepInfo> LocalizationSteps { get; private set; }
 
-	public virtual void DownloadProjectFromLocalizationProvider(string ProjectName, ProjectImportExportInfo ProjectImportInfo)
-	{
-		throw new AutomationException("Unimplemented DownloadProjectFromLocalizationProvider.");
-	}
+		/** Config data used by the PO file import process */
+		public ProjectImportExportInfo ImportInfo { get; private set; }
 
-	public virtual void UploadProjectToLocalizationProvider(string ProjectName, ProjectImportExportInfo ProjectExportInfo)
-	{
-		throw new AutomationException("Unimplemented UploadProjectToLocalizationProvider.");
-	}
+		/** Config data used by the PO file export process */
+		public ProjectImportExportInfo ExportInfo { get; private set; }
+	};
 
-	public static LocalizationProvider GetLocalizationProvider(string InLocalizationProviderId, LocalizationProvider.LocalizationProviderArgs InLocalizationProviderArgs)
+	public abstract class LocalizationProvider
 	{
-		if (String.IsNullOrEmpty(InLocalizationProviderId))
+		public struct LocalizationProviderArgs
 		{
-			return null;
+			public string RootWorkingDirectory;
+			public string RootLocalizationTargetDirectory;
+			public string RemoteFilenamePrefix;
+			public BuildCommand Command;
+			public int PendingChangeList;
+		};
+
+		public LocalizationProvider(LocalizationProviderArgs InArgs)
+		{
+			RootWorkingDirectory = InArgs.RootWorkingDirectory;
+			RemoteFilenamePrefix = InArgs.RemoteFilenamePrefix;
+			Command = InArgs.Command;
+			PendingChangeList = InArgs.PendingChangeList;
+
+			LocalizationBranchName = Command.ParseParamValue("LocalizationBranch");
+			bUploadAllCultures = Command.ParseParam("UploadAllCultures");
+		}
+
+		public virtual string GetLocalizationProviderId()
+		{
+			throw new AutomationException("Unimplemented GetLocalizationProviderId.");
 		}
 
-		if (CachedLocalizationProviderTypes == null)
+		public virtual void DownloadProjectFromLocalizationProvider(string ProjectName, ProjectImportExportInfo ProjectImportInfo)
 		{
-			// Find all types that derive from LocalizationProvider in any of our DLLs
-			CachedLocalizationProviderTypes = new Dictionary<string, Type>();
-			var LoadedAssemblies = AppDomain.CurrentDomain.GetAssemblies();
-			foreach (var Dll in LoadedAssemblies)
+			throw new AutomationException("Unimplemented DownloadProjectFromLocalizationProvider.");
+		}
+
+		public virtual void UploadProjectToLocalizationProvider(string ProjectName, ProjectImportExportInfo ProjectExportInfo)
+		{
+			throw new AutomationException("Unimplemented UploadProjectToLocalizationProvider.");
+		}
+
+		public static LocalizationProvider GetLocalizationProvider(string InLocalizationProviderId, LocalizationProvider.LocalizationProviderArgs InLocalizationProviderArgs)
+		{
+			if (String.IsNullOrEmpty(InLocalizationProviderId))
+			{
+				return null;
+			}
+
+			if (CachedLocalizationProviderTypes == null)
 			{
-				var AllTypes = Dll.GetTypes();
-				foreach (var PotentialLocalizationNodeType in AllTypes)
+				// Find all types that derive from LocalizationProvider in any of our DLLs
+				CachedLocalizationProviderTypes = new Dictionary<string, Type>();
+				var LoadedAssemblies = AppDomain.CurrentDomain.GetAssemblies();
+				foreach (var Dll in LoadedAssemblies)
 				{
-					if (PotentialLocalizationNodeType != typeof(LocalizationProvider) && !PotentialLocalizationNodeType.IsAbstract && !PotentialLocalizationNodeType.IsInterface && typeof(LocalizationProvider).IsAssignableFrom(PotentialLocalizationNodeType))
+					var AllTypes = Dll.GetTypes();
+					foreach (var PotentialLocalizationNodeType in AllTypes)
 					{
-						// Types should implement a static StaticGetLocalizationProviderId method
-						var Method = PotentialLocalizationNodeType.GetMethod("StaticGetLocalizationProviderId");
-						if (Method != null)
+						if (PotentialLocalizationNodeType != typeof(LocalizationProvider) && !PotentialLocalizationNodeType.IsAbstract && !PotentialLocalizationNodeType.IsInterface && typeof(LocalizationProvider).IsAssignableFrom(PotentialLocalizationNodeType))
 						{
-							try
+							// Types should implement a static StaticGetLocalizationProviderId method
+							var Method = PotentialLocalizationNodeType.GetMethod("StaticGetLocalizationProviderId");
+							if (Method != null)
 							{
-								var LocalizationProviderId = Method.Invoke(null, null) as string;
-								CachedLocalizationProviderTypes.Add(LocalizationProviderId, PotentialLocalizationNodeType);
+								try
+								{
+									var LocalizationProviderId = Method.Invoke(null, null) as string;
+									CachedLocalizationProviderTypes.Add(LocalizationProviderId, PotentialLocalizationNodeType);
+								}
+								catch
+								{
+									BuildCommand.LogWarning("Type '{0}' threw when calling its StaticGetLocalizationProviderId method.", PotentialLocalizationNodeType.FullName);
+								}
 							}
-							catch
+							else
 							{
-								BuildCommand.LogWarning("Type '{0}' threw when calling its StaticGetLocalizationProviderId method.", PotentialLocalizationNodeType.FullName);
+								BuildCommand.LogWarning("Type '{0}' derives from LocalizationProvider but is missing its StaticGetLocalizationProviderId method.", PotentialLocalizationNodeType.FullName);
 							}
 						}
-						else
-						{
-							BuildCommand.LogWarning("Type '{0}' derives from LocalizationProvider but is missing its StaticGetLocalizationProviderId method.", PotentialLocalizationNodeType.FullName);
-						}
 					}
 				}
 			}
-		}
 
-		Type LocalizationNodeType;
-		CachedLocalizationProviderTypes.TryGetValue(InLocalizationProviderId, out LocalizationNodeType);
-		if (LocalizationNodeType != null)
-		{
-			try
+			Type LocalizationNodeType;
+			CachedLocalizationProviderTypes.TryGetValue(InLocalizationProviderId, out LocalizationNodeType);
+			if (LocalizationNodeType != null)
 			{
-				return Activator.CreateInstance(LocalizationNodeType, new object[] { InLocalizationProviderArgs }) as LocalizationProvider;
+				try
+				{
+					return Activator.CreateInstance(LocalizationNodeType, new object[] { InLocalizationProviderArgs }) as LocalizationProvider;
+				}
+				catch (Exception e)
+				{
+					BuildCommand.LogWarning("Unable to create an instance of the type '{0}'. {1}", LocalizationNodeType.FullName, e.ToString());
+				}
 			}
-			catch (Exception e)
+			else
 			{
-				BuildCommand.LogWarning("Unable to create an instance of the type '{0}'. {1}", LocalizationNodeType.FullName, e.ToString());
+				BuildCommand.LogWarning("Could not find a localization provider for '{0}'", InLocalizationProviderId);
 			}
-		}
-		else
-		{
-			BuildCommand.LogWarning("Could not find a localization provider for '{0}'", InLocalizationProviderId);
-		}
-
-		return null;
-	}
 
-	protected string RootWorkingDirectory;
-	protected string LocalizationBranchName;
-	protected string RemoteFilenamePrefix;
-	protected bool bUploadAllCultures;
-	protected BuildCommand Command;
-	protected int PendingChangeList;
+			return null;
+		}
 
-	private static Dictionary<string, Type> CachedLocalizationProviderTypes;
-};
+		protected string RootWorkingDirectory;
+		protected string LocalizationBranchName;
+		protected string RemoteFilenamePrefix;
+		protected bool bUploadAllCultures;
+		protected BuildCommand Command;
+		protected int PendingChangeList;
 
+		private static Dictionary<string, Type> CachedLocalizationProviderTypes;
+	};
 }
diff --git a/Engine/Source/Programs/AutomationTool/OneSkyLocalization/OneSkyLocalizationProvider.cs b/Engine/Source/Programs/AutomationTool/OneSkyLocalization/OneSkyLocalizationProvider.cs
--- a/Engine/Source/Programs/AutomationTool/OneSkyLocalization/OneSkyLocalizationProvider.cs
+++ b/Engine/Source/Programs/AutomationTool/OneSkyLocalization/OneSkyLocalizationProvider.cs
@@ -11,337 +11,335 @@ using EpicGames.OneSkyLocalization.Config;
 
 namespace EpicGames.OneSkyLocalization
 {
-
-[Help("OneSkyConfigName", "Name of the config data to use (see OneSkyConfigHelper).")]
-[Help("OneSkyProjectGroupName", "Name of the project group in OneSky.")]
-public class OneSkyLocalizationProvider : LocalizationProvider
-{
-	public OneSkyLocalizationProvider(LocalizationProviderArgs InArgs)
-		: base(InArgs)
+	[Help("OneSkyConfigName", "Name of the config data to use (see OneSkyConfigHelper).")]
+	[Help("OneSkyProjectGroupName", "Name of the project group in OneSky.")]
+	public class OneSkyLocalizationProvider : LocalizationProvider
 	{
-		var OneSkyConfigName = Command.ParseParamValue("OneSkyConfigName");
-		if (OneSkyConfigName == null)
+		public OneSkyLocalizationProvider(LocalizationProviderArgs InArgs)
+			: base(InArgs)
 		{
-			throw new AutomationException("Missing required command line argument: 'OneSkyConfigName'");
+			var OneSkyConfigName = Command.ParseParamValue("OneSkyConfigName");
+			if (OneSkyConfigName == null)
+			{
+				throw new AutomationException("Missing required command line argument: 'OneSkyConfigName'");
+			}
+
+			var OneSkyProjectGroupName = Command.ParseParamValue("OneSkyProjectGroupName");
+			if (OneSkyProjectGroupName == null)
+			{
+				throw new AutomationException("Missing required command line argument: 'OneSkyProjectGroupName'");
+			}
+
+			OneSkyConfigData OneSkyConfig = OneSkyConfigHelper.Find(OneSkyConfigName);
+			OneSkyService = new OneSkyService(OneSkyConfig.ApiKey, OneSkyConfig.ApiSecret);
+			OneSkyProjectGroup = GetOneSkyProjectGroup(OneSkyProjectGroupName);
 		}
 
-		var OneSkyProjectGroupName = Command.ParseParamValue("OneSkyProjectGroupName");
-		if (OneSkyProjectGroupName == null)
+		public static string StaticGetLocalizationProviderId()
 		{
-			throw new AutomationException("Missing required command line argument: 'OneSkyProjectGroupName'");
+			return "OneSky";
 		}
 
-		OneSkyConfigData OneSkyConfig = OneSkyConfigHelper.Find(OneSkyConfigName);
-		OneSkyService = new OneSkyService(OneSkyConfig.ApiKey, OneSkyConfig.ApiSecret);
-		OneSkyProjectGroup = GetOneSkyProjectGroup(OneSkyProjectGroupName);
-	}
-
-	public static string StaticGetLocalizationProviderId()
-	{
-		return "OneSky";
-	}
-
-	public override string GetLocalizationProviderId()
-	{
-		return StaticGetLocalizationProviderId();
-	}
-
-	public override void DownloadProjectFromLocalizationProvider(string ProjectName, ProjectImportExportInfo ProjectImportInfo)
-	{
-		var OneSkyProject = GetOneSkyProject(ProjectName);
-		DownloadLatestPOFiles(OneSkyProject, null, ProjectImportInfo);
-		foreach (var Platform in ProjectImportInfo.SplitPlatformNames)
+		public override string GetLocalizationProviderId()
 		{
-			DownloadLatestPOFiles(OneSkyProject, Platform, ProjectImportInfo);
+			return StaticGetLocalizationProviderId();
 		}
-	}
 
-	private void DownloadLatestPOFiles(Project OneSkyProject, string Platform, ProjectImportExportInfo ProjectImportInfo)
-	{
-		var OneSkyFileName = GetOneSkyFilename(ProjectImportInfo.PortableObjectName, Platform);
-		var OneSkyFile = OneSkyProject.UploadedFiles.FirstOrDefault(f => f.Filename == OneSkyFileName);
+		public override void DownloadProjectFromLocalizationProvider(string ProjectName, ProjectImportExportInfo ProjectImportInfo)
+		{
+			var OneSkyProject = GetOneSkyProject(ProjectName);
+			DownloadLatestPOFiles(OneSkyProject, null, ProjectImportInfo);
+			foreach (var Platform in ProjectImportInfo.SplitPlatformNames)
+			{
+				DownloadLatestPOFiles(OneSkyProject, Platform, ProjectImportInfo);
+			}
+		}
 
-		// Export
-		if (OneSkyFile != null)
+		private void DownloadLatestPOFiles(Project OneSkyProject, string Platform, ProjectImportExportInfo ProjectImportInfo)
 		{
-			var CulturesToExport = new List<string>();
-			foreach (var OneSkyCulture in OneSkyProject.EnabledCultures)
+			var OneSkyFileName = GetOneSkyFilename(ProjectImportInfo.PortableObjectName, Platform);
+			var OneSkyFile = OneSkyProject.UploadedFiles.FirstOrDefault(f => f.Filename == OneSkyFileName);
+
+			// Export
+			if (OneSkyFile != null)
 			{
-				// Skip the native culture, as OneSky has mangled it
-				if (OneSkyCulture == ProjectImportInfo.NativeCulture)
+				var CulturesToExport = new List<string>();
+				foreach (var OneSkyCulture in OneSkyProject.EnabledCultures)
 				{
-					continue;
-				}
+					// Skip the native culture, as OneSky has mangled it
+					if (OneSkyCulture == ProjectImportInfo.NativeCulture)
+					{
+						continue;
+					}
 
-				// Only export the OneSky cultures that we care about for this project
-				if (ProjectImportInfo.CulturesToGenerate.Contains(OneSkyCulture))
-				{
-					CulturesToExport.Add(OneSkyCulture);
+					// Only export the OneSky cultures that we care about for this project
+					if (ProjectImportInfo.CulturesToGenerate.Contains(OneSkyCulture))
+					{
+						CulturesToExport.Add(OneSkyCulture);
+					}
 				}
-			}
 
-			var DestinationDirectory = String.IsNullOrEmpty(Platform) 
-					? new DirectoryInfo(CommandUtils.CombinePaths(RootWorkingDirectory, ProjectImportInfo.DestinationPath))
-					: new DirectoryInfo(CommandUtils.CombinePaths(RootWorkingDirectory, ProjectImportInfo.DestinationPath, ProjectImportExportInfo.PlatformLocalizationFolderName, Platform));
-			ExportOneSkyFileToDirectory(OneSkyFile, DestinationDirectory, ProjectImportInfo.PortableObjectName, CulturesToExport, ProjectImportInfo.bUseCultureDirectory);
+				var DestinationDirectory = String.IsNullOrEmpty(Platform)
+						? new DirectoryInfo(CommandUtils.CombinePaths(RootWorkingDirectory, ProjectImportInfo.DestinationPath))
+						: new DirectoryInfo(CommandUtils.CombinePaths(RootWorkingDirectory, ProjectImportInfo.DestinationPath, ProjectImportExportInfo.PlatformLocalizationFolderName, Platform));
+				ExportOneSkyFileToDirectory(OneSkyFile, DestinationDirectory, ProjectImportInfo.PortableObjectName, CulturesToExport, ProjectImportInfo.bUseCultureDirectory);
+			}
 		}
-	}
-
-	private UploadedFile.ExportTranslationState ExportOneSkyTranslationWithRetry(UploadedFile OneSkyFile, string Culture, MemoryStream MemoryStream)
-	{
-		const int MAX_COUNT = 3;
 
-		long StartingMemPos = MemoryStream.Position;
-		int Count = 0;
-		for (;;)
+		private UploadedFile.ExportTranslationState ExportOneSkyTranslationWithRetry(UploadedFile OneSkyFile, string Culture, MemoryStream MemoryStream)
 		{
-			try
-			{
-				return OneSkyFile.ExportTranslation(Culture, MemoryStream).Result;
-			}
-			catch (Exception)
+			const int MAX_COUNT = 3;
+
+			long StartingMemPos = MemoryStream.Position;
+			int Count = 0;
+			for (;;)
 			{
-				if (++Count < MAX_COUNT)
+				try
 				{
-					MemoryStream.Position = StartingMemPos;
-					Console.WriteLine("ExportOneSkyTranslation attempt {0}/{1} failed. Retrying...", Count, MAX_COUNT);
-					continue;
+					return OneSkyFile.ExportTranslation(Culture, MemoryStream).Result;
 				}
+				catch (Exception)
+				{
+					if (++Count < MAX_COUNT)
+					{
+						MemoryStream.Position = StartingMemPos;
+						Console.WriteLine("ExportOneSkyTranslation attempt {0}/{1} failed. Retrying...", Count, MAX_COUNT);
+						continue;
+					}
 
-				Console.WriteLine("ExportOneSkyTranslation attempt {0}/{1} failed.", Count, MAX_COUNT);
-				break;
+					Console.WriteLine("ExportOneSkyTranslation attempt {0}/{1} failed.", Count, MAX_COUNT);
+					break;
+				}
 			}
-		}
 
-		return UploadedFile.ExportTranslationState.Failure;
-	}
+			return UploadedFile.ExportTranslationState.Failure;
+		}
 
-	private void ExportOneSkyFileToDirectory(UploadedFile OneSkyFile, DirectoryInfo DestinationDirectory, string DestinationFilename, IEnumerable<string> Cultures, bool bUseCultureDirectory)
-	{
-		foreach (var Culture in Cultures)
+		private void ExportOneSkyFileToDirectory(UploadedFile OneSkyFile, DirectoryInfo DestinationDirectory, string DestinationFilename, IEnumerable<string> Cultures, bool bUseCultureDirectory)
 		{
-			var CultureDirectory = (bUseCultureDirectory) ? new DirectoryInfo(Path.Combine(DestinationDirectory.FullName, Culture)) : DestinationDirectory;
-			if (!CultureDirectory.Exists)
-			{
-				CultureDirectory.Create();
-			}
-
-			using (var MemoryStream = new MemoryStream())
+			foreach (var Culture in Cultures)
 			{
-				var ExportTranslationState = ExportOneSkyTranslationWithRetry(OneSkyFile, Culture, MemoryStream);
-				if (ExportTranslationState == UploadedFile.ExportTranslationState.Success)
+				var CultureDirectory = (bUseCultureDirectory) ? new DirectoryInfo(Path.Combine(DestinationDirectory.FullName, Culture)) : DestinationDirectory;
+				if (!CultureDirectory.Exists)
 				{
-					var ExportFile = new FileInfo(Path.Combine(CultureDirectory.FullName, DestinationFilename));
+					CultureDirectory.Create();
+				}
 
-					// Write out the updated PO file so that the gather commandlet will import the new data from it
+				using (var MemoryStream = new MemoryStream())
+				{
+					var ExportTranslationState = ExportOneSkyTranslationWithRetry(OneSkyFile, Culture, MemoryStream);
+					if (ExportTranslationState == UploadedFile.ExportTranslationState.Success)
 					{
-						var ExportFileWasReadOnly = false;
-						if (ExportFile.Exists)
-						{
-							// We're going to clobber the existing PO file, so make sure it's writable (it may be read-only if in Perforce)
-							ExportFileWasReadOnly = ExportFile.IsReadOnly;
-							ExportFile.IsReadOnly = false;
-						}
+						var ExportFile = new FileInfo(Path.Combine(CultureDirectory.FullName, DestinationFilename));
 
-						MemoryStream.Position = 0;
-						using (Stream FileStream = ExportFile.Open(FileMode.Create))
+						// Write out the updated PO file so that the gather commandlet will import the new data from it
 						{
-							MemoryStream.CopyTo(FileStream);
-							Console.WriteLine("[SUCCESS] Exported '{0}' as '{1}' ({2})", OneSkyFile.Filename, ExportFile.FullName, Culture);
+							var ExportFileWasReadOnly = false;
+							if (ExportFile.Exists)
+							{
+								// We're going to clobber the existing PO file, so make sure it's writable (it may be read-only if in Perforce)
+								ExportFileWasReadOnly = ExportFile.IsReadOnly;
+								ExportFile.IsReadOnly = false;
+							}
+
+							MemoryStream.Position = 0;
+							using (Stream FileStream = ExportFile.Open(FileMode.Create))
+							{
+								MemoryStream.CopyTo(FileStream);
+								Console.WriteLine("[SUCCESS] Exported '{0}' as '{1}' ({2})", OneSkyFile.Filename, ExportFile.FullName, Culture);
+							}
+
+							if (ExportFileWasReadOnly)
+							{
+								ExportFile.IsReadOnly = true;
+							}
 						}
 
-						if (ExportFileWasReadOnly)
+						// Also update the back-up copy so we can diff against what we got from OneSky, and what the gather commandlet produced
 						{
-							ExportFile.IsReadOnly = true;
+							var ExportFileCopy = new FileInfo(Path.Combine(ExportFile.DirectoryName, String.Format("{0}_FromOneSky{1}", Path.GetFileNameWithoutExtension(ExportFile.Name), ExportFile.Extension)));
+
+							var ExportFileCopyWasReadOnly = false;
+							if (ExportFileCopy.Exists)
+							{
+								// We're going to clobber the existing PO file, so make sure it's writable (it may be read-only if in Perforce)
+								ExportFileCopyWasReadOnly = ExportFileCopy.IsReadOnly;
+								ExportFileCopy.IsReadOnly = false;
+							}
+
+							ExportFile.CopyTo(ExportFileCopy.FullName, true);
+
+							if (ExportFileCopyWasReadOnly)
+							{
+								ExportFileCopy.IsReadOnly = true;
+							}
+
+							// Add/check out backed up POs from OneSky.
+							if (CommandUtils.P4Enabled)
+							{
+								UE4Build.AddBuildProductsToChangelist(PendingChangeList, new List<string>() { ExportFileCopy.FullName });
+							}
 						}
 					}
-
-					// Also update the back-up copy so we can diff against what we got from OneSky, and what the gather commandlet produced
+					else if (ExportTranslationState == UploadedFile.ExportTranslationState.NoContent)
 					{
-						var ExportFileCopy = new FileInfo(Path.Combine(ExportFile.DirectoryName, String.Format("{0}_FromOneSky{1}", Path.GetFileNameWithoutExtension(ExportFile.Name), ExportFile.Extension)));
-
-						var ExportFileCopyWasReadOnly = false;
-						if (ExportFileCopy.Exists)
-						{
-							// We're going to clobber the existing PO file, so make sure it's writable (it may be read-only if in Perforce)
-							ExportFileCopyWasReadOnly = ExportFileCopy.IsReadOnly;
-							ExportFileCopy.IsReadOnly = false;
-						}
-
-						ExportFile.CopyTo(ExportFileCopy.FullName, true);
-
-						if (ExportFileCopyWasReadOnly)
-						{
-							ExportFileCopy.IsReadOnly = true;
-						}
-
-						// Add/check out backed up POs from OneSky.
-						if (CommandUtils.P4Enabled)
-						{
-							UE4Build.AddBuildProductsToChangelist(PendingChangeList, new List<string>() { ExportFileCopy.FullName });
-						}
+						Console.WriteLine("[SUCCESS] Skipped '{0}' ({1}) as it has no translations", OneSkyFile.Filename, Culture);
+					}
+					else
+					{
+						Console.WriteLine("[FAILED] Failed to get translation data for '{0}' ({1})", OneSkyFile.Filename, Culture);
 					}
-				}
-				else if (ExportTranslationState == UploadedFile.ExportTranslationState.NoContent)
-				{
-					Console.WriteLine("[SUCCESS] Skipped '{0}' ({1}) as it has no translations", OneSkyFile.Filename, Culture);
-				}
-				else
-				{
-					Console.WriteLine("[FAILED] Failed to get translation data for '{0}' ({1})", OneSkyFile.Filename, Culture);
 				}
 			}
 		}
-	}
 
-	public override void UploadProjectToLocalizationProvider(string ProjectName, ProjectImportExportInfo ProjectExportInfo)
-	{
-		var OneSkyProject = GetOneSkyProject(ProjectName);
-
-		// Upload the .po file for the native culture first
-		UploadFileToOneSky(OneSkyProject, ProjectExportInfo.NativeCulture, null, ProjectExportInfo);
-		foreach (var Platform in ProjectExportInfo.SplitPlatformNames)
+		public override void UploadProjectToLocalizationProvider(string ProjectName, ProjectImportExportInfo ProjectExportInfo)
 		{
-			UploadFileToOneSky(OneSkyProject, ProjectExportInfo.NativeCulture, Platform, ProjectExportInfo);
-		}
+			var OneSkyProject = GetOneSkyProject(ProjectName);
 
-		if (bUploadAllCultures)
-		{
-			// Upload the remaining .po files for the other cultures
-			foreach (var Culture in ProjectExportInfo.CulturesToGenerate)
+			// Upload the .po file for the native culture first
+			UploadFileToOneSky(OneSkyProject, ProjectExportInfo.NativeCulture, null, ProjectExportInfo);
+			foreach (var Platform in ProjectExportInfo.SplitPlatformNames)
 			{
-				// Skip native culture as we uploaded it above
-				if (Culture != ProjectExportInfo.NativeCulture)
+				UploadFileToOneSky(OneSkyProject, ProjectExportInfo.NativeCulture, Platform, ProjectExportInfo);
+			}
+
+			if (bUploadAllCultures)
+			{
+				// Upload the remaining .po files for the other cultures
+				foreach (var Culture in ProjectExportInfo.CulturesToGenerate)
 				{
-					UploadFileToOneSky(OneSkyProject, Culture, null, ProjectExportInfo);
-					foreach (var Platform in ProjectExportInfo.SplitPlatformNames)
+					// Skip native culture as we uploaded it above
+					if (Culture != ProjectExportInfo.NativeCulture)
 					{
-						UploadFileToOneSky(OneSkyProject, Culture, Platform, ProjectExportInfo);
+						UploadFileToOneSky(OneSkyProject, Culture, null, ProjectExportInfo);
+						foreach (var Platform in ProjectExportInfo.SplitPlatformNames)
+						{
+							UploadFileToOneSky(OneSkyProject, Culture, Platform, ProjectExportInfo);
+						}
 					}
 				}
 			}
 		}
-	}
-
-	private UploadedFile UploadOneSkyTranslationWithRetry(OneSky.Project OneSkyProject, string OneSkyFileName, string Culture, FileStream FileStream)
-	{
-		const int MAX_COUNT = 3;
 
-		long StartingFilePos = FileStream.Position;
-		int Count = 0;
-		for (;;)
+		private UploadedFile UploadOneSkyTranslationWithRetry(OneSky.Project OneSkyProject, string OneSkyFileName, string Culture, FileStream FileStream)
 		{
-			try
-			{
-				return OneSkyProject.Upload(OneSkyFileName, FileStream, Culture).Result;
-			}
-			catch (Exception)
+			const int MAX_COUNT = 3;
+
+			long StartingFilePos = FileStream.Position;
+			int Count = 0;
+			for (;;)
 			{
-				if (++Count < MAX_COUNT)
+				try
 				{
-					FileStream.Position = StartingFilePos;
-					Console.WriteLine("UploadOneSkyTranslation attempt {0}/{1} failed. Retrying...", Count, MAX_COUNT);
-					continue;
+					return OneSkyProject.Upload(OneSkyFileName, FileStream, Culture).Result;
 				}
+				catch (Exception)
+				{
+					if (++Count < MAX_COUNT)
+					{
+						FileStream.Position = StartingFilePos;
+						Console.WriteLine("UploadOneSkyTranslation attempt {0}/{1} failed. Retrying...", Count, MAX_COUNT);
+						continue;
+					}
 
-				Console.WriteLine("UploadOneSkyTranslation attempt {0}/{1} failed.", Count, MAX_COUNT);
-				break;
+					Console.WriteLine("UploadOneSkyTranslation attempt {0}/{1} failed.", Count, MAX_COUNT);
+					break;
+				}
 			}
-		}
-
-		return null;
-	}
 
-	private void UploadFileToOneSky(OneSky.Project OneSkyProject, string Culture, string Platform, ProjectImportExportInfo ProjectExportInfo)
-	{
-		var SourceDirectory = String.IsNullOrEmpty(Platform)
-				? new DirectoryInfo(CommandUtils.CombinePaths(RootWorkingDirectory, ProjectExportInfo.DestinationPath))
-				: new DirectoryInfo(CommandUtils.CombinePaths(RootWorkingDirectory, ProjectExportInfo.DestinationPath, ProjectImportExportInfo.PlatformLocalizationFolderName, Platform));
-		var CultureDirectory = (ProjectExportInfo.bUseCultureDirectory) ? new DirectoryInfo(Path.Combine(SourceDirectory.FullName, Culture)) : SourceDirectory;
+			return null;
+		}
 
-		var FileToUpload = new FileInfo(Path.Combine(CultureDirectory.FullName, ProjectExportInfo.PortableObjectName));
-		using (var FileStream = FileToUpload.OpenRead())
+		private void UploadFileToOneSky(OneSky.Project OneSkyProject, string Culture, string Platform, ProjectImportExportInfo ProjectExportInfo)
 		{
-			// Read the BOM
-			var UTF8BOM = new byte[3];
-			FileStream.Read(UTF8BOM, 0, 3);
+			var SourceDirectory = String.IsNullOrEmpty(Platform)
+					? new DirectoryInfo(CommandUtils.CombinePaths(RootWorkingDirectory, ProjectExportInfo.DestinationPath))
+					: new DirectoryInfo(CommandUtils.CombinePaths(RootWorkingDirectory, ProjectExportInfo.DestinationPath, ProjectImportExportInfo.PlatformLocalizationFolderName, Platform));
+			var CultureDirectory = (ProjectExportInfo.bUseCultureDirectory) ? new DirectoryInfo(Path.Combine(SourceDirectory.FullName, Culture)) : SourceDirectory;
 
-			// We want to ignore the utf8 BOM
-			if (UTF8BOM[0] != 0xef || UTF8BOM[1] != 0xbb || UTF8BOM[2] != 0xbf)
+			var FileToUpload = new FileInfo(Path.Combine(CultureDirectory.FullName, ProjectExportInfo.PortableObjectName));
+			using (var FileStream = FileToUpload.OpenRead())
 			{
-				FileStream.Position = 0;
-			}
+				// Read the BOM
+				var UTF8BOM = new byte[3];
+				FileStream.Read(UTF8BOM, 0, 3);
 
-			var OneSkyFileName = GetOneSkyFilename(ProjectExportInfo.PortableObjectName, Platform);
+				// We want to ignore the utf8 BOM
+				if (UTF8BOM[0] != 0xef || UTF8BOM[1] != 0xbb || UTF8BOM[2] != 0xbf)
+				{
+					FileStream.Position = 0;
+				}
 
-			Console.WriteLine("Uploading: '{0}' as '{1}' ({2})", FileToUpload.FullName, OneSkyFileName, Culture);
+				var OneSkyFileName = GetOneSkyFilename(ProjectExportInfo.PortableObjectName, Platform);
 
-			var UploadedFile = UploadOneSkyTranslationWithRetry(OneSkyProject, OneSkyFileName, Culture, FileStream);
+				Console.WriteLine("Uploading: '{0}' as '{1}' ({2})", FileToUpload.FullName, OneSkyFileName, Culture);
 
-			if (UploadedFile == null)
-			{
-				Console.WriteLine("[FAILED] Uploading: '{0}' ({1})", FileToUpload.FullName, Culture);
-			}
-			else
-			{
-				Console.WriteLine("[SUCCESS] Uploading: '{0}' ({1})", FileToUpload.FullName, Culture);
+				var UploadedFile = UploadOneSkyTranslationWithRetry(OneSkyProject, OneSkyFileName, Culture, FileStream);
+
+				if (UploadedFile == null)
+				{
+					Console.WriteLine("[FAILED] Uploading: '{0}' ({1})", FileToUpload.FullName, Culture);
+				}
+				else
+				{
+					Console.WriteLine("[SUCCESS] Uploading: '{0}' ({1})", FileToUpload.FullName, Culture);
+				}
 			}
 		}
-	}
-
-	private ProjectGroup GetOneSkyProjectGroup(string ProjectGroupName)
-	{
-		var OneSkyProjectGroup = OneSkyService.ProjectGroups.FirstOrDefault(g => g.Name == ProjectGroupName);
 
-		if (OneSkyProjectGroup == null)
+		private ProjectGroup GetOneSkyProjectGroup(string ProjectGroupName)
 		{
-			OneSkyProjectGroup = new ProjectGroup(ProjectGroupName, "en");
-			OneSkyService.ProjectGroups.Add(OneSkyProjectGroup);
-		}
+			var OneSkyProjectGroup = OneSkyService.ProjectGroups.FirstOrDefault(g => g.Name == ProjectGroupName);
 
-		return OneSkyProjectGroup;
-	}
+			if (OneSkyProjectGroup == null)
+			{
+				OneSkyProjectGroup = new ProjectGroup(ProjectGroupName, "en");
+				OneSkyService.ProjectGroups.Add(OneSkyProjectGroup);
+			}
 
-	private OneSky.Project GetOneSkyProject(string ProjectName, string ProjectDescription = "")
-	{
-		OneSky.Project OneSkyProject = OneSkyProjectGroup.Projects.FirstOrDefault(p => p.Name == ProjectName);
+			return OneSkyProjectGroup;
+		}
 
-		if (OneSkyProject == null)
+		private OneSky.Project GetOneSkyProject(string ProjectName, string ProjectDescription = "")
 		{
-			ProjectType projectType = OneSkyService.ProjectTypes.First(pt => pt.Code == "website");
+			OneSky.Project OneSkyProject = OneSkyProjectGroup.Projects.FirstOrDefault(p => p.Name == ProjectName);
 
-			OneSkyProject = new OneSky.Project(ProjectName, ProjectDescription, projectType);
-			OneSkyProjectGroup.Projects.Add(OneSkyProject);
-		}
+			if (OneSkyProject == null)
+			{
+				ProjectType projectType = OneSkyService.ProjectTypes.First(pt => pt.Code == "website");
 
-		return OneSkyProject;
-	}
+				OneSkyProject = new OneSky.Project(ProjectName, ProjectDescription, projectType);
+				OneSkyProjectGroup.Projects.Add(OneSkyProject);
+			}
 
-	private string GetOneSkyFilename(string BaseFilename, string Platform)
-	{
-		var OneSkyFileName = BaseFilename;
-		if (!String.IsNullOrEmpty(Platform))
-		{
-			// Apply the platform suffix. XLoc will take care of merging the files from different platforms together.
-			var OneSkyFileNameWithSuffix = Path.GetFileNameWithoutExtension(OneSkyFileName) + "_" + Platform + Path.GetExtension(OneSkyFileName);
-			OneSkyFileName = OneSkyFileNameWithSuffix;
+			return OneSkyProject;
 		}
-		if (!String.IsNullOrEmpty(LocalizationBranchName))
-		{
-			// Apply the branch suffix. OneSky will take care of merging the files from different branches together.
-			var OneSkyFileNameWithSuffix = Path.GetFileNameWithoutExtension(OneSkyFileName) + "_" + LocalizationBranchName + Path.GetExtension(OneSkyFileName);
-			OneSkyFileName = OneSkyFileNameWithSuffix;
-		}
-		if (!String.IsNullOrEmpty(RemoteFilenamePrefix))
+
+		private string GetOneSkyFilename(string BaseFilename, string Platform)
 		{
-			// Apply the prefix (this is used to avoid collisions with plugins that use the same name for their PO files)
-			OneSkyFileName = RemoteFilenamePrefix + "_" + OneSkyFileName;
+			var OneSkyFileName = BaseFilename;
+			if (!String.IsNullOrEmpty(Platform))
+			{
+				// Apply the platform suffix. XLoc will take care of merging the files from different platforms together.
+				var OneSkyFileNameWithSuffix = Path.GetFileNameWithoutExtension(OneSkyFileName) + "_" + Platform + Path.GetExtension(OneSkyFileName);
+				OneSkyFileName = OneSkyFileNameWithSuffix;
+			}
+			if (!String.IsNullOrEmpty(LocalizationBranchName))
+			{
+				// Apply the branch suffix. OneSky will take care of merging the files from different branches together.
+				var OneSkyFileNameWithSuffix = Path.GetFileNameWithoutExtension(OneSkyFileName) + "_" + LocalizationBranchName + Path.GetExtension(OneSkyFileName);
+				OneSkyFileName = OneSkyFileNameWithSuffix;
+			}
+			if (!String.IsNullOrEmpty(RemoteFilenamePrefix))
+			{
+				// Apply the prefix (this is used to avoid collisions with plugins that use the same name for their PO files)
+				OneSkyFileName = RemoteFilenamePrefix + "_" + OneSkyFileName;
+			}
+			return OneSkyFileName;
 		}
-		return OneSkyFileName;
-	}
-
-	private OneSkyService OneSkyService;
-	private ProjectGroup OneSkyProjectGroup;
-}
 
+		private OneSkyService OneSkyService;
+		private ProjectGroup OneSkyProjectGroup;
+	}
 }
diff --git a/Engine/Source/Programs/AutomationTool/Scripts/Localisation.Automation.cs b/Engine/Source/Programs/AutomationTool/Scripts/Localisation.Automation.cs
--- a/Engine/Source/Programs/AutomationTool/Scripts/Localisation.Automation.cs
+++ b/Engine/Source/Programs/AutomationTool/Scripts/Localisation.Automation.cs
@@ -24,9 +24,9 @@ using Tools.DotNETCommon;
 [Help("AdditionalCommandletArguments", "Optional arguments to pass to the gather process.")]
 class Localize : BuildCommand
 {
-	private struct LocalizationBatch
+	private class LocalizationBatch
 	{
-		public LocalizationBatch(string InUEProjectDirectory, string InLocalizationTargetDirectory, string InRemoteFilenamePrefix, List<string> InLocalizationProjectNames)
+		public LocalizationBatch(string InUEProjectDirectory, string InLocalizationTargetDirectory, string InRemoteFilenamePrefix, IReadOnlyList<string> InLocalizationProjectNames)
 		{
 			UEProjectDirectory = InUEProjectDirectory;
 			LocalizationTargetDirectory = InLocalizationTargetDirectory;
@@ -34,10 +34,10 @@ class Localize : BuildCommand
 			LocalizationProjectNames = InLocalizationProjectNames;
 		}
 
-		public string UEProjectDirectory;
-		public string LocalizationTargetDirectory;
-		public string RemoteFilenamePrefix;
-		public List<string> LocalizationProjectNames;
+		public string UEProjectDirectory { get; private set; }
+		public string LocalizationTargetDirectory { get; private set; }
+		public string RemoteFilenamePrefix { get; private set; }
+		public IReadOnlyList<string> LocalizationProjectNames { get; private set; }
 	};
 
 	public override void ExecuteBuild()
@@ -78,21 +78,21 @@ class Localize : BuildCommand
 			LocalizationProviderName = "";
 		}
 
-		var LocalizationSteps = new List<string>();
+		var LocalizationStepNames = new List<string>();
 		{
-			var LocalizationStepsStr = ParseParamValue("LocalizationSteps");
-			if (LocalizationStepsStr == null)
+			var LocalizationStepNamesStr = ParseParamValue("LocalizationSteps");
+			if (LocalizationStepNamesStr == null)
 			{
-				LocalizationSteps.AddRange(new string[] { "Download", "Gather", "Import", "Export", "Compile", "GenerateReports", "Upload" });
+				LocalizationStepNames.AddRange(new string[] { "Download", "Gather", "Import", "Export", "Compile", "GenerateReports", "Upload" });
 			}
 			else
 			{
-				foreach (var StepName in LocalizationStepsStr.Split(','))
+				foreach (var StepName in LocalizationStepNamesStr.Split(','))
 				{
-					LocalizationSteps.Add(StepName.Trim());
+					LocalizationStepNames.Add(StepName.Trim());
 				}
 			}
-			LocalizationSteps.Add("Monolithic"); // Always allow the monolithic scripts to run as we don't know which steps they do
+			LocalizationStepNames.Add("Monolithic"); // Always allow the monolithic scripts to run as we don't know which steps they do
 		}
 
 		var ShouldGatherPlugins = ParseParam("IncludePlugins");
@@ -183,7 +183,7 @@ class Localize : BuildCommand
 		// Process each localization batch
 		foreach (var LocalizationBatch in LocalizationBatches)
 		{
-			ProcessLocalizationProjects(LocalizationBatch, PendingChangeList, UEProjectRoot, UEProjectName, LocalizationProviderName, LocalizationSteps, AdditionalCommandletArguments);
+			ProcessLocalizationProjects(LocalizationBatch, PendingChangeList, UEProjectRoot, UEProjectName, LocalizationProviderName, LocalizationStepNames, AdditionalCommandletArguments);
 		}
 
 		// Clean-up the changelist so it only contains the changed files, and then submit it (if we were asked to)
@@ -232,7 +232,7 @@ class Localize : BuildCommand
 		}
 	}
 
-	private void ProcessLocalizationProjects(LocalizationBatch LocalizationBatch, int PendingChangeList, string UEProjectRoot, string UEProjectName, string LocalizationProviderName, List<string> LocalizationSteps, string AdditionalCommandletArguments)
+	private void ProcessLocalizationProjects(LocalizationBatch LocalizationBatch, int PendingChangeList, string UEProjectRoot, string UEProjectName, string LocalizationProviderName, IReadOnlyList<string> LocalizationSteps, string AdditionalCommandletArguments)
 	{
 		var EditorExe = CombinePaths(CmdEnv.LocalRoot, @"Engine/Binaries/Win64/UE4Editor-Cmd.exe");
 		var RootWorkingDirectory = CombinePaths(UEProjectRoot, LocalizationBatch.UEProjectDirectory);
@@ -348,12 +348,11 @@ class Localize : BuildCommand
 		}
 	}
 
-	private ProjectInfo GenerateProjectInfo(string RootWorkingDirectory, string ProjectName, List<string> LocalizationSteps)
+	private ProjectInfo GenerateProjectInfo(string RootWorkingDirectory, string ProjectName, IReadOnlyList<string> LocalizationStepNames)
 	{
-		var ProjectInfo = new ProjectInfo();
-
-		ProjectInfo.ProjectName = ProjectName;
-		ProjectInfo.LocalizationSteps = new List<ProjectStepInfo>();
+		var LocalizationSteps = new List<ProjectStepInfo>();
+		ProjectImportExportInfo ImportInfo = null;
+		ProjectImportExportInfo ExportInfo = null;
 
 		// Projects generated by the localization dashboard will use multiple config files that must be run in a specific order
 		// Older projects (such as the Engine) would use a single config file containing all the steps
@@ -361,18 +360,18 @@ class Localize : BuildCommand
 		var MonolithicConfigFile = CombinePaths(RootWorkingDirectory, String.Format(@"Config/Localization/{0}.ini", ProjectName));
 		if (File.Exists(MonolithicConfigFile))
 		{
-			ProjectInfo.LocalizationSteps.Add(new ProjectStepInfo("Monolithic", MonolithicConfigFile));
+			LocalizationSteps.Add(new ProjectStepInfo("Monolithic", MonolithicConfigFile));
 
-			ProjectInfo.ImportInfo = GenerateProjectImportExportInfo(RootWorkingDirectory, MonolithicConfigFile);
-			ProjectInfo.ExportInfo = ProjectInfo.ImportInfo;
+			ImportInfo = GenerateProjectImportExportInfo(RootWorkingDirectory, MonolithicConfigFile);
+			ExportInfo = ImportInfo;
 		}
 		else
 		{
 			var FileSuffixes = new[] { 
-				new { Suffix = "Gather", Required = LocalizationSteps.Contains("Gather") }, 
-				new { Suffix = "Import", Required = LocalizationSteps.Contains("Import") || LocalizationSteps.Contains("Download") },	// Downloading needs the parsed ImportInfo
-				new { Suffix = "Export", Required = LocalizationSteps.Contains("Gather") || LocalizationSteps.Contains("Upload")},		// Uploading needs the parsed ExportInfo
-				new { Suffix = "Compile", Required = LocalizationSteps.Contains("Compile") }, 
+				new { Suffix = "Gather", Required = LocalizationStepNames.Contains("Gather") }, 
+				new { Suffix = "Import", Required = LocalizationStepNames.Contains("Import") || LocalizationStepNames.Contains("Download") },	// Downloading needs the parsed ImportInfo
+				new { Suffix = "Export", Required = LocalizationStepNames.Contains("Gather") || LocalizationStepNames.Contains("Upload")},		// Uploading needs the parsed ExportInfo
+				new { Suffix = "Compile", Required = LocalizationStepNames.Contains("Compile") }, 
 				new { Suffix = "GenerateReports", Required = false } 
 			};
 
@@ -382,15 +381,15 @@ class Localize : BuildCommand
 
 				if (File.Exists(ModularConfigFile))
 				{
-					ProjectInfo.LocalizationSteps.Add(new ProjectStepInfo(FileSuffix.Suffix, ModularConfigFile));
+					LocalizationSteps.Add(new ProjectStepInfo(FileSuffix.Suffix, ModularConfigFile));
 
 					if (FileSuffix.Suffix == "Import")
 					{
-						ProjectInfo.ImportInfo = GenerateProjectImportExportInfo(RootWorkingDirectory, ModularConfigFile);
+						ImportInfo = GenerateProjectImportExportInfo(RootWorkingDirectory, ModularConfigFile);
 					}
 					else if (FileSuffix.Suffix == "Export")
 					{
-						ProjectInfo.ExportInfo = GenerateProjectImportExportInfo(RootWorkingDirectory, ModularConfigFile);
+						ExportInfo = GenerateProjectImportExportInfo(RootWorkingDirectory, ModularConfigFile);
 					}
 				}
 				else if (FileSuffix.Required)
@@ -400,58 +399,63 @@ class Localize : BuildCommand
 			}
 		}
 
-		return ProjectInfo;
+		return new ProjectInfo(ProjectName, LocalizationSteps, ImportInfo, ExportInfo);
 	}
 
 	private ProjectImportExportInfo GenerateProjectImportExportInfo(string RootWorkingDirectory, string LocalizationConfigFile)
 	{
-		var ProjectImportExportInfo = new ProjectImportExportInfo();
-
 		ConfigFile File = new ConfigFile(new FileReference(LocalizationConfigFile), ConfigLineAction.Add);
 		var LocalizationConfig = new ConfigHierarchy(new ConfigFile[] { File });
 
-		if (!LocalizationConfig.GetString("CommonSettings", "DestinationPath", out ProjectImportExportInfo.DestinationPath))
+		string DestinationPath;
+		if (!LocalizationConfig.GetString("CommonSettings", "DestinationPath", out DestinationPath))
 		{
 			throw new AutomationException("Failed to find a required config key! Section: 'CommonSettings', Key: 'DestinationPath', File: '{0}'", LocalizationConfigFile);
 		}
 
-		if (!LocalizationConfig.GetString("CommonSettings", "ManifestName", out ProjectImportExportInfo.ManifestName))
+		string ManifestName;
+		if (!LocalizationConfig.GetString("CommonSettings", "ManifestName", out ManifestName))
 		{
 			throw new AutomationException("Failed to find a required config key! Section: 'CommonSettings', Key: 'ManifestName', File: '{0}'", LocalizationConfigFile);
 		}
 
-		if (!LocalizationConfig.GetString("CommonSettings", "ArchiveName", out ProjectImportExportInfo.ArchiveName))
+		string ArchiveName;
+		if (!LocalizationConfig.GetString("CommonSettings", "ArchiveName", out ArchiveName))
 		{
 			throw new AutomationException("Failed to find a required config key! Section: 'CommonSettings', Key: 'ArchiveName', File: '{0}'", LocalizationConfigFile);
 		}
 
-		if (!LocalizationConfig.GetString("CommonSettings", "PortableObjectName", out ProjectImportExportInfo.PortableObjectName))
+		string PortableObjectName;
+		if (!LocalizationConfig.GetString("CommonSettings", "PortableObjectName", out PortableObjectName))
 		{
 			throw new AutomationException("Failed to find a required config key! Section: 'CommonSettings', Key: 'PortableObjectName', File: '{0}'", LocalizationConfigFile);
 		}
 
-		if (!LocalizationConfig.GetString("CommonSettings", "NativeCulture", out ProjectImportExportInfo.NativeCulture))
+		string NativeCulture;
+		if (!LocalizationConfig.GetString("CommonSettings", "NativeCulture", out NativeCulture))
 		{
 			throw new AutomationException("Failed to find a required config key! Section: 'CommonSettings', Key: 'NativeCulture', File: '{0}'", LocalizationConfigFile);
 		}
 
-		if (!LocalizationConfig.GetArray("CommonSettings", "CulturesToGenerate", out ProjectImportExportInfo.CulturesToGenerate))
+		List<string> CulturesToGenerate;
+		if (!LocalizationConfig.GetArray("CommonSettings", "CulturesToGenerate", out CulturesToGenerate))
 		{
 			throw new AutomationException("Failed to find a required config key! Section: 'CommonSettings', Key: 'CulturesToGenerate', File: '{0}'", LocalizationConfigFile);
 		}
 
-		if (!LocalizationConfig.GetBool("CommonSettings", "bUseCultureDirectory", out ProjectImportExportInfo.bUseCultureDirectory))
+		bool bUseCultureDirectory;
+		if (!LocalizationConfig.GetBool("CommonSettings", "bUseCultureDirectory", out bUseCultureDirectory))
 		{
 			// bUseCultureDirectory is optional, default is true
-			ProjectImportExportInfo.bUseCultureDirectory = true;
+			bUseCultureDirectory = true;
 		}
 
+		var ProjectImportExportInfo = new ProjectImportExportInfo(DestinationPath, ManifestName, ArchiveName, PortableObjectName, NativeCulture, CulturesToGenerate, bUseCultureDirectory);
 		ProjectImportExportInfo.CalculateSplitPlatformNames(RootWorkingDirectory);
-
 		return ProjectImportExportInfo;
 	}
 
-	private Dictionary<string, byte[]> GetPOFileHashes(List<LocalizationBatch> LocalizationBatches, string UEProjectRoot)
+	private Dictionary<string, byte[]> GetPOFileHashes(IReadOnlyList<LocalizationBatch> LocalizationBatches, string UEProjectRoot)
 	{
 		var AllFiles = new Dictionary<string, byte[]>();
 
diff --git a/Engine/Source/Programs/AutomationTool/XLocLocalization/XLocLocalizationProvider.cs b/Engine/Source/Programs/AutomationTool/XLocLocalization/XLocLocalizationProvider.cs
--- a/Engine/Source/Programs/AutomationTool/XLocLocalization/XLocLocalizationProvider.cs
+++ b/Engine/Source/Programs/AutomationTool/XLocLocalization/XLocLocalizationProvider.cs
@@ -14,422 +14,420 @@ using EpicGames.Localization;
 
 namespace EpicGames.XLocLocalization
 {
-
-public struct XLocConfig
-{
-	public string Server;
-	public string APIKey;
-	public string APISecret;
-	public string LocalizationId;
-};
-
-public struct XLocUtils
-{
-	public static string MD5HashString(string Str)
+	public struct XLocConfig
 	{
-		var HashBytes = MD5.Create().ComputeHash(Encoding.UTF8.GetBytes(Str));
-
-		StringBuilder HashString = new StringBuilder();
-		foreach (var HashByte in HashBytes)
-		{
-			HashString.Append(HashByte.ToString("x2"));
-		}
-		return HashString.ToString();
-	}
+		public string Server;
+		public string APIKey;
+		public string APISecret;
+		public string LocalizationId;
+	};
 
-	public static WebResponse GetWebResponse(WebRequest Request)
+	public struct XLocUtils
 	{
-		// GetResponse sometimes throws rather than just letting you use the error code in the response
-		try
+		public static string MD5HashString(string Str)
 		{
-			return Request.GetResponse();
+			var HashBytes = MD5.Create().ComputeHash(Encoding.UTF8.GetBytes(Str));
+
+			StringBuilder HashString = new StringBuilder();
+			foreach (var HashByte in HashBytes)
+			{
+				HashString.Append(HashByte.ToString("x2"));
+			}
+			return HashString.ToString();
 		}
-		catch (WebException Ex)
+
+		public static WebResponse GetWebResponse(WebRequest Request)
 		{
-			return Ex.Response;
+			// GetResponse sometimes throws rather than just letting you use the error code in the response
+			try
+			{
+				return Request.GetResponse();
+			}
+			catch (WebException Ex)
+			{
+				return Ex.Response;
+			}
 		}
 	}
-}
 
-public abstract class XLocLocalizationProvider : LocalizationProvider
-{
-	public XLocLocalizationProvider(LocalizationProviderArgs InArgs)
-		: base(InArgs)
+	public abstract class XLocLocalizationProvider : LocalizationProvider
 	{
-		Config = new XLocConfig();
-	}
-
-	public override void DownloadProjectFromLocalizationProvider(string ProjectName, ProjectImportExportInfo ProjectImportInfo)
-	{
-		var XLocApiClient = CreateXLocApiClient();
+		public XLocLocalizationProvider(LocalizationProviderArgs InArgs)
+			: base(InArgs)
+		{
+			Config = new XLocConfig();
+		}
 
-		try
+		public override void DownloadProjectFromLocalizationProvider(string ProjectName, ProjectImportExportInfo ProjectImportInfo)
 		{
-			var AuthToken = RequestAuthTokenWithRetry(XLocApiClient);
+			var XLocApiClient = CreateXLocApiClient();
 
-			// Get the latest files for each culture.
-			foreach (var Culture in ProjectImportInfo.CulturesToGenerate)
+			try
 			{
-				// Skip the native culture, as XLoc doesn't have an entry for it
-				if (Culture == ProjectImportInfo.NativeCulture)
-				{
-					continue;
-				}
+				var AuthToken = RequestAuthTokenWithRetry(XLocApiClient);
 
-				DownloadLatestPOFile(XLocApiClient, AuthToken, Culture, null, ProjectImportInfo);
-				foreach (var Platform in ProjectImportInfo.SplitPlatformNames)
+				// Get the latest files for each culture.
+				foreach (var Culture in ProjectImportInfo.CulturesToGenerate)
 				{
-					DownloadLatestPOFile(XLocApiClient, AuthToken, Culture, Platform, ProjectImportInfo);
+					// Skip the native culture, as XLoc doesn't have an entry for it
+					if (Culture == ProjectImportInfo.NativeCulture)
+					{
+						continue;
+					}
+
+					DownloadLatestPOFile(XLocApiClient, AuthToken, Culture, null, ProjectImportInfo);
+					foreach (var Platform in ProjectImportInfo.SplitPlatformNames)
+					{
+						DownloadLatestPOFile(XLocApiClient, AuthToken, Culture, Platform, ProjectImportInfo);
+					}
 				}
 			}
+			finally
+			{
+				XLocApiClient.Close();
+			}
 		}
-		finally
-		{
-			XLocApiClient.Close();
-		}
-	}
 
-	private void DownloadLatestPOFile(XLocApiClient XLocApiClient, string AuthToken, string Culture, string Platform, ProjectImportExportInfo ProjectImportInfo)
-	{
-		var XLocFilename = GetXLocFilename(ProjectImportInfo.PortableObjectName, Platform);
-
-		// This will throw if the requested culture is invalid, but we don't want to let that kill the whole gather
-		var LatestBuildXml = "";
-		try
-		{
-			var EpicCultureToXLocLanguageId = GetEpicCultureToXLocLanguageId();
-			LatestBuildXml = RequestLatestBuild(XLocApiClient, AuthToken, EpicCultureToXLocLanguageId[Culture], XLocFilename);
-		}
-		catch (Exception Ex)
-		{
-			BuildCommand.LogWarning("RequestLatestBuild failed for {0}. {1}", Culture, Ex);
-			return;
-		}
-		if (String.IsNullOrEmpty(LatestBuildXml))
+		private void DownloadLatestPOFile(XLocApiClient XLocApiClient, string AuthToken, string Culture, string Platform, ProjectImportExportInfo ProjectImportInfo)
 		{
-			Console.WriteLine("[IGNORED] '{0}' has no build data ({1})", XLocFilename, Culture);
-			return;
-		}
+			var XLocFilename = GetXLocFilename(ProjectImportInfo.PortableObjectName, Platform);
 
-		var POFileUri = "";
-		var BuildsXmlDoc = new XmlDocument();
-		BuildsXmlDoc.LoadXml(LatestBuildXml);
+			// This will throw if the requested culture is invalid, but we don't want to let that kill the whole gather
+			var LatestBuildXml = "";
+			try
+			{
+				var EpicCultureToXLocLanguageId = GetEpicCultureToXLocLanguageId();
+				LatestBuildXml = RequestLatestBuild(XLocApiClient, AuthToken, EpicCultureToXLocLanguageId[Culture], XLocFilename);
+			}
+			catch (Exception Ex)
+			{
+				BuildCommand.LogWarning("RequestLatestBuild failed for {0}. {1}", Culture, Ex);
+				return;
+			}
+			if (String.IsNullOrEmpty(LatestBuildXml))
+			{
+				Console.WriteLine("[IGNORED] '{0}' has no build data ({1})", XLocFilename, Culture);
+				return;
+			}
 
-		var BuildElem = BuildsXmlDoc["Build"];
-		if (BuildElem != null)
-		{
-			var BuildFilesElem = BuildElem["BuildFiles"];
-			if (BuildFilesElem != null)
+			var POFileUri = "";
+			var BuildsXmlDoc = new XmlDocument();
+			BuildsXmlDoc.LoadXml(LatestBuildXml);
+
+			var BuildElem = BuildsXmlDoc["Build"];
+			if (BuildElem != null)
 			{
-				foreach (XmlNode BuildFile in BuildFilesElem)
+				var BuildFilesElem = BuildElem["BuildFiles"];
+				if (BuildFilesElem != null)
 				{
-					bool IsCorrectFile = false;
-
-					// Is this the file we want?
-					var GameFileElem = BuildFile["GameFile"];
-					if (GameFileElem != null)
+					foreach (XmlNode BuildFile in BuildFilesElem)
 					{
-						var GameFileNameElem = GameFileElem["Name"];
-						if (GameFileNameElem != null && GameFileNameElem.InnerText.Equals(XLocFilename, StringComparison.OrdinalIgnoreCase))
+						bool IsCorrectFile = false;
+
+						// Is this the file we want?
+						var GameFileElem = BuildFile["GameFile"];
+						if (GameFileElem != null)
 						{
-							IsCorrectFile = true;
+							var GameFileNameElem = GameFileElem["Name"];
+							if (GameFileNameElem != null && GameFileNameElem.InnerText.Equals(XLocFilename, StringComparison.OrdinalIgnoreCase))
+							{
+								IsCorrectFile = true;
+							}
 						}
-					}
 
-					if (IsCorrectFile)
-					{
-						var BuildFileDownloadUriElem = BuildFile["DownloadUri"];
-						if (BuildFileDownloadUriElem != null)
+						if (IsCorrectFile)
 						{
-							POFileUri = BuildFileDownloadUriElem.InnerText;
-							break;
+							var BuildFileDownloadUriElem = BuildFile["DownloadUri"];
+							if (BuildFileDownloadUriElem != null)
+							{
+								POFileUri = BuildFileDownloadUriElem.InnerText;
+								break;
+							}
 						}
 					}
 				}
 			}
-		}
 
-		if (String.IsNullOrEmpty(POFileUri))
-		{
-			Console.WriteLine("[IGNORED] '{0}' was not found in the build data ({1})", XLocFilename, Culture);
-		}
-		else
-		{
-			var DestinationDirectory = String.IsNullOrEmpty(Platform) 
-					? new DirectoryInfo(CommandUtils.CombinePaths(RootWorkingDirectory, ProjectImportInfo.DestinationPath))
-					: new DirectoryInfo(CommandUtils.CombinePaths(RootWorkingDirectory, ProjectImportInfo.DestinationPath, ProjectImportExportInfo.PlatformLocalizationFolderName, Platform));
-			var CultureDirectory = (ProjectImportInfo.bUseCultureDirectory) ? new DirectoryInfo(Path.Combine(DestinationDirectory.FullName, Culture)) : DestinationDirectory;
-			if (!CultureDirectory.Exists)
+			if (String.IsNullOrEmpty(POFileUri))
 			{
-				CultureDirectory.Create();
+				Console.WriteLine("[IGNORED] '{0}' was not found in the build data ({1})", XLocFilename, Culture);
 			}
-
-			var HTTPRequest = WebRequest.Create(POFileUri);
-			HTTPRequest.Method = "GET";
-			using (var Response = (HttpWebResponse)XLocUtils.GetWebResponse(HTTPRequest))
+			else
 			{
-				if (Response.StatusCode != HttpStatusCode.OK)
+				var DestinationDirectory = String.IsNullOrEmpty(Platform)
+						? new DirectoryInfo(CommandUtils.CombinePaths(RootWorkingDirectory, ProjectImportInfo.DestinationPath))
+						: new DirectoryInfo(CommandUtils.CombinePaths(RootWorkingDirectory, ProjectImportInfo.DestinationPath, ProjectImportExportInfo.PlatformLocalizationFolderName, Platform));
+				var CultureDirectory = (ProjectImportInfo.bUseCultureDirectory) ? new DirectoryInfo(Path.Combine(DestinationDirectory.FullName, Culture)) : DestinationDirectory;
+				if (!CultureDirectory.Exists)
 				{
-					BuildCommand.LogWarning("HTTP Request to '{0}' failed. {1}", POFileUri, Response.StatusDescription);
-					return;
+					CultureDirectory.Create();
 				}
 
-				using (var ResponseStream = Response.GetResponseStream())
+				var HTTPRequest = WebRequest.Create(POFileUri);
+				HTTPRequest.Method = "GET";
+				using (var Response = (HttpWebResponse)XLocUtils.GetWebResponse(HTTPRequest))
 				{
-					var ExportFile = new FileInfo(Path.Combine(CultureDirectory.FullName, ProjectImportInfo.PortableObjectName));
-
-					// Write out the updated PO file so that the gather commandlet will import the new data from it
+					if (Response.StatusCode != HttpStatusCode.OK)
 					{
-						var ExportFileWasReadOnly = false;
-						if (ExportFile.Exists)
-						{
-							// We're going to clobber the existing PO file, so make sure it's writable (it may be read-only if in Perforce)
-							ExportFileWasReadOnly = ExportFile.IsReadOnly;
-							ExportFile.IsReadOnly = false;
-						}
-
-						using (var FileStream = ExportFile.Open(FileMode.Create))
-						{
-							ResponseStream.CopyTo(FileStream);
-							Console.WriteLine("[SUCCESS] Exporting: '{0}' as '{1}' ({2})", XLocFilename, ExportFile.FullName, Culture);
-						}
-
-						if (ExportFileWasReadOnly)
-						{
-							ExportFile.IsReadOnly = true;
-						}
+						BuildCommand.LogWarning("HTTP Request to '{0}' failed. {1}", POFileUri, Response.StatusDescription);
+						return;
 					}
 
-					// Also update the back-up copy so we can diff against what we got from XLoc, and what the gather commandlet produced
+					using (var ResponseStream = Response.GetResponseStream())
 					{
-						var ExportFileCopy = new FileInfo(Path.Combine(ExportFile.DirectoryName, String.Format("{0}_FromXLoc{1}", Path.GetFileNameWithoutExtension(ExportFile.Name), ExportFile.Extension)));
+						var ExportFile = new FileInfo(Path.Combine(CultureDirectory.FullName, ProjectImportInfo.PortableObjectName));
 
-						var ExportFileCopyWasReadOnly = false;
-						if (ExportFileCopy.Exists)
+						// Write out the updated PO file so that the gather commandlet will import the new data from it
 						{
-							// We're going to clobber the existing PO file, so make sure it's writable (it may be read-only if in Perforce)
-							ExportFileCopyWasReadOnly = ExportFileCopy.IsReadOnly;
-							ExportFileCopy.IsReadOnly = false;
+							var ExportFileWasReadOnly = false;
+							if (ExportFile.Exists)
+							{
+								// We're going to clobber the existing PO file, so make sure it's writable (it may be read-only if in Perforce)
+								ExportFileWasReadOnly = ExportFile.IsReadOnly;
+								ExportFile.IsReadOnly = false;
+							}
+
+							using (var FileStream = ExportFile.Open(FileMode.Create))
+							{
+								ResponseStream.CopyTo(FileStream);
+								Console.WriteLine("[SUCCESS] Exporting: '{0}' as '{1}' ({2})", XLocFilename, ExportFile.FullName, Culture);
+							}
+
+							if (ExportFileWasReadOnly)
+							{
+								ExportFile.IsReadOnly = true;
+							}
 						}
 
-						ExportFile.CopyTo(ExportFileCopy.FullName, true);
-
-						if (ExportFileCopyWasReadOnly)
+						// Also update the back-up copy so we can diff against what we got from XLoc, and what the gather commandlet produced
 						{
-							ExportFileCopy.IsReadOnly = true;
-						}
-
-						// Add/check out backed up POs from OneSky.
-						if (CommandUtils.P4Enabled)
-						{
-							UE4Build.AddBuildProductsToChangelist(PendingChangeList, new List<string>() { ExportFileCopy.FullName });
+							var ExportFileCopy = new FileInfo(Path.Combine(ExportFile.DirectoryName, String.Format("{0}_FromXLoc{1}", Path.GetFileNameWithoutExtension(ExportFile.Name), ExportFile.Extension)));
+
+							var ExportFileCopyWasReadOnly = false;
+							if (ExportFileCopy.Exists)
+							{
+								// We're going to clobber the existing PO file, so make sure it's writable (it may be read-only if in Perforce)
+								ExportFileCopyWasReadOnly = ExportFileCopy.IsReadOnly;
+								ExportFileCopy.IsReadOnly = false;
+							}
+
+							ExportFile.CopyTo(ExportFileCopy.FullName, true);
+
+							if (ExportFileCopyWasReadOnly)
+							{
+								ExportFileCopy.IsReadOnly = true;
+							}
+
+							// Add/check out backed up POs from OneSky.
+							if (CommandUtils.P4Enabled)
+							{
+								UE4Build.AddBuildProductsToChangelist(PendingChangeList, new List<string>() { ExportFileCopy.FullName });
+							}
 						}
 					}
 				}
 			}
 		}
-	}
-
-	public override void UploadProjectToLocalizationProvider(string ProjectName, ProjectImportExportInfo ProjectExportInfo)
-	{
-		var XLocApiClient = CreateXLocApiClient();
-		var TransferServiceClient = CreateTransferServiceClient();
 
-		try
+		public override void UploadProjectToLocalizationProvider(string ProjectName, ProjectImportExportInfo ProjectExportInfo)
 		{
-			var AuthToken = RequestAuthTokenWithRetry(XLocApiClient);
+			var XLocApiClient = CreateXLocApiClient();
+			var TransferServiceClient = CreateTransferServiceClient();
 
-			// Upload the .po file for the native culture first
-			UploadLatestPOFile(TransferServiceClient, AuthToken, ProjectExportInfo.NativeCulture, null, ProjectExportInfo);
-			foreach (var Platform in ProjectExportInfo.SplitPlatformNames)
+			try
 			{
-				UploadLatestPOFile(TransferServiceClient, AuthToken, ProjectExportInfo.NativeCulture, Platform, ProjectExportInfo);
-			}
+				var AuthToken = RequestAuthTokenWithRetry(XLocApiClient);
 
-			if (bUploadAllCultures)
-			{
-				// Upload the remaining .po files for the other cultures
-				foreach (var Culture in ProjectExportInfo.CulturesToGenerate)
+				// Upload the .po file for the native culture first
+				UploadLatestPOFile(TransferServiceClient, AuthToken, ProjectExportInfo.NativeCulture, null, ProjectExportInfo);
+				foreach (var Platform in ProjectExportInfo.SplitPlatformNames)
 				{
-					// Skip native culture as we uploaded it above
-					if (Culture != ProjectExportInfo.NativeCulture)
+					UploadLatestPOFile(TransferServiceClient, AuthToken, ProjectExportInfo.NativeCulture, Platform, ProjectExportInfo);
+				}
+
+				if (bUploadAllCultures)
+				{
+					// Upload the remaining .po files for the other cultures
+					foreach (var Culture in ProjectExportInfo.CulturesToGenerate)
 					{
-						UploadLatestPOFile(TransferServiceClient, AuthToken, Culture, null, ProjectExportInfo);
-						foreach (var Platform in ProjectExportInfo.SplitPlatformNames)
+						// Skip native culture as we uploaded it above
+						if (Culture != ProjectExportInfo.NativeCulture)
 						{
-							UploadLatestPOFile(TransferServiceClient, AuthToken, Culture, Platform, ProjectExportInfo);
+							UploadLatestPOFile(TransferServiceClient, AuthToken, Culture, null, ProjectExportInfo);
+							foreach (var Platform in ProjectExportInfo.SplitPlatformNames)
+							{
+								UploadLatestPOFile(TransferServiceClient, AuthToken, Culture, Platform, ProjectExportInfo);
+							}
 						}
 					}
 				}
 			}
+			finally
+			{
+				XLocApiClient.Close();
+				TransferServiceClient.Close();
+			}
 		}
-		finally
-		{
-			XLocApiClient.Close();
-			TransferServiceClient.Close();
-		}
-	}
 
-	private void UploadLatestPOFile(TransferServiceClient TransferServiceClient, string AuthToken, string Culture, string Platform, ProjectImportExportInfo ProjectExportInfo)
-	{
-		var SourceDirectory = String.IsNullOrEmpty(Platform)
-				? new DirectoryInfo(CommandUtils.CombinePaths(RootWorkingDirectory, ProjectExportInfo.DestinationPath))
-				: new DirectoryInfo(CommandUtils.CombinePaths(RootWorkingDirectory, ProjectExportInfo.DestinationPath, ProjectImportExportInfo.PlatformLocalizationFolderName, Platform));
-		var CultureDirectory = (ProjectExportInfo.bUseCultureDirectory) ? new DirectoryInfo(Path.Combine(SourceDirectory.FullName, Culture)) : SourceDirectory;
+		private void UploadLatestPOFile(TransferServiceClient TransferServiceClient, string AuthToken, string Culture, string Platform, ProjectImportExportInfo ProjectExportInfo)
+		{
+			var SourceDirectory = String.IsNullOrEmpty(Platform)
+					? new DirectoryInfo(CommandUtils.CombinePaths(RootWorkingDirectory, ProjectExportInfo.DestinationPath))
+					: new DirectoryInfo(CommandUtils.CombinePaths(RootWorkingDirectory, ProjectExportInfo.DestinationPath, ProjectImportExportInfo.PlatformLocalizationFolderName, Platform));
+			var CultureDirectory = (ProjectExportInfo.bUseCultureDirectory) ? new DirectoryInfo(Path.Combine(SourceDirectory.FullName, Culture)) : SourceDirectory;
 
-		var FileToUpload = new FileInfo(Path.Combine(CultureDirectory.FullName, ProjectExportInfo.PortableObjectName));
-		var XLocFilename = GetXLocFilename(ProjectExportInfo.PortableObjectName, Platform);
+			var FileToUpload = new FileInfo(Path.Combine(CultureDirectory.FullName, ProjectExportInfo.PortableObjectName));
+			var XLocFilename = GetXLocFilename(ProjectExportInfo.PortableObjectName, Platform);
 
-		using (var FileStream = FileToUpload.OpenRead())
-		{
-			// We need to leave the language ID field blank for the native culture to avoid XLoc trying to process it as translated source, rather than raw source text
-			var EpicCultureToXLocLanguageId = GetEpicCultureToXLocLanguageId();
-			var LanguageId = (Culture == ProjectExportInfo.NativeCulture) ? "" : EpicCultureToXLocLanguageId[Culture];
+			using (var FileStream = FileToUpload.OpenRead())
+			{
+				// We need to leave the language ID field blank for the native culture to avoid XLoc trying to process it as translated source, rather than raw source text
+				var EpicCultureToXLocLanguageId = GetEpicCultureToXLocLanguageId();
+				var LanguageId = (Culture == ProjectExportInfo.NativeCulture) ? "" : EpicCultureToXLocLanguageId[Culture];
 
-			var FileUploadMetaData = new XLoc.Contracts.GameFileUploadInfo();
-			FileUploadMetaData.CaseSensitive = false;
-			FileUploadMetaData.FileName = XLocFilename;
-			FileUploadMetaData.HistoricalTranslation = false;
-			FileUploadMetaData.LanguageId = LanguageId;
-			FileUploadMetaData.LocalizationId = Config.LocalizationId;
-			FileUploadMetaData.PlatformId = "!";
+				var FileUploadMetaData = new XLoc.Contracts.GameFileUploadInfo();
+				FileUploadMetaData.CaseSensitive = false;
+				FileUploadMetaData.FileName = XLocFilename;
+				FileUploadMetaData.HistoricalTranslation = false;
+				FileUploadMetaData.LanguageId = LanguageId;
+				FileUploadMetaData.LocalizationId = Config.LocalizationId;
+				FileUploadMetaData.PlatformId = "!";
 
-			Console.WriteLine("Uploading: '{0}' as '{1}' ({2})", FileToUpload.FullName, XLocFilename, Culture);
+				Console.WriteLine("Uploading: '{0}' as '{1}' ({2})", FileToUpload.FullName, XLocFilename, Culture);
 
-			try
-			{
-				TransferServiceClient.UploadGameFile(Config.APIKey, AuthToken, FileUploadMetaData, FileToUpload.Length, FileStream);
-				Console.WriteLine("[SUCCESS] Uploading: '{0}' ({1})", FileToUpload.FullName, Culture);
-			}
-			catch (Exception Ex)
-			{
-				Console.WriteLine("[FAILED] Uploading: '{0}' ({1}) - {2}", FileToUpload.FullName, Culture, Ex);
+				try
+				{
+					TransferServiceClient.UploadGameFile(Config.APIKey, AuthToken, FileUploadMetaData, FileToUpload.Length, FileStream);
+					Console.WriteLine("[SUCCESS] Uploading: '{0}' ({1})", FileToUpload.FullName, Culture);
+				}
+				catch (Exception Ex)
+				{
+					Console.WriteLine("[FAILED] Uploading: '{0}' ({1}) - {2}", FileToUpload.FullName, Culture, Ex);
+				}
 			}
 		}
-	}
 
-	protected XLocApiClient CreateXLocApiClient()
-	{
-		var Binding = new BasicHttpBinding();
-		Binding.Name = "basicHttpXLocApi";
-		Binding.CloseTimeout = new TimeSpan(0, 1, 0);
-		Binding.OpenTimeout = new TimeSpan(0, 1, 0);
-		Binding.ReceiveTimeout = new TimeSpan(0, 10, 0);
-		Binding.SendTimeout = new TimeSpan(0, 1, 0);
-		Binding.AllowCookies = false;
-		Binding.BypassProxyOnLocal = false;
-		Binding.HostNameComparisonMode = HostNameComparisonMode.StrongWildcard;
-		Binding.MaxBufferSize = 2147483647;
-		Binding.MaxBufferPoolSize = 5242880;
-		Binding.MaxReceivedMessageSize = 2147483647;
-		Binding.MessageEncoding = WSMessageEncoding.Text;
-		Binding.TextEncoding = Encoding.UTF8;
-		Binding.TransferMode = TransferMode.Buffered;
-		Binding.UseDefaultWebProxy = true;
-		Binding.ReaderQuotas.MaxDepth = 32;
-		Binding.ReaderQuotas.MaxStringContentLength = 2147483647;
-		Binding.ReaderQuotas.MaxArrayLength = 65536;
-		Binding.ReaderQuotas.MaxBytesPerRead = 4096;
-		Binding.ReaderQuotas.MaxNameTableCharCount = 65536;
-
-		var Endpoint = new EndpointAddress(new Uri(String.Format("http://{0}/api/XLocApiService.svc", Config.Server)));
-
-		return new XLocApiClient(Binding, Endpoint);
-	}
+		protected XLocApiClient CreateXLocApiClient()
+		{
+			var Binding = new BasicHttpBinding();
+			Binding.Name = "basicHttpXLocApi";
+			Binding.CloseTimeout = new TimeSpan(0, 1, 0);
+			Binding.OpenTimeout = new TimeSpan(0, 1, 0);
+			Binding.ReceiveTimeout = new TimeSpan(0, 10, 0);
+			Binding.SendTimeout = new TimeSpan(0, 1, 0);
+			Binding.AllowCookies = false;
+			Binding.BypassProxyOnLocal = false;
+			Binding.HostNameComparisonMode = HostNameComparisonMode.StrongWildcard;
+			Binding.MaxBufferSize = 2147483647;
+			Binding.MaxBufferPoolSize = 5242880;
+			Binding.MaxReceivedMessageSize = 2147483647;
+			Binding.MessageEncoding = WSMessageEncoding.Text;
+			Binding.TextEncoding = Encoding.UTF8;
+			Binding.TransferMode = TransferMode.Buffered;
+			Binding.UseDefaultWebProxy = true;
+			Binding.ReaderQuotas.MaxDepth = 32;
+			Binding.ReaderQuotas.MaxStringContentLength = 2147483647;
+			Binding.ReaderQuotas.MaxArrayLength = 65536;
+			Binding.ReaderQuotas.MaxBytesPerRead = 4096;
+			Binding.ReaderQuotas.MaxNameTableCharCount = 65536;
+
+			var Endpoint = new EndpointAddress(new Uri(String.Format("http://{0}/api/XLocApiService.svc", Config.Server)));
+
+			return new XLocApiClient(Binding, Endpoint);
+		}
 
-	protected TransferServiceClient CreateTransferServiceClient()
-	{
-		var Binding = new BasicHttpBinding();
-		Binding.Name = "transfer";
-		Binding.MaxReceivedMessageSize = 67108864;
-		Binding.TransferMode = TransferMode.Streamed;
+		protected TransferServiceClient CreateTransferServiceClient()
+		{
+			var Binding = new BasicHttpBinding();
+			Binding.Name = "transfer";
+			Binding.MaxReceivedMessageSize = 67108864;
+			Binding.TransferMode = TransferMode.Streamed;
 
-		var Endpoint = new EndpointAddress(new Uri(String.Format("http://{0}/api/XLocTransferService.svc", Config.Server)));
+			var Endpoint = new EndpointAddress(new Uri(String.Format("http://{0}/api/XLocTransferService.svc", Config.Server)));
 
-		return new TransferServiceClient(Binding, Endpoint);
-	}
-
-	protected string RequestAuthToken(XLocApiClient XLocApiClient)
-	{
-		return XLocApiClient.GetAuthToken(Config.APIKey, XLocUtils.MD5HashString(Config.APIKey + Config.APISecret));
-	}
+			return new TransferServiceClient(Binding, Endpoint);
+		}
 
-	protected string RequestAuthTokenWithRetry(XLocApiClient XLocApiClient)
-	{
-		const int MAX_COUNT = 3;
+		protected string RequestAuthToken(XLocApiClient XLocApiClient)
+		{
+			return XLocApiClient.GetAuthToken(Config.APIKey, XLocUtils.MD5HashString(Config.APIKey + Config.APISecret));
+		}
 
-		int Count = 0;
-		for (;;)
+		protected string RequestAuthTokenWithRetry(XLocApiClient XLocApiClient)
 		{
-			try
-			{
-				return RequestAuthToken(XLocApiClient);
-			}
-			catch (Exception Ex)
+			const int MAX_COUNT = 3;
+
+			int Count = 0;
+			for (; ; )
 			{
-				if (++Count < MAX_COUNT)
+				try
 				{
-					BuildCommand.LogWarning("RequestAuthToken attempt {0}/{1} failed. Retrying...", Count, MAX_COUNT);
-					continue;
+					return RequestAuthToken(XLocApiClient);
 				}
+				catch (Exception Ex)
+				{
+					if (++Count < MAX_COUNT)
+					{
+						BuildCommand.LogWarning("RequestAuthToken attempt {0}/{1} failed. Retrying...", Count, MAX_COUNT);
+						continue;
+					}
 
-				BuildCommand.LogWarning("RequestAuthToken attempt {0}/{1} failed.", Count, MAX_COUNT);
-				throw Ex;
+					BuildCommand.LogWarning("RequestAuthToken attempt {0}/{1} failed.", Count, MAX_COUNT);
+					throw Ex;
+				}
 			}
 		}
-	}
-
-	protected string RequestLatestBuild(XLocApiClient XLocApiClient, string AuthToken, string LanguageId, string RemoteFilename)
-	{
-		return XLocApiClient.GetLatestBuildByFile(Config.APIKey, AuthToken, XLocUtils.MD5HashString(Config.APIKey + Config.APISecret + Config.LocalizationId + LanguageId + RemoteFilename), Config.LocalizationId, LanguageId, RemoteFilename);
-	}
 
-	private string GetXLocFilename(string BaseFilename, string Platform)
-	{
-		var XLocFilename = BaseFilename;
-		if (!String.IsNullOrEmpty(Platform))
-		{
-			// Apply the platform suffix. XLoc will take care of merging the files from different platforms together.
-			var XLocFilenameWithSuffix = Path.GetFileNameWithoutExtension(XLocFilename) + "_" + Platform + Path.GetExtension(XLocFilename);
-			XLocFilename = XLocFilenameWithSuffix;
-		}
-		if (!String.IsNullOrEmpty(LocalizationBranchName))
+		protected string RequestLatestBuild(XLocApiClient XLocApiClient, string AuthToken, string LanguageId, string RemoteFilename)
 		{
-			// Apply the branch suffix. XLoc will take care of merging the files from different branches together.
-			var XLocFilenameWithSuffix = Path.GetFileNameWithoutExtension(XLocFilename) + "_" + LocalizationBranchName + Path.GetExtension(XLocFilename);
-			XLocFilename = XLocFilenameWithSuffix;
+			return XLocApiClient.GetLatestBuildByFile(Config.APIKey, AuthToken, XLocUtils.MD5HashString(Config.APIKey + Config.APISecret + Config.LocalizationId + LanguageId + RemoteFilename), Config.LocalizationId, LanguageId, RemoteFilename);
 		}
-		if (!String.IsNullOrEmpty(RemoteFilenamePrefix))
+
+		private string GetXLocFilename(string BaseFilename, string Platform)
 		{
-			// Apply the prefix (this is used to avoid collisions with plugins that use the same name for their PO files)
-			XLocFilename = RemoteFilenamePrefix + "_" + XLocFilename;
+			var XLocFilename = BaseFilename;
+			if (!String.IsNullOrEmpty(Platform))
+			{
+				// Apply the platform suffix. XLoc will take care of merging the files from different platforms together.
+				var XLocFilenameWithSuffix = Path.GetFileNameWithoutExtension(XLocFilename) + "_" + Platform + Path.GetExtension(XLocFilename);
+				XLocFilename = XLocFilenameWithSuffix;
+			}
+			if (!String.IsNullOrEmpty(LocalizationBranchName))
+			{
+				// Apply the branch suffix. XLoc will take care of merging the files from different branches together.
+				var XLocFilenameWithSuffix = Path.GetFileNameWithoutExtension(XLocFilename) + "_" + LocalizationBranchName + Path.GetExtension(XLocFilename);
+				XLocFilename = XLocFilenameWithSuffix;
+			}
+			if (!String.IsNullOrEmpty(RemoteFilenamePrefix))
+			{
+				// Apply the prefix (this is used to avoid collisions with plugins that use the same name for their PO files)
+				XLocFilename = RemoteFilenamePrefix + "_" + XLocFilename;
+			}
+			return XLocFilename;
 		}
-		return XLocFilename;
-	}
 
-	protected XLocConfig Config;
+		protected XLocConfig Config;
 
-	virtual protected Dictionary<string, string> GetEpicCultureToXLocLanguageId()
-	{
-		return new Dictionary<string, string>
+		virtual protected Dictionary<string, string> GetEpicCultureToXLocLanguageId()
 		{
-			{ "en", "E" },		// English
-			{ "fr", "F" },		// French
-			{ "it", "I" },		// Italian
-			{ "de", "G" },		// German
-			{ "es", "S" },		// Spanish
-			{ "es-419", "Y" },	// LatAm Spanish
-			{ "ru", "R" },		// Russian
-			{ "pl", "P" },		// Polish
-			{ "ar", "A" },		// Arabic
-			{ "ko", "K" },		// Korean
-			{ "ja", "J" },		// Japanese
-			{ "zh-Hans", "2" },	// Simp Chinese
-			{ "zh-Hant", "1" },	// Trad Chinese
-			{ "tr", "3" },		// Turkish
-			{ "pt-BR", "$" },	// Brazilian Portuguese
-		};
+			return new Dictionary<string, string>
+			{
+				{ "en", "E" },		// English
+				{ "fr", "F" },		// French
+				{ "it", "I" },		// Italian
+				{ "de", "G" },		// German
+				{ "es", "S" },		// Spanish
+				{ "es-419", "Y" },	// LatAm Spanish
+				{ "ru", "R" },		// Russian
+				{ "pl", "P" },		// Polish
+				{ "ar", "A" },		// Arabic
+				{ "ko", "K" },		// Korean
+				{ "ja", "J" },		// Japanese
+				{ "zh-Hans", "2" },	// Simp Chinese
+				{ "zh-Hant", "1" },	// Trad Chinese
+				{ "tr", "3" },		// Turkish
+				{ "pt-BR", "$" },	// Brazilian Portuguese
+			};
+		}
 	}
 }
-
-}
