Gauntlet iOS device support for running parallel tests

diff --git a/Engine/Source/Programs/AutomationTool/Gauntlet/Platform/IOS/Gauntlet.TargetDeviceIOS.cs b/Engine/Source/Programs/AutomationTool/Gauntlet/Platform/IOS/Gauntlet.TargetDeviceIOS.cs
--- a/Engine/Source/Programs/AutomationTool/Gauntlet/Platform/IOS/Gauntlet.TargetDeviceIOS.cs
+++ b/Engine/Source/Programs/AutomationTool/Gauntlet/Platform/IOS/Gauntlet.TargetDeviceIOS.cs
@@ -186,13 +186,15 @@
 
 		public TargetDeviceIOS(string InName)
 		{
+			KillZombies();
+
+			var DefaultDevices = GetConnectedDeviceUUID();			
+
 			IsDefaultDevice = (String.IsNullOrEmpty(InName) || InName.Equals("default", StringComparison.OrdinalIgnoreCase));
 
 			Name = InName;
 			LocalDirectoryMappings = new Dictionary<EIntendedBaseCopyDirectory, string>();
 
-			var DefaultDevices = GetConnectedDeviceUUID();			
-
 			// If no device name or its 'default' then use the first default device
 			if (IsDefaultDevice)
 			{				
@@ -422,11 +424,11 @@
 			if (!File.Exists(LocalExecutable))
 			{
 				throw new AutomationException("Local executable not found for -dev argument: {0}", LocalExecutable);
-			}
-
-			File.WriteAllText(CacheResignedFilename, "The application has been resigned");
-
-			// copy local executable
+			}
+
+			File.WriteAllText(CacheResignedFilename, "The application has been resigned");
+
+			// copy local executable
 			FileInfo SrcInfo = new FileInfo(LocalExecutable);			
 			string DestPath = Path.Combine(CachedAppPath, BundleName);
 			SrcInfo.CopyTo(DestPath, true);
@@ -468,6 +470,9 @@
 
 		}
 
+		// We need to lock around setting up the IPA
+		static object IPALock = new object();
+
 		public IAppInstall InstallApplication(UnrealAppConfig AppConfig)
 		{            
             IOSBuild Build = AppConfig.Build as IOSBuild;
@@ -476,33 +481,38 @@
 			if (Build == null)
 			{
 				throw new AutomationException("Invalid build for IOS!");
-			}
+			}	
 
-			Log.Info("Installing using IPA {0}", Build.SourceIPAPath);
-
-			// device artifact path
-			DeviceArtifactPath = string.Format("/Documents/{0}/Saved", AppConfig.ProjectName);
-
-			bool CacheResigned = File.Exists(CacheResignedFilename);
+			bool CacheResigned = false;
 			bool UseLocalExecutable = Globals.Params.ParseParam("dev");
 
-			if (CacheResigned && !UseLocalExecutable)
+			lock(IPALock)
 			{
-				if (File.Exists(IPAHashFilename))
+				Log.Info("Installing using IPA {0}", Build.SourceIPAPath);
+
+				// device artifact path
+				DeviceArtifactPath = string.Format("/Documents/{0}/Saved", AppConfig.ProjectName);
+
+				CacheResigned = File.Exists(CacheResignedFilename);				
+
+				if (CacheResigned && !UseLocalExecutable)
 				{
-					Log.Verbose("App was resigned, invalidating app cache");
-					File.Delete(IPAHashFilename);
-				}								
+					if (File.Exists(IPAHashFilename))
+					{
+						Log.Verbose("App was resigned, invalidating app cache");
+						File.Delete(IPAHashFilename);
+					}								
+				}
+
+				PrepareIPA(Build);
+
+				// local executable support			
+				if (UseLocalExecutable)
+				{					
+					ResignApplication(AppConfig);				
+				}
 			}
 
-			PrepareIPA(Build);
-
-			// local executable support			
-			if (UseLocalExecutable)
-			{					
-				ResignApplication(AppConfig);				
-			}
-
 			if (CacheResigned || UseLocalExecutable || !CheckDeployedIPA(Build))
 			{
 				// uninstall will clean all device artifacts
@@ -617,8 +627,7 @@
 		public bool PowerOff() { return true; }
 		public bool Reboot() { return true; }
 
-		// catch attempts to connect multiple iOS devices in parallel, see https://jira.it.epicgames.net/browse/UEATM-219
-		static int ConnectedDeviceCount = 0;
+		static Dictionary<string, bool> ConnectedDevices = new Dictionary<string, bool>();
 		bool Connected = false;
 
 		public bool Connect() 
@@ -630,22 +639,18 @@
 					return true;								
 				}
 
-				ConnectedDeviceCount++;
-
-				if (ConnectedDeviceCount > 1)
+				bool ExistingConnection = false;
+				if (ConnectedDevices.TryGetValue(DeviceName, out ExistingConnection))
 				{
-
-					throw new AutomationException("Parallel iOS device connections are not currently supported, see https://jira.it.epicgames.net/browse/UEATM-219");
+					if (ExistingConnection)
+					{
+						throw new AutomationException("Connected to already connected device");
+					}
 				}
 
+				ConnectedDevices[DeviceName] = true;
+
 				Connected = true;
-
-				// Get rid of any zombies, this needs to be reworked to use tracked process id's when running parallel tests
-				// may need to kill by processid for spawned children
-				IOSBuild.ExecuteCommand("killall", "ios-deploy");
-				Thread.Sleep(2500);
-				IOSBuild.ExecuteCommand("killall", "lldb");			
-				Thread.Sleep(2500);
 			}
 
 			return true; 
@@ -661,12 +666,11 @@
 
 				Connected = false;
 
-				ConnectedDeviceCount--;
+				if (ConnectedDevices.ContainsKey(DeviceName))
+				{
+					ConnectedDevices.Remove(DeviceName);
+				}
 
-				if (ConnectedDeviceCount < 0)
-				{
-					throw new AutomationException("iOS device connection mismatch");
-				}
 			}
 
 			return true; 
@@ -685,7 +689,7 @@
 		/// </summary>
 		List<string> GetConnectedDeviceUUID()
 		{
-			var Result = ExecuteIOSDeployCommand("--detect");
+			var Result = ExecuteIOSDeployCommand("--detect", 60, true, false);
 
 			if (Result.ExitCode != 0)
 			{
@@ -698,6 +702,25 @@
 				M => M.Groups[2].ToString()
 			).ToList();
 		}
+
+		static bool ZombiesKilled = false;
+
+		// Get rid of any zombie lldb/iosdeploy processes, this needs to be reworked to use tracked process id's when running parallel tests across multiple AutomationTool.exe processes on test workers
+		void KillZombies()
+		{
+
+			if (ZombiesKilled)
+			{
+				return;
+			}
+
+			ZombiesKilled = true;
+
+			IOSBuild.ExecuteCommand("killall", "ios-deploy");
+			Thread.Sleep(2500);
+			IOSBuild.ExecuteCommand("killall", "lldb");			
+			Thread.Sleep(2500);
+		}
 		
 		// Gauntlet cache folder for tracking device/ipa state
 		string GauntletAppCache { get { return Path.Combine(Environment.GetFolderPath(Environment.SpecialFolder.Personal), ".gauntletappcache");} }
@@ -821,9 +844,9 @@
 			return LocalDirectoryMappings;
 		}
 
-		public IProcessResult ExecuteIOSDeployCommand(String CommandLine, int WaitTime = 60, bool WarnOnTimeout = true)
+		public IProcessResult ExecuteIOSDeployCommand(String CommandLine, int WaitTime = 60, bool WarnOnTimeout = true, bool UseDeviceID = true)
 		{
-			if (!IsDefaultDevice)
+			if (UseDeviceID && !IsDefaultDevice)
 			{
 				CommandLine = String.Format("--id {0} {1}", DeviceName, CommandLine);
 			}
