Added support for using UE_SDK_ROOT (carefullyredist) in build scripts.

diff --git a/Engine/Source/Programs/AutomationTool/Android/AndroidPlatform.Automation.cs b/Engine/Source/Programs/AutomationTool/Android/AndroidPlatform.Automation.cs
--- a/Engine/Source/Programs/AutomationTool/Android/AndroidPlatform.Automation.cs
+++ b/Engine/Source/Programs/AutomationTool/Android/AndroidPlatform.Automation.cs
@@ -555,7 +555,8 @@ public class AndroidPlatform : Platform
             BatchLines = new string[] {
 						"setlocal",
                         "set ANDROIDHOME=%ANDROID_HOME%",
-                        "if \"%ANDROIDHOME%\"==\"\" set ANDROIDHOME="+Environment.GetEnvironmentVariable("ANDROID_HOME"),
+						"if NOT \"UE_SDKS_ROOT\"==\"\" (call %UE_SDKS_ROOT%\\HostWin64\\Android\\SetupEnvironmentVars.bat)",
+						"if \"%ANDROIDHOME%\"==\"\" set ANDROIDHOME="+Environment.GetEnvironmentVariable("ANDROID_HOME"),
 						"set ADB=%ANDROIDHOME%\\platform-tools\\adb.exe",
 						"set DEVICE=",
                         "if not \"%1\"==\"\" set DEVICE=-s %1",
