Small fix for Gauntlet HtmlBuilder.

#rnx

diff --git a/Engine/Source/Programs/AutomationTool/Gauntlet/Framework/Utils/Gauntlet.HtmlBuilder.cs b/Engine/Source/Programs/AutomationTool/Gauntlet/Framework/Utils/Gauntlet.HtmlBuilder.cs
--- a/Engine/Source/Programs/AutomationTool/Gauntlet/Framework/Utils/Gauntlet.HtmlBuilder.cs
+++ b/Engine/Source/Programs/AutomationTool/Gauntlet/Framework/Utils/Gauntlet.HtmlBuilder.cs
@@ -120,7 +120,7 @@ namespace Gauntlet
 		/// <returns></returns>
 		public HtmlBuilder EnsureEndsWithRowClose()
 		{
-			if (SB.Length > 0 && !EndsWithNewLine)
+			if (SB.Length > 0 && !EndsWithRowClose)
 			{
 				SB.Append("</tr>");
 			}
