Merging //UE4/Dev-Main@4624481 to Dev-Rendering (//UE4/Dev-Rendering)

diff --git a/Engine/Config/BaseEngine.ini b/Engine/Config/BaseEngine.ini
--- a/Engine/Config/BaseEngine.ini
+++ b/Engine/Config/BaseEngine.ini
@@ -765,7 +765,7 @@ gc.NumRetriesBeforeForcingGC=10
 gc.AllowParallelGC=True
 ; pick a fractional number to keep phase shifting and avoid collisions
 gc.TimeBetweenPurgingPendingKillObjects=61.1
-gc.MaxObjectsInEditor=8388607
+gc.MaxObjectsInEditor=16777216
 gc.IncrementalBeginDestroyEnabled=True
 gc.CreateGCClusters=True
 gc.MergeGCClusters=False
diff --git a/Engine/Plugins/Experimental/SkeletalReduction/Source/Private/SkeletalSimplifierLinearAlgebra.h b/Engine/Plugins/Experimental/SkeletalReduction/Source/Private/SkeletalSimplifierLinearAlgebra.h
--- a/Engine/Plugins/Experimental/SkeletalReduction/Source/Private/SkeletalSimplifierLinearAlgebra.h
+++ b/Engine/Plugins/Experimental/SkeletalReduction/Source/Private/SkeletalSimplifierLinearAlgebra.h
@@ -248,12 +248,10 @@ namespace SkeletalSimplifier
 			typedef ArrayBase<6, SymmetricMatrix>            MyBase;
 			typedef typename MyBase::ScalarType              ScalarType;
 
-			using MyBase::Data;
-
 			SymmetricMatrix() 
 			{
 				// Zero
-				this->Reset();
+				Reset();
 			}
 
 			/**
@@ -966,6 +964,9 @@ namespace SkeletalSimplifier
 			typedef ArrayBase<SIZE, TDenseVecD>           MyBase;
 			typedef typename MyBase::ScalarType           ScalarType;
 
+			using MyBase::Data;
+			using MyBase::Reset;
+
 			enum { Size = SIZE };
 
 			TDenseVecD()
diff --git a/Engine/Source/Editor/UnrealEd/Classes/CookOnTheSide/CookOnTheFlyServer.h b/Engine/Source/Editor/UnrealEd/Classes/CookOnTheSide/CookOnTheFlyServer.h
--- a/Engine/Source/Editor/UnrealEd/Classes/CookOnTheSide/CookOnTheFlyServer.h
+++ b/Engine/Source/Editor/UnrealEd/Classes/CookOnTheSide/CookOnTheFlyServer.h
@@ -135,6 +135,8 @@ private:
 	int32 MaxNumPackagesBeforePartialGC;
 	/** Max number of concurrent shader jobs reducing this too low will increase cook time */
 	int32 MaxConcurrentShaderJobs;
+	/** Min number of free UObject indices before the cooker should partial gc */
+	int32 MinFreeUObjectIndicesBeforeGC;
 
 	ECookInitializationFlags CookFlags = ECookInitializationFlags::None;
 	TUniquePtr<class FSandboxPlatformFile> SandboxFile;
diff --git a/Engine/Source/Editor/UnrealEd/Private/CookOnTheFlyServer.cpp b/Engine/Source/Editor/UnrealEd/Private/CookOnTheFlyServer.cpp
--- a/Engine/Source/Editor/UnrealEd/Private/CookOnTheFlyServer.cpp
+++ b/Engine/Source/Editor/UnrealEd/Private/CookOnTheFlyServer.cpp
@@ -3205,7 +3205,13 @@ bool UCookOnTheFlyServer::HasExceededMaxMemory() const
 		return true;
 	}
 
-	
+#if UE_GC_TRACK_OBJ_AVAILABLE
+	if (GUObjectArray.GetObjectArrayEstimatedAvailable() < MinFreeUObjectIndicesBeforeGC)
+	{
+		UE_LOG(LogCook, Display, TEXT("Running out of available UObject indices (%d remaining)"), GUObjectArray.GetObjectArrayEstimatedAvailable());
+		return true;
+	}
+#endif // UE_GC_TRACK_OBJ_AVAILABLE
 
 	return false;
 }
@@ -4084,6 +4090,10 @@ void UCookOnTheFlyServer::Initialize( ECookMode::Type DesiredCookMode, ECookInit
 	MinMemoryBeforeGC = MinMemoryBeforeGCInMB * 1024LL * 1024LL;
 	MinMemoryBeforeGC = FMath::Min(MaxMemoryAllowance, MinMemoryBeforeGC);
 
+	MinFreeUObjectIndicesBeforeGC = 5000;
+	GConfig->GetInt(TEXT("CookSettings"), TEXT("MinFreeUObjectIndicesBeforeGC"), MinFreeUObjectIndicesBeforeGC, GEditorIni);
+	MinFreeUObjectIndicesBeforeGC = FMath::Max(MinFreeUObjectIndicesBeforeGC, 0);
+
 	int32 MinFreeMemoryInMB = 0;
 	GConfig->GetInt(TEXT("CookSettings"), TEXT("MinFreeMemory"), MinFreeMemoryInMB, GEditorIni);
 	MinFreeMemoryInMB = FMath::Max(MinFreeMemoryInMB, 0);
diff --git a/Engine/Source/Programs/AutomationTool/Scripts/RebuildLightMapsCommand.Automation.cs b/Engine/Source/Programs/AutomationTool/Scripts/RebuildLightMapsCommand.Automation.cs
--- a/Engine/Source/Programs/AutomationTool/Scripts/RebuildLightMapsCommand.Automation.cs
+++ b/Engine/Source/Programs/AutomationTool/Scripts/RebuildLightMapsCommand.Automation.cs
@@ -260,43 +260,46 @@ namespace AutomationScripts.Automation
 		 */
 		void SendCompletionMessage(bool bWasSuccessful, String MessageBody)
 		{
-			MailMessage Message = new System.Net.Mail.MailMessage();
-			Message.Priority = MailPriority.High;
-			Message.From = new MailAddress("[REDACTED_EMAIL]");
+			if (StakeholdersEmailAddresses != null)
+			{
+				MailMessage Message = new System.Net.Mail.MailMessage();
+				Message.Priority = MailPriority.High;
+				Message.From = new MailAddress("[REDACTED_EMAIL]");
 
-            string Branch = "Unknown";
-            if ( P4Enabled )
-            {
-                Branch = P4Env.Branch;
-            }
+				string Branch = "Unknown";
+				if (P4Enabled)
+				{
+					Branch = P4Env.Branch;
+				}
 
-			foreach (String NextStakeHolder in StakeholdersEmailAddresses)
-			{
-				Message.To.Add(new MailAddress(NextStakeHolder));
-			}
+				foreach (String NextStakeHolder in StakeholdersEmailAddresses)
+				{
+					Message.To.Add(new MailAddress(NextStakeHolder));
+				}
 
-			ConfigHierarchy Ini = ConfigCache.ReadHierarchy(ConfigHierarchyType.EditorPerProjectUserSettings, DirectoryReference.FromFile(ProjectFullPath), UnrealTargetPlatform.Win64);
-			List<string> Emails;
-			Ini.GetArray("RebuildlightingSettings", "EmailNotification", out Emails);
-			foreach (var Email in Emails)
-			{
-				Message.CC.Add(new MailAddress(Email));
-			}
+				ConfigHierarchy Ini = ConfigCache.ReadHierarchy(ConfigHierarchyType.EditorPerProjectUserSettings, DirectoryReference.FromFile(ProjectFullPath), UnrealTargetPlatform.Win64);
+				List<string> Emails;
+				Ini.GetArray("RebuildlightingSettings", "EmailNotification", out Emails);
+				foreach (var Email in Emails)
+				{
+					Message.CC.Add(new MailAddress(Email));
+				}
 
-			Message.Subject = String.Format("Nightly lightmap rebuild {0} for {1}", bWasSuccessful ? "[SUCCESS]" : "[FAILED]", Branch);
-			Message.Body = MessageBody;
-            /*Attachment Attach = new Attachment();
-            Message.Attachments.Add()*/
-			try
-			{
-				#pragma warning disable CS0618 // Mono 4.6.x obsoletes this class
-				SmtpClient MailClient = new SmtpClient("smtp.epicgames.net");
-				MailClient.Send(Message);
-				#pragma warning restore CS0618
-			}
-			catch (Exception Ex)
-			{
-				LogError("Failed to send notify email to {0} ({1})", String.Join(", ", StakeholdersEmailAddresses.ToArray()), Ex.Message);
+				Message.Subject = String.Format("Nightly lightmap rebuild {0} for {1}", bWasSuccessful ? "[SUCCESS]" : "[FAILED]", Branch);
+				Message.Body = MessageBody;
+				/*Attachment Attach = new Attachment();
+				Message.Attachments.Add()*/
+				try
+				{
+#pragma warning disable CS0618 // Mono 4.6.x obsoletes this class
+					SmtpClient MailClient = new SmtpClient("smtp.epicgames.net");
+					MailClient.Send(Message);
+#pragma warning restore CS0618
+				}
+				catch (Exception Ex)
+				{
+					LogError("Failed to send notify email to {0} ({1})", String.Join(", ", StakeholdersEmailAddresses.ToArray()), Ex.Message);
+				}
 			}
 		}
 
diff --git a/Engine/Source/Programs/AutomationTool/Scripts/ResavePackagesCommand.Automation.cs b/Engine/Source/Programs/AutomationTool/Scripts/ResavePackagesCommand.Automation.cs
--- a/Engine/Source/Programs/AutomationTool/Scripts/ResavePackagesCommand.Automation.cs
+++ b/Engine/Source/Programs/AutomationTool/Scripts/ResavePackagesCommand.Automation.cs
@@ -263,38 +263,38 @@ namespace AutomationScripts.Automation
 		 */
 		void SendCompletionMessage(bool bWasSuccessful, String MessageBody)
 		{
-			MailMessage Message = new System.Net.Mail.MailMessage();
-			Message.Priority = MailPriority.High;
-			Message.From = new MailAddress("[REDACTED_EMAIL]");
+			if (StakeholdersEmailAddresses != null)
+			{
+				MailMessage Message = new System.Net.Mail.MailMessage();
+				Message.Priority = MailPriority.High;
+				Message.From = new MailAddress("[REDACTED_EMAIL]");
 
-            string Branch = "Unknown";
-            if ( P4Enabled )
-            {
-                Branch = P4Env.Branch;
-            }
+				string Branch = "Unknown";
+				if (P4Enabled)
+				{
+					Branch = P4Env.Branch;
+				}
 
-			foreach (String NextStakeHolder in StakeholdersEmailAddresses)
-			{
-				Message.To.Add(new MailAddress(NextStakeHolder));
-			}
+				foreach (String NextStakeHolder in StakeholdersEmailAddresses)
+				{
+					Message.To.Add(new MailAddress(NextStakeHolder));
+				}
 
-			Message.CC.Add(new MailAddress("[REDACTED_EMAIL]"));
-            Message.CC.Add(new MailAddress("[REDACTED_EMAIL]"));
-            Message.CC.Add(new MailAddress("[REDACTED_EMAIL]"));
-			Message.Subject = String.Format("Nightly lightmap rebuild {0} for {1}", bWasSuccessful ? "[SUCCESS]" : "[FAILED]", Branch);
-			Message.Body = MessageBody;
-            /*Attachment Attach = new Attachment();
-            Message.Attachments.Add()*/
-			try
-			{
-				#pragma warning disable CS0618 // Mono 4.6.x obsoletes this class
-				SmtpClient MailClient = new SmtpClient("smtp.epicgames.net");
-				MailClient.Send(Message);
-				#pragma warning restore CS0618
-			}
-			catch (Exception Ex)
-			{
-				LogError("Failed to send notify email to {0} ({1})", String.Join(", ", StakeholdersEmailAddresses.ToArray()), Ex.Message);
+				Message.Subject = String.Format("Nightly lightmap rebuild {0} for {1}", bWasSuccessful ? "[SUCCESS]" : "[FAILED]", Branch);
+				Message.Body = MessageBody;
+				/*Attachment Attach = new Attachment();
+				Message.Attachments.Add()*/
+				try
+				{
+#pragma warning disable CS0618 // Mono 4.6.x obsoletes this class
+					SmtpClient MailClient = new SmtpClient("smtp.epicgames.net");
+					MailClient.Send(Message);
+#pragma warning restore CS0618
+				}
+				catch (Exception Ex)
+				{
+					LogError("Failed to send notify email to {0} ({1})", String.Join(", ", StakeholdersEmailAddresses.ToArray()), Ex.Message);
+				}
 			}
 		}
 
diff --git a/Engine/Source/Runtime/CoreUObject/Public/UObject/UObjectArray.h b/Engine/Source/Runtime/CoreUObject/Public/UObject/UObjectArray.h
--- a/Engine/Source/Runtime/CoreUObject/Public/UObject/UObjectArray.h
+++ b/Engine/Source/Runtime/CoreUObject/Public/UObject/UObjectArray.h
@@ -272,6 +272,17 @@ public:
 	{
 		return NumElements;
 	}
+
+	/**
+	* Return the number max capacity of the array
+	* Thread safe, but you know, someone might have added more elements before this even returns
+	* @return	the maximum number of elements in the array
+	**/
+	FORCEINLINE int32 Capacity() const TSAN_SAFE
+	{
+		return MaxElements;
+	}
+
 	/**
 	* Return if this index is valid
 	* Thread safe, if it is valid now, it is valid forever. Other threads might be adding during this call.
@@ -426,6 +437,17 @@ public:
 	{
 		return NumElements;
 	}
+
+	/**
+	* Return the number max capacity of the array
+	* Thread safe, but you know, someone might have added more elements before this even returns
+	* @return	the maximum number of elements in the array
+	**/
+	FORCEINLINE int32 Capacity() const TSAN_SAFE
+	{
+		return MaxElements;
+	}
+
 	/**
 	* Return if this index is valid
 	* Thread safe, if it is valid now, it is valid forever. Other threads might be adding during this call.
@@ -802,10 +824,18 @@ public:
 	 *
 	 * @return	The number of objects claimed
 	 */
-	FORCEINLINE int32 GetObjectArrayNumMinusAvailable()
+	int32 GetObjectArrayNumMinusAvailable()
 	{
 		return ObjObjects.Num() - ObjAvailableCount.GetValue();
 	}
+
+	/**
+	* Returns the estimated number of object indices available for allocation
+	*/
+	int32 GetObjectArrayEstimatedAvailable()
+	{
+		return ObjObjects.Capacity() - GetObjectArrayNumMinusAvailable();
+	}
 #endif
 
 	/**
