Allow Localization Automation to be be submitted to Perforce

diff --git a/Engine/Source/Programs/AutomationTool/Scripts/Localisation.Automation.cs b/Engine/Source/Programs/AutomationTool/Scripts/Localisation.Automation.cs
--- a/Engine/Source/Programs/AutomationTool/Scripts/Localisation.Automation.cs
+++ b/Engine/Source/Programs/AutomationTool/Scripts/Localisation.Automation.cs
@@ -176,7 +176,13 @@ class Localize : BuildCommand
 		Dictionary<string, byte[]> InitalPOFileHashes = null;
 		if (P4Enabled)
 		{
-			PendingChangeList = P4.CreateChange(P4Env.Client, "Localization Automation");
+			var ChangeListCommitMessage = "Localization Automation";
+			if (File.Exists(CombinePaths(CmdEnv.LocalRoot, @"Engine/Build/NotForLicensees/EpicInternal.txt")))
+			{
+				ChangeListCommitMessage += "\n#okforgithub ignore";
+			}
+
+			PendingChangeList = P4.CreateChange(P4Env.Client, ChangeListCommitMessage);
 			InitalPOFileHashes = GetPOFileHashes(LocalizationBatches, UEProjectRoot);
 		}
 
