Made Gauntlet's build path source support  multiple folders to search. This is to better support packaged builds which end up in Binaries/<platform> so they can be found at the same time as staged builds without moving them under /saved

-build=local is now a (clearer) synonym for -build=staged. Both will search /Saved/StagedBuilds and /Binaries for matching builds

#tests BuildCookFort and RunFort with Android

diff --git a/Engine/Source/Programs/AutomationTool/Gauntlet/Platform/Android/Gauntlet.TargetDeviceAndroid.cs b/Engine/Source/Programs/AutomationTool/Gauntlet/Platform/Android/Gauntlet.TargetDeviceAndroid.cs
--- a/Engine/Source/Programs/AutomationTool/Gauntlet/Platform/Android/Gauntlet.TargetDeviceAndroid.cs
+++ b/Engine/Source/Programs/AutomationTool/Gauntlet/Platform/Android/Gauntlet.TargetDeviceAndroid.cs
@@ -715,7 +715,7 @@ namespace Gauntlet
 				}
 				else
 				{
-					Log.Info("Installing {0} to {1} via adb push", SourcePath, DestPath);
+					Log.Info("Copying {0} to {1} via adb push", QuotedSourcePath, DestPath);
 					string AdbCommand = string.Format("push {0} {1}", QuotedSourcePath, DestPath);
 					AdbResult = RunAdbDeviceCommand(AdbCommand);
 
@@ -1035,7 +1035,7 @@ namespace Gauntlet
 		/// <returns></returns>
 		public static IProcessResult RunAdbGlobalCommand(string Args, bool Wait = true, bool bShouldLogCommand = false)
 		{
-			CommandUtils.ERunOptions RunOptions = CommandUtils.ERunOptions.AppMustExist;
+			CommandUtils.ERunOptions RunOptions = CommandUtils.ERunOptions.AppMustExist | CommandUtils.ERunOptions.NoWaitForExit;
 
 			if (Log.IsVeryVerbose)
 			{
@@ -1045,19 +1045,20 @@ namespace Gauntlet
 			{
 				RunOptions |= CommandUtils.ERunOptions.NoLoggingOfRunCommand;
 			}
-
-			if (Wait == false)
-			{
-				RunOptions |= CommandUtils.ERunOptions.NoWaitForExit;
-			}
-
+	
 			if (bShouldLogCommand)
 			{
 				Log.Verbose("Running ADB Command: adb {0}", Args);
 			}
 			
 			IProcessResult Process = AndroidPlatform.RunAdbCommand(null, null, Args, null, RunOptions);
-			return Process;
+
+			if (Wait)
+			{
+				Process.WaitForExit();
+			}
+			
+return Process;
 		}
 
 		public void AllowDeviceSleepState(bool bAllowSleep)
diff --git a/Engine/Source/Programs/AutomationTool/Gauntlet/Unreal/Base/Gauntlet.UnrealTestNode.cs b/Engine/Source/Programs/AutomationTool/Gauntlet/Unreal/Base/Gauntlet.UnrealTestNode.cs
--- a/Engine/Source/Programs/AutomationTool/Gauntlet/Unreal/Base/Gauntlet.UnrealTestNode.cs
+++ b/Engine/Source/Programs/AutomationTool/Gauntlet/Unreal/Base/Gauntlet.UnrealTestNode.cs
@@ -536,21 +536,23 @@ namespace Gauntlet
 
 			try
 			{
-                // Basic pre-existing directory check.
-                if (!Globals.Params.ParseParam("dev") && Directory.Exists(OutputPath))
-                {
-                    string NewOutputPath = OutputPath;
-                    int i = 0;
-                    while (Directory.Exists(NewOutputPath))
-                    {
-                        i++;
-                        NewOutputPath = string.Format("{0}_{1}", OutputPath, i);
-                    }
-                    Log.Info("Directory already exists at {0}", OutputPath);
-                    OutputPath = NewOutputPath;
-                }
+				// Basic pre-existing directory check.
+				if (CommandUtils.IsBuildMachine && Directory.Exists(OutputPath))
+				{
+					string NewOutputPath = OutputPath;
+					int i = 0;
+					while (Directory.Exists(NewOutputPath))
+					{
+						i++;
+						NewOutputPath = string.Format("{0}_{1}", OutputPath, i);
+					}
+					Log.Info("Directory already exists at {0}", OutputPath);
+					OutputPath = NewOutputPath;
+				}
 				Log.Info("Saving artifacts to {0}", OutputPath);
 				Directory.CreateDirectory(OutputPath);
+				Utils.SystemHelpers.MarkDirectoryForCleanup(OutputPath);
+
 				SessionArtifacts = SaveRoleArtifacts(OutputPath);
 
 				// call legacy version
diff --git a/Engine/Source/Programs/AutomationTool/Gauntlet/Unreal/BuildSource/Gauntlet.UnrealBuildSource.cs b/Engine/Source/Programs/AutomationTool/Gauntlet/Unreal/BuildSource/Gauntlet.UnrealBuildSource.cs
--- a/Engine/Source/Programs/AutomationTool/Gauntlet/Unreal/BuildSource/Gauntlet.UnrealBuildSource.cs
+++ b/Engine/Source/Programs/AutomationTool/Gauntlet/Unreal/BuildSource/Gauntlet.UnrealBuildSource.cs
@@ -20,7 +20,7 @@ namespace Gauntlet
 
 		public string BuildName { get; protected set; }
 
-		public string BuildPath { get; protected set; }
+		public IEnumerable<string> BuildPaths { get; protected set; }
 
 		public string Branch { get; protected set; }
 
@@ -32,24 +32,23 @@ namespace Gauntlet
 
 		public int BuildCount { get { return DiscoveredBuilds.Count; } }
 
-		public UnrealBuildSource(string InProjectName, bool InUsesSharedBuildType, string InUnrealPath, string StagedReference) 
+		public UnrealBuildSource(string InProjectName, bool InUsesSharedBuildType, string InUnrealPath, string BuildReference) 
 		{
-			InitBuildSource(InProjectName, InUsesSharedBuildType, InUnrealPath, StagedReference, null);
+			InitBuildSource(InProjectName, InUsesSharedBuildType, InUnrealPath, BuildReference, null);
 		}
 
-		public UnrealBuildSource(string InProjectName, bool InUsesSharedBuildType, string InUnrealPath, string StagedReference, Func<string, string> ResolutionDelegate)
+		public UnrealBuildSource(string InProjectName, bool InUsesSharedBuildType, string InUnrealPath, string BuildReference, Func<string, string> ResolutionDelegate)
 		{
-			InitBuildSource(InProjectName, InUsesSharedBuildType, InUnrealPath, StagedReference, ResolutionDelegate);
-
+			InitBuildSource(InProjectName, InUsesSharedBuildType, InUnrealPath, BuildReference, ResolutionDelegate);
 		}
 
-		public UnrealBuildSource(string InProjectName, bool InUsesSharedBuildType, string InUnrealPath, string StagedReference, IEnumerable<string> InSearchPaths)
+		public UnrealBuildSource(string InProjectName, bool InUsesSharedBuildType, string InUnrealPath, string BuildReference, IEnumerable<string> InSearchPaths)
 		{
-			InitBuildSource(InProjectName, InUsesSharedBuildType, InUnrealPath, StagedReference, (string StageRef) =>
+			InitBuildSource(InProjectName, InUsesSharedBuildType, InUnrealPath, BuildReference, (string BuildRef) =>
 			{
 				foreach (string SearchPath in InSearchPaths)
 				{
-					DirectoryInfo SearchDir = new DirectoryInfo(Path.Combine(SearchPath, StageRef));
+					DirectoryInfo SearchDir = new DirectoryInfo(Path.Combine(SearchPath, BuildRef));
 
 					if (SearchDir.Exists)
 					{
@@ -74,67 +73,23 @@ namespace Gauntlet
 
 		protected void InitBuildSource(string InProjectName, bool InUsesSharedBuildType, string InUnrealPath, string InBuildArgument, Func<string, string> ResolutionDelegate)
 		{
-			if (InBuildArgument.Equals("AutoP4", StringComparison.InvariantCultureIgnoreCase))
-			{
-				if (!CommandUtils.P4Enabled)
-				{
-					throw new AutomationException("-Build=AutoP4 requires -P4");
-				}
-				if (CommandUtils.P4Env.Changelist < 1000)
-				{
-					throw new AutomationException("-Build=AutoP4 requires a CL from P4 and we have {0}", CommandUtils.P4Env.Changelist);
-				}
-
-				string BuildRoot = CommandUtils.CombinePaths(CommandUtils.RootBuildStorageDirectory());
-				string CachePath = InternalUtils.GetEnvironmentVariable("UE-BuildCachePath", "");
-
-				string SrcBuildPath = CommandUtils.CombinePaths(BuildRoot, InProjectName);
-				string SrcBuildPath2 = CommandUtils.CombinePaths(BuildRoot, InProjectName.Replace("Game", "").Replace("game", ""));
-
-				string SrcBuildPath_Cache = CommandUtils.CombinePaths(CachePath, InProjectName);
-				string SrcBuildPath2_Cache = CommandUtils.CombinePaths(CachePath, InProjectName.Replace("Game", "").Replace("game", ""));
-
-				if (!InternalUtils.SafeDirectoryExists(SrcBuildPath))
-				{
-					if (!InternalUtils.SafeDirectoryExists(SrcBuildPath2))
-					{
-						throw new AutomationException("-Build=AutoP4: Neither {0} nor {1} exists.", SrcBuildPath, SrcBuildPath2);
-					}
-					SrcBuildPath = SrcBuildPath2;
-					SrcBuildPath_Cache = SrcBuildPath2_Cache;
-				}
-				string SrcCLPath = CommandUtils.CombinePaths(SrcBuildPath, CommandUtils.EscapePath(CommandUtils.P4Env.Branch) + "-CL-" + CommandUtils.P4Env.Changelist.ToString());
-				string SrcCLPath_Cache = CommandUtils.CombinePaths(SrcBuildPath_Cache, CommandUtils.EscapePath(CommandUtils.P4Env.Branch) + "-CL-" + CommandUtils.P4Env.Changelist.ToString());
-				if (!InternalUtils.SafeDirectoryExists(SrcCLPath))
-				{
-					throw new AutomationException("-Build=AutoP4: {0} does not exist.", SrcCLPath);
-				}
-
-				if (InternalUtils.SafeDirectoryExists(SrcCLPath_Cache))
-				{
-					InBuildArgument = SrcCLPath_Cache;
-				}
-				else
-				{
-					InBuildArgument = SrcCLPath;
-				}
-				Log.Verbose("Using AutoP4 path {0}", InBuildArgument);
-			}
-
 			UnrealPath = InUnrealPath;
 			UsesSharedBuildType = InUsesSharedBuildType;
 
 			ResolveProjectName(InProjectName);
 
 			// Resolve the build argument into something meaningful
-			string ResolvedBuildPath, ResolvedBuildName;
-			if (!ResolveBuildReference(InBuildArgument, ResolutionDelegate, out ResolvedBuildPath, out ResolvedBuildName))
+			string ResolvedBuildName;
+			IEnumerable<string> ResolvedPaths = null;
+
+			if (!ResolveBuildReference(InBuildArgument, ResolutionDelegate, out ResolvedPaths, out ResolvedBuildName))
 			{
 				throw new AutomationException("Unable to resolve {0} to a valid build", InBuildArgument);
 			}
 
 			BuildName = ResolvedBuildName;
-			BuildPath = ResolvedBuildPath;
+			BuildPaths = ResolvedPaths;
+
 			DiscoveredBuilds = DiscoverBuilds();
 
 			if (DiscoveredBuilds.Count() == 0)
@@ -185,16 +140,63 @@ namespace Gauntlet
 			}
 		}
 
-		virtual protected bool ResolveBuildReference(string InBuildReference, Func<string, string> ResolutionDelegate, out string OutBuildPath, out string OutBuildName)
+		virtual protected bool ResolveBuildReference(string InBuildReference, Func<string, string> ResolutionDelegate, out IEnumerable<string> OutBuildPaths, out string OutBuildName)
 		{
 			OutBuildName = null;
-			OutBuildPath = null;
+			OutBuildPaths = null;
 
 			if (string.IsNullOrEmpty(InBuildReference))
 			{
 				return false;
 			}
 
+			if (InBuildReference.Equals("AutoP4", StringComparison.InvariantCultureIgnoreCase))
+			{
+				if (!CommandUtils.P4Enabled)
+				{
+					throw new AutomationException("-Build=AutoP4 requires -P4");
+				}
+				if (CommandUtils.P4Env.Changelist < 1000)
+				{
+					throw new AutomationException("-Build=AutoP4 requires a CL from P4 and we have {0}", CommandUtils.P4Env.Changelist);
+				}
+
+				string BuildRoot = CommandUtils.CombinePaths(CommandUtils.RootBuildStorageDirectory());
+				string CachePath = InternalUtils.GetEnvironmentVariable("UE-BuildCachePath", "");
+
+				string SrcBuildPath = CommandUtils.CombinePaths(BuildRoot, ProjectName);
+				string SrcBuildPath2 = CommandUtils.CombinePaths(BuildRoot, ProjectName.Replace("Game", "").Replace("game", ""));
+
+				string SrcBuildPath_Cache = CommandUtils.CombinePaths(CachePath, ProjectName);
+				string SrcBuildPath2_Cache = CommandUtils.CombinePaths(CachePath, ProjectName.Replace("Game", "").Replace("game", ""));
+
+				if (!InternalUtils.SafeDirectoryExists(SrcBuildPath))
+				{
+					if (!InternalUtils.SafeDirectoryExists(SrcBuildPath2))
+					{
+						throw new AutomationException("-Build=AutoP4: Neither {0} nor {1} exists.", SrcBuildPath, SrcBuildPath2);
+					}
+					SrcBuildPath = SrcBuildPath2;
+					SrcBuildPath_Cache = SrcBuildPath2_Cache;
+				}
+				string SrcCLPath = CommandUtils.CombinePaths(SrcBuildPath, CommandUtils.EscapePath(CommandUtils.P4Env.Branch) + "-CL-" + CommandUtils.P4Env.Changelist.ToString());
+				string SrcCLPath_Cache = CommandUtils.CombinePaths(SrcBuildPath_Cache, CommandUtils.EscapePath(CommandUtils.P4Env.Branch) + "-CL-" + CommandUtils.P4Env.Changelist.ToString());
+				if (!InternalUtils.SafeDirectoryExists(SrcCLPath))
+				{
+					throw new AutomationException("-Build=AutoP4: {0} does not exist.", SrcCLPath);
+				}
+
+				if (InternalUtils.SafeDirectoryExists(SrcCLPath_Cache))
+				{
+					InBuildReference = SrcCLPath_Cache;
+				}
+				else
+				{
+					InBuildReference = SrcCLPath;
+				}
+				Log.Verbose("Using AutoP4 path {0}", InBuildReference);
+			}
+
 			// BuildParam could be a path, a name that we should resolve to a path, Staged, or Editor
 			DirectoryInfo BuildDir = new DirectoryInfo(InBuildReference);
 
@@ -202,9 +204,9 @@ namespace Gauntlet
 			{
 				// Easy option first - is this a full path?
 				OutBuildName = BuildDir.Name;
-				OutBuildPath = BuildDir.FullName;
+				OutBuildPaths = new string[] { BuildDir.FullName };
 			}
-			else if (BuildDir.Name.Equals("staged", StringComparison.OrdinalIgnoreCase))
+			else if (BuildDir.Name.Equals("local", StringComparison.OrdinalIgnoreCase) || BuildDir.Name.Equals("staged", StringComparison.OrdinalIgnoreCase))
 			{
 				string ProjectDir = Path.GetDirectoryName(UnrealHelpers.GetProjectPath(ProjectName));
 
@@ -214,20 +216,25 @@ namespace Gauntlet
 				}
 
 				// First special case - "Staged" means use whats locally staged
-				OutBuildName = "Staged";
-				OutBuildPath = Path.Combine(ProjectDir, "Saved", "StagedBuilds");
+				OutBuildName = "Local";
+				string StagedPath = Path.Combine(ProjectDir, "Saved", "StagedBuilds");
 
-				if (Directory.Exists(OutBuildPath) == false)
+				if (Directory.Exists(StagedPath) == false)
 				{
-					Log.Error("BuildReference was Staged but staged directory {0} not found", OutBuildPath);
+					Log.Error("BuildReference was Staged but staged directory {0} not found", StagedPath);
 					return false;
 				}
+				
+				// include binaries path for packaged builds
+				string BinariesPath = Path.Combine(ProjectDir, "Binaries");
+
+				OutBuildPaths = new string[] { StagedPath, BinariesPath };
 			}
 			else if (BuildDir.Name.Equals("editor", StringComparison.OrdinalIgnoreCase))
 			{
 				// Second special case - "Editor" means run using the editor, no path needed
 				OutBuildName = "Editor";
-				OutBuildPath = Environment.CurrentDirectory;
+				OutBuildPaths = new string[] { Environment.CurrentDirectory };
 			}
 			else
 			{
@@ -256,12 +263,12 @@ namespace Gauntlet
 						}
 
 						OutBuildName = Di.Name;
-						OutBuildPath = Di.FullName;
+						OutBuildPaths = new string[] { Di.FullName };
 					}
 				}
 			}
 
-			if (string.IsNullOrEmpty(OutBuildName) || string.IsNullOrEmpty(OutBuildPath))
+			if (string.IsNullOrEmpty(OutBuildName) || (OutBuildPaths == null || OutBuildPaths.Count() == 0))
 			{
 				Log.Error("Unable to resolve build argument '{0}'", InBuildReference);
 				return false;
@@ -279,15 +286,18 @@ namespace Gauntlet
 		{
 			var BuildList = new List<IBuild>();
 
-			if (string.IsNullOrEmpty(BuildPath) == false)
+			if (BuildPaths.Count() > 0)
 			{
-				IEnumerable<IFolderBuildSource> BuildSources = Gauntlet.Utils.InterfaceHelpers.FindImplementations<IFolderBuildSource>();
-
-				foreach (var BS in BuildSources)
+				foreach (string Path in BuildPaths)
 				{
-					IEnumerable<IBuild> Builds = BS.GetBuildsAtPath(ProjectName, BuildPath);
+					IEnumerable<IFolderBuildSource> BuildSources = Gauntlet.Utils.InterfaceHelpers.FindImplementations<IFolderBuildSource>();
 
-					BuildList.AddRange(Builds);
+					foreach (var BS in BuildSources)
+					{
+						IEnumerable<IBuild> Builds = BS.GetBuildsAtPath(ProjectName, Path);
+
+						BuildList.AddRange(Builds);
+					}
 				}
 			}
 
@@ -375,7 +385,7 @@ namespace Gauntlet
 				return true;
 			}
 
-			Reasons.Add(string.Format("No build at {0} that matches {1}", BuildPath, Role.ToString()));
+			Reasons.Add(string.Format("No build at {0} that matches {1}", string.Join(",", BuildPaths), Role.ToString()));
 
 			return false;
 		}
@@ -594,7 +604,7 @@ namespace Gauntlet
 					}
 					else
 					{
-						throw new AutomationException("No suitable build for {0} found at {1}", TargetPlatform, BuildPath);
+						throw new AutomationException("No suitable build for {0} found at {1}", TargetPlatform, string.Join(",", BuildPaths));
 					}
 
 					//ExePath = AndroidSource.SourceApkPath;			
@@ -660,6 +670,8 @@ namespace Gauntlet
 				return UnrealPath;
 			}
 
+			string BuildPath = BuildPaths.ElementAt(0);
+
 			if (string.IsNullOrEmpty(BuildPath))
 			{
 				return null;
@@ -673,8 +685,8 @@ namespace Gauntlet
 				PlatformPath = Path.Combine(PlatformPath, "Staged");
 			}
 
-			// Urgh
-			if (Platform == UnrealTargetPlatform.Android && BuildName.Equals("Staged", StringComparison.OrdinalIgnoreCase) == false)
+			// Urgh - build share uses a different style...
+			if (Platform == UnrealTargetPlatform.Android && BuildName.Equals("Local", StringComparison.OrdinalIgnoreCase) == false)
 			{
 				PlatformPath = PlatformPath.Replace("Android_ETC2Client", "Android\\FullPackages");
 			}
@@ -720,95 +732,5 @@ namespace Gauntlet
 			//List<string> Empty = new List<string>();
 			//return CanSupportRole(new UnrealSessionRole(UnrealRoleType.Editor, BuildHostPlatform.Current.Platform, UnrealTargetConfiguration.Development), ref Empty); ;
 		}
-
-		bool ResolveStagedReference(string InBuildReference, Func<string, string> ResolutionDelegate)
-		{ 
-			if (string.IsNullOrEmpty(InBuildReference))
-			{
-				return false;
-			}
-
-			// BuildParam could be a path, a name that we should resolve to a path, Staged, or Editor
-			DirectoryInfo BuildDir = new DirectoryInfo(InBuildReference);
-
-			string ResolvedName = null;
-			string ResolvedPath = null;
-
-			if (BuildDir.Exists)
-			{
-				// Easy option first - is this a full path?
-				ResolvedName = BuildDir.Name;
-				ResolvedPath = BuildDir.FullName;
-			}
-			else if (BuildDir.Name.Equals("staged", StringComparison.OrdinalIgnoreCase))
-			{
-				// First special case - "Staged" means use whats locally staged
-				ResolvedName = "Staged";
-				ResolvedPath = Path.Combine(UnrealPath, ProjectName, "Saved", "StagedBuilds");
-
-				if (Directory.Exists(ResolvedPath) == false)
-				{
-					Log.Error("BuildReference was Staged but staged directory {0} not found", ResolvedPath);
-					return false;
-				}
-			}
-			else if (BuildDir.Name.Equals("editor", StringComparison.OrdinalIgnoreCase))
-			{
-				// Second special case - "Editor" means run using the editor
-				ResolvedName = "Editor";
-				ResolvedPath = UnrealPath;
-
-				if (Directory.Exists(ResolvedPath) == false)
-				{
-					Log.Error("BuildReference was Editor but editor directory {0} not found", ResolvedPath);
-					return false;
-				}
-			}
-			else
-			{
-				// todo - make this more generic
-				if (BuildDir.Name.Equals("usesyncedbuild", StringComparison.OrdinalIgnoreCase))
-				{
-					BuildVersion Version;
-					if (BuildVersion.TryRead(BuildVersion.GetDefaultFileName(), out Version))
-					{
-						InBuildReference = Version.BranchName + "-CL-" + Version.Changelist.ToString();
-					}
-				}
-
-				// See if it's in the passed locations
-				if (ResolutionDelegate != null)
-				{
-					string FullPath = ResolutionDelegate(InBuildReference);
-
-					if (string.IsNullOrEmpty(FullPath) == false)
-					{
-						DirectoryInfo Di = new DirectoryInfo(FullPath);
-
-						if (Di.Exists == false)
-						{
-							throw new AutomationException("Resolution delegate returned non existent path");
-						}
-
-						ResolvedName = Di.Name;
-						ResolvedPath = Di.FullName;
-					}
-				}
-			}
-
-			if (string.IsNullOrEmpty(ResolvedPath))
-			{
-				Log.Error("Unable to resolve build argument '{0}'", InBuildReference);
-				return false;
-			}
-
-			BuildName = ResolvedName;
-			BuildPath = ResolvedPath;
-
-			return true;
-		}
 	}
-
-
-
 }
