Gauntlet -
native plugin now marked as runtime so it's available in more builds.
android device now uses a path under the Gauntlet tempdir for pulling artifacts. Dir is also marked for cleanup.

diff --git a/Engine/Plugins/Experimental/Gauntlet/Gauntlet.uplugin b/Engine/Plugins/Experimental/Gauntlet/Gauntlet.uplugin
--- a/Engine/Plugins/Experimental/Gauntlet/Gauntlet.uplugin
+++ b/Engine/Plugins/Experimental/Gauntlet/Gauntlet.uplugin
@@ -18,7 +18,7 @@
 	[
 		{
 			"Name" : "Gauntlet",
-			"Type" : "Developer",
+			"Type" : "Runtime",
 			"LoadingPhase": "PreDefault"
 		}
 	]
diff --git a/Engine/Source/Programs/AutomationTool/Gauntlet/Platform/Android/Gauntlet.TargetDeviceAndroid.cs b/Engine/Source/Programs/AutomationTool/Gauntlet/Platform/Android/Gauntlet.TargetDeviceAndroid.cs
--- a/Engine/Source/Programs/AutomationTool/Gauntlet/Platform/Android/Gauntlet.TargetDeviceAndroid.cs
+++ b/Engine/Source/Programs/AutomationTool/Gauntlet/Platform/Android/Gauntlet.TargetDeviceAndroid.cs
@@ -190,9 +190,27 @@ namespace Gauntlet
 		protected void SaveArtifacts()
 		{
 			// copy remote artifacts to local
+			if (Directory.Exists(Install.AndroidDevice.LocalCachePath))
+			{
+				try
+				{
+					// don't consider this fatal, people often have the directory or a file open
+					Directory.Delete(Install.AndroidDevice.LocalCachePath, true);
+				}
+				catch
+				{
+					Log.Warning("Failed to remove old cache folder {0}", Install.AndroidDevice.LocalCachePath);
+				}
+			}
+
+			// mark it as a temp dir (will also create it)
+			Utils.SystemHelpers.MarkDirectoryForCleanup(Install.AndroidDevice.LocalCachePath);
+
 			string LocalSaved = Path.Combine(Install.AndroidDevice.LocalCachePath, "Saved");
 			Directory.CreateDirectory(LocalSaved);
-			string ArtifactPullCommand = string.Format("pull {0} {1}", Install.AndroidDevice.DeviceArtifactPath, Install.AndroidDevice.LocalCachePath);
+
+			// pull all the artifacts
+			string ArtifactPullCommand = string.Format("pull {0} {1}", Install.AndroidDevice.DeviceArtifactPath, LocalSaved);
 			IProcessResult PullCmd = Install.AndroidDevice.RunAdbDeviceCommand(ArtifactPullCommand);
 
 			if (PullCmd.ExitCode != 0)
@@ -430,7 +448,8 @@ namespace Gauntlet
 			// for IP devices need to sanitize this
 			Name = DeviceName.Replace(":", "_");
 
-			LocalCachePath = Path.Combine(Path.GetTempPath(), "AndroidDevice_" + Name);
+			// Path we use for artifacts, we'll create it later when we need it
+			LocalCachePath = Path.Combine(Globals.TempDir, "AndroidDevice_" + Name);
 
 			ConnectedDevices = GetAllConnectedDevices();
 
@@ -469,12 +488,7 @@ namespace Gauntlet
 						RunAdbGlobalCommand(string.Format("disconnect {0}", DeviceName));
 
 						Log.Info("Disconnected {0}", DeviceName);
-					}
-
-					if (Directory.Exists(LocalCachePath))
-					{
-						Directory.Delete(LocalCachePath, true);
-					}
+					}					
 				}
 				catch (Exception Ex)
 				{
