Gauntlet -
Moved SubmitToDashboard into its own function to ensure derived classes don't accidentally skip it.
Renamed a few classes for future work

diff --git a/Engine/Source/Programs/AutomationTool/Gauntlet/Unreal/Base/Gauntlet.UnrealTestNode.cs b/Engine/Source/Programs/AutomationTool/Gauntlet/Unreal/Base/Gauntlet.UnrealTestNode.cs
--- a/Engine/Source/Programs/AutomationTool/Gauntlet/Unreal/Base/Gauntlet.UnrealTestNode.cs
+++ b/Engine/Source/Programs/AutomationTool/Gauntlet/Unreal/Base/Gauntlet.UnrealTestNode.cs
@@ -63,6 +63,11 @@ namespace Gauntlet
 		/// </summary>
 		public IEnumerable<UnrealRoleArtifacts> SessionArtifacts { get; private set; }
 
+		/// <summary>
+		/// Whether we submit to the dashboard
+		/// </summary>
+		public virtual bool ShouldSubmitDashboardResult { get { return CommandUtils.IsBuildMachine; } }
+
 		/// <summary>
 		/// Helper class that turns our wishes into reallity
 		/// </summary>
@@ -571,6 +576,15 @@ namespace Gauntlet
 			{
 				Log.Warning("Failed to save completion report. {0}", Ex);
 			}
+
+			try
+			{
+				SubmitToDashboard(GetTestResult(), Context, Context.BuildInfo, SessionArtifacts, OutputPath);
+			}
+			catch (Exception Ex)
+			{
+				Log.Warning("Failed to submit results to dashboard. {0}", Ex);
+			}
 		}
 
 		/// <summary>
@@ -583,6 +597,16 @@ namespace Gauntlet
 		{
 		}
 
+		/// <summary>
+		/// Optional function that is called on test completion and gives an opportunity to create a report
+		/// </summary>
+		/// <param name="Result"></param>
+		/// <param name="Contex"></param>
+		/// <param name="Build"></param>
+		public virtual void SubmitToDashboard(TestResult Result, UnrealTestContext Contex, UnrealBuildSource Build, IEnumerable<UnrealRoleArtifacts> Artifacts, string ArtifactPath)
+		{
+		}
+
 		/// <summary>
 		/// Called to request that the test save all artifacts from the completed test to the specified 
 		/// output path. By default the app will save all logs and crash dumps
