Made SaveRoleArtifacts virtual so it can be overridden

diff --git a/Engine/Source/Programs/AutomationTool/Gauntlet/Unreal/Base/Gauntlet.UnrealTestNode.cs b/Engine/Source/Programs/AutomationTool/Gauntlet/Unreal/Base/Gauntlet.UnrealTestNode.cs
--- a/Engine/Source/Programs/AutomationTool/Gauntlet/Unreal/Base/Gauntlet.UnrealTestNode.cs
+++ b/Engine/Source/Programs/AutomationTool/Gauntlet/Unreal/Base/Gauntlet.UnrealTestNode.cs
@@ -628,7 +628,7 @@ namespace Gauntlet
 		/// <param name="Node"></param>
 		/// <param name="OutputPath"></param>
 		/// <returns></returns>
-		public IEnumerable<UnrealRoleArtifacts> SaveRoleArtifacts(string OutputPath)
+		public virtual IEnumerable<UnrealRoleArtifacts> SaveRoleArtifacts(string OutputPath)
 		{
 			CachedArtifactPath = OutputPath;
 			return UnrealApp.SaveRoleArtifacts(Context, TestInstance, CachedArtifactPath);
