UAT: Add a build command to parse the output from MSVC's /d1reportTime argument, and convert it into a trace that can be visualized in Chrome's chrome://tracing tab.

To generate an input file, compile a single file with UnrealBuildTool and redirect the output to a text file:

    UnrealBuildTool.exe UE4Editor Win64 Development -SingleFile="E:\P4\Engine\Source\Runtime\Core\Private\Misc\CoreGlobals.cpp" -Timing > E:\Output.txt

To convert into a trace readable by chrome:

   AutomationTool.exe ParseMsvcTimingInfo -File=E:\Output.txt

Output file will be written to the same directory as the input file, but with a json extension.

diff --git a/Engine/Source/Programs/AutomationTool/Scripts/AutomationScripts.Automation.csproj b/Engine/Source/Programs/AutomationTool/Scripts/AutomationScripts.Automation.csproj
--- a/Engine/Source/Programs/AutomationTool/Scripts/AutomationScripts.Automation.csproj
+++ b/Engine/Source/Programs/AutomationTool/Scripts/AutomationScripts.Automation.csproj
@@ -62,6 +62,7 @@
     </Compile>
     <Compile Include="AnalyzeThirdPartyLibs.Automation.cs" />
     <Compile Include="ArchiveCommand.Automation.cs" />
+    <Compile Include="ParseMsvcTimingInfo.cs" />
     <Compile Include="StashTarget.cs" />
     <Compile Include="BlameKeyword.Automation.cs" />
     <Compile Include="BuildCommonTools.Automation.cs" />
diff --git a/Engine/Source/Programs/AutomationTool/Scripts/ParseMsvcTimingInfo.cs b/Engine/Source/Programs/AutomationTool/Scripts/ParseMsvcTimingInfo.cs
new file mode 100644
--- /dev/null
+++ b/Engine/Source/Programs/AutomationTool/Scripts/ParseMsvcTimingInfo.cs
@@ -0,0 +1,88 @@
+using System;
+using System.Collections.Generic;
+using System.IO;
+using System.Linq;
+using System.Text;
+using System.Text.RegularExpressions;
+using System.Threading.Tasks;
+using Tools.DotNETCommon;
+
+namespace AutomationTool
+{
+	[Help("Parses Visual C++ timing information (as generated by UBT with the -Timing flag), and converts it into JSON format which can be visualized using the chrome://tracing tab")]
+	[Help("-File=<Path>", "Path to the input file")]
+	class ParseMsvcTimingInfo : BuildCommand
+	{
+		public override void ExecuteBuild()
+		{
+			FileReference InputFile = ParseRequiredFileReferenceParam("File");
+
+			string[] Lines = FileReference.ReadAllLines(InputFile);
+
+			int LineIdx = 0;
+			while(LineIdx < Lines.Length && !Lines[LineIdx].StartsWith("Include Headers:", StringComparison.Ordinal))
+			{
+				LineIdx++;
+			}
+			if(LineIdx == Lines.Length)
+			{
+				throw new AutomationException("Couldn't find 'Include Headers' section in {0}", InputFile);
+			}
+
+			LineIdx++;
+			if(LineIdx < Lines.Length && Lines[LineIdx].StartsWith("\tCount:", StringComparison.Ordinal))
+			{
+				LineIdx++;
+			}
+
+			FileReference OutputFile = InputFile.ChangeExtension(".json");
+			using(JsonWriter Writer = new JsonWriter(OutputFile))
+			{
+				Writer.WriteObjectStart();
+				Writer.WriteArrayStart("traceEvents");
+
+				Stack<float> FinishTimesForIndent = new Stack<float>();
+				FinishTimesForIndent.Push(0.0f);
+
+				float StartTime = 0.0f;
+				for(; LineIdx < Lines.Length; LineIdx++)
+				{
+					Match Match = Regex.Match(Lines[LineIdx], "^\t\t(\t*)([^\t]+):\\s*([0-9\\.]+)s$");
+					if(!Match.Success)
+					{
+						break;
+					}
+
+					int Indent = Match.Groups[1].Length;
+					string FileName = Match.Groups[2].Value;
+					float Duration = float.Parse(Match.Groups[3].Value);
+
+					while(Indent <= FinishTimesForIndent.Count - 1)
+					{
+						StartTime = FinishTimesForIndent.Pop();
+					}
+
+					Writer.WriteObjectStart();
+					Writer.WriteValue("pid", 1);
+					Writer.WriteValue("tid", 1);
+					Writer.WriteValue("ts", (int)(StartTime * 1000.0f));
+					Writer.WriteValue("dur", (int)(Duration * 1000.0f));
+					Writer.WriteValue("ph", "X");
+					Writer.WriteValue("name", Path.GetFileName(FileName));
+					Writer.WriteObjectStart("args");
+					Writer.WriteValue("path", FileName);
+					Writer.WriteObjectEnd();
+					Writer.WriteObjectEnd();
+
+					while(Indent >= FinishTimesForIndent.Count)
+					{
+						FinishTimesForIndent.Push(StartTime + Duration);
+					}
+				}
+
+				Writer.WriteArrayEnd();
+				Writer.WriteObjectEnd();
+			}
+		}
+	}
+}
