Parse out a sorted list of classes and functions from the MSVC timing info. Groups templates together, to allow better understanding the most expensive code to compile.

diff --git a/Engine/Source/Programs/AutomationTool/Scripts/ParseMsvcTimingInfo.cs b/Engine/Source/Programs/AutomationTool/Scripts/ParseMsvcTimingInfo.cs
--- a/Engine/Source/Programs/AutomationTool/Scripts/ParseMsvcTimingInfo.cs
+++ b/Engine/Source/Programs/AutomationTool/Scripts/ParseMsvcTimingInfo.cs
@@ -18,24 +18,35 @@ namespace AutomationTool
 			FileReference InputFile = ParseRequiredFileReferenceParam("File");
 
 			string[] Lines = FileReference.ReadAllLines(InputFile);
-
-			int LineIdx = 0;
-			while(LineIdx < Lines.Length && !Lines[LineIdx].StartsWith("Include Headers:", StringComparison.Ordinal))
+			for(int LineIdx = 0; LineIdx < Lines.Length; )
 			{
-				LineIdx++;
-			}
-			if(LineIdx == Lines.Length)
-			{
-				throw new AutomationException("Couldn't find 'Include Headers' section in {0}", InputFile);
+				string Line = Lines[LineIdx];
+				if(Line.StartsWith("Include Headers:", StringComparison.Ordinal))
+				{
+					LineIdx = ParseIncludeHeaders(Lines, LineIdx + 1, InputFile.ChangeExtension(".json"));
+				}
+				else if(Line.StartsWith("Class Definitions:", StringComparison.Ordinal))
+				{
+					LineIdx = ParseDefinitions(Lines, LineIdx + 1, InputFile.ChangeExtension(".classes.txt"));
+				}
+				else if(Line.StartsWith("Function Definitions:", StringComparison.Ordinal))
+				{
+					LineIdx = ParseDefinitions(Lines, LineIdx + 1, InputFile.ChangeExtension(".functions.txt"));
+				}
+				else
+				{
+					LineIdx++;
+				}
 			}
+		}
 
-			LineIdx++;
+		int ParseIncludeHeaders(string[] Lines, int LineIdx, FileReference OutputFile)
+		{
 			if(LineIdx < Lines.Length && Lines[LineIdx].StartsWith("\tCount:", StringComparison.Ordinal))
 			{
 				LineIdx++;
 			}
 
-			FileReference OutputFile = InputFile.ChangeExtension(".json");
 			using(JsonWriter Writer = new JsonWriter(OutputFile))
 			{
 				Writer.WriteObjectStart();
@@ -83,6 +94,49 @@ namespace AutomationTool
 				Writer.WriteArrayEnd();
 				Writer.WriteObjectEnd();
 			}
+			return LineIdx;
+		}
+
+		int ParseDefinitions(string[] Lines, int LineIdx, FileReference OutputFile)
+		{
+			if(LineIdx < Lines.Length && Lines[LineIdx].StartsWith("\tCount:", StringComparison.Ordinal))
+			{
+				LineIdx++;
+			}
+
+			Dictionary<string, float> ClassNameToTime = new Dictionary<string, float>();
+			for(; LineIdx < Lines.Length; LineIdx++)
+			{
+				Match Match = Regex.Match(Lines[LineIdx], "^\t\t\t*([^\t]+):\\s*([0-9\\.]+)s$");
+				if(!Match.Success)
+				{
+					break;
+				}
+
+				string ClassName = Match.Groups[1].Value;
+
+				int TemplateIdx = ClassName.IndexOf('<');
+				if(TemplateIdx != -1)
+				{
+					ClassName = ClassName.Substring(0, TemplateIdx) + "<>";
+				}
+
+				float Time;
+				ClassNameToTime.TryGetValue(ClassName, out Time);
+
+				Time += float.Parse(Match.Groups[2].Value);
+				ClassNameToTime[ClassName] = Time;
+			}
+
+			using(StreamWriter Writer = new StreamWriter(OutputFile.FullName))
+			{
+				foreach(KeyValuePair<string, float> Pair in ClassNameToTime.OrderByDescending(x => x.Value))
+				{
+					Writer.WriteLine("{0,7:0.000}: {1}", Pair.Value, Pair.Key);
+				}
+			}
+
+			return LineIdx;
 		}
 	}
 }
