Minor fix for incorrect error message when a file reference parameter is missing in UAT.

diff --git a/Engine/Source/Programs/AutomationTool/AutomationUtils/BuildCommand.cs b/Engine/Source/Programs/AutomationTool/AutomationUtils/BuildCommand.cs
--- a/Engine/Source/Programs/AutomationTool/AutomationUtils/BuildCommand.cs
+++ b/Engine/Source/Programs/AutomationTool/AutomationUtils/BuildCommand.cs
@@ -96,7 +96,7 @@ namespace AutomationTool
 			FileReference Value = ParseOptionalFileReferenceParam(Param);
 			if(Value == null)
 			{
-				throw new AutomationException("Missing -{0}=... parameter");
+				throw new AutomationException("Missing -{0}=... parameter", Param);
 			}
 			return Value;
 		}
