Fix times for MSVC timing info being out by a factor of 1000.

diff --git a/Engine/Source/Programs/AutomationTool/Scripts/ParseMsvcTimingInfo.cs b/Engine/Source/Programs/AutomationTool/Scripts/ParseMsvcTimingInfo.cs
--- a/Engine/Source/Programs/AutomationTool/Scripts/ParseMsvcTimingInfo.cs
+++ b/Engine/Source/Programs/AutomationTool/Scripts/ParseMsvcTimingInfo.cs
@@ -78,8 +78,8 @@ namespace AutomationTool
 					Writer.WriteObjectStart();
 					Writer.WriteValue("pid", 1);
 					Writer.WriteValue("tid", 1);
-					Writer.WriteValue("ts", (int)(StartTime * 1000.0f));
-					Writer.WriteValue("dur", (int)(Duration * 1000.0f));
+					Writer.WriteValue("ts", (long)(StartTime * 1000.0f * 1000.0f));
+					Writer.WriteValue("dur", (long)(Duration * 1000.0f * 1000.0f));
 					Writer.WriteValue("ph", "X");
 					Writer.WriteValue("name", Path.GetFileName(FileName));
 					Writer.WriteObjectStart("args");
