HTML5 - pthread_main.js name changed by emscripten 1.38.24 upgrade

also, inform developer about the long link time

#rnx

diff --git a/Engine/Source/Programs/AutomationTool/HTML5/HTML5Platform.Automation.cs b/Engine/Source/Programs/AutomationTool/HTML5/HTML5Platform.Automation.cs
--- a/Engine/Source/Programs/AutomationTool/HTML5/HTML5Platform.Automation.cs
+++ b/Engine/Source/Programs/AutomationTool/HTML5/HTML5Platform.Automation.cs
@@ -149,13 +149,13 @@ public class HTML5Platform : Platform
 				File.Copy(      SrcUE4GameBasename + ".js.mem", UE4GameBasename + ".js.mem", true);
 				File.SetAttributes(UE4GameBasename + ".js.mem", FileAttributes.Normal);
 
-				File.Copy(Path.Combine(Path.GetDirectoryName(_ProjectFullpath), "pthread-main.js"), Path.Combine(PackagePath, "pthread-main.js"), true);
-				File.SetAttributes(Path.Combine(PackagePath, "pthread-main.js"), FileAttributes.Normal);
+				File.Copy(      SrcUE4GameBasename + ".worker.js", UE4GameBasename + ".worker.js", true);
+				File.SetAttributes(UE4GameBasename + ".worker.js", FileAttributes.Normal);
 			}
 			else
 			{
-				// nuke possibly old deployed pthread-main.js, which is not needed for singlethreaded build.
-				File.Delete(Path.Combine(PackagePath, "pthread-main.js"));
+				// nuke possibly old deployed .worker.js, which is not needed for singlethreaded build.
+				File.Delete(UE4GameBasename + ".worker.js");
 			}
 		}
 		// else, c++ projects will compile "to" PackagePath
@@ -249,7 +249,7 @@ public class HTML5Platform : Platform
 			if (enableMultithreading)
 			{
 				CompressionTasks.Add(Task.Factory.StartNew(() => CompressFile(UE4GameBasename + ".js.mem", UE4GameBasename + ".js.memgz")));			// mem init file
-				CompressionTasks.Add(Task.Factory.StartNew(() => CompressFile(PackagePath + "/pthread-main.js", PackagePath + "/pthread-main.jsgz")));	// pthread file
+				CompressionTasks.Add(Task.Factory.StartNew(() => CompressFile(UE4GameBasename + ".worker.js", UE4GameBasename + ".worker.jsgz")));		// pthread file
 			}
 			Task.WaitAll(CompressionTasks.ToArray());
 		}
@@ -262,7 +262,7 @@ public class HTML5Platform : Platform
 			File.Delete(UE4GameBasename + ".jsgz");
 			File.Delete(UE4GameBasename + ".js.symbolsgz");
 			File.Delete(UE4GameBasename + ".js.memgz");
-			File.Delete(PackagePath + "/pthread-main.jsgz");
+			File.Delete(UE4GameBasename + ".worker.jsgz");
 			File.Delete(PackagePath + "/Utility.jsgz");
 			File.Delete(ProjectBasename + ".datagz");
 			File.Delete(ProjectBasename + ".data.jsgz");
@@ -516,7 +516,7 @@ public class HTML5Platform : Platform
 		if (enableMultithreading)
 		{
 			SC.ArchiveFiles(PackagePath, UE4GameBasename + ".js.mem");		// memory init file
-			SC.ArchiveFiles(PackagePath, "pthread-main.js");
+			SC.ArchiveFiles(PackagePath, UE4GameBasename + ".worker.js");
 		}
 
 		// Archive HTML5 Server and a Readme.
@@ -541,7 +541,7 @@ public class HTML5Platform : Platform
 			if (enableMultithreading)
 			{
 				SC.ArchiveFiles(PackagePath, UE4GameBasename + ".js.memgz");
-				SC.ArchiveFiles(PackagePath, "pthread-main.jsgz");
+				SC.ArchiveFiles(PackagePath, UE4GameBasename + ".worker.jsgz");
 			}
 		}
 		else
@@ -550,7 +550,7 @@ public class HTML5Platform : Platform
 			File.Delete(UE4GameBasename + ".wasmgz");
 			File.Delete(UE4GameBasename + ".jsgz");
 			File.Delete(UE4GameBasename + ".js.symbolsgz");
-			File.Delete("pthread-main.jsgz");
+			File.Delete(UE4GameBasename + ".worker.jsgz");
 			File.Delete("Utility.jsgz");
 			File.Delete(ProjectBasename + ".datagz");
 			File.Delete(ProjectBasename + ".data.jsgz");
diff --git a/Engine/Source/Programs/UnrealBuildTool/Platform/HTML5/HTML5ToolChain.cs b/Engine/Source/Programs/UnrealBuildTool/Platform/HTML5/HTML5ToolChain.cs
--- a/Engine/Source/Programs/UnrealBuildTool/Platform/HTML5/HTML5ToolChain.cs
+++ b/Engine/Source/Programs/UnrealBuildTool/Platform/HTML5/HTML5ToolChain.cs
@@ -644,6 +644,9 @@ namespace UnrealBuildTool
 			LinkAction.PrerequisiteItems.Add(ResponseFileItem);
 			Actions.Add(LinkAction);
 
+			Log.TraceInformation("NOTE: about to link for HTML5 -- this takes at least 7 minutes (and up to 20 minutes on older machines) to complete.");
+			Log.TraceInformation("      we are workig with the Emscripten makers to speed this up.");
+
 			return OutputFile;
 		}
 
