PhysX: change v13 Linux toolchain to v12 one.

- There are issues when the libraries are compiled with clang 7.0.1 (tracked as UE-69848).

#rnx

diff --git a/Engine/Source/Programs/AutomationTool/Scripts/BuildPhysX.Automation.cs b/Engine/Source/Programs/AutomationTool/Scripts/BuildPhysX.Automation.cs
--- a/Engine/Source/Programs/AutomationTool/Scripts/BuildPhysX.Automation.cs
+++ b/Engine/Source/Programs/AutomationTool/Scripts/BuildPhysX.Automation.cs
@@ -650,6 +650,22 @@ class BuildPhysX : BuildCommand
 					{
 						LogInformation("Generating projects for lib " + TargetLib.ToString() + ", " + TargetData.ToString());
 
+						if (TargetData.Platform == UnrealTargetPlatform.Linux)
+						{
+							// the libraries are broken when compiled with clang 7.0.1
+							string OriginalToolchainPath = Environment.GetEnvironmentVariable("LINUX_MULTIARCH_ROOT");
+							if (!string.IsNullOrEmpty(OriginalToolchainPath))
+							{
+								string ToolchainPathToUse = OriginalToolchainPath.Replace("v13_clang-7.0.1-centos7", "v12_clang-6.0.1-centos7");
+								LogInformation("Working around problems with newer clangs: {0} -> {1}", OriginalToolchainPath, ToolchainPathToUse);
+								Environment.SetEnvironmentVariable("LINUX_MULTIARCH_ROOT", ToolchainPathToUse);
+							}
+							else
+							{
+								LogWarning("LINUX_MULTIARCH_ROOT is not set!");
+							}
+						}
+
 						ProcessStartInfo StartInfo = new ProcessStartInfo();
 						StartInfo.FileName = CMakeName;
 						StartInfo.WorkingDirectory = CMakeTargetDirectory.ToString();
