Fix target not being rebuilt when an plugin not required for this platform is disabled. Also add log output showing why a temp target was required.

diff --git a/Engine/Source/Programs/AutomationTool/AutomationUtils/ProjectUtils.cs b/Engine/Source/Programs/AutomationTool/AutomationUtils/ProjectUtils.cs
--- a/Engine/Source/Programs/AutomationTool/AutomationUtils/ProjectUtils.cs
+++ b/Engine/Source/Programs/AutomationTool/AutomationUtils/ProjectUtils.cs
@@ -128,10 +128,22 @@ namespace AutomationTool
 		}
 
 		private static bool RequiresTempTarget(FileReference RawProjectPath, List<UnrealTargetPlatform> ClientTargetPlatforms, List<UnrealTargetConfiguration> ClientTargetConfigurations, bool AssetNativizationRequested)
+		{
+			string Reason;
+			if(RequiresTempTarget(RawProjectPath, ClientTargetPlatforms, ClientTargetConfigurations, AssetNativizationRequested, out Reason))
+			{
+				Log.TraceInformation("{0} requires a temporary target.cs to be generated ({1})", RawProjectPath.GetFileName(), Reason);
+				return true;
+			}
+			return false;
+		}
+
+		private static bool RequiresTempTarget(FileReference RawProjectPath, List<UnrealTargetPlatform> ClientTargetPlatforms, List<UnrealTargetConfiguration> ClientTargetConfigurations, bool AssetNativizationRequested, out string Reason)
 		{
 			// check to see if we already have a Target.cs file
 			if (File.Exists (Path.Combine (Path.GetDirectoryName (RawProjectPath.FullName), "Source", RawProjectPath.GetFileNameWithoutExtension() + ".Target.cs")))
 			{
+				Reason = null;
 				return false;
 			}
 			else if (Directory.Exists(Path.Combine(Path.GetDirectoryName(RawProjectPath.FullName), "Source")))
@@ -141,6 +153,7 @@ namespace AutomationTool
 				FileInfo[] Files = (new DirectoryInfo( Path.Combine (Path.GetDirectoryName (RawProjectPath.FullName), "Source")).GetFiles ("*.Target.cs", SearchOption.AllDirectories));
 				if (Files.Length > 0)
 				{
+					Reason = null;
 					return false;
 				}
 			}
@@ -154,6 +167,7 @@ namespace AutomationTool
                 // we're going to be converting some of the project's assets 
                 // into native code, so we require a distinct target (executable) 
                 // be generated for this project
+				Reason = "asset nativization is enabled";
                 return true;
             }
 
@@ -164,6 +178,7 @@ namespace AutomationTool
 					EncryptionAndSigning.CryptoSettings Settings = EncryptionAndSigning.ParseCryptoSettings(RawProjectPath.Directory, ClientPlatform);
 					if (Settings.IsAnyEncryptionEnabled() || Settings.bEnablePakSigning)
 					{
+						Reason = "encryption/signing is enabled";
 						return true;
 					}
 				}
@@ -198,64 +213,60 @@ namespace AutomationTool
 				}
 			}
 
-			// Change the working directory to be the Engine/Source folder. We are running from Engine/Binaries/DotNET
-			DirectoryReference oldCWD = DirectoryReference.GetCurrentDirectory();
-			try
+			// Read the project descriptor, and find all the plugins available to this project
+			ProjectDescriptor Project = ProjectDescriptor.FromFile(RawProjectPath);
+			List<PluginInfo> AvailablePlugins = Plugins.ReadAvailablePlugins(CommandUtils.EngineDirectory, RawProjectPath, Project.AdditionalPluginDirectories);
+
+			// check the target platforms for any differences in build settings or additional plugins
+			foreach (UnrealTargetPlatform TargetPlatformType in TargetPlatforms)
 			{
-				DirectoryReference EngineSourceDirectory = DirectoryReference.Combine(CommandUtils.EngineDirectory, "Source");
-				if (!DirectoryReference.Exists(EngineSourceDirectory)) // only set the directory if it exists, this should only happen if we are launching the editor from an artist sync
+				if(!CommandUtils.IsEngineInstalled() && !PlatformExports.HasDefaultBuildConfig(RawProjectPath, TargetPlatformType))
 				{
-					EngineSourceDirectory = DirectoryReference.Combine(CommandUtils.EngineDirectory, "Binaries");
+					Reason = "project has non-default build configuration";
+					return true;
+				}
+				if(PlatformExports.RequiresBuild(RawProjectPath, TargetPlatformType))
+				{
+					Reason = "overriden by target platform";
+					return true;
 				}
-				Directory.SetCurrentDirectory(EngineSourceDirectory.FullName);
-
-				// Read the project descriptor, and find all the plugins available to this project
-				ProjectDescriptor Project = ProjectDescriptor.FromFile(RawProjectPath);
-				List<PluginInfo> AvailablePlugins = Plugins.ReadAvailablePlugins(CommandUtils.EngineDirectory, RawProjectPath, Project.AdditionalPluginDirectories);
 
-				// check the target platforms for any differences in build settings or additional plugins
-				bool RetVal = false;
-				foreach (UnrealTargetPlatform TargetPlatformType in TargetPlatforms)
+				// find if there are any plugins enabled or disabled which differ from the default
+				foreach(PluginInfo Plugin in AvailablePlugins)
 				{
-					if((!CommandUtils.IsEngineInstalled() && !PlatformExports.HasDefaultBuildConfig(RawProjectPath, TargetPlatformType)) || PlatformExports.RequiresBuild(RawProjectPath, TargetPlatformType))
+					bool bPluginEnabledForTarget = false;
+					foreach (UnrealTargetConfiguration TargetConfigType in TargetConfigurations)
 					{
-						RetVal = true;
-						break;
+						bPluginEnabledForTarget |= Plugins.IsPluginEnabledForProject(Plugin, Project, TargetPlatformType, TargetConfigType, TargetRules.TargetType.Game);
 					}
 
-					// find if there are any plugins enabled or disabled which differ from the default
-					foreach(PluginInfo Plugin in AvailablePlugins)
+					bool bPluginEnabledForBaseTarget = false;
+					if(!Plugin.Descriptor.bInstalled)
 					{
-						bool bPluginEnabledForProject = false;
 						foreach (UnrealTargetConfiguration TargetConfigType in TargetConfigurations)
 						{
-							bPluginEnabledForProject |= Plugins.IsPluginEnabledForProject(Plugin, Project, TargetPlatformType, TargetConfigType, TargetRules.TargetType.Game);
+							bPluginEnabledForBaseTarget |= Plugins.IsPluginEnabledForProject(Plugin, null, TargetPlatformType, TargetConfigType, TargetRules.TargetType.Game);
 						}
-						if ((bPluginEnabledForProject != Plugin.EnabledByDefault) || (bPluginEnabledForProject && Plugin.Descriptor.bInstalled))
+					}
+
+					if (bPluginEnabledForTarget != bPluginEnabledForBaseTarget)
+					{
+						if(bPluginEnabledForTarget)
 						{
-							// NOTE: this code was only marking plugins that compiled for the platform to upgrade to code project, however
-							// this doesn't work in practice, because the runtime code will look for the plugin, without a .uplugin file,
-							// and will fail. This is the safest way to make sure all platforms are acting the same. However, if you 
-							// whitelist the plugin in the .uproject file, the above UProjectInfo.IsPluginEnabledForProject check won't pass
-							// so you won't get in here. Leaving this commented out code in there, because someone is bound to come looking 
-							// for why a non-whitelisted platform module is causing a project to convert to code-based. 
-							// As an aside, if you run the project with UE4Game (not your Project's binary) from the debugger, it will work
-							// _in this case_ because the .uplugin file will have been staged, and there is no needed library 
-							// if(Plugin.Descriptor.Modules.Any(Module => Module.IsCompiledInConfiguration(TargetPlatformType, TargetType.Game, bBuildDeveloperTools: false, bBuildEditor: false)))
-							{
-								RetVal = true;
-								break;
-							}
+							Reason = String.Format("{0} plugin is enabled", Plugin.Name);
+							return true;
+						}
+						else
+						{
+							Reason = String.Format("{0} plugin is disabled", Plugin.Name);
+							return true;
 						}
 					}
 				}
-				return RetVal;
-			}
-			finally
-			{
-				// Change back to the original directory
-				Directory.SetCurrentDirectory(oldCWD.FullName);
 			}
+
+			Reason = null;
+			return false;
 		}
 
 		private static void GenerateTempTarget(FileReference RawProjectPath)
diff --git a/Engine/Source/Programs/UnrealBuildTool/System/Plugins.cs b/Engine/Source/Programs/UnrealBuildTool/System/Plugins.cs
--- a/Engine/Source/Programs/UnrealBuildTool/System/Plugins.cs
+++ b/Engine/Source/Programs/UnrealBuildTool/System/Plugins.cs
@@ -374,7 +374,7 @@ namespace UnrealBuildTool
 		/// <summary>
 		/// Determine if a plugin is enabled for a given project
 		/// </summary>
-		/// <param name="Project">The project to check</param>
+		/// <param name="Project">The project to check. May be null.</param>
 		/// <param name="Plugin">Information about the plugin</param>
 		/// <param name="Platform">The target platform</param>
 		/// <param name="TargetConfiguration">The target configuration</param>
@@ -382,7 +382,7 @@ namespace UnrealBuildTool
 		/// <returns>True if the plugin should be enabled for this project</returns>
 		public static bool IsPluginEnabledForProject(PluginInfo Plugin, ProjectDescriptor Project, UnrealTargetPlatform Platform, UnrealTargetConfiguration TargetConfiguration, TargetType Target)
 		{
-			bool bEnabled = Plugin.EnabledByDefault;
+			bool bEnabled = Plugin.EnabledByDefault && Plugin.Descriptor.SupportsTargetPlatform(Platform);
 			if (Project != null && Project.Plugins != null)
 			{
 				foreach (PluginReferenceDescriptor PluginReference in Project.Plugins)
