- Proper fix to merge conflicts

diff --git a/Engine/Source/Programs/AutomationTool/Scripts/CopyBuildToStagingDirectory.Automation.cs b/Engine/Source/Programs/AutomationTool/Scripts/CopyBuildToStagingDirectory.Automation.cs
--- a/Engine/Source/Programs/AutomationTool/Scripts/CopyBuildToStagingDirectory.Automation.cs
+++ b/Engine/Source/Programs/AutomationTool/Scripts/CopyBuildToStagingDirectory.Automation.cs
@@ -1886,7 +1886,27 @@ public partial class Project : CommandUtils
 				{
 					FileReference PrimaryOrderFile   = (PakOrderFileLocations.Count >= 1) ? PakOrderFileLocations[0] : null;
 					FileReference SecondaryOrderFile = (PakOrderFileLocations.Count >= 2) ? PakOrderFileLocations[1] : null;
-					Commands.Add(GetUnrealPakArguments(PakParams.UnrealPakResponseFile, OutputLocation, PrimaryOrderFile, SC.StageTargetPlatform.GetPlatformPakCommandLine(Params, SC) + CompressionFormats + " " + Params.AdditionalPakOptions, PakParams.bCompressed, CryptoSettings, CryptoKeysCacheFilename, PatchSourceContentPath, PakParams.EncryptionKeyGuid, SecondaryOrderFile));
+
+                    string BulkOption = "";
+                    {
+                        ConfigHierarchy PlatformGameConfig;
+                        if (Params.EngineConfigs.TryGetValue(SC.StageTargetPlatform.PlatformType, out PlatformGameConfig))
+                        {
+                            bool bMasterEnable = false;
+                            PlatformGameConfig.GetBool("MemoryMappedFiles", "MasterEnable", out bMasterEnable);
+                            if (bMasterEnable)
+                            {
+                                int Value = 0;
+                                PlatformGameConfig.GetInt32("MemoryMappedFiles", "Alignment", out Value);
+                                if (Value > 0)
+                                {
+                                    BulkOption = String.Format(" -AlignForMemoryMapping={0}", Value);
+                                }
+                            }
+                        }
+                    }
+
+					Commands.Add(GetUnrealPakArguments(PakParams.UnrealPakResponseFile, OutputLocation, PrimaryOrderFile, SC.StageTargetPlatform.GetPlatformPakCommandLine(Params, SC) + BulkOption + CompressionFormats + " " + Params.AdditionalPakOptions, PakParams.bCompressed, CryptoSettings, CryptoKeysCacheFilename, PatchSourceContentPath, PakParams.EncryptionKeyGuid, SecondaryOrderFile));
 					LogNames.Add(OutputLocation.GetFileNameWithoutExtension());
 				}
 			}
