- Final fix now for merge issue (I had compiled wrong program to test my fix, oops)

diff --git a/Engine/Source/Programs/AutomationTool/Scripts/CopyBuildToStagingDirectory.Automation.cs b/Engine/Source/Programs/AutomationTool/Scripts/CopyBuildToStagingDirectory.Automation.cs
--- a/Engine/Source/Programs/AutomationTool/Scripts/CopyBuildToStagingDirectory.Automation.cs
+++ b/Engine/Source/Programs/AutomationTool/Scripts/CopyBuildToStagingDirectory.Automation.cs
@@ -1889,15 +1889,15 @@ public partial class Project : CommandUtils
 
                     string BulkOption = "";
                     {
-                        ConfigHierarchy PlatformGameConfig;
-                        if (Params.EngineConfigs.TryGetValue(SC.StageTargetPlatform.PlatformType, out PlatformGameConfig))
+                        ConfigHierarchy PlatformEngineConfig;
+                        if (Params.EngineConfigs.TryGetValue(SC.StageTargetPlatform.PlatformType, out PlatformEngineConfig))
                         {
                             bool bMasterEnable = false;
-                            PlatformGameConfig.GetBool("MemoryMappedFiles", "MasterEnable", out bMasterEnable);
+							PlatformEngineConfig.GetBool("MemoryMappedFiles", "MasterEnable", out bMasterEnable);
                             if (bMasterEnable)
                             {
                                 int Value = 0;
-                                PlatformGameConfig.GetInt32("MemoryMappedFiles", "Alignment", out Value);
+								PlatformEngineConfig.GetInt32("MemoryMappedFiles", "Alignment", out Value);
                                 if (Value > 0)
                                 {
                                     BulkOption = String.Format(" -AlignForMemoryMapping={0}", Value);
