Added support for querying for updates from our backend.

diff --git a/Engine/Plugins/Online/OnlineSubsystem/Source/Private/OnlineSubsystemModule.cpp b/Engine/Plugins/Online/OnlineSubsystem/Source/Private/OnlineSubsystemModule.cpp
--- a/Engine/Plugins/Online/OnlineSubsystem/Source/Private/OnlineSubsystemModule.cpp
+++ b/Engine/Plugins/Online/OnlineSubsystem/Source/Private/OnlineSubsystemModule.cpp
@@ -233,7 +233,6 @@ void FOnlineSubsystemModule::EnumerateOnlineSubsystems(FEnumerateOnlineSubsystem
 
 FName FOnlineSubsystemModule::ParseOnlineSubsystemName(const FName& FullName, FName& SubsystemName, FName& InstanceName) const
 {
-#if !(UE_GAME || UE_SERVER)
 	SubsystemName = DefaultPlatformService;
 	InstanceName = FOnlineSubsystemImpl::DefaultInstanceName;
 
@@ -262,18 +261,6 @@ FName FOnlineSubsystemModule::ParseOnlineSubsystemName(const FName& FullName, FN
 	}
 
 	return FName(*FString::Printf(TEXT("%s:%s"), *SubsystemName.ToString(), *InstanceName.ToString()));
-#else	
-	
-	SubsystemName = FullName.IsNone() ? DefaultPlatformService : FullName;
-	InstanceName = FOnlineSubsystemImpl::DefaultInstanceName;
-
-#if !UE_BUILD_SHIPPING
-	int32 DelimIdx = INDEX_NONE;
-	static const TCHAR InstanceDelim = ':';
-	ensure(!FullName.ToString().FindChar(InstanceDelim, DelimIdx) && DelimIdx == INDEX_NONE);
-#endif
-	return SubsystemName;
-#endif // !(UE_GAME || UE_SERVER)
 }
 
 IOnlineSubsystem* FOnlineSubsystemModule::GetOnlineSubsystem(const FName InSubsystemName)
diff --git a/Engine/Source/Programs/AutomationTool/Android/AndroidPlatform.Automation.cs b/Engine/Source/Programs/AutomationTool/Android/AndroidPlatform.Automation.cs
--- a/Engine/Source/Programs/AutomationTool/Android/AndroidPlatform.Automation.cs
+++ b/Engine/Source/Programs/AutomationTool/Android/AndroidPlatform.Automation.cs
@@ -554,8 +554,8 @@ public class AndroidPlatform : Platform
 			LogInformation("Writing bat for install with {0}", bPackageDataInsideApk ? "data in APK" : "separate OBB");
             BatchLines = new string[] {
 						"setlocal",
-                        "set ANDROIDHOME=%ANDROID_HOME%",
 						"if NOT \"UE_SDKS_ROOT\"==\"\" (call %UE_SDKS_ROOT%\\HostWin64\\Android\\SetupEnvironmentVars.bat)",
+						"set ANDROIDHOME=%ANDROID_HOME%",		
 						"if \"%ANDROIDHOME%\"==\"\" set ANDROIDHOME="+Environment.GetEnvironmentVariable("ANDROID_HOME"),
 						"set ADB=%ANDROIDHOME%\\platform-tools\\adb.exe",
 						"set DEVICE=",
@@ -632,6 +632,7 @@ public class AndroidPlatform : Platform
 			LogInformation("Writing bat for uninstall with {0}", bPackageDataInsideApk ? "data in APK" : "separate OBB");
 			BatchLines = new string[] {
 						"setlocal",
+						"if NOT \"UE_SDKS_ROOT\"==\"\" (call %UE_SDKS_ROOT%\\HostWin64\\Android\\SetupEnvironmentVars.bat)",
 						"set ANDROIDHOME=%ANDROID_HOME%",
 						"if \"%ANDROIDHOME%\"==\"\" set ANDROIDHOME="+Environment.GetEnvironmentVariable("ANDROID_HOME"),
 						"set ADB=%ANDROIDHOME%\\platform-tools\\adb.exe",
