HTML5 - downgrading emscripten back to 1.38.20

- engine changes to use 1.38.20 toolchain

#rnx

diff --git a/Engine/Source/Programs/AutomationTool/HTML5/HTML5Platform.Automation.cs b/Engine/Source/Programs/AutomationTool/HTML5/HTML5Platform.Automation.cs
--- a/Engine/Source/Programs/AutomationTool/HTML5/HTML5Platform.Automation.cs
+++ b/Engine/Source/Programs/AutomationTool/HTML5/HTML5Platform.Automation.cs
@@ -149,13 +149,13 @@ public class HTML5Platform : Platform
 				File.Copy(      SrcUE4GameBasename + ".js.mem", UE4GameBasename + ".js.mem", true);
 				File.SetAttributes(UE4GameBasename + ".js.mem", FileAttributes.Normal);
 
-				File.Copy(      SrcUE4GameBasename + ".worker.js", UE4GameBasename + ".worker.js", true);
-				File.SetAttributes(UE4GameBasename + ".worker.js", FileAttributes.Normal);
+				File.Copy(Path.Combine(Path.GetDirectoryName(_ProjectFullpath), "pthread-main.js"), Path.Combine(PackagePath, "pthread-main.js"), true);
+				File.SetAttributes(Path.Combine(PackagePath, "pthread-main.js"), FileAttributes.Normal);
 			}
 			else
 			{
-				// nuke possibly old deployed .worker.js, which is not needed for singlethreaded build.
-				File.Delete(UE4GameBasename + ".worker.js");
+				// nuke possibly old deployed pthread-main.js, which is not needed for singlethreaded build.
+				File.Delete(Path.Combine(PackagePath, "pthread-main.js"));
 			}
 		}
 		// else, c++ projects will compile "to" PackagePath
@@ -249,7 +249,7 @@ public class HTML5Platform : Platform
 			if (enableMultithreading)
 			{
 				CompressionTasks.Add(Task.Factory.StartNew(() => CompressFile(UE4GameBasename + ".js.mem", UE4GameBasename + ".js.memgz")));			// mem init file
-				CompressionTasks.Add(Task.Factory.StartNew(() => CompressFile(UE4GameBasename + ".worker.js", UE4GameBasename + ".worker.jsgz")));		// pthread file
+				CompressionTasks.Add(Task.Factory.StartNew(() => CompressFile(PackagePath + "/pthread-main.js", PackagePath + "/pthread-main.jsgz")));	// pthread file
 			}
 			Task.WaitAll(CompressionTasks.ToArray());
 		}
@@ -262,7 +262,7 @@ public class HTML5Platform : Platform
 			File.Delete(UE4GameBasename + ".jsgz");
 			File.Delete(UE4GameBasename + ".js.symbolsgz");
 			File.Delete(UE4GameBasename + ".js.memgz");
-			File.Delete(UE4GameBasename + ".worker.jsgz");
+			File.Delete(PackagePath + "/pthread-main.jsgz");
 			File.Delete(PackagePath + "/Utility.jsgz");
 			File.Delete(ProjectBasename + ".datagz");
 			File.Delete(ProjectBasename + ".data.jsgz");
@@ -516,7 +516,7 @@ public class HTML5Platform : Platform
 		if (enableMultithreading)
 		{
 			SC.ArchiveFiles(PackagePath, UE4GameBasename + ".js.mem");		// memory init file
-			SC.ArchiveFiles(PackagePath, UE4GameBasename + ".worker.js");
+			SC.ArchiveFiles(PackagePath, "pthread-main.js");
 		}
 
 		// Archive HTML5 Server and a Readme.
@@ -541,7 +541,7 @@ public class HTML5Platform : Platform
 			if (enableMultithreading)
 			{
 				SC.ArchiveFiles(PackagePath, UE4GameBasename + ".js.memgz");
-				SC.ArchiveFiles(PackagePath, UE4GameBasename + ".worker.jsgz");
+				SC.ArchiveFiles(PackagePath, "pthread-main.jsgz");
 			}
 		}
 		else
@@ -550,7 +550,7 @@ public class HTML5Platform : Platform
 			File.Delete(UE4GameBasename + ".wasmgz");
 			File.Delete(UE4GameBasename + ".jsgz");
 			File.Delete(UE4GameBasename + ".js.symbolsgz");
-			File.Delete(UE4GameBasename + ".worker.jsgz");
+			File.Delete("pthread-main.jsgz");
 			File.Delete("Utility.jsgz");
 			File.Delete(ProjectBasename + ".datagz");
 			File.Delete(ProjectBasename + ".data.jsgz");
diff --git a/Engine/Source/Programs/UnrealBuildTool/Platform/HTML5/HTML5SDKInfo.cs b/Engine/Source/Programs/UnrealBuildTool/Platform/HTML5/HTML5SDKInfo.cs
--- a/Engine/Source/Programs/UnrealBuildTool/Platform/HTML5/HTML5SDKInfo.cs
+++ b/Engine/Source/Programs/UnrealBuildTool/Platform/HTML5/HTML5SDKInfo.cs
@@ -18,8 +18,8 @@ namespace UnrealBuildTool
 		static string NODE_VER = "8.9.1_64bit";
 		static string PYTHON_VER = "2.7.13.1_64bit"; // Only used on Windows; other platforms use built-in Python.
 
-		static string LLVM_VER = "e1.38.24_64bit";
-		static string SDKVersion = "1.38.24";
+		static string LLVM_VER = "e1.38.20_64bit";
+		static string SDKVersion = "1.38.20";
 
 //		static string LLVM_VER = "incoming";
 //		static string SDKVersion = "incoming";
diff --git a/Engine/Source/Runtime/ApplicationCore/Private/HTML5/HTML5InputInterface.cpp b/Engine/Source/Runtime/ApplicationCore/Private/HTML5/HTML5InputInterface.cpp
--- a/Engine/Source/Runtime/ApplicationCore/Private/HTML5/HTML5InputInterface.cpp
+++ b/Engine/Source/Runtime/ApplicationCore/Private/HTML5/HTML5InputInterface.cpp
@@ -186,7 +186,7 @@ void FHTML5InputInterface::SendControllerEvents()
 //	/// XXX TODO Restore this, emscripten_get_num_gamepads() currently has a moderately high performance impact, so disabled for local testing
 //	static bool GamepadSupported = false;
 //#else
-	static bool GamepadSupported = emscripten_sample_gamepad_data() == EMSCRIPTEN_RESULT_SUCCESS;
+	static bool GamepadSupported = true;
 //#endif
 
 	const double CurrentTime = FPlatformTime::Seconds();
