Prevent UAT scripts from being rebuilt in an installed engine build.

diff --git a/Engine/Source/Programs/AutomationTool/AutomationUtils/ScriptCompiler.cs b/Engine/Source/Programs/AutomationTool/AutomationUtils/ScriptCompiler.cs
--- a/Engine/Source/Programs/AutomationTool/AutomationUtils/ScriptCompiler.cs
+++ b/Engine/Source/Programs/AutomationTool/AutomationUtils/ScriptCompiler.cs
@@ -60,7 +60,14 @@ namespace AutomationTool
 			bool DoCompile = false;
 			if (GlobalCommandLine.Compile)
 			{
-				DoCompile = true;
+				if(CommandUtils.IsEngineInstalled())
+				{
+					CommandUtils.LogWarning("Ignoring -Compile argument because engine is installed.");
+				}
+				else
+				{
+					DoCompile = true;
+				}
 			}
 
 			// Compile only if not disallowed.
