UAT: Always force a changelist to be submitted, even if it contains unchanged files.

diff --git a/Engine/Source/Programs/AutomationTool/AutomationUtils/P4Utils.cs b/Engine/Source/Programs/AutomationTool/AutomationUtils/P4Utils.cs
--- a/Engine/Source/Programs/AutomationTool/AutomationUtils/P4Utils.cs
+++ b/Engine/Source/Programs/AutomationTool/AutomationUtils/P4Utils.cs
@@ -2189,7 +2189,7 @@ namespace AutomationTool
                     return;
                 }
 				string CmdOutput;
-				if (!LogP4Output(out CmdOutput, String.Format("submit -c {0}", CL)))
+				if (!LogP4Output(out CmdOutput, String.Format("submit -c {0} -f submitunchanged", CL)))
 				{
 					if (!Force)
 					{
