Wrap Gauntlet CopyDirectory FileInfo creation in an exception handler for better error logging

diff --git a/Engine/Source/Programs/AutomationTool/Gauntlet/Framework/Gauntlet.Utils.cs b/Engine/Source/Programs/AutomationTool/Gauntlet/Framework/Gauntlet.Utils.cs
--- a/Engine/Source/Programs/AutomationTool/Gauntlet/Framework/Gauntlet.Utils.cs
+++ b/Engine/Source/Programs/AutomationTool/Gauntlet/Framework/Gauntlet.Utils.cs
@@ -881,8 +881,20 @@ namespace Gauntlet
 						DestPath = Transform(DestPath);
 					}
 
-					FileInfo DestInfo = new FileInfo(DestPath);
-					FileInfo SrcInfo = new FileInfo(Path.Combine(SourceDir.FullName, RelativePath));
+					string SourcePath = Path.Combine(SourceDir.FullName, RelativePath);
+					FileInfo DestInfo;
+					FileInfo SrcInfo;
+
+					// wrap FileInfo creation with exception handler as can throw and want informative error
+					try
+					{
+						DestInfo = new FileInfo(DestPath);
+						SrcInfo = new FileInfo(SourcePath);
+					}
+					catch (Exception Ex)
+					{
+						throw new Exception(string.Format("FileInfo creation failed for Source:{0}, Dest:{1}, with: {2}", SourcePath, DestPath, Ex.Message));
+					}
 
 					// ensure directory exists
 					DestInfo.Directory.Create();
