Fix for Gauntlet copying Android artifacts to the wrong directory.

#rnx
#test RunFort -test=memorylistdiff -platform=android

diff --git a/Engine/Source/Programs/AutomationTool/Gauntlet/Platform/Android/Gauntlet.TargetDeviceAndroid.cs b/Engine/Source/Programs/AutomationTool/Gauntlet/Platform/Android/Gauntlet.TargetDeviceAndroid.cs
--- a/Engine/Source/Programs/AutomationTool/Gauntlet/Platform/Android/Gauntlet.TargetDeviceAndroid.cs
+++ b/Engine/Source/Programs/AutomationTool/Gauntlet/Platform/Android/Gauntlet.TargetDeviceAndroid.cs
@@ -210,7 +210,7 @@ namespace Gauntlet
 			Directory.CreateDirectory(LocalSaved);
 
 			// pull all the artifacts
-			string ArtifactPullCommand = string.Format("pull {0} {1}", Install.AndroidDevice.DeviceArtifactPath, LocalSaved);
+			string ArtifactPullCommand = string.Format("pull {0} {1}", Install.AndroidDevice.DeviceArtifactPath, Install.AndroidDevice.LocalCachePath);
 			IProcessResult PullCmd = Install.AndroidDevice.RunAdbDeviceCommand(ArtifactPullCommand);
 
 			if (PullCmd.ExitCode != 0)
