Suspend EC error parsing when logging file copy errors (may try on new kit, etc and shouldn't fail test)

diff --git a/Engine/Source/Programs/AutomationTool/Gauntlet/Framework/Gauntlet.Utils.cs b/Engine/Source/Programs/AutomationTool/Gauntlet/Framework/Gauntlet.Utils.cs
--- a/Engine/Source/Programs/AutomationTool/Gauntlet/Framework/Gauntlet.Utils.cs
+++ b/Engine/Source/Programs/AutomationTool/Gauntlet/Framework/Gauntlet.Utils.cs
@@ -933,7 +933,10 @@ namespace Gauntlet
 							}
 							else
 							{
-								Log.Error("File Copy failed with {0}.", ex.Message);
+								using (var PauseEC = new ScopedSuspendECErrorParsing())
+								{
+									Log.Error("File Copy failed with {0}.", ex.Message);
+								}
 								throw new Exception(string.Format("File Copy failed with {0}.", ex.Message));
 							}
 						}
