Gauntlet: Add support for new uuid format to iOS device

diff --git a/Engine/Source/Programs/AutomationTool/Gauntlet/Platform/IOS/Gauntlet.TargetDeviceIOS.cs b/Engine/Source/Programs/AutomationTool/Gauntlet/Platform/IOS/Gauntlet.TargetDeviceIOS.cs
--- a/Engine/Source/Programs/AutomationTool/Gauntlet/Platform/IOS/Gauntlet.TargetDeviceIOS.cs
+++ b/Engine/Source/Programs/AutomationTool/Gauntlet/Platform/IOS/Gauntlet.TargetDeviceIOS.cs
@@ -696,7 +696,7 @@ namespace Gauntlet
 				return new List<string>();
 			}
 
-			MatchCollection DeviceMatches = Regex.Matches(Result.Output, @"(.?)Found\ ([a-z0-9]{40})");
+			MatchCollection DeviceMatches = Regex.Matches(Result.Output, @"(.?)Found\ ([a-z0-9]{40}|[A-Z0-9]{8}-[A-Z0-9]{16})");
 
 			return DeviceMatches.Cast<Match>().Select<Match, string>(
 				M => M.Groups[2].ToString()
