#UE4 Add AdditionalPackageOptions option "ForceMaxIPACompression" for forcing max compression when packaging IPAs, even if they are not "distributed". Useful for reducing the size of "bulk" non-distributed ipas.

diff --git a/Engine/Source/Programs/AutomationTool/IOS/IOSPlatform.Automation.cs b/Engine/Source/Programs/AutomationTool/IOS/IOSPlatform.Automation.cs
--- a/Engine/Source/Programs/AutomationTool/IOS/IOSPlatform.Automation.cs
+++ b/Engine/Source/Programs/AutomationTool/IOS/IOSPlatform.Automation.cs
@@ -626,6 +626,28 @@ public class IOSPlatform : Platform
 		}
 	}
 
+	private bool ShouldUseMaxIPACompression(ProjectParams Params)
+	{
+		if (Params.Distribution)
+		{
+			return true;
+		}
+
+		if (!string.IsNullOrEmpty(Params.AdditionalPackageOptions))
+		{
+			string[] OptionsArray = Params.AdditionalPackageOptions.Split(' ');
+			foreach (string Option in OptionsArray)
+			{
+				if (Option.Equals("-ForceMaxIPACompression", StringComparison.InvariantCultureIgnoreCase))
+				{
+					return true;
+				}
+			}
+		}
+
+		return false;
+	}
+
 	private void PackageIPA(ProjectParams Params, string ProjectGameExeFilename, DeploymentContext SC)
 	{
 		string BaseDirectory = Path.GetDirectoryName(ProjectGameExeFilename);
@@ -661,7 +683,7 @@ public class IOSPlatform : Platform
             Zip.AlternateEncoding = Encoding.UTF8;
 
 			// set the compression level
-			if (Params.Distribution)
+			if (ShouldUseMaxIPACompression(Params))
 			{
 				Zip.CompressionLevel = CompressionLevel.BestCompression;
 			}
