Adding support for Gauntlet Test nodes to query what run they are on and how many total runs there are for the current test instance.
#rnx

diff --git a/Engine/Source/Programs/AutomationTool/Gauntlet/Unreal/Base/Gauntlet.UnrealTestNode.cs b/Engine/Source/Programs/AutomationTool/Gauntlet/Unreal/Base/Gauntlet.UnrealTestNode.cs
--- a/Engine/Source/Programs/AutomationTool/Gauntlet/Unreal/Base/Gauntlet.UnrealTestNode.cs
+++ b/Engine/Source/Programs/AutomationTool/Gauntlet/Unreal/Base/Gauntlet.UnrealTestNode.cs
@@ -854,6 +854,23 @@ namespace Gauntlet
 			return ExitCode;
 		}
 
+		/// <summary>
+		/// Returns the current Pass count for this node
+		/// </summary>
+		public int GetCurrentPass()
+		{
+			return CurrentPass;
+		}
+
+		/// <summary>
+		/// Returns the current total Pass count for this node
+		/// </summary>
+		/// <returns></returns>
+		public int GetNumPasses()
+		{
+			return NumPasses;
+		}
+
 		/// <summary>
 		/// Result of the test once completed. Nodes inheriting from us should override
 		/// this if custom results are necessary
