#UE4 Forcing max ipa compression now also compresses the executable.

diff --git a/Engine/Source/Programs/AutomationTool/IOS/IOSPlatform.Automation.cs b/Engine/Source/Programs/AutomationTool/IOS/IOSPlatform.Automation.cs
--- a/Engine/Source/Programs/AutomationTool/IOS/IOSPlatform.Automation.cs
+++ b/Engine/Source/Programs/AutomationTool/IOS/IOSPlatform.Automation.cs
@@ -628,11 +628,6 @@ public class IOSPlatform : Platform
 
 	private bool ShouldUseMaxIPACompression(ProjectParams Params)
 	{
-		if (Params.Distribution)
-		{
-			return true;
-		}
-
 		if (!string.IsNullOrEmpty(Params.AdditionalPackageOptions))
 		{
 			string[] OptionsArray = Params.AdditionalPackageOptions.Split(' ');
@@ -683,7 +678,8 @@ public class IOSPlatform : Platform
             Zip.AlternateEncoding = Encoding.UTF8;
 
 			// set the compression level
-			if (ShouldUseMaxIPACompression(Params))
+			bool bUseMaxIPACompression = ShouldUseMaxIPACompression(Params);
+			if (Params.Distribution || bUseMaxIPACompression)
 			{
 				Zip.CompressionLevel = CompressionLevel.BestCompression;
 			}
@@ -728,7 +724,7 @@ public class IOSPlatform : Platform
 					bIsExecutable = true;
 				}
 
-				if (bIsExecutable)
+				if (bIsExecutable && !bUseMaxIPACompression)
 				{
 					// The executable will be encrypted in the final distribution IPA and will compress very poorly, so keeping it
 					// uncompressed gives a better indicator of IPA size for our distro builds
