Fix for not having permissions to run the Install / Uninstall batch files for Lumin packages on Macs.  Execute permissions were missing

diff --git a/Engine/Source/Programs/AutomationTool/Lumin/LuminPlatform.Automation.cs b/Engine/Source/Programs/AutomationTool/Lumin/LuminPlatform.Automation.cs
--- a/Engine/Source/Programs/AutomationTool/Lumin/LuminPlatform.Automation.cs
+++ b/Engine/Source/Programs/AutomationTool/Lumin/LuminPlatform.Automation.cs
@@ -659,6 +659,13 @@ public class LuminPlatform : Platform
 		BatchLines = GenerateUninstallBatchFile(PackageName, Params);
 		File.WriteAllLines(UninstallBatchName, BatchLines);
 
+		// If needed, make the batch files able to execute
+		if (Utils.IsRunningOnMono)
+		{
+			CommandUtils.FixUnixFilePermissions(BatchName);
+			CommandUtils.FixUnixFilePermissions(UninstallBatchName);
+		}
+
 		PrintRunTime();
 	}
 
