Gauntlet packaged XB1 build source support

diff --git a/Engine/Source/Programs/AutomationTool/Gauntlet/Unreal/Base/Gauntlet.UnrealSession.cs b/Engine/Source/Programs/AutomationTool/Gauntlet/Unreal/Base/Gauntlet.UnrealSession.cs
--- a/Engine/Source/Programs/AutomationTool/Gauntlet/Unreal/Base/Gauntlet.UnrealSession.cs
+++ b/Engine/Source/Programs/AutomationTool/Gauntlet/Unreal/Base/Gauntlet.UnrealSession.cs
@@ -140,7 +140,7 @@ namespace Gauntlet
                 RequiredBuildFlags |= BuildFlags.Bulk;
             }
 
-			if (Globals.Params.ParseParam("packaged") && (InPlatform == UnrealTargetPlatform.Switch))
+			if (Globals.Params.ParseParam("packaged") && (InPlatform == UnrealTargetPlatform.Switch || InPlatform == UnrealTargetPlatform.XboxOne))
 			{
 				RequiredBuildFlags |= BuildFlags.Packaged;
 			}
