Allow multiple nodes to produce dxil.dll and dxcompiler.dll.

#rnx

diff --git a/Engine/Source/Programs/AutomationTool/BuildGraph/TempStorage.cs b/Engine/Source/Programs/AutomationTool/BuildGraph/TempStorage.cs
--- a/Engine/Source/Programs/AutomationTool/BuildGraph/TempStorage.cs
+++ b/Engine/Source/Programs/AutomationTool/BuildGraph/TempStorage.cs
@@ -1112,6 +1112,14 @@ namespace AutomationTool
 			{
 				return true;
 			}
+			if(FileName.Equals("dxil.dll", StringComparison.OrdinalIgnoreCase))
+			{
+				return true;
+			}
+			if(FileName.Equals("dxcompiler.dll", StringComparison.OrdinalIgnoreCase))
+			{
+				return true;
+			}
 			return false;
 		}
 	}
