Gauntlet: Adding reboot functionality to iOS device

diff --git a/Engine/Source/Programs/AutomationTool/Gauntlet/Platform/IOS/Gauntlet.TargetDeviceIOS.cs b/Engine/Source/Programs/AutomationTool/Gauntlet/Platform/IOS/Gauntlet.TargetDeviceIOS.cs
--- a/Engine/Source/Programs/AutomationTool/Gauntlet/Platform/IOS/Gauntlet.TargetDeviceIOS.cs
+++ b/Engine/Source/Programs/AutomationTool/Gauntlet/Platform/IOS/Gauntlet.TargetDeviceIOS.cs
@@ -625,7 +625,50 @@ namespace Gauntlet
 		public bool IsOn { get { return true; } }
 		public bool PowerOn() { return true; }
 		public bool PowerOff() { return true; }
-		public bool Reboot() { return true; }
+		
+		public bool Reboot() 
+		{ 
+			const string Cmd = "/usr/local/bin/idevicediagnostics";
+			if (!File.Exists(Cmd))
+			{
+				Log.Verbose("Rebooting iOS device requires idevicediagnostics binary");
+				return true;
+			}
+
+			var Result = IOSBuild.ExecuteCommand(Cmd, string.Format("restart -u {0}", DeviceName));
+			if (Result.ExitCode != 0)
+			{
+				Log.Warning(string.Format("Failed to reboot iOS device {0}, restart command failed", DeviceName));
+				return true;
+			}
+
+			// initial wait 20 seconds
+			Thread.Sleep(20 * 1000);
+
+			const int WaitPeriod = 10;
+			int WaitTime = 120;
+			bool rebooted = false;
+			do
+			{
+				Result = IOSBuild.ExecuteCommand(Cmd, string.Format("diagnostics WiFi -u {0}", DeviceName));
+				if (Result.ExitCode == 0)
+				{
+					rebooted = true;
+					break;
+				}
+				
+				Thread.Sleep(WaitPeriod * 1000);
+				WaitTime -= WaitPeriod;
+
+			} while (WaitTime > 0);
+
+			if (!rebooted) 
+			{
+				Log.Warning("Failed to reboot iOS device {0}, device didn't come back after restart", DeviceName);
+			}
+
+			return true; 
+		}
 
 		static Dictionary<string, bool> ConnectedDevices = new Dictionary<string, bool>();
 		bool Connected = false;
