Fixed an issue with Hyperlinks in HtmlBuilder.

#rnx

diff --git a/Engine/Source/Programs/AutomationTool/Gauntlet/Framework/Utils/Gauntlet.HtmlBuilder.cs b/Engine/Source/Programs/AutomationTool/Gauntlet/Framework/Utils/Gauntlet.HtmlBuilder.cs
--- a/Engine/Source/Programs/AutomationTool/Gauntlet/Framework/Utils/Gauntlet.HtmlBuilder.cs
+++ b/Engine/Source/Programs/AutomationTool/Gauntlet/Framework/Utils/Gauntlet.HtmlBuilder.cs
@@ -212,8 +212,7 @@ namespace Gauntlet
 		/// <returns></returns>
 		public HtmlBuilder Hyperlink(string URL, string Text)
 		{
-			EnsureEndsWithNewLine();
-			SB.AppendFormat("<a href=\"{0}\">{1}</h1>", Text);
+			SB.AppendFormat("<a href=\"{0}\">{1}</a>", URL, Text);
 			return this;
 		}
 
