Adding Gauntlet workerid parameter which workers can use for configuration

diff --git a/Engine/Source/Programs/AutomationTool/Gauntlet/Framework/Gauntlet.Utils.cs b/Engine/Source/Programs/AutomationTool/Gauntlet/Framework/Gauntlet.Utils.cs
--- a/Engine/Source/Programs/AutomationTool/Gauntlet/Framework/Gauntlet.Utils.cs
+++ b/Engine/Source/Programs/AutomationTool/Gauntlet/Framework/Gauntlet.Utils.cs
@@ -35,6 +35,32 @@ namespace Gauntlet
 		static List<Action> InnerPostAbortHandlers = new List<Action>();
 		public static bool CancelSignalled { get; private set; }
 
+		/// <summary>
+		/// Get the worker id of this Gauntlet instance
+		/// returns -1 if instance is not a member of a worker group
+		/// </summary>
+		public static int WorkerID
+		{
+			get
+			{
+				int Default = -1;
+				return Params.ParseValue("workerid", Default);
+			}
+
+		}
+
+		/// <summary>
+		/// Returns true if Gauntlet instance is a member of a worker group
+		/// </summary>
+		public static bool IsWorker
+		{
+			get
+			{
+				return WorkerID != -1;
+			}
+		}
+
+
 		public static string TempDir
 		{
 			get
diff --git a/Engine/Source/Programs/AutomationTool/Gauntlet/Platform/IOS/Gauntlet.TargetDeviceIOS.cs b/Engine/Source/Programs/AutomationTool/Gauntlet/Platform/IOS/Gauntlet.TargetDeviceIOS.cs
--- a/Engine/Source/Programs/AutomationTool/Gauntlet/Platform/IOS/Gauntlet.TargetDeviceIOS.cs
+++ b/Engine/Source/Programs/AutomationTool/Gauntlet/Platform/IOS/Gauntlet.TargetDeviceIOS.cs
@@ -752,7 +752,7 @@ namespace Gauntlet
 		void KillZombies()
 		{
 
-			if (ZombiesKilled)
+			if (Globals.IsWorker || ZombiesKilled)
 			{
 				return;
 			}
