Add some FileReference/DirectoryReference overloads to CommandUtils.

diff --git a/Engine/Source/Programs/AutomationTool/AutomationUtils/CommandUtils.cs b/Engine/Source/Programs/AutomationTool/AutomationUtils/CommandUtils.cs
--- a/Engine/Source/Programs/AutomationTool/AutomationUtils/CommandUtils.cs
+++ b/Engine/Source/Programs/AutomationTool/AutomationUtils/CommandUtils.cs
@@ -589,6 +589,16 @@ namespace AutomationTool
 			DeleteDirectory(false, Directory);
 		}
 
+		/// <summary>
+		/// Deletes a directory(or directories) including its contents (recursively, will delete read-only files).
+		/// If the deletion of the directory fails, this function throws an Exception.
+		/// </summary>
+        /// <param name="Directories">Directories</param>
+        public static void DeleteDirectory(DirectoryReference Directory)
+		{
+			DeleteDirectory(Directory.FullName);
+		}
+
 		/// <summary>
 		/// Deletes a directory(or directories) including its contents (recursively, will delete read-only files).
 		/// If the deletion of the directory fails, prints a warning.
@@ -654,6 +664,17 @@ namespace AutomationTool
 			}
 		}
 
+		/// <summary>
+		/// Attempts to delete a directory, if that fails deletes all files and folder from the specified directory.
+		/// This works around the issue when the user has a file open in a notepad from that directory. Somehow deleting the file works but
+		/// deleting the directory with the file that's open, doesn't.
+		/// </summary>
+		/// <param name="DirectoryName"></param>
+		public static void DeleteDirectoryContents(DirectoryReference DirectoryName)
+		{
+			DeleteDirectoryContents(DirectoryName.FullName);
+		}
+
 		/// <summary>
 		/// Checks if a directory(or directories) exists.
 		/// </summary>
@@ -1639,6 +1660,11 @@ namespace AutomationTool
 			ThreadedCopyFiles(SourceFiles, DestFiles, MaxThreads);
 		}
 
+		public static void ThreadedCopyFiles(DirectoryReference SourceDirectory, DirectoryReference DestDirectory, int MaxThreads = 64)
+		{
+			ThreadedCopyFiles(SourceDirectory.FullName, DestDirectory.FullName, MaxThreads);
+		}
+
 		/// <summary>
 		/// Copies files using multiple threads
         /// </summary>
diff --git a/Engine/Source/Programs/AutomationTool/AutomationUtils/CommandletUtils.cs b/Engine/Source/Programs/AutomationTool/AutomationUtils/CommandletUtils.cs
--- a/Engine/Source/Programs/AutomationTool/AutomationUtils/CommandletUtils.cs
+++ b/Engine/Source/Programs/AutomationTool/AutomationUtils/CommandletUtils.cs
@@ -163,7 +163,7 @@ namespace AutomationTool
         /// <param name="Maps">List of maps to cook, can be null in which case -MapIniSection=AllMaps is used.</param>
         /// <param name="TargetPlatform">Target platform.</param>
         /// <param name="Parameters">List of additional parameters.</param>
-        public static List<string> GenerateDistillFileSetsCommandlet(FileReference ProjectName, string ManifestFile, string UE4Exe = "UE4Editor-Cmd.exe", string[] Maps = null, string Parameters = "")
+        public static List<FileReference> GenerateDistillFileSetsCommandlet(FileReference ProjectName, string ManifestFile, string UE4Exe = "UE4Editor-Cmd.exe", string[] Maps = null, string Parameters = "")
         {
             string MapsToCook = "";
             if (!IsNullOrEmpty(Maps))
@@ -193,7 +193,7 @@ namespace AutomationTool
             {
                 throw new AutomationException("GenerateDistillFileSets for {0} did not produce any files.", ProjectName);
             }
-            var Result = new List<string>();
+            var Result = new List<FileReference>();
             foreach (var ThisFile in Lines)
             {
                 var TestFile = CombinePaths(ThisFile);
@@ -203,8 +203,8 @@ namespace AutomationTool
                 }
                 // we correct the case here
                 var TestFileInfo = new FileInfo(TestFile);
-                var FinalFile = CombinePaths(TestFileInfo.FullName);
-                if (!FileExists_NoExceptions(FinalFile))
+                var FinalFile = new FileReference(TestFileInfo);
+                if (!FileReference.Exists(FinalFile))
                 {
                     throw new AutomationException("GenerateDistillFileSets produced {0}, but {1} doesn't exist.", ThisFile, FinalFile);
                 }
diff --git a/Engine/Source/Programs/AutomationTool/Scripts/AutomationScripts.Automation.csproj b/Engine/Source/Programs/AutomationTool/Scripts/AutomationScripts.Automation.csproj
--- a/Engine/Source/Programs/AutomationTool/Scripts/AutomationScripts.Automation.csproj
+++ b/Engine/Source/Programs/AutomationTool/Scripts/AutomationScripts.Automation.csproj
@@ -86,7 +86,6 @@
     <Compile Include="CopyBuildToStagingDirectory.Automation.cs" />
     <Compile Include="CopySharedCookedBuild.Automation.cs" />
     <Compile Include="CopyUAT.Automation.cs" />
-    <Compile Include="CopyUsingDistillFileSet.Automation.cs" />
     <Compile Include="CryptoKeys.Automation.cs" />
     <Compile Include="DeployCommand.Automation.cs" />
     <Compile Include="ExtractPaks.Automation.cs" />
diff --git a/Engine/Source/Programs/AutomationTool/Scripts/CopyUsingDistillFileSet.Automation.cs b/Engine/Source/Programs/AutomationTool/Scripts/CopyUsingDistillFileSet.Automation.cs
deleted file mode 100644
