Gauntlet: release devices back to pool once artifacts have been saved

diff --git a/Engine/Source/Programs/AutomationTool/Gauntlet/Unreal/Base/Gauntlet.UnrealPGONode.cs b/Engine/Source/Programs/AutomationTool/Gauntlet/Unreal/Base/Gauntlet.UnrealPGONode.cs
--- a/Engine/Source/Programs/AutomationTool/Gauntlet/Unreal/Base/Gauntlet.UnrealPGONode.cs
+++ b/Engine/Source/Programs/AutomationTool/Gauntlet/Unreal/Base/Gauntlet.UnrealPGONode.cs
@@ -93,11 +93,34 @@ namespace Gauntlet
 		private const float ScreenshotScale = 1.0f / 3.0f;
 		private const int ScreenshotQuality = 30;
 		protected string ScreenshotDirectory;
+		protected bool ResultsGathered = false;
 
 		public override void TickTest()
 		{
 			base.TickTest();
 
+			if (GetTestStatus() == TestStatus.Complete || ResultsGathered)
+			{
+				if (!ResultsGathered)
+				{
+					ResultsGathered = true;
+
+					try
+					{
+						// Gather results and merge PGO data
+						Log.Info("Gathering profiling results to {0}", TestInstance.ClientApps[0].ArtifactPath);
+						PGOPlatform.GatherResults(TestInstance.ClientApps[0].ArtifactPath);
+					}
+					catch (Exception Ex)
+					{
+						Log.Error("Error getting PGO results: {0}", Ex);
+					}
+
+				}
+
+				return;
+			}
+
 			// Handle device screenshot update
 			TimeSpan Delta = DateTime.Now - ScreenshotTime;
 			ITargetDevice Device = TestInstance.ClientApps[0].Device;
@@ -126,18 +149,6 @@ namespace Gauntlet
 			}
 
 		}
-
-		public override void CreateReport(TestResult Result, UnrealTestContext Context, UnrealBuildSource Build, IEnumerable<UnrealRoleArtifacts> Artifacts, string ArtifactPath)
-		{
-			if (Result != TestResult.Passed)
-			{
-				return;
-			}
-
-			// Gather results and merge PGO data
-			Log.Info("Gathering profiling results...");
-			PGOPlatform.GatherResults(TestInstance.ClientApps[0].ArtifactPath);
-		}
 	}
 
 	internal interface IPGOPlatform
diff --git a/Engine/Source/Programs/AutomationTool/Gauntlet/Unreal/Base/Gauntlet.UnrealSession.cs b/Engine/Source/Programs/AutomationTool/Gauntlet/Unreal/Base/Gauntlet.UnrealSession.cs
--- a/Engine/Source/Programs/AutomationTool/Gauntlet/Unreal/Base/Gauntlet.UnrealSession.cs
+++ b/Engine/Source/Programs/AutomationTool/Gauntlet/Unreal/Base/Gauntlet.UnrealSession.cs
@@ -494,7 +494,7 @@ namespace Gauntlet
 		}
 		#endregion
 
-		void ReleaseDevices()
+		public void ReleaseDevices()
 		{
 			if (ReservedDevices.Count() > 0)
 			{
diff --git a/Engine/Source/Programs/AutomationTool/Gauntlet/Unreal/Base/Gauntlet.UnrealTestNode.cs b/Engine/Source/Programs/AutomationTool/Gauntlet/Unreal/Base/Gauntlet.UnrealTestNode.cs
--- a/Engine/Source/Programs/AutomationTool/Gauntlet/Unreal/Base/Gauntlet.UnrealTestNode.cs
+++ b/Engine/Source/Programs/AutomationTool/Gauntlet/Unreal/Base/Gauntlet.UnrealTestNode.cs
@@ -568,6 +568,16 @@ namespace Gauntlet
 				Log.Warning("Failed to save artifacts. {0}", Ex);
 			}
 
+			try
+			{
+				// Artifacts have been saved, release devices back to pool for other tests to use
+				UnrealApp.ReleaseDevices();
+			}
+			catch (Exception Ex)
+			{
+				Log.Warning("Failed to release devices. {0}", Ex);
+			}
+
 			try
 			{
 				CreateReport(GetTestResult(), Context, Context.BuildInfo, SessionArtifacts, OutputPath);
