Unshelved from Hongyi.Yu

Fix numbering for patch paks when skipping a patch tier on new files

diff --git a/Engine/Source/Programs/AutomationTool/Scripts/CopyBuildToStagingDirectory.Automation.cs b/Engine/Source/Programs/AutomationTool/Scripts/CopyBuildToStagingDirectory.Automation.cs
--- a/Engine/Source/Programs/AutomationTool/Scripts/CopyBuildToStagingDirectory.Automation.cs
+++ b/Engine/Source/Programs/AutomationTool/Scripts/CopyBuildToStagingDirectory.Automation.cs
@@ -1736,17 +1736,20 @@ public partial class Project : CommandUtils
 		List<string> Commands = new List<string>();
 		List<string> LogNames = new List<string>();
 
-		List<Tuple<FileReference, StagedFileReference, string>> Outputs = new List<Tuple<FileReference, StagedFileReference, string>>();
-		foreach (CreatePakParams PakParams in PakParamsList)
+		// Calculate target patch index by iterating all pak files in source build
+		int TargetPatchIndex = 0;
+		int NumPakFiles = 0;
+		if (bShouldGeneratePatch)
 		{
-			string OutputFilename = PakParams.PakName + "-" + SC.FinalCookPlatform;
-			if (bShouldGeneratePatch)
+			foreach (CreatePakParams PakParams in PakParamsList)
 			{
-				int TargetPatchIndex = 0;
+				string OutputFilename = PakParams.PakName + "-" + SC.FinalCookPlatform;
 				string ExistingPatchSearchPath = SC.StageTargetPlatform.GetReleasePakFilePath(SC, Params, null);
 				if (Directory.Exists(ExistingPatchSearchPath))
 				{
 					IEnumerable<string> PakFileSet = Directory.EnumerateFiles(ExistingPatchSearchPath, OutputFilename + "*" + PostFix + OutputFilenameExtension);
+					NumPakFiles += PakFileSet.Count();
+
 					foreach (string PakFilePath in PakFileSet)
 					{
 						string PakFileName = Path.GetFileName(PakFilePath);
@@ -1765,11 +1768,21 @@ public partial class Project : CommandUtils
 							}
 						}
 					}
-					if (Params.ShouldAddPatchLevel && PakFileSet.Count() > 0)
-					{
-						TargetPatchIndex++;
-					}
 				}
+			}
+
+			if (Params.ShouldAddPatchLevel && NumPakFiles > 0)
+			{
+				TargetPatchIndex++;
+			}
+		}
+
+		List<Tuple<FileReference, StagedFileReference, string>> Outputs = new List<Tuple<FileReference, StagedFileReference, string>>();
+		foreach (CreatePakParams PakParams in PakParamsList)
+		{
+			string OutputFilename = PakParams.PakName + "-" + SC.FinalCookPlatform;
+			if (bShouldGeneratePatch)
+			{
 				OutputFilename = OutputFilename + "_" + TargetPatchIndex;
 			}
 			OutputFilename = OutputFilename + PostFix;
