Added ArtifactPath field to UnrealTestNode so that tests keep a reference to this file path. Create this path string within StartTest instead of StopTest so that it is populated sooner.

#rnx

diff --git a/Engine/Source/Programs/AutomationTool/Gauntlet/Unreal/Base/Gauntlet.UnrealTestNode.cs b/Engine/Source/Programs/AutomationTool/Gauntlet/Unreal/Base/Gauntlet.UnrealTestNode.cs
--- a/Engine/Source/Programs/AutomationTool/Gauntlet/Unreal/Base/Gauntlet.UnrealTestNode.cs
+++ b/Engine/Source/Programs/AutomationTool/Gauntlet/Unreal/Base/Gauntlet.UnrealTestNode.cs
@@ -84,6 +84,11 @@ namespace Gauntlet
 
 		static protected DateTime SessionStartTime = DateTime.MinValue;
 
+		/// <summary>
+		/// Path to the directory that logs and other artifacts are copied to after the test run.
+		/// </summary>
+		protected string ArtifactPath { get; private set; }
+
 		/// <summary>
 		/// Our test result. May be set directly, or by overriding GetUnrealTestResult()
 		/// </summary>
@@ -91,8 +96,6 @@ namespace Gauntlet
 
 		protected TConfigClass CachedConfig = null;
 
-		private string CachedArtifactPath = null;
-
 		/// <summary>
 		/// If our test should exit suddenly, this is the process that caused it
 		/// </summary>
@@ -115,6 +118,7 @@ namespace Gauntlet
 			LastLogCount = 0;
 			CurrentPass = 0;
 			NumPasses = 0;
+			ArtifactPath = string.Empty;
 		}
 
 		 ~UnrealTestNode()
@@ -411,7 +415,33 @@ namespace Gauntlet
 			TConfigClass Config = GetCachedConfiguration();
 
 			CurrentPass = Pass;
-			NumPasses = InNumPasses;			
+			NumPasses = InNumPasses;
+
+			string TestFolder = ToString();
+			TestFolder = TestFolder.Replace(" ", "_");
+			TestFolder = TestFolder.Replace(",", "");
+
+			ArtifactPath = Path.Combine(Context.Options.LogDir, TestFolder);
+
+			// if doing multiple passes, put each in a subdir
+			if (NumPasses > 1)
+			{
+				ArtifactPath = Path.Combine(ArtifactPath, string.Format("Pass_{0}_of_{1}", CurrentPass, NumPasses));
+			}
+
+			// Basic pre-existing directory check.
+			if (CommandUtils.IsBuildMachine && Directory.Exists(ArtifactPath))
+			{
+				string NewOutputPath = ArtifactPath;
+				int i = 0;
+				while (Directory.Exists(NewOutputPath))
+				{
+					i++;
+					NewOutputPath = string.Format("{0}_{1}", ArtifactPath, i);
+				}
+				Log.Info("Directory already exists at {0}", ArtifactPath);
+				ArtifactPath = NewOutputPath;
+			}
 
 			// Launch the test
 			TestInstance = UnrealApp.LaunchSession();
@@ -526,42 +556,16 @@ namespace Gauntlet
 			// access to these objects and their resources! Final cleanup is done in CleanupTest()
 			TestInstance.Shutdown();
 
-			//string TestFolder = string.Format("{0}-{1:yyyy.MM.dd-HH.mm}", Name, SessionStartTime);
-			string TestFolder = ToString();
-			TestFolder = TestFolder.Replace(" ", "_");
-			TestFolder = TestFolder.Replace(",", "");
-
-			string OutputPath = Path.Combine(Context.Options.LogDir, TestFolder);
-
-			// if doing multiple passes, put each in a subdir
-			if (NumPasses > 1)
-			{
-				OutputPath = Path.Combine(OutputPath, string.Format("Pass_{0}_of_{1}", CurrentPass, NumPasses));
-			}
-
 			try
 			{
-				// Basic pre-existing directory check.
-				if (CommandUtils.IsBuildMachine && Directory.Exists(OutputPath))
-				{
-					string NewOutputPath = OutputPath;
-					int i = 0;
-					while (Directory.Exists(NewOutputPath))
-					{
-						i++;
-						NewOutputPath = string.Format("{0}_{1}", OutputPath, i);
-					}
-					Log.Info("Directory already exists at {0}", OutputPath);
-					OutputPath = NewOutputPath;
-				}
-				Log.Info("Saving artifacts to {0}", OutputPath);
-				Directory.CreateDirectory(OutputPath);
-				Utils.SystemHelpers.MarkDirectoryForCleanup(OutputPath);
+				Log.Info("Saving artifacts to {0}", ArtifactPath);
+				Directory.CreateDirectory(ArtifactPath);
+				Utils.SystemHelpers.MarkDirectoryForCleanup(ArtifactPath);
 
-				SessionArtifacts = SaveRoleArtifacts(OutputPath);
+				SessionArtifacts = SaveRoleArtifacts(ArtifactPath);
 
 				// call legacy version
-				SaveArtifacts_DEPRECATED(OutputPath);
+				SaveArtifacts_DEPRECATED(ArtifactPath);
 			}
 			catch (Exception Ex)
 			{
@@ -580,7 +584,7 @@ namespace Gauntlet
 
 			try
 			{
-				CreateReport(GetTestResult(), Context, Context.BuildInfo, SessionArtifacts, OutputPath);
+				CreateReport(GetTestResult(), Context, Context.BuildInfo, SessionArtifacts, ArtifactPath);
 			}
 			catch (Exception Ex)
 			{
@@ -589,7 +593,7 @@ namespace Gauntlet
 
 			try
 			{
-				SubmitToDashboard(GetTestResult(), Context, Context.BuildInfo, SessionArtifacts, OutputPath);
+				SubmitToDashboard(GetTestResult(), Context, Context.BuildInfo, SessionArtifacts, ArtifactPath);
 			}
 			catch (Exception Ex)
 			{
@@ -640,8 +644,7 @@ namespace Gauntlet
 		/// <returns></returns>
 		public virtual IEnumerable<UnrealRoleArtifacts> SaveRoleArtifacts(string OutputPath)
 		{
-			CachedArtifactPath = OutputPath;
-			return UnrealApp.SaveRoleArtifacts(Context, TestInstance, CachedArtifactPath);
+			return UnrealApp.SaveRoleArtifacts(Context, TestInstance, ArtifactPath);
 		}
 
 		/// <summary>
