Allow for selection of shared cooked builds within the most recent code CL

diff --git a/Engine/Source/Programs/AutomationTool/Scripts/SharedCookedBuild.cs b/Engine/Source/Programs/AutomationTool/Scripts/SharedCookedBuild.cs
--- a/Engine/Source/Programs/AutomationTool/Scripts/SharedCookedBuild.cs
+++ b/Engine/Source/Programs/AutomationTool/Scripts/SharedCookedBuild.cs
@@ -15,20 +15,32 @@ public class SharedCookedBuild
 {
 	private static Task CopySharedCookedBuildTask = null;
 
-	private static bool FindBestSharedCookedBuild(ref string FinalCookedBuildPath, string ProjectFullPath, UnrealTargetPlatform TargetPlatform, string CookPlatform, string SharedCookedBuildCL )
+	public enum SharedCookType
+	{
+		Exact,
+		RecentContentOnly,
+		Recent
+	}
+
+	public static bool FindBestSharedCookedBuild(ref string FinalCookedBuildPath, string ProjectFullPath, UnrealTargetPlatform TargetPlatform, string CookPlatform, SharedCookType BuildType)
 	{
 		string BuildRoot = CommandUtils.P4Enabled ? CommandUtils.P4Env.Branch.Replace("/", "+") : "";
 		int CurrentCLInt = CommandUtils.P4Enabled ? CommandUtils.P4Env.Changelist : 0;
+		int CurrentCodeCLInt = CommandUtils.P4Enabled ? CommandUtils.P4Env.CodeChangelist : 0;
+
 
 		BuildVersion Version;
 		if (BuildVersion.TryRead(BuildVersion.GetDefaultFileName(), out Version))
 		{
 			CurrentCLInt = Version.Changelist;
+			CurrentCodeCLInt = Version.EffectiveCompatibleChangelist;
 			BuildRoot = Version.BranchName;
 		}
 		System.GC.Collect();
 		string CurrentCL = CurrentCLInt.ToString();
 
+		CommandUtils.LogInformation("Search for best shared cooked build...");
+		CommandUtils.LogInformation("CurrentCL: {0}, Type: {1}", CurrentCodeCLInt, BuildType);
 
 		FileReference ProjectFileRef = new FileReference(ProjectFullPath);
 		// get network location 
@@ -43,7 +55,7 @@ public class SharedCookedBuild
 		const string MetaDataFilename = "\\Metadata\\DevelopmentAssetRegistry.bin";
 
 
-		if (SharedCookedBuildCL == "usesyncedbuild")
+		if (BuildType == SharedCookType.Exact)
 		{
 			foreach (string CookedBuildPath in CookedBuildPaths)
 			{
@@ -74,13 +86,10 @@ public class SharedCookedBuild
 				}
 			}
 		}
-		else if (SharedCookedBuildCL == "userecentbuild")
+		else if (BuildType == SharedCookType.Recent || BuildType == SharedCookType.RecentContentOnly)
 		{
-
 			// build our CookedBUildPath into a regex which we can execute on the directories and extract info from
 
-
-
 			string BestBuild = null;
 			int BestCLNumber = 0;
 
@@ -141,15 +150,23 @@ public class SharedCookedBuild
 						}
 					}
 				}
-				
 			}
 
 			if ( string.IsNullOrEmpty(BestBuild) )
 			{
+				CommandUtils.LogError("Cannot find a recent shared cooked build");
+				return false;
+			}
+
+			if (BuildType == SharedCookType.RecentContentOnly && BestCLNumber < CurrentCodeCLInt)
+			{
+				CommandUtils.LogError("Cannot find a recent shared cooked build without code changes");
+				CommandUtils.LogError("Your CL: {0}, Last code CL: {1}, Best fit: {2}", CurrentCLInt, CurrentCodeCLInt, BestCLNumber);
 				return false;
 			}
 
 			FinalCookedBuildPath = BestBuild;
+			CommandUtils.LogInformation("Selected build: {0}");
 			return true;
 		}
 
@@ -241,22 +258,32 @@ public class SharedCookedBuild
 		return true;
 	}
 
-	public static void CopySharedCookedBuildForTarget(string ProjectFullPath, UnrealTargetPlatform TargetPlatform, string CookPlatform, string BuildCl, bool bOnlyCopyAssetRegistry = false)
+	public static void CopySharedCookedBuildForTarget(string ProjectFullPath, UnrealTargetPlatform TargetPlatform, string CookPlatform, string BuildCL, bool bOnlyCopyAssetRegistry = false)
 	{
 		var LocalPath = CommandUtils.CombinePaths(Path.GetDirectoryName(ProjectFullPath), "Saved", "SharedIterativeBuild", CookPlatform);
-		
+
+		SharedCookType SharedBuildType = SharedCookType.Exact;
+		if (BuildCL.Equals("userecentbuild"))
+		{
+			SharedBuildType = SharedCookType.Recent;
+		}
+		else if (BuildCL.Equals("userecentcontentonly"))
+		{
+			SharedBuildType = SharedCookType.RecentContentOnly;
+		}
+
 		string CookedBuildPath = null;
-		if ( FindBestSharedCookedBuild(ref CookedBuildPath, ProjectFullPath, TargetPlatform, CookPlatform, BuildCl) )
+		if ( FindBestSharedCookedBuild(ref CookedBuildPath, ProjectFullPath, TargetPlatform, CookPlatform, SharedBuildType) )
 		{
 			CopySharedCookedBuildForTargetInternal(CookedBuildPath, CookPlatform, LocalPath, bOnlyCopyAssetRegistry);
 		}
-		else if( FindBestSharedCookedBuild(ref CookedBuildPath, ProjectFullPath, TargetPlatform, TargetPlatform.ToString(), BuildCl) )
+		else if( FindBestSharedCookedBuild(ref CookedBuildPath, ProjectFullPath, TargetPlatform, TargetPlatform.ToString(), SharedBuildType) )
 		{
 			CopySharedCookedBuildForTargetInternal(CookedBuildPath, TargetPlatform.ToString(), LocalPath, bOnlyCopyAssetRegistry);
 		}
 		else
 		{
-			throw new AutomationException("Can't find cooked build for {0} {1} {2}", BuildCl, CookPlatform, TargetPlatform.ToString());
+			throw new AutomationException("Can't find cooked build for {0} {1} {2}", SharedBuildType,  CookPlatform, TargetPlatform.ToString());
 		}
 
 		return;
