Allow iterate shared cooked builds to re-stage non-UFS files over the base build

diff --git a/Engine/Source/Programs/AutomationTool/Scripts/CopyBuildToStagingDirectory.Automation.cs b/Engine/Source/Programs/AutomationTool/Scripts/CopyBuildToStagingDirectory.Automation.cs
--- a/Engine/Source/Programs/AutomationTool/Scripts/CopyBuildToStagingDirectory.Automation.cs
+++ b/Engine/Source/Programs/AutomationTool/Scripts/CopyBuildToStagingDirectory.Automation.cs
@@ -891,7 +891,21 @@ public partial class Project : CommandUtils
 		}
 
 		// Remap all the non-ufs files if not using a PAK file
-		SC.FilesToStage.NonUFSFiles = SC.FilesToStage.NonUFSFiles.ToDictionary(x => SC.StageTargetPlatform.Remap(x.Key), x => x.Value);
+		if (Params.HasIterateSharedCookedBuild)
+		{
+			// Shared NonUFS files are staged in their remapped location, and may be duplicated in the to-stage list.
+			Dictionary<StagedFileReference, FileReference> NonUFSToStage = new Dictionary<StagedFileReference, FileReference>();
+			foreach (KeyValuePair<StagedFileReference, FileReference> StagedFilePair in SC.FilesToStage.NonUFSFiles)
+			{
+				NonUFSToStage[SC.StageTargetPlatform.Remap(StagedFilePair.Key)] = StagedFilePair.Value;
+			}
+			SC.FilesToStage.NonUFSFiles = NonUFSToStage;
+		}
+		else
+		{
+			SC.FilesToStage.NonUFSFiles = SC.FilesToStage.NonUFSFiles.ToDictionary(x => SC.StageTargetPlatform.Remap(x.Key), x => x.Value);
+		}
+
 		if (!Params.UsePak(SC.StageTargetPlatform))
 		{
 			SC.FilesToStage.UFSFiles = SC.FilesToStage.UFSFiles.ToDictionary(x => SC.StageTargetPlatform.Remap(x.Key), x => x.Value);
