Add ability for multiple server roles within one test to load different maps.

diff --git a/Engine/Source/Programs/AutomationTool/Gauntlet/Unreal/Base/Gauntlet.UnrealSession.cs b/Engine/Source/Programs/AutomationTool/Gauntlet/Unreal/Base/Gauntlet.UnrealSession.cs
--- a/Engine/Source/Programs/AutomationTool/Gauntlet/Unreal/Base/Gauntlet.UnrealSession.cs
+++ b/Engine/Source/Programs/AutomationTool/Gauntlet/Unreal/Base/Gauntlet.UnrealSession.cs
@@ -59,6 +59,11 @@ namespace Gauntlet
 		/// </summary>
 		public string CommandLine;
 
+		/// <summary>
+		/// Map override to use on a server in case we don't want them all running the same map.
+		/// </summary>
+		public string MapOverride;
+
 		/// <summary>
 		/// List of files to copy to the device.
 		/// </summary>
@@ -116,10 +121,11 @@ namespace Gauntlet
 
 			Platform = InPlatform;
 			Configuration = InConfiguration;
+			MapOverride = string.Empty;
 
 			if (string.IsNullOrEmpty(InCommandLine))
 			{
-				CommandLine = "";
+				CommandLine = string.Empty;
 			}
 			else
 			{
diff --git a/Engine/Source/Programs/AutomationTool/Gauntlet/Unreal/Base/Gauntlet.UnrealTestConfiguration.cs b/Engine/Source/Programs/AutomationTool/Gauntlet/Unreal/Base/Gauntlet.UnrealTestConfiguration.cs
--- a/Engine/Source/Programs/AutomationTool/Gauntlet/Unreal/Base/Gauntlet.UnrealTestConfiguration.cs
+++ b/Engine/Source/Programs/AutomationTool/Gauntlet/Unreal/Base/Gauntlet.UnrealTestConfiguration.cs
@@ -66,8 +66,9 @@ namespace Gauntlet
 		{
 			Type = InType;
             PlatformOverride = InPlatformOverride;
-			CommandLine = "";
-			ExplicitClientCommandLine = "";
+			CommandLine = string.Empty;
+			MapOverride = string.Empty;
+			ExplicitClientCommandLine = string.Empty;
 			Controllers = new List<string>();
             FilesToCopy = new List<UnrealFileToCopy>();
             RoleType = ERoleModifier.None;
@@ -103,6 +104,11 @@ namespace Gauntlet
 
         public List<UnrealFileToCopy> FilesToCopy { get; set; }
 
+		/// <summary>
+		/// A map value passed in per server in case a test needs multiple servers on different maps.
+		/// </summary>
+		public string MapOverride { get; set; }
+
 		/// <summary>
 		/// Role device configuration 
 		/// </summary>
diff --git a/Engine/Source/Programs/AutomationTool/Gauntlet/Unreal/Base/Gauntlet.UnrealTestNode.cs b/Engine/Source/Programs/AutomationTool/Gauntlet/Unreal/Base/Gauntlet.UnrealTestNode.cs
--- a/Engine/Source/Programs/AutomationTool/Gauntlet/Unreal/Base/Gauntlet.UnrealTestNode.cs
+++ b/Engine/Source/Programs/AutomationTool/Gauntlet/Unreal/Base/Gauntlet.UnrealTestNode.cs
@@ -398,6 +398,7 @@ namespace Gauntlet
 					// copy over relevant settings from test role
                     SessionRole.FilesToCopy = TestRole.FilesToCopy;
 					SessionRole.ConfigureDevice = TestRole.ConfigureDevice;
+					SessionRole.MapOverride = TestRole.MapOverride;
 
 					SessionRoles.Add(SessionRole);
 				}
diff --git a/Engine/Source/Programs/AutomationTool/Gauntlet/Unreal/Game/UE4Game.TestConfig.cs b/Engine/Source/Programs/AutomationTool/Gauntlet/Unreal/Game/UE4Game.TestConfig.cs
--- a/Engine/Source/Programs/AutomationTool/Gauntlet/Unreal/Game/UE4Game.TestConfig.cs
+++ b/Engine/Source/Programs/AutomationTool/Gauntlet/Unreal/Game/UE4Game.TestConfig.cs
@@ -29,7 +29,17 @@ namespace UE4Game
 		{
 			base.ApplyToConfig(AppConfig, ConfigRole, OtherRoles);
 
-			if (string.IsNullOrEmpty(Map) == false)
+			if (string.IsNullOrEmpty(ConfigRole.MapOverride) == false)
+			{
+				// If map is specified, pass it to the server, or just the client if there is no server
+				if (AppConfig.ProcessType.IsServer()
+				|| (AppConfig.ProcessType.IsClient() && RoleCount(UnrealTargetRole.Server) == 0))
+				{
+					// must be the first argument!
+					AppConfig.CommandLine = ConfigRole.MapOverride + " " + AppConfig.CommandLine;
+				}
+			}
+			else if (string.IsNullOrEmpty(Map) == false)
 			{
 				// If map is specified, pass it to the server, or just the client if there is no server
 				if (AppConfig.ProcessType.IsServer()
