HTML5 - fix log spam

this fixes a bug when packaging from windows

#rnx

diff --git a/Engine/Source/Programs/AutomationTool/HTML5/HTML5Platform.Automation.cs b/Engine/Source/Programs/AutomationTool/HTML5/HTML5Platform.Automation.cs
--- a/Engine/Source/Programs/AutomationTool/HTML5/HTML5Platform.Automation.cs
+++ b/Engine/Source/Programs/AutomationTool/HTML5/HTML5Platform.Automation.cs
@@ -398,12 +398,14 @@ public class HTML5Platform : Platform
 
 				if (LineStr.Contains("%EMSDK_VERSION%"))
 				{
-					LineStr = LineStr.Replace("%EMSDK_VERSION%", (ProjectConfiguration == "Shipping") ? HTML5SDKInfo.EmscriptenVersion() : HTML5SDKInfo.EMSCRIPTEN_ROOT);
+					string escpath = HTML5SDKInfo.EMSCRIPTEN_ROOT.Replace("\\", "/");
+					LineStr = LineStr.Replace("%EMSDK_VERSION%", (ProjectConfiguration == "Shipping") ? HTML5SDKInfo.EmscriptenVersion() : escpath);
 				}
 
 				if (LineStr.Contains("%EMSDK_CONFIG%"))
 				{
-					LineStr = LineStr.Replace("%EMSDK_CONFIG%", (ProjectConfiguration == "Shipping") ? "" : HTML5SDKInfo.DOT_EMSCRIPTEN);
+					string escpath = HTML5SDKInfo.DOT_EMSCRIPTEN.Replace("\\", "/");
+					LineStr = LineStr.Replace("%EMSDK_CONFIG%", (ProjectConfiguration == "Shipping") ? "" : escpath);
 				}
 
 				outputContents.AppendLine(LineStr);
