- Fixed the order of UnrealPak parameters, with the parallel pak file creation (required so UnrealPak can load project-enabled plugins)

diff --git a/Engine/Source/Programs/AutomationTool/Scripts/CopyBuildToStagingDirectory.Automation.cs b/Engine/Source/Programs/AutomationTool/Scripts/CopyBuildToStagingDirectory.Automation.cs
--- a/Engine/Source/Programs/AutomationTool/Scripts/CopyBuildToStagingDirectory.Automation.cs
+++ b/Engine/Source/Programs/AutomationTool/Scripts/CopyBuildToStagingDirectory.Automation.cs
@@ -2033,7 +2033,7 @@ public partial class Project : CommandUtils
 			for(int Idx = 0; Idx < Commands.Count; Idx++)
 			{
 				string LogFile = LogUtils.GetUniqueLogName(CombinePaths(CmdEnv.LogFolder, String.Format("UnrealPak-{0}", LogNames[Idx])));
-				string Arguments = Commands[Idx] = String.Format("{0} {1} -multiprocess -abslog={2} {3}", Commands[Idx], MakePathSafeToUseWithCommandLine(ProjectFile.FullName), MakePathSafeToUseWithCommandLine(LogFile), AdditionalCompressionOptionsOnCommandLine);
+				string Arguments = Commands[Idx] = String.Format("{0} {1} -multiprocess -abslog={2} {3}", MakePathSafeToUseWithCommandLine(ProjectFile.FullName), Commands[Idx], MakePathSafeToUseWithCommandLine(LogFile), AdditionalCompressionOptionsOnCommandLine);
 
 				int LocalIdx = Idx;
 				Queue.Enqueue(() => Results[LocalIdx] = Run(UnrealPakExe, Arguments, Options: ERunOptions.AppMustExist | ERunOptions.NoLoggingOfRunCommand));
