Fixed a problem with code signing iOS binaries for blueprint projects with plugins enabled

diff --git a/Engine/Source/Programs/AutomationTool/IOS/IOSPlatform.Automation.cs b/Engine/Source/Programs/AutomationTool/IOS/IOSPlatform.Automation.cs
--- a/Engine/Source/Programs/AutomationTool/IOS/IOSPlatform.Automation.cs
+++ b/Engine/Source/Programs/AutomationTool/IOS/IOSPlatform.Automation.cs
@@ -535,7 +535,7 @@ public class IOSPlatform : Platform
 		{
 			// project.xcodeproj doesn't exist, so generate temp project
 			string Arguments = "-project=\"" + RawProjectPath + "\"";
-			Arguments += " -platforms=" + PlatformName + " -game -nointellisense -" + PlatformName + "deployonly -ignorejunk -projectfileformat=XCode";
+			Arguments += " -platforms=" + PlatformName + " -game -nointellisense -" + PlatformName + "deployonly -ignorejunk -projectfileformat=XCode -includetemptargets";
 
 			// If engine is installed then UBT doesn't need to be built
 			if (CommandUtils.IsEngineInstalled())
