Missing files to fix UAT compile errors.

#rnx

diff --git a/Engine/Source/Programs/AutomationTool/Scripts/CopyBuildToStagingDirectory.Automation.cs b/Engine/Source/Programs/AutomationTool/Scripts/CopyBuildToStagingDirectory.Automation.cs
--- a/Engine/Source/Programs/AutomationTool/Scripts/CopyBuildToStagingDirectory.Automation.cs
+++ b/Engine/Source/Programs/AutomationTool/Scripts/CopyBuildToStagingDirectory.Automation.cs
@@ -1458,7 +1458,7 @@ public partial class Project : CommandUtils
 				continue;
 			}
 
-			IEnumerable<string> FilesEnumberable;
+			IReadOnlyList<string> FilesEnumberable;
 			if (PakRulesConfig.TryGetValues(SectionName, "Files", out FilesEnumberable))
 			{
 				// Only add if we have actual files, we can end up with none due to config overriding
diff --git a/Engine/Source/Programs/AutomationTool/Scripts/SharedCookedBuild.cs b/Engine/Source/Programs/AutomationTool/Scripts/SharedCookedBuild.cs
--- a/Engine/Source/Programs/AutomationTool/Scripts/SharedCookedBuild.cs
+++ b/Engine/Source/Programs/AutomationTool/Scripts/SharedCookedBuild.cs
@@ -139,7 +139,7 @@ public class SharedCookedBuild
 		// Attempt manifest searching first
 		ConfigHierarchy Hierarchy = ConfigCache.ReadHierarchy(ConfigHierarchyType.Engine, DirectoryReference.FromFile(ProjectFile), UnrealTargetPlatform.Win64);
 
-		IEnumerable<string> RawSharedCookedSources = null;
+		IReadOnlyList<string> RawSharedCookedSources = null;
 		Hierarchy.TryGetValues("SharedCookedBuildSettings", "SharedCookedSources", out RawSharedCookedSources);
 		if (RawSharedCookedSources == null)
 		{
