Properly validate platforms passed into BuildCookRun. Comment around the place where platform names are validated implies that there may be some pathway that incorrectly passes a cook platform through as this value, but I'm unable to find such a case (and it should be fixed if it exists, because it's not being parsed correctly at the moment).

Also improved error path for passing an incorrect configuration.

diff --git a/Engine/Source/Programs/AutomationTool/AutomationUtils/ProjectParams.cs b/Engine/Source/Programs/AutomationTool/AutomationUtils/ProjectParams.cs
--- a/Engine/Source/Programs/AutomationTool/AutomationUtils/ProjectParams.cs
+++ b/Engine/Source/Programs/AutomationTool/AutomationUtils/ProjectParams.cs
@@ -142,7 +142,7 @@ namespace AutomationTool
 					// or (if AllowPlatformParams is true)
 					// -Platform_1 -Platform_2 ... -Platform_k
 					string CmdLinePlatform = null;
-					foreach (var ParamName in PlatformParamNames)
+					foreach (string ParamName in PlatformParamNames)
 					{
 						string ParamValue = Command.ParseParamValue(ParamName);
 						if (!string.IsNullOrEmpty(ParamValue))
@@ -169,36 +169,34 @@ namespace AutomationTool
 					{
 						// Get all platforms from the param value: Platform_1+Platform_2+...+Platform_k
 						TargetPlatforms = new List<TargetPlatformDescriptor>();
-						var PlatformNames = (new HashSet<string>(CmdLinePlatform.Split(new char[] { '+' }, StringSplitOptions.RemoveEmptyEntries))).ToList();
-						foreach (var PlatformName in PlatformNames)
+						List<string> PlatformNames = (new HashSet<string>(CmdLinePlatform.Split(new char[] { '+' }, StringSplitOptions.RemoveEmptyEntries))).ToList();
+						foreach (string PlatformName in PlatformNames)
 						{
                             // Look for dependent platforms, Source_1.Dependent_1+Source_2.Dependent_2+Standalone_3
-                            var SubPlatformNames = new List<string>(PlatformName.Split(new char[] { '.' }, StringSplitOptions.RemoveEmptyEntries));
+                            List<string> SubPlatformNames = new List<string>(PlatformName.Split(new char[] { '.' }, StringSplitOptions.RemoveEmptyEntries));
 
-                            foreach (var SubPlatformName in SubPlatformNames)
+                            foreach (string SubPlatformName in SubPlatformNames)
                             {
-								// Need to tolerate cook platform names here, which UFE likes to pass in. (TODO: Not sure if it's right to do that, but it does pass -targetplatform as well)
-								UnrealTargetPlatform NewPlatformType;
-								if (UnrealTargetPlatform.TryParse(SubPlatformName, out NewPlatformType))
-								{
-									// generate all valid platform descriptions for this platform type + cook flavors
-									List<TargetPlatformDescriptor> PlatformDescriptors = Platform.GetValidTargetPlatforms(NewPlatformType, CookFlavors);
-									TargetPlatforms.AddRange(PlatformDescriptors);
+								// Require this to be a valid platform name
+								UnrealTargetPlatform NewPlatformType = UnrealTargetPlatform.Parse(SubPlatformName);
+
+								// generate all valid platform descriptions for this platform type + cook flavors
+								List<TargetPlatformDescriptor> PlatformDescriptors = Platform.GetValidTargetPlatforms(NewPlatformType, CookFlavors);
+								TargetPlatforms.AddRange(PlatformDescriptors);
                                                               
-									if (SubPlatformName != SubPlatformNames[0])
+								if (SubPlatformName != SubPlatformNames[0])
+								{
+									// This is not supported with cook flavors
+									if (!CommandUtils.IsNullOrEmpty(CookFlavors))
 									{
-										// This is not supported with cook flavors
-										if (!CommandUtils.IsNullOrEmpty(CookFlavors))
-										{
-											throw new AutomationException("Cook flavors are not supported for dependent platforms!");
-										}
-
-										// We're a dependent platform so add ourselves to the map, pointing to the first element in the list
-										UnrealTargetPlatform SubPlatformType;
-										if (UnrealTargetPlatform.TryParse(SubPlatformNames[0], out SubPlatformType))
-										{
-											DependentPlatformMap.Add(new TargetPlatformDescriptor(NewPlatformType), new TargetPlatformDescriptor(SubPlatformType));
-										}
+										throw new AutomationException("Cook flavors are not supported for dependent platforms!");
+									}
+
+									// We're a dependent platform so add ourselves to the map, pointing to the first element in the list
+									UnrealTargetPlatform SubPlatformType;
+									if (UnrealTargetPlatform.TryParse(SubPlatformNames[0], out SubPlatformType))
+									{
+										DependentPlatformMap.Add(new TargetPlatformDescriptor(NewPlatformType), new TargetPlatformDescriptor(SubPlatformType));
 									}
 								}
                             }
@@ -826,7 +824,7 @@ namespace AutomationTool
 			{
 				if (Command != null)
 				{
-					var ClientConfig = Command.ParseParamValue("clientconfig");
+					string ClientConfig = Command.ParseParamValue("clientconfig");
 
                     if (ClientConfig == null)
                         ClientConfig = Command.ParseParamValue("config");
@@ -834,10 +832,10 @@ namespace AutomationTool
                     if (ClientConfig != null)
 					{
 						this.ClientConfigsToBuild = new List<UnrealTargetConfiguration>();
-						var Configs = new ParamList<string>(ClientConfig.Split(new char[] { '+' }, StringSplitOptions.RemoveEmptyEntries));
-						foreach (var ConfigName in Configs)
+						ParamList<string> Configs = new ParamList<string>(ClientConfig.Split(new char[] { '+' }, StringSplitOptions.RemoveEmptyEntries));
+						foreach (string ConfigName in Configs)
 						{
-							this.ClientConfigsToBuild.Add((UnrealTargetConfiguration)Enum.Parse(typeof(UnrealTargetConfiguration), ConfigName, true));
+							this.ClientConfigsToBuild.Add(ParseConfig(ConfigName));
 						}
 					}
 				}
@@ -989,7 +987,7 @@ namespace AutomationTool
 			{
 				if (Command != null)
 				{
-					var ServerConfig = Command.ParseParamValue("serverconfig");
+					string ServerConfig = Command.ParseParamValue("serverconfig");
 
                     if (ServerConfig == null)
                         ServerConfig = Command.ParseParamValue("config");
@@ -997,10 +995,10 @@ namespace AutomationTool
                     if (ServerConfig != null)
 					{
 						this.ServerConfigsToBuild = new List<UnrealTargetConfiguration>();
-						var Configs = new ParamList<string>(ServerConfig.Split(new char[] { '+' }, StringSplitOptions.RemoveEmptyEntries));
-						foreach (var ConfigName in Configs)
+						ParamList<string> Configs = new ParamList<string>(ServerConfig.Split(new char[] { '+' }, StringSplitOptions.RemoveEmptyEntries));
+						foreach (string ConfigName in Configs)
 						{
-							this.ServerConfigsToBuild.Add((UnrealTargetConfiguration)Enum.Parse(typeof(UnrealTargetConfiguration), ConfigName, true));
+							this.ServerConfigsToBuild.Add(ParseConfig(ConfigName));
 						}
 					}
 				}
@@ -1078,14 +1076,24 @@ namespace AutomationTool
             }
         }
 
-        #endregion
+		static UnrealTargetConfiguration ParseConfig(string ConfigName)
+		{
+			UnrealTargetConfiguration ConfigValue;
+			if (!Enum.TryParse(ConfigName, true, out ConfigValue))
+			{
+				throw new AutomationException("Invalid configuration '{0}'. Valid configurations are '{1}'.", ConfigName, String.Join("', '", Enum.GetNames(typeof(UnrealTargetConfiguration)).Where(x => x != nameof(UnrealTargetConfiguration.Unknown))));
+			}
+			return ConfigValue;
+		}
 
-        #region Shared
+		#endregion
 
-        /// <summary>
-        /// Shared: Full path to the .uproject file
-        /// </summary>
-        public FileReference RawProjectPath { private set; get; }
+		#region Shared
+
+		/// <summary>
+		/// Shared: Full path to the .uproject file
+		/// </summary>
+		public FileReference RawProjectPath { private set; get; }
 
 		/// <summary>
 		/// Shared: The current project is a foreign project, commandline: -foreign
