Added "build Content Only Patch" to ec jobs list.

diff --git a/Engine/Source/Programs/AutomationTool/AutomationUtils/BundleUtils.cs b/Engine/Source/Programs/AutomationTool/AutomationUtils/BundleUtils.cs
--- a/Engine/Source/Programs/AutomationTool/AutomationUtils/BundleUtils.cs
+++ b/Engine/Source/Programs/AutomationTool/AutomationUtils/BundleUtils.cs
@@ -42,6 +42,7 @@ namespace AutomationUtils.Automation
 			public bool bFoundParent { get; set; }
 			public bool bContainsShaderLibrary { get; set; }
 			public int Priority { get; set; }
+			public string ExecFileName { get; set; }
 		};
 		public static bool LoadBundleConfig(string BundleIniFile, ref Dictionary<string, BundleSettings> Bundles)
 		{
@@ -78,6 +79,11 @@ namespace AutomationUtils.Automation
 					BundleConfig.GetString(SectionName, "Parent", out ParentName);
 					Bundle.ParentName = ParentName;
 				}
+				{
+					string ExecFileName;
+					BundleConfig.GetString(SectionName, "ExecFileName", out ExecFileName);
+					Bundle.ExecFileName = ExecFileName;
+				}
 				{
 					List<string> Tags;
 					BundleConfig.GetArray(SectionName, "Tags", out Tags);
