Exposing new Install Time Coefficients to C# BPT classes.

diff --git a/Engine/Source/Programs/AutomationTool/AutomationUtils/MCPPublic.cs b/Engine/Source/Programs/AutomationTool/AutomationUtils/MCPPublic.cs
--- a/Engine/Source/Programs/AutomationTool/AutomationUtils/MCPPublic.cs
+++ b/Engine/Source/Programs/AutomationTool/AutomationUtils/MCPPublic.cs
@@ -792,6 +792,22 @@ namespace EpicGames.MCP.Automation
 				/// Note that the sum of these can be higher than the actual total due to possibility of shares files.
 				/// </summary>
 				public Dictionary<string, ulong> IndividualTagDeltaSizes;
+				/// <summary>
+				/// Install time coefficients represent an estimation for time to install the patch. These are not accurate timing representations, but are comparable between runs with different versions.
+				/// They can be used to spot out of the ordinary time requirements for installing an update.
+				/// The list if non-null will contain 6 entries as follows:
+				///   InstallTimeCoefficients[0] - Low-Spec using DestructiveInstall.
+				///   InstallTimeCoefficients[1] - Low-Spec using NonDestructiveInstall.
+				///   InstallTimeCoefficients[2] - Mid-Spec using DestructiveInstall.
+				///   InstallTimeCoefficients[3] - Mid-Spec using NonDestructiveInstall.
+				///   InstallTimeCoefficients[4] - High-Spec using DestructiveInstall.
+				///   InstallTimeCoefficients[5] - High-Spec using NonDestructiveInstall.
+				/// Low-Spec was taken from 25 percentile as of July 2019.
+				/// Mid-Spec was taken from 50 percentile as of July 2019.
+				/// High-Spec was taken from 75 percentile as of July 2019.
+				/// If the BPT version being used does not support this feature, or a problem occurred, InstallTimeCoefficients will be null.
+				/// </summary>
+				public List<float> InstallTimeCoefficients;
 			}
 			/// <summary>
 			/// The manifest detail for the source build of the differential.
