Add support for store version offset by architecture

#ue4
#android

diff --git a/Engine/Config/BaseEngine.ini b/Engine/Config/BaseEngine.ini
--- a/Engine/Config/BaseEngine.ini
+++ b/Engine/Config/BaseEngine.ini
@@ -1951,6 +1951,10 @@ InstallLocation=InternalOnly
 DepthBufferPreference=Default
 PackageName=com.YourCompany.[PROJECT]
 StoreVersion=1
+StoreVersionOffsetArmV7=0
+StoreVersionOffsetArm64=0
+StoreVersionOffsetX86=0
+StoreVersionOffsetX8664=0
 VersionDisplayName=1.0
 MinSDKVersion=9
 TargetSDKVersion=26
diff --git a/Engine/Source/Programs/AutomationTool/Android/AndroidPlatform.Automation.cs b/Engine/Source/Programs/AutomationTool/Android/AndroidPlatform.Automation.cs
--- a/Engine/Source/Programs/AutomationTool/Android/AndroidPlatform.Automation.cs
+++ b/Engine/Source/Programs/AutomationTool/Android/AndroidPlatform.Automation.cs
@@ -107,7 +107,7 @@ public class AndroidPlatform : Platform
 
 	private static string GetFinalSymbolizedSODirectory(string ApkName, DeploymentContext SC, string Architecture, string GPUArchitecture)
 	{
-		string PackageVersion = GetPackageInfo(ApkName, true);
+		string PackageVersion = GetPackageInfo(ApkName, SC, true);
 		if (PackageVersion == null || PackageVersion.Length == 0)
 		{
 			throw new AutomationException(ExitCode.Error_FailureGettingPackageInfo, "Failed to get package version from " + ApkName);
@@ -116,16 +116,16 @@ public class AndroidPlatform : Platform
 		return SC.ShortProjectName + "_Symbols_v" + PackageVersion + "/" + SC.ShortProjectName + Architecture + GPUArchitecture;
 	}
 
-	private static string GetFinalObbName(string ApkName, bool bUseAppType = true)
+	private static string GetFinalObbName(string ApkName, DeploymentContext SC, bool bUseAppType = true)
 	{
 		// calculate the name for the .obb file
-		string PackageName = GetPackageInfo(ApkName, false);
+		string PackageName = GetPackageInfo(ApkName, SC, false);
 		if (PackageName == null)
 		{
 			throw new AutomationException(ExitCode.Error_FailureGettingPackageInfo, "Failed to get package name from " + ApkName);
 		}
 
-		string PackageVersion = GetPackageInfo(ApkName, true);
+		string PackageVersion = GetPackageInfo(ApkName, SC, true);
 		if (PackageVersion == null || PackageVersion.Length == 0)
 		{
 			throw new AutomationException(ExitCode.Error_FailureGettingPackageInfo, "Failed to get package version from " + ApkName);
@@ -151,16 +151,16 @@ public class AndroidPlatform : Platform
 		return ObbName;
 	}
 
-	private static string GetFinalPatchName(string ApkName, bool bUseAppType = true)
+	private static string GetFinalPatchName(string ApkName, DeploymentContext SC, bool bUseAppType = true)
 	{
 		// calculate the name for the .obb file
-		string PackageName = GetPackageInfo(ApkName, false);
+		string PackageName = GetPackageInfo(ApkName, SC, false);
 		if (PackageName == null)
 		{
 			throw new AutomationException(ExitCode.Error_FailureGettingPackageInfo, "Failed to get package name from " + ApkName);
 		}
 
-		string PackageVersion = GetPackageInfo(ApkName, true);
+		string PackageVersion = GetPackageInfo(ApkName, SC, true);
 		if (PackageVersion == null || PackageVersion.Length == 0)
 		{
 			throw new AutomationException(ExitCode.Error_FailureGettingPackageInfo, "Failed to get package version from " + ApkName);
@@ -200,17 +200,17 @@ public class AndroidPlatform : Platform
 		return PakParams;
 	}
 
-	private static string GetDeviceObbName(string ApkName)
+	private static string GetDeviceObbName(string ApkName, DeploymentContext SC)
 	{
-        string ObbName = GetFinalObbName(ApkName, false);
-        string PackageName = GetPackageInfo(ApkName, false);
+        string ObbName = GetFinalObbName(ApkName, SC, false);
+        string PackageName = GetPackageInfo(ApkName, SC, false);
         return TargetAndroidLocation + PackageName + "/" + Path.GetFileName(ObbName);
 	}
 
-	private static string GetDevicePatchName(string ApkName)
+	private static string GetDevicePatchName(string ApkName, DeploymentContext SC)
 	{
-		string PatchName = GetFinalPatchName(ApkName, false);
-		string PackageName = GetPackageInfo(ApkName, false);
+		string PatchName = GetFinalPatchName(ApkName, SC, false);
+		string PackageName = GetPackageInfo(ApkName, SC, false);
 		return TargetAndroidLocation + PackageName + "/" + Path.GetFileName(PatchName);
 	}
 
@@ -499,7 +499,7 @@ public class AndroidPlatform : Platform
 		Ini.GetInt32("/Script/AndroidRuntimeSettings.AndroidRuntimeSettings", "MinSDKVersion", out MinSDKVersion);
 		int TargetSDKVersion = MinSDKVersion;
 		Ini.GetInt32("/Script/AndroidRuntimeSettings.AndroidRuntimeSettings", "TargetSDKVersion", out TargetSDKVersion);
-		LogInformation("Target SDK Version" + TargetSDKVersion);
+		LogInformation("Target SDK Version " + TargetSDKVersion);
 
 		foreach (string Architecture in Architectures)
 		{
@@ -554,14 +554,14 @@ public class AndroidPlatform : Platform
 				string PatchName = "";
 				if (!bPackageDataInsideApk)
 			    {
-				    DeviceObbName = GetDeviceObbName(ApkName);
-				    ObbName = GetFinalObbName(ApkName);
+				    DeviceObbName = GetDeviceObbName(ApkName, SC);
+				    ObbName = GetFinalObbName(ApkName, SC);
 					CopyFile(LocalObbName, ObbName);
 
 					if (File.Exists(LocalPatchName))
 					{
-						DevicePatchName = GetDevicePatchName(ApkName);
-						PatchName = GetFinalPatchName(ApkName);
+						DevicePatchName = GetDevicePatchName(ApkName, SC);
+						PatchName = GetFinalPatchName(ApkName, SC);
 						CopyFile(LocalPatchName, PatchName);
 					}
 				}
@@ -577,7 +577,7 @@ public class AndroidPlatform : Platform
 				{
 					bool bIsPC = (Target == UnrealTargetPlatform.Win64);
 					// Write install batch file(s).
-					string PackageName = GetPackageInfo(ApkName, false);
+					string PackageName = GetPackageInfo(ApkName, SC, false);
 					string BatchName = GetFinalBatchName(ApkName, SC, bMakeSeparateApks ? Architecture : "", bMakeSeparateApks ? GPUArchitecture : "", false, EBatchType.Install, Target);
 					// make a batch file that can be used to install the .apk and .obb files
 					string[] BatchLines = GenerateInstallBatchFile(bPackageDataInsideApk, PackageName, ApkName, Params, ObbName, DeviceObbName, false, PatchName, DevicePatchName, false, bIsPC, Params.Distribution, TargetSDKVersion > 22);
@@ -875,14 +875,14 @@ public class AndroidPlatform : Platform
 		bool bMakeSeparateApks = UnrealBuildTool.AndroidExports.ShouldMakeSeparateApks();
 		bool bPackageDataInsideApk = UnrealBuildTool.AndroidExports.CreateDeploymentHandler(Params.RawProjectPath, Params.ForcePackageData).GetPackageDataInsideApk();
 
-		bool bAddedOBB = false;
+		List<string> AddedObbFiles = new List<string>();
 		foreach (string Architecture in Architectures)
 		{
 			foreach (string GPUArchitecture in GPUArchitectures)
 			{
 				string ApkName = GetFinalApkName(Params, SC.StageExecutables[0], true, bMakeSeparateApks ? Architecture : "", bMakeSeparateApks ? GPUArchitecture : "");
-				string ObbName = GetFinalObbName(ApkName);
-				string PatchName = GetFinalPatchName(ApkName);
+				string ObbName = GetFinalObbName(ApkName, SC);
+				string PatchName = GetFinalPatchName(ApkName, SC);
 				bool bBuildWithHiddenSymbolVisibility = BuildWithHiddenSymbolVisibility(SC);
 				bool bSaveSymbols = GetSaveSymbols(SC);
 				//string NoOBBBatchName = GetFinalBatchName(ApkName, Params, bMakeSeparateApks ? Architecture : "", bMakeSeparateApks ? GPUArchitecture : "", true, false);
@@ -911,14 +911,18 @@ public class AndroidPlatform : Platform
 				}
 
 				SC.ArchiveFiles(Path.GetDirectoryName(ApkName), Path.GetFileName(ApkName));
-				if (!bPackageDataInsideApk && !bAddedOBB)
+				if (!bPackageDataInsideApk)
 				{
-					bAddedOBB = true;
-					SC.ArchiveFiles(Path.GetDirectoryName(ObbName), Path.GetFileName(ObbName));
-
-					if (FileExists(PatchName))
+					// only add if not already in archive list
+					if (!AddedObbFiles.Contains(ObbName))
 					{
-						SC.ArchiveFiles(Path.GetDirectoryName(PatchName), Path.GetFileName(PatchName));
+						AddedObbFiles.Add(ObbName);
+
+						SC.ArchiveFiles(Path.GetDirectoryName(ObbName), Path.GetFileName(ObbName));
+						if (FileExists(PatchName))
+						{
+							SC.ArchiveFiles(Path.GetDirectoryName(PatchName), Path.GetFileName(PatchName));
+						}
 					}
 				}
 
@@ -1206,14 +1210,14 @@ public class AndroidPlatform : Platform
             }
 
             // now we can use the apk to get more info
-            string PackageName = GetPackageInfo(ApkName, false);
+            string PackageName = GetPackageInfo(ApkName, SC, false);
 
             // Setup the OBB name and add the storage path (queried from the device) to it
             string DeviceStorageQueryCommand = GetStorageQueryCommand();
             IProcessResult Result = RunAdbCommand(Params, DeviceName, DeviceStorageQueryCommand, null, ERunOptions.AppMustExist);
             String StorageLocation = Result.Output.Trim(); // "/mnt/sdcard";
-            string DeviceObbName = StorageLocation + "/" + GetDeviceObbName(ApkName);
-			string DevicePatchName = StorageLocation + "/" + GetDevicePatchName(ApkName);
+            string DeviceObbName = StorageLocation + "/" + GetDeviceObbName(ApkName, SC);
+			string DevicePatchName = StorageLocation + "/" + GetDevicePatchName(ApkName, SC);
 			string RemoteDir = StorageLocation + "/UE4Game/" + Params.ShortProjectName;
 
             // determine if APK out of date
@@ -1594,7 +1598,7 @@ public class AndroidPlatform : Platform
 			else if (SC.Archive)
             {
                 // deploy the obb if there is one
-                string ObbPath = Path.Combine(SC.StageDirectory.FullName, GetFinalObbName(ApkName));
+                string ObbPath = Path.Combine(SC.StageDirectory.FullName, GetFinalObbName(ApkName, SC));
                 if (File.Exists(ObbPath))
                 {
                     // cache some strings
@@ -1605,7 +1609,7 @@ public class AndroidPlatform : Platform
                 }
 
 				// deploy the patch if there is one
-				string PatchPath = Path.Combine(SC.StageDirectory.FullName, GetFinalPatchName(ApkName));
+				string PatchPath = Path.Combine(SC.StageDirectory.FullName, GetFinalPatchName(ApkName, SC));
 				if (File.Exists(PatchPath))
 				{
 					// cache some strings
@@ -1689,17 +1693,24 @@ public class AndroidPlatform : Platform
 			LogInformation("GetPackageInfo ReturnValue: {0}", ReturnValue);
 		}
 
+		return ReturnValue;
+	}
+
+	public static string GetPackageInfo(string ApkName, DeploymentContext SC, bool bRetrieveVersionCode)
+	{
+		string ReturnValue = GetPackageInfo(ApkName, bRetrieveVersionCode);
+
 		if (ReturnValue == null || ReturnValue.Length == 0)
 		{
 			/** If APK does not exist or we cant find package info in apk use the packageInfo file */
-			ReturnValue = GetPackageInfoFromInfoFile(ApkName, bRetrieveVersionCode);
+			ReturnValue = GetPackageInfoFromInfoFile(ApkName, SC, bRetrieveVersionCode);
 		}
 
 		return ReturnValue;
 	}
 
 	/** Lookup package info in packageInfo.txt file in same directory as the APK would have been */
-	private static string GetPackageInfoFromInfoFile(string ApkName, bool bRetrieveVersionCode)
+	private static string GetPackageInfoFromInfoFile(string ApkName, DeploymentContext SC, bool bRetrieveVersionCode)
 	{
 		string ReturnValue = null;
 		String PackageInfoPath = Path.Combine(Path.GetDirectoryName(ApkName), "packageInfo.txt");
@@ -1716,6 +1727,30 @@ public class AndroidPlatform : Platform
 			// parse extra info that the aapt-based method got
 			MetaAppTypeLine = Lines[3];
 		}
+
+		if (bRetrieveVersionCode)
+		{
+			int StoreVersion = 1;
+			int.TryParse(ReturnValue, out StoreVersion);
+
+			int StoreVersionOffset = 0;
+			ConfigHierarchy Ini = ConfigCache.ReadHierarchy(ConfigHierarchyType.Engine, DirectoryReference.FromFile(SC.RawProjectPath), SC.StageTargetPlatform.PlatformType);
+			if (ApkName.Contains("-armv7-"))
+			{
+				Ini.GetInt32("/Script/AndroidRuntimeSettings.AndroidRuntimeSettings", "StoreVersionOffsetArmV7", out StoreVersionOffset);
+			}
+			else if (ApkName.Contains("-arm64-"))
+			{
+				Ini.GetInt32("/Script/AndroidRuntimeSettings.AndroidRuntimeSettings", "StoreVersionOffsetArm64", out StoreVersionOffset);
+			}
+			else if (ApkName.Contains("-x64-"))
+			{
+				Ini.GetInt32("/Script/AndroidRuntimeSettings.AndroidRuntimeSettings", "StoreVersionOffsetX8664", out StoreVersionOffset);
+			}
+			StoreVersion += StoreVersionOffset;
+			ReturnValue = StoreVersion.ToString("0");
+		}
+
 		LogInformation("packageInfo.txt file exists: {0}", fileExists);
 		LogInformation("packageInfo return MetaAppTypeLine: {0}", MetaAppTypeLine);
 		LogInformation("packageInfo return value: {0}", ReturnValue);
diff --git a/Engine/Source/Programs/UnrealBuildTool/Platform/Android/UEDeployAndroid.cs b/Engine/Source/Programs/UnrealBuildTool/Platform/Android/UEDeployAndroid.cs
--- a/Engine/Source/Programs/UnrealBuildTool/Platform/Android/UEDeployAndroid.cs
+++ b/Engine/Source/Programs/UnrealBuildTool/Platform/Android/UEDeployAndroid.cs
@@ -869,8 +869,11 @@ namespace UnrealBuildTool
 		}
 
 		private int CachedStoreVersion = -1;
+		private int CachedStoreVersionOffsetArmV7 = 0;
+		private int CachedStoreVersionOffsetArm64 = 0;
+		private int CachedStoreVersionOffsetX8664= 0;
 
-		public int GetStoreVersion()
+		public int GetStoreVersion(string UE4Arch)
 		{
 			if (CachedStoreVersion < 1)
 			{
@@ -901,6 +904,17 @@ namespace UnrealBuildTool
 				Log.TraceInformation("GotStoreVersion found v{0}. (bUseChangeListAsStoreVersion={1} IsBuildMachine={2} EngineChangeList={3})", StoreVersion, bUseChangeListAsStoreVersion, IsBuildMachine, EngineChangelist);
 
 				CachedStoreVersion = StoreVersion;
+
+				Ini.GetInt32("/Script/AndroidRuntimeSettings.AndroidRuntimeSettings", "StoreVersionOffsetArmV7", out CachedStoreVersionOffsetArmV7);
+				Ini.GetInt32("/Script/AndroidRuntimeSettings.AndroidRuntimeSettings", "StoreVersionOffsetArm64", out CachedStoreVersionOffsetArm64);
+				Ini.GetInt32("/Script/AndroidRuntimeSettings.AndroidRuntimeSettings", "StoreVersionOffsetX8664", out CachedStoreVersionOffsetX8664);
+			}
+
+			switch (UE4Arch)
+			{
+				case "-armv7": return CachedStoreVersion + CachedStoreVersionOffsetArmV7;
+				case "-arm64": return CachedStoreVersion + CachedStoreVersionOffsetArm64;
+				case "-x64": return CachedStoreVersion + CachedStoreVersionOffsetX8664;
 			}
 
 			return CachedStoreVersion;
@@ -949,7 +963,7 @@ namespace UnrealBuildTool
 			return CachedVersionDisplayName;
 		}
 
-		public void WriteJavaOBBDataFile(string FileName, string PackageName, List<string> ObbSources, string CookFlavor, bool bPackageDataInsideApk)
+		public void WriteJavaOBBDataFile(string FileName, string PackageName, List<string> ObbSources, string CookFlavor, bool bPackageDataInsideApk, string UE4Arch)
 		{
 			Log.TraceInformation("\n==== Writing to OBB data file {0} ====", FileName);
 
@@ -963,7 +977,7 @@ namespace UnrealBuildTool
 //				AppType = ".Client";		// should always be empty now; fix up the name in batch file instead
 			}
 
-			int StoreVersion = GetStoreVersion();
+			int StoreVersion = GetStoreVersion(UE4Arch);
 
 			StringBuilder obbData = new StringBuilder("package " + PackageName + ";\n\n");
 			obbData.Append("public class OBBData\n{\n");
@@ -1141,7 +1155,7 @@ namespace UnrealBuildTool
 			}
 		}
 
-		public void WriteCrashlyticsResources(string UEBuildPath, string PackageName, string ApplicationDisplayName, bool bIsEmbedded)
+		public void WriteCrashlyticsResources(string UEBuildPath, string PackageName, string ApplicationDisplayName, bool bIsEmbedded, string UE4Arch)
 		{
 			System.DateTime CurrentDateTime = System.DateTime.Now;
 			string BuildID = Guid.NewGuid().ToString();
@@ -1158,7 +1172,7 @@ namespace UnrealBuildTool
 			CrashPropertiesContent.Append("version_name=" + VersionDisplayName + "\n");
 			CrashPropertiesContent.Append("package_name=" + PackageName + "\n");
 			CrashPropertiesContent.Append("build_id=" + BuildID + "\n");
-			CrashPropertiesContent.Append("version_code=" + GetStoreVersion().ToString() + "\n");
+			CrashPropertiesContent.Append("version_code=" + GetStoreVersion(UE4Arch).ToString() + "\n");
 
 			string CrashPropertiesFileName = Path.Combine(UEBuildPath, "assets", "crashlytics-build.properties");
 			MakeDirectoryIfRequired(CrashPropertiesFileName);
@@ -2177,7 +2191,7 @@ namespace UnrealBuildTool
 			// Read the engine version
 			string EngineVersion = ReadEngineVersion();
 
-			int StoreVersion = GetStoreVersion();
+			int StoreVersion = GetStoreVersion(UE4Arch);
 
 			string Arch = GetNDKArch(UE4Arch);
 			int NDKLevelInt = 0;
@@ -3137,7 +3151,7 @@ namespace UnrealBuildTool
 			// Create gradle.properties
 			StringBuilder GradleProperties = new StringBuilder();
 
-			int StoreVersion = GetStoreVersion();
+			int StoreVersion = GetStoreVersion(GetUE4Arch(NDKArch));
 			string VersionDisplayName = GetVersionDisplayName(bIsEmbedded);
 
 			ConfigHierarchy Ini = GetConfigCacheIni(ConfigHierarchyType.Engine);
@@ -3399,6 +3413,10 @@ namespace UnrealBuildTool
 			string UE4OBBDataFileName = GetUE4JavaOBBDataFileName(TemplateDestinationBase);
 			string UE4DownloadShimFileName = GetUE4JavaDownloadShimFileName(UE4JavaFilePath);
 
+			// Get list of all architecture and GPU targets for build
+			List<string> Arches = ToolChain.GetAllArchitectures();
+			List<string> GPUArchitectures = ToolChain.GetAllGPUArchitectures();
+
 			// Template generated files
 			string JavaTemplateSourceDir = GetUE4TemplateJavaSourceDir(EngineDirectory);
 			IEnumerable<TemplateFile> templates = from template in Directory.EnumerateFiles(JavaTemplateSourceDir, "*.template")
@@ -3415,7 +3433,7 @@ namespace UnrealBuildTool
 			}
 
 			// Generate the OBBData.java file if out of date (can skip rewriting it if packaging inside Apk in some cases)
-			WriteJavaOBBDataFile(UE4OBBDataFileName, PackageName, RequiredOBBFiles, CookFlavor, bPackageDataInsideApk);
+			WriteJavaOBBDataFile(UE4OBBDataFileName, PackageName, RequiredOBBFiles, CookFlavor, bPackageDataInsideApk, Arches[0]);
 
 			// Make sure any existing proguard file in project is NOT used (back it up)
 			string ProjectBuildProguardFile = Path.Combine(GameBuildFilesPath, "proguard-project.txt");
@@ -3526,13 +3544,6 @@ namespace UnrealBuildTool
 				Log.TraceInformation("Application display name is different than last build, forcing repackage.");
 			}
 
-			// Write Crashlytics data if enabled (not needed for Gradle)
-			if (CrashlyticsPluginEnabled && !bGradleEnabled)
-			{
-				Trace.TraceInformation("Writing Crashlytics resources");
-				WriteCrashlyticsResources(Path.Combine(ProjectDirectory, "Build", "Android"), PackageName, ApplicationDisplayName, bSkipGradleBuild);
-			}
-
 			// if the manifest matches, look at other settings stored in a file
 			if (bBuildSettingsMatch)
 			{
@@ -3564,9 +3575,6 @@ namespace UnrealBuildTool
 
 			}
 
-			List<string> Arches = ToolChain.GetAllArchitectures();
-			List<string> GPUArchitectures = ToolChain.GetAllGPUArchitectures();
-
 			// Initialize UPL contexts for each architecture enabled
 			List<string> NDKArches = new List<string>();
 			foreach (string Arch in Arches)
@@ -3781,8 +3789,18 @@ namespace UnrealBuildTool
 
 					// modify the generated build.xml before the final include
 					UpdateBuildXML(Arch, NDKArch, EngineDirectory, UE4BuildPath);
+
+					// Write Crashlytics data if enabled (not needed for Gradle)
+					if (CrashlyticsPluginEnabled)
+					{
+						Trace.TraceInformation("Writing Crashlytics resources");
+						WriteCrashlyticsResources(Path.Combine(ProjectDirectory, "Build", "Android"), PackageName, ApplicationDisplayName, bSkipGradleBuild, Arch);
+					}
 				}
 
+				// Generate the OBBData.java file again in case architecture has different store version
+				WriteJavaOBBDataFile(UE4OBBDataFileName, PackageName, RequiredOBBFiles, CookFlavor, bPackageDataInsideApk, Arch);
+
 				// update GameActivity.java and GameApplication.java if out of date
 				UpdateGameActivity(Arch, NDKArch, EngineDirectory, UE4BuildPath);
 				UpdateGameApplication(Arch, NDKArch, EngineDirectory, UE4BuildPath, bGradleEnabled);
@@ -4107,7 +4125,7 @@ namespace UnrealBuildTool
 					if (bSaveSymbols || (Configuration == UnrealTargetConfiguration.Shipping && bBuildWithHiddenSymbolVisibility))
 					{
 						// Copy .so with symbols to 
-						int StoreVersion = GetStoreVersion();
+						int StoreVersion = GetStoreVersion(Arch);
 						string SymbolSODirectory = Path.Combine(DestApkDirectory, ProjectName + "_Symbols_v" + StoreVersion + "/" + ProjectName + Arch + GPUArchitecture);
 						string SymbolifiedSOPath = Path.Combine(SymbolSODirectory, Path.GetFileName(FinalSOName));
 						MakeDirectoryIfRequired(SymbolifiedSOPath);
@@ -4217,7 +4235,7 @@ namespace UnrealBuildTool
 
 			string[] PackageInfoSource = new string[4];
 			PackageInfoSource[0] = PackageName;
-			PackageInfoSource[1] = GetStoreVersion().ToString();
+			PackageInfoSource[1] = GetStoreVersion("").ToString();
 			PackageInfoSource[2] = GetVersionDisplayName(bIsEmbedded);
 			PackageInfoSource[3] = string.Format("name='com.epicgames.ue4.GameActivity.AppType' value='{0}'", InTargetType == TargetType.Game ? "" : InTargetType.ToString());
 
diff --git a/Engine/Source/Runtime/Android/AndroidRuntimeSettings/Classes/AndroidRuntimeSettings.h b/Engine/Source/Runtime/Android/AndroidRuntimeSettings/Classes/AndroidRuntimeSettings.h
--- a/Engine/Source/Runtime/Android/AndroidRuntimeSettings/Classes/AndroidRuntimeSettings.h
+++ b/Engine/Source/Runtime/Android/AndroidRuntimeSettings/Classes/AndroidRuntimeSettings.h
@@ -201,6 +201,18 @@ public:
 	UPROPERTY(GlobalConfig, EditAnywhere, Category = "APK Packaging", Meta = (DisplayName = "Store Version (1-2147483647)", ClampMin="1", ClampMax="2147483647"))
 	int32 StoreVersion;
 
+	// Offset to add to store version for APKs generated for armv7
+	UPROPERTY(GlobalConfig, EditAnywhere, Category = "APK Packaging", meta = (DisplayName = "Store Version offset (armv7)"))
+	int32 StoreVersionOffsetArmV7;
+
+	// Offset to add to store version for APKs generated for arm64
+	UPROPERTY(GlobalConfig, EditAnywhere, Category = "APK Packaging", meta = (DisplayName = "Store Version offset (arm64)"))
+	int32 StoreVersionOffsetArm64;
+
+	// Offset to add to store version for APKs generated for x86_64
+	UPROPERTY(GlobalConfig, EditAnywhere, Category = "APK Packaging", meta = (DisplayName = "Store Version offset (x86_64)"))
+	int32 StoreVersionOffsetX8664;
+
 	// The visual application name displayed for end users
 	UPROPERTY(GlobalConfig, EditAnywhere, Category = "APK Packaging", Meta = (DisplayName = "Application Display Name (app_name), project name if blank"))
 	FString ApplicationDisplayName;
