Added string formatting to BuildTarget
Added ability for Help() markup in build script to span multiple lines

diff --git a/Engine/Source/Programs/AutomationTool/AutomationUtils/HelpUtils.cs b/Engine/Source/Programs/AutomationTool/AutomationUtils/HelpUtils.cs
--- a/Engine/Source/Programs/AutomationTool/AutomationUtils/HelpUtils.cs
+++ b/Engine/Source/Programs/AutomationTool/AutomationUtils/HelpUtils.cs
@@ -343,7 +343,11 @@ namespace AutomationTool
 					}
 					else
 					{
-						ObjectDescription = HelpAttribute.Description;
+						if (!String.IsNullOrEmpty(ObjectDescription))
+						{
+							ObjectDescription += Environment.NewLine;
+						}
+						ObjectDescription += HelpAttribute.Description;
 					}
 				}
 			}
diff --git a/Engine/Source/Programs/AutomationTool/AutomationUtils/UE4Build.cs b/Engine/Source/Programs/AutomationTool/AutomationUtils/UE4Build.cs
--- a/Engine/Source/Programs/AutomationTool/AutomationUtils/UE4Build.cs
+++ b/Engine/Source/Programs/AutomationTool/AutomationUtils/UE4Build.cs
@@ -1,4 +1,4 @@
-// Copyright 1998-2019 Epic Games, Inc. All Rights Reserved.
+﻿// Copyright 1998-2019 Epic Games, Inc. All Rights Reserved.
 using System;
 using System.Collections.Generic;
 using System.Text;
@@ -447,6 +447,15 @@ namespace AutomationTool
 			/// Whether to clean this target. If not specified, the target will be cleaned if -Clean is on the command line.
 			/// </summary>
 			public bool? Clean;
+
+			/// <summary>
+			/// Format as string
+			/// </summary>
+			/// <returns></returns>
+			public override string ToString()
+			{
+				return string.Format("{0} {1} {2}", TargetName, Platform, Config);
+			}
 		}
 
 
