Updated TempStorage in UAT to use Zip64 when necessary

diff --git a/Engine/Source/Programs/AutomationTool/BuildGraph/TempStorage.cs b/Engine/Source/Programs/AutomationTool/BuildGraph/TempStorage.cs
--- a/Engine/Source/Programs/AutomationTool/BuildGraph/TempStorage.cs
+++ b/Engine/Source/Programs/AutomationTool/BuildGraph/TempStorage.cs
@@ -927,6 +927,7 @@ namespace AutomationTool
 							// Create one zip per thread using the given basename
 							using (var ZipArchive = new Ionic.Zip.ZipFile(ZipFileName.FullName) { CompressionLevel = Ionic.Zlib.CompressionLevel.BestSpeed })
 							{
+								ZipArchive.UseZip64WhenSaving = Ionic.Zip.Zip64Option.AsNecessary;
 
 								// pull from the queue until we are out of files.
 								do
