Now automatically detect if mono5 is installed on Mac and use msbuild if so, otherwise stick to xbuild

diff --git a/Engine/Build/BatchFiles/Mac/SetupMono.sh b/Engine/Build/BatchFiles/Mac/SetupMono.sh
--- a/Engine/Build/BatchFiles/Mac/SetupMono.sh
+++ b/Engine/Build/BatchFiles/Mac/SetupMono.sh
@@ -5,6 +5,7 @@ sh FixMonoFiles.sh
 sh FixDependencyFiles.sh
 
 IS_MONO_INSTALLED=0
+IS_MONO_5_INSTALLED=0
 MONO_VERSION_PATH=`which mono` || true
 
 # if we can't find mono path, try one last hail mary of a standard install location
@@ -26,6 +27,7 @@ if [ ! $MONO_VERSION_PATH == "" ] && [ -f $MONO_VERSION_PATH ]; then
 	MONO_VERSION=(`echo ${MONO_VERSION:MONO_VERSION_PREFIX_LEN} |tr '.' ' '`)
 	if [ ${MONO_VERSION[0]} -ge 5 ]; then # Allow any Mono 5.x and up
 		IS_MONO_INSTALLED=1
+		IS_MONO_5_INSTALLED=1
 	elif [ ${MONO_VERSION[0]} -eq 4 ]; then
 		if [ ${MONO_VERSION[1]} -eq 0 ] && [ ${MONO_VERSION[2]} -ge 2 ]; then
 			IS_MONO_INSTALLED=1
@@ -43,6 +45,9 @@ if [ $IS_MONO_INSTALLED -eq 0 ]; then
 	export PATH=$UE_MONO_DIR/bin:$PATH
 	export MONO_PATH=$UE_MONO_DIR/lib:$MONO_PATH
 	export LD_LIBRARY_PATH=$UE_MONO_DIR/lib:$LD_LIBRARY_PATH
+else
+	export IS_MONO_INSTALLED=$IS_MONO_INSTALLED
+	export IS_MONO_5_INSTALLED=$IS_MONO_5_INSTALLED
 fi
 
 cd "$START_DIR"
diff --git a/Engine/Build/BatchFiles/RunUAT.sh b/Engine/Build/BatchFiles/RunUAT.sh
--- a/Engine/Build/BatchFiles/RunUAT.sh
+++ b/Engine/Build/BatchFiles/RunUAT.sh
@@ -65,6 +65,7 @@ if [ "$(uname)" = "Linux" ]; then
 	source "$SCRIPT_DIR/Linux/SetupMono.sh" "$SCRIPT_DIR/Linux"
 fi
 
+
 if [ "$UATCompileArg" = "-compile" ]; then
   # see if the .csproj exists to be compiled
 	if [ ! -f Source/Programs/AutomationTool/AutomationTool.csproj ]; then
@@ -75,11 +76,18 @@ if [ "$UATCompileArg" = "-compile" ]; then
 		# make sure the UBT project has references to auto-discovered platform extension source files
 		"${SCRIPT_DIR}/FindPlatformExtensionSources.sh"
 
+		# mono 5.0 and up include msbuild
+		if [ ! $IS_MONO_5_INSTALLED == "" ]; then
+			BUILD_TOOL=msbuild
+		else
+			BUILD_TOOL=xbuild
+		fi
+
 		ARGS="/p:Configuration=Development /p:Platform=AnyCPU /verbosity:quiet /nologo /p:NoWarn=1591 /property:AutomationToolProjectOnly=true"
 		ARGS="${ARGS} /p:TargetFrameworkProfile="
 
-		echo "xbuild Source/Programs/AutomationTool/AutomationTool.csproj $ARGS"
-		xbuild Source/Programs/AutomationTool/AutomationTool.csproj $ARGS
+		echo "$BUILD_TOOL Source/Programs/AutomationTool/AutomationTool.csproj $ARGS"
+		$BUILD_TOOL Source/Programs/AutomationTool/AutomationTool.csproj $ARGS
 		# make sure it succeeded
 		if [ $? -ne 0 ]; then
 			echo RunUAT ERROR: AutomationTool failed to compile.
diff --git a/Engine/Source/Programs/AutomationTool/AutomationUtils/MacHostPlatform.cs b/Engine/Source/Programs/AutomationTool/AutomationUtils/MacHostPlatform.cs
--- a/Engine/Source/Programs/AutomationTool/AutomationUtils/MacHostPlatform.cs
+++ b/Engine/Source/Programs/AutomationTool/AutomationUtils/MacHostPlatform.cs
@@ -7,14 +7,52 @@ using System.Text;
 using System.Diagnostics;
 using System.IO;
 using UnrealBuildTool;
+using System.Text.RegularExpressions;
+using Tools.DotNETCommon;
 
 namespace AutomationTool
 {
 	class MacHostPlatform : HostPlatform
 	{
+		static string CachedMsBuildTool = "";
+
 		public override string GetMsBuildExe()
 		{
-			return "msbuild";
+			// As of 5.0 mono comes with msbuild which performs better. If that's installed then
+			if (string.IsNullOrEmpty(CachedMsBuildTool))
+			{
+				bool CanUseMsBuild = false;
+			
+				string VersionString = Utils.RunLocalProcessAndReturnStdOut("monoh", "--version");
+
+				// Mono JIT compiler version 5.18.1.89 (2018-08/223ea7ef92e Tue May 21 12:03:39 EDT 2019) etc
+				Match M = Regex.Match(VersionString, @"version\s+(\d+\..+?)\s");
+
+				if (M.Success)
+				{
+					try
+					{
+						Version Ver = Version.Parse(M.Groups[1].ToString());
+						CanUseMsBuild = Ver.Major >= 5;
+					}
+					catch
+					{
+						Log.TraceWarning("Unable to parse version string {0}", M.Groups[1]);
+					}
+				}
+	
+				if (CanUseMsBuild)
+				{
+					CachedMsBuildTool = "msbuild";
+				}
+				else
+				{
+					Log.TraceWarning("Using xbuild. Install Mono 5.0 or greater for faster builds!");
+					CachedMsBuildTool = "xbuild";
+				}
+			}
+
+			return CachedMsBuildTool;
 		}
 
 		public override string RelativeBinariesFolder
diff --git a/Engine/Source/Programs/AutomationTool/AutomationUtils/ScriptCompiler.cs b/Engine/Source/Programs/AutomationTool/AutomationUtils/ScriptCompiler.cs
--- a/Engine/Source/Programs/AutomationTool/AutomationUtils/ScriptCompiler.cs
+++ b/Engine/Source/Programs/AutomationTool/AutomationUtils/ScriptCompiler.cs
@@ -284,9 +284,10 @@ namespace AutomationTool
 
 			DateTime StartTime = DateTime.Now;
 
-			bool UseParallelMsBuild = UnrealBuildTool.BuildHostPlatform.Current.Platform == UnrealBuildTool.UnrealTargetPlatform.Win64
-				|| UnrealBuildTool.BuildHostPlatform.Current.Platform == UnrealBuildTool.UnrealTargetPlatform.Win32
-				|| UnrealBuildTool.BuildHostPlatform.Current.Platform == UnrealBuildTool.UnrealTargetPlatform.Mac;
+			string BuildTool = CommandUtils.CmdEnv.MsBuildExe;
+
+			// msbuild (standard on windows, in mono >5.0 is preferred due to speed and parallel compilation)
+			bool UseParallelMsBuild = Path.GetFileNameWithoutExtension(BuildTool).ToLower() == "msbuild";
 
 			if (UseParallelMsBuild)
 			{
@@ -308,7 +309,7 @@ namespace AutomationTool
 
 				var CmdLine = String.Format("\"{0}\" /p:Configuration={1} /verbosity:{2} /nologo", UATProjFile, BuildConfig, MsBuildVerbosity);
 				// suppress the run command because it can be long and intimidating, making the logs around this code harder to read.
-				var Result = CommandUtils.Run(CommandUtils.CmdEnv.MsBuildExe, CmdLine, Options: CommandUtils.ERunOptions.Default | CommandUtils.ERunOptions.NoLoggingOfRunCommand | CommandUtils.ERunOptions.LoggingOfRunDuration);
+				var Result = CommandUtils.Run(BuildTool, CmdLine, Options: CommandUtils.ERunOptions.Default | CommandUtils.ERunOptions.NoLoggingOfRunCommand | CommandUtils.ERunOptions.LoggingOfRunDuration);
 				if (Result.ExitCode != 0)
 				{
 					throw new AutomationException(String.Format("Failed to build \"{0}\":{1}{2}", UATProjFile, Environment.NewLine, Result.Output));
