Removed a piece of test code from previous checkin

diff --git a/Engine/Source/Programs/AutomationTool/AutomationUtils/MacHostPlatform.cs b/Engine/Source/Programs/AutomationTool/AutomationUtils/MacHostPlatform.cs
--- a/Engine/Source/Programs/AutomationTool/AutomationUtils/MacHostPlatform.cs
+++ b/Engine/Source/Programs/AutomationTool/AutomationUtils/MacHostPlatform.cs
@@ -23,7 +23,7 @@ namespace AutomationTool
 			{
 				bool CanUseMsBuild = false;
 			
-				string VersionString = Utils.RunLocalProcessAndReturnStdOut("monoh", "--version");
+				string VersionString = Utils.RunLocalProcessAndReturnStdOut("mono", "--version");
 
 				// Mono JIT compiler version 5.18.1.89 (2018-08/223ea7ef92e Tue May 21 12:03:39 EDT 2019) etc
 				Match M = Regex.Match(VersionString, @"version\s+(\d+\..+?)\s");
