UAT: Add a utility function for converting a P4 tag record to a string.

#rnx

diff --git a/Engine/Source/Programs/AutomationTool/AutomationUtils/P4Utils.cs b/Engine/Source/Programs/AutomationTool/AutomationUtils/P4Utils.cs
--- a/Engine/Source/Programs/AutomationTool/AutomationUtils/P4Utils.cs
+++ b/Engine/Source/Programs/AutomationTool/AutomationUtils/P4Utils.cs
@@ -3208,6 +3208,25 @@ namespace AutomationTool
 			return Tags;
 		}
 
+		/// <summary>
+		/// Formats a tagged record as a string
+		/// </summary>
+		/// <param name="Record">The record to format</param>
+		/// <returns>Single string containing the record</returns>
+		public static string FormatTaggedOutput(Dictionary<string, string> Record)
+		{
+			StringBuilder Result = new StringBuilder();
+			foreach (KeyValuePair<string, string> Pair in Record)
+			{
+				if (Result.Length > 0)
+				{
+					Result.Append('\n');
+				}
+				Result.AppendFormat("{0}: {1}", Pair.Key, Pair.Value);
+			}
+			return Result.ToString();
+		}
+
 		/// <summary>
 		/// Checks if the client exists in P4.
 		/// </summary>
