Cleaned up some comments (thanks MS!)

diff --git a/Engine/Source/Programs/AutomationTool/AutomationUtils/MacHostPlatform.cs b/Engine/Source/Programs/AutomationTool/AutomationUtils/MacHostPlatform.cs
--- a/Engine/Source/Programs/AutomationTool/AutomationUtils/MacHostPlatform.cs
+++ b/Engine/Source/Programs/AutomationTool/AutomationUtils/MacHostPlatform.cs
@@ -18,7 +18,8 @@ namespace AutomationTool
 
 		public override string GetMsBuildExe()
 		{
-			// As of 5.0 mono comes with msbuild which performs better. If that's installed then
+			// As of 5.0 mono comes with msbuild which performs better. If that's installed then we'll
+			// use it, if not then fallback to xbuild and watn
 			if (string.IsNullOrEmpty(CachedMsBuildTool))
 			{
 				bool CanUseMsBuild = false;
diff --git a/Engine/Source/Programs/AutomationTool/AutomationUtils/ScriptCompiler.cs b/Engine/Source/Programs/AutomationTool/AutomationUtils/ScriptCompiler.cs
--- a/Engine/Source/Programs/AutomationTool/AutomationUtils/ScriptCompiler.cs
+++ b/Engine/Source/Programs/AutomationTool/AutomationUtils/ScriptCompiler.cs
@@ -286,7 +286,7 @@ namespace AutomationTool
 
 			string BuildTool = CommandUtils.CmdEnv.MsBuildExe;
 
-			// msbuild (standard on windows, in mono >5.0 is preferred due to speed and parallel compilation)
+			// msbuild (standard on windows, in mono >= 5.0 is preferred due to speed and parallel compilation)
 			bool UseParallelMsBuild = Path.GetFileNameWithoutExtension(BuildTool).ToLower() == "msbuild";
 
 			if (UseParallelMsBuild)
