Append cook platform to PrePak manifest log so that platforms don't stomp each other's logs.

diff --git a/Engine/Source/Programs/AutomationTool/Scripts/CopyBuildToStagingDirectory.Automation.cs b/Engine/Source/Programs/AutomationTool/Scripts/CopyBuildToStagingDirectory.Automation.cs
--- a/Engine/Source/Programs/AutomationTool/Scripts/CopyBuildToStagingDirectory.Automation.cs
+++ b/Engine/Source/Programs/AutomationTool/Scripts/CopyBuildToStagingDirectory.Automation.cs
@@ -1540,7 +1540,7 @@ public partial class Project : CommandUtils
 	{
 		LogInformation("Creating pak using staging manifest.");
 
-		DumpManifest(SC, CombinePaths(CmdEnv.LogFolder, "PrePak" + (SC.DedicatedServer ? "_Server" : "")));
+		DumpManifest(SC, CombinePaths(CmdEnv.LogFolder, "PrePak" + (SC.DedicatedServer ? "_Server" : "") + "_" + SC.CookPlatform));
 
 		var UnrealPakResponseFile = CreatePakResponseFileFromStagingManifest(SC, SC.FilesToStage.UFSFiles);
 
@@ -2332,7 +2332,7 @@ public partial class Project : CommandUtils
 	private static void CreatePaksUsingChunkManifests(ProjectParams Params, DeploymentContext SC)
 	{
 		LogInformation("Creating pak using streaming install manifests.");
-		DumpManifest(SC, CombinePaths(CmdEnv.LogFolder, "PrePak" + (SC.DedicatedServer ? "_Server" : "")));
+		DumpManifest(SC, CombinePaths(CmdEnv.LogFolder, "PrePak" + (SC.DedicatedServer ? "_Server" : "") + "_" + SC.CookPlatform));
 
 		ConfigHierarchy PlatformGameConfig = ConfigCache.ReadHierarchy(ConfigHierarchyType.Game, DirectoryReference.FromFile(Params.RawProjectPath), SC.StageTargetPlatform.IniPlatformType);
 
