Allow pak rules to only apply to chunked or non-chunked builds.
Fix issue in non-chunked builds that could cause rules to be misapplied - Each staged file should have its pak chunk list.

diff --git a/Engine/Config/BasePakFileRules.ini b/Engine/Config/BasePakFileRules.ini
--- a/Engine/Config/BasePakFileRules.ini
+++ b/Engine/Config/BasePakFileRules.ini
@@ -1,10 +1,12 @@
 ; These rules are applied in order, the first rule that applies per file is taken and no others are evaluated
 ; [SectionName]
-; bOverrideChunkManifest=false 			; If true this allows overriding assignments from the cooker
+; bOverrideChunkManifest=false 		; If true this allows overriding assignments from the cooker
 ; bExcludeFromPaks=false 			; If true this removes entirely, cannot coexist with overridepaks
 ; OverridePaks="pakchunk1" 			; If set this will override pak list, comma separated
 ; Platforms="iOS,Android"			; If set this rule will only apply to these platforms
 ; Targets="Shipping,Test"			; If set this rule will only apply to these configurations
+; bOnlyChunkedBuilds=true			; If set this rule will only apply to chunked builds
+; bOnlyNonChunkedBuilds=true		; If set this rule will only apply to non-chunked builds
 ; +Files=".../*FileMask*.*"			; List of file masks to apply to, using the C# FileFilter class
 
 
diff --git a/Engine/Source/Programs/AutomationTool/Scripts/CopyBuildToStagingDirectory.Automation.cs b/Engine/Source/Programs/AutomationTool/Scripts/CopyBuildToStagingDirectory.Automation.cs
--- a/Engine/Source/Programs/AutomationTool/Scripts/CopyBuildToStagingDirectory.Automation.cs
+++ b/Engine/Source/Programs/AutomationTool/Scripts/CopyBuildToStagingDirectory.Automation.cs
@@ -1368,6 +1368,8 @@ public partial class Project : CommandUtils
 			return null;
 		}
 
+		bool bChunkedBuild = SC.PlatformUsesChunkManifests && DoesChunkPakManifestExist(Params, SC);
+
 		ConfigHierarchy PakRulesConfig = new ConfigHierarchy(ConfigFiles);
 
 		List<PakFileRules> RulesList = new List<PakFileRules>();
@@ -1375,8 +1377,29 @@ public partial class Project : CommandUtils
 		{
 			//LogInformation("Building PakFileRules for Section {0}", SectionName);
 
-			string PlatformString, TargetString, PakString;
+			bool bOnlyChunkedBuilds = false;
+			if (!PakRulesConfig.TryGetValue(SectionName, "bOnlyChunkedBuilds", out bOnlyChunkedBuilds))
+			{
+				bOnlyChunkedBuilds = false;
+			}
+
+			bool bOnlyNonChunkedBuilds = false;
+			if (!PakRulesConfig.TryGetValue(SectionName, "bOnlyNonChunkedBuilds", out bOnlyNonChunkedBuilds))
+			{
+				bOnlyNonChunkedBuilds = false;
+			}
+
+			if (bChunkedBuild && bOnlyNonChunkedBuilds)
+			{
+				continue;
+			}
+
+			if(!bChunkedBuild && bOnlyChunkedBuilds)
+			{
+				continue;
+			}
 
+			string PlatformString;
 			if (PakRulesConfig.TryGetValue(SectionName, "Platforms", out PlatformString))
 			{
 				string[] PlatformStrings = PlatformString.Split(new char[] { ',', ' ' }, StringSplitOptions.RemoveEmptyEntries);
@@ -1404,6 +1427,7 @@ public partial class Project : CommandUtils
 				}
 			}
 
+			string TargetString;
 			if (PakRulesConfig.TryGetValue(SectionName, "Targets", out TargetString))
 			{
 				string[] TargetStrings = TargetString.Split(new char[] { ',', ' ' }, StringSplitOptions.RemoveEmptyEntries);
@@ -1437,6 +1461,7 @@ public partial class Project : CommandUtils
 			PakRules.Name = SectionName;
 			PakRulesConfig.TryGetValue(SectionName, "bExcludeFromPaks", out PakRules.bExcludeFromPaks);
 			PakRulesConfig.TryGetValue(SectionName, "bOverrideChunkManifest", out PakRules.bOverrideChunkManifest);
+			string PakString;
 			PakRulesConfig.TryGetValue(SectionName, "OverridePaks", out PakString);
 
 			if (PakString != null && PakString.Length > 0)
@@ -1503,6 +1528,7 @@ public partial class Project : CommandUtils
 					{
 						return;
 					}
+
 					if (PakRules.OverridePaks != null)
 					{
 						LogInformation("Overridding chunk assignment {0} to {1}, this can cause broken references", StagingFile.Key, string.Join(", ", PakRules.OverridePaks));
@@ -1511,7 +1537,6 @@ public partial class Project : CommandUtils
 					{
 						LogInformation("Removing {0} from pak despite chunk assignment, this can cause broken references", StagingFile.Key);
 					}
-
 				}
 
 				bExcludeFromPaks = PakRules.bExcludeFromPaks;
@@ -1525,6 +1550,7 @@ public partial class Project : CommandUtils
 				{
 					//LogInformation("Excluding {0} from pak file", StagingFile.Key);
 				}
+
 				return;
 			}
 		}
@@ -1546,12 +1572,13 @@ public partial class Project : CommandUtils
 
 		List<PakFileRules> PakRulesList = GetPakFileRules(Params, SC);
 
-		List<string> PakList = new List<string>();
 		List<string> FilesToRemove = new List<string>();
 
 		// Apply the pak file rules, this can remove things but will not override the pak file name
 		foreach (var StagingFile in UnrealPakResponseFile)
 		{
+			List<string> PakList = new List<string>();
+
 			bool bExcludeFromPaks = false;
 			ApplyPakFileRules(PakRulesList, StagingFile, PakList, out bExcludeFromPaks);
 
@@ -2421,9 +2448,9 @@ public partial class Project : CommandUtils
 					ChunkDefinition Chunk = ChunkDefinitions[ChunkIndex];
 
 					if (Chunk.Manifest.Contains(OriginalFilename) ||
-							Chunk.Manifest.Contains(OriginalReplaceSlashes) ||
-							Chunk.Manifest.Contains(NoExtension) ||
-							Chunk.Manifest.Contains(NoExtensionReplaceSlashes))
+						Chunk.Manifest.Contains(OriginalReplaceSlashes) ||
+						Chunk.Manifest.Contains(NoExtension) ||
+						Chunk.Manifest.Contains(NoExtensionReplaceSlashes))
 					{
 						PakList.Add(Chunk.ChunkName);
 					}
