QOL improvements for Mac from UE4/Main that missed original branch

- Use msbuild where possible for building UAT by checking for mono5
- Hash state of UAT scritps post-compile to skip redundant rebuilds

diff --git a/Engine/Source/Programs/AutomationTool/AutomationUtils/ScriptCompiler.cs b/Engine/Source/Programs/AutomationTool/AutomationUtils/ScriptCompiler.cs
--- a/Engine/Source/Programs/AutomationTool/AutomationUtils/ScriptCompiler.cs
+++ b/Engine/Source/Programs/AutomationTool/AutomationUtils/ScriptCompiler.cs
@@ -187,6 +187,13 @@ namespace AutomationTool
 
 			foreach (string Module in Modules)
 			{
+				// this project is built by the AutomationTool project that the RunUAT script builds so 
+				// it will always be newer than was last built
+				if (Module.Contains("AutomationUtils.Automation"))
+				{
+					continue;
+				}
+
 				CsProjectInfo Proj;
 
 				Dictionary<string, string> Properties = new Dictionary<string, string>();
@@ -286,7 +293,7 @@ namespace AutomationTool
 
 			string BuildTool = CommandUtils.CmdEnv.MsBuildExe;
 
-			// msbuild (standard on windows, in mono >= 5.0 is preferred due to speed and parallel compilation)
+			// msbuild (standard on windows, in mono >=5.0 is preferred due to speed and parallel compilation)
 			bool UseParallelMsBuild = Path.GetFileNameWithoutExtension(BuildTool).ToLower() == "msbuild";
 
 			if (UseParallelMsBuild)
diff --git a/Engine/Source/Programs/DotNETCommon/DotNETUtilities/HashCollection.cs b/Engine/Source/Programs/DotNETCommon/DotNETUtilities/HashCollection.cs
--- a/Engine/Source/Programs/DotNETCommon/DotNETUtilities/HashCollection.cs
+++ b/Engine/Source/Programs/DotNETCommon/DotNETUtilities/HashCollection.cs
@@ -179,7 +179,7 @@ namespace Tools.DotNETCommon
 				{
 					HashEntry RHSEntry = RHS.GetEntryForName(Entry.Name);
 					Log.TraceInformation("RHS hash mismatch for {0}", Entry.Name);
-					Log.TraceInformation("Current:\n\t{0}\n\t{1}\n\t{2}", Entry.Hash, Entry.Name, Entry.MetaData);
+					Log.TraceInformation("LHS:\n\t{0}\n\t{1}\n\t{2}", Entry.Hash, Entry.Name, Entry.MetaData);
 					Log.TraceInformation("RHS:\n\t{0}\n\t{1}\n\t{2}", RHSEntry.Hash, RHSEntry.Name, RHSEntry.MetaData);
 				}
 			}
