Gauntlet: Add support for dervice reseration pools
#rnx

diff --git a/Engine/Source/Programs/AutomationTool/AutomationUtils/DeviceReservation.cs b/Engine/Source/Programs/AutomationTool/AutomationUtils/DeviceReservation.cs
--- a/Engine/Source/Programs/AutomationTool/AutomationUtils/DeviceReservation.cs
+++ b/Engine/Source/Programs/AutomationTool/AutomationUtils/DeviceReservation.cs
@@ -51,13 +51,13 @@ namespace AutomationTool.DeviceReservation
 		/// </summary>
 		/// <param name="InWorkingDirectory">Working directory which contains the devices.xml and reservations.xml files. Usually a network share.</param>
 		/// <param name="InDeviceTypes">An array of device types to reserve, one for each device requested. These must match the device types listed in devices.xml.</param>
-		public DeviceReservationAutoRenew(string InReservationBaseUri, int RetryMax, params string[] InDeviceTypes)
+		public DeviceReservationAutoRenew(string InReservationBaseUri, int RetryMax, string PoolID, params string[] InDeviceTypes)
 		{
 			ReservationBaseUri = new Uri(InReservationBaseUri);
 
 			// Make a device reservation for all the required device types.
 			// This blocks until the reservation is successful.
-			ActiveReservation = Reservation.Create(ReservationBaseUri, InDeviceTypes, ReserveTime, RetryMax);
+			ActiveReservation = Reservation.Create(ReservationBaseUri, InDeviceTypes, ReserveTime, RetryMax, PoolID);
 
 			// Resolve the device IPs
 			ReservedDevices = new List<Device>();
@@ -203,6 +203,7 @@ namespace AutomationTool.DeviceReservation
 			public string Hostname;
 			public TimeSpan Duration;
 			public string ReservationDetails;
+			public string PoolID;
 		}
 
 		private static string SanitizeErrorMessage(string Message)
@@ -221,7 +222,7 @@ namespace AutomationTool.DeviceReservation
 			return Message;
 		}
 
-		public static Reservation Create(Uri BaseUri, string[] DeviceTypes, TimeSpan Duration, int RetryMax = 5)
+		public static Reservation Create(Uri BaseUri, string[] DeviceTypes, TimeSpan Duration, int RetryMax = 5, string PoolID = "")
 		{
 			bool bFirst = true;
 			TimeSpan RetryTime = TimeSpan.FromMinutes(1);
@@ -247,7 +248,8 @@ namespace AutomationTool.DeviceReservation
 						DeviceTypes = DeviceTypes,
 						Hostname = Environment.MachineName,
 						Duration = Duration,
-						ReservationDetails = ReservationDetails
+						ReservationDetails = ReservationDetails,
+						PoolID = PoolID
 					});
 				}
 				catch (WebException WebEx)
diff --git a/Engine/Source/Programs/AutomationTool/Gauntlet/Framework/Devices/Gauntlet.DevicePool.cs b/Engine/Source/Programs/AutomationTool/Gauntlet/Framework/Devices/Gauntlet.DevicePool.cs
--- a/Engine/Source/Programs/AutomationTool/Gauntlet/Framework/Devices/Gauntlet.DevicePool.cs
+++ b/Engine/Source/Programs/AutomationTool/Gauntlet/Framework/Devices/Gauntlet.DevicePool.cs
@@ -578,9 +578,11 @@ namespace Gauntlet
 			{
 				DeviceReservationAutoRenew DeviceReservation = null;
 
+				string PoolID = Globals.WorkerPoolID != -1 ? Globals.WorkerPoolID.ToString() : "";
+
 				try
 				{
-					DeviceReservation = new DeviceReservationAutoRenew(DeviceURL, 0, Devices.ToArray());
+					DeviceReservation = new DeviceReservationAutoRenew(DeviceURL, 0, PoolID, Devices.ToArray());
 				}
 				catch (Exception Ex)
 				{
diff --git a/Engine/Source/Programs/AutomationTool/Gauntlet/Framework/Gauntlet.Utils.cs b/Engine/Source/Programs/AutomationTool/Gauntlet/Framework/Gauntlet.Utils.cs
--- a/Engine/Source/Programs/AutomationTool/Gauntlet/Framework/Gauntlet.Utils.cs
+++ b/Engine/Source/Programs/AutomationTool/Gauntlet/Framework/Gauntlet.Utils.cs
@@ -49,6 +49,21 @@ namespace Gauntlet
 
 		}
 
+		/// <summary>
+		/// Get the worker pool id of the host worker, pools are assigned to teams such as QA, Automation, etc
+		/// returns -1 if instance is not running on a worker pool
+		/// </summary>
+		public static int WorkerPoolID
+		{
+			get
+			{
+				int Default = -1;
+				return Params.ParseValue("workerpoolid", Default);
+			}
+
+		}
+
+
 		/// <summary>
 		/// Returns true if Gauntlet instance is a member of a worker group
 		/// </summary>
