Gauntlet: Adding "NotBulk" flag, so we can specifically exclude bulk builds if we wish to
#rnx

diff --git a/Engine/Source/Programs/AutomationTool/Gauntlet/Framework/Base/Gauntlet.BuildSource.cs b/Engine/Source/Programs/AutomationTool/Gauntlet/Framework/Base/Gauntlet.BuildSource.cs
--- a/Engine/Source/Programs/AutomationTool/Gauntlet/Framework/Base/Gauntlet.BuildSource.cs
+++ b/Engine/Source/Programs/AutomationTool/Gauntlet/Framework/Base/Gauntlet.BuildSource.cs
@@ -19,7 +19,8 @@ namespace Gauntlet
 		CanReplaceCommandLine = (1 << 2),		// build can run arbitrary command lines
 		CanReplaceExecutable = (1 << 3),		// build can use exes from elsewhere
         Bulk = (1 << 4),						// build is full-content, now startup download
-		ContentOnlyProject = (1 << 5)			// build is a content-only project, e.g. exe is ue4game.exe
+		NotBulk = (1 << 5),						// build is not bulk
+		ContentOnlyProject = (1 << 6)			// build is a content-only project, e.g. exe is ue4game.exe
     }
 
 	/// <summary>
diff --git a/Engine/Source/Programs/AutomationTool/Gauntlet/Platform/Android/Gauntlet.AndroidBuildSource.cs b/Engine/Source/Programs/AutomationTool/Gauntlet/Platform/Android/Gauntlet.AndroidBuildSource.cs
--- a/Engine/Source/Programs/AutomationTool/Gauntlet/Platform/Android/Gauntlet.AndroidBuildSource.cs
+++ b/Engine/Source/Programs/AutomationTool/Gauntlet/Platform/Android/Gauntlet.AndroidBuildSource.cs
@@ -142,12 +142,16 @@ namespace Gauntlet
 				{
 					Flags |= BuildFlags.CanReplaceExecutable;
 				}
-                if (AbsPath.Contains("Bulk"))
-                {
-                    Flags |= BuildFlags.Bulk;
-                }
+				if (AbsPath.Contains("Bulk"))
+				{
+					Flags |= BuildFlags.Bulk;
+				}
+				else
+				{
+					Flags |= BuildFlags.NotBulk;
+				}
 
-                AndroidBuild NewBuild = new AndroidBuild(UnrealConfig, AndroidPackageName, SourceApkPath, FilesToInstall, Flags);
+				AndroidBuild NewBuild = new AndroidBuild(UnrealConfig, AndroidPackageName, SourceApkPath, FilesToInstall, Flags);
 
 				DiscoveredBuilds.Add(NewBuild);
 
diff --git a/Engine/Source/Programs/AutomationTool/Gauntlet/Platform/IOS/Gauntlet.IOSBuildSource.cs b/Engine/Source/Programs/AutomationTool/Gauntlet/Platform/IOS/Gauntlet.IOSBuildSource.cs
--- a/Engine/Source/Programs/AutomationTool/Gauntlet/Platform/IOS/Gauntlet.IOSBuildSource.cs
+++ b/Engine/Source/Programs/AutomationTool/Gauntlet/Platform/IOS/Gauntlet.IOSBuildSource.cs
@@ -148,6 +148,10 @@ namespace Gauntlet
 				{
 					Flags |= BuildFlags.Bulk;
 				}
+				else
+				{
+					Flags |= BuildFlags.NotBulk;
+				}
 
 				string SourceIPAPath = Fi.FullName;
 				string PackageName = GetBundleIdentifier(SourceIPAPath);
diff --git a/Engine/Source/Programs/AutomationTool/Gauntlet/Unreal/Base/Gauntlet.UnrealSession.cs b/Engine/Source/Programs/AutomationTool/Gauntlet/Unreal/Base/Gauntlet.UnrealSession.cs
--- a/Engine/Source/Programs/AutomationTool/Gauntlet/Unreal/Base/Gauntlet.UnrealSession.cs
+++ b/Engine/Source/Programs/AutomationTool/Gauntlet/Unreal/Base/Gauntlet.UnrealSession.cs
@@ -145,10 +145,17 @@ namespace Gauntlet
 
 			//@todo: for bulk/packaged builds, we should mark the platform as capable of these as we're catching global and not test specific flags
 			// where we may be running parallel tests on multiple platforms
-			if (Globals.Params.ParseParam("bulk") && (InPlatform == UnrealTargetPlatform.Android || InPlatform == UnrealTargetPlatform.IOS))
-            {
-                RequiredBuildFlags |= BuildFlags.Bulk;
-            }
+			if (InPlatform == UnrealTargetPlatform.Android || InPlatform == UnrealTargetPlatform.IOS)
+			{
+				if (Globals.Params.ParseParam("notbulk"))
+				{
+					RequiredBuildFlags |= BuildFlags.NotBulk;
+				}
+				else if (Globals.Params.ParseParam("bulk"))
+				{
+					RequiredBuildFlags |= BuildFlags.Bulk;
+				}
+			}
 
 			if (Globals.Params.ParseParam("packaged") && (InPlatform == UnrealTargetPlatform.Switch || InPlatform == UnrealTargetPlatform.XboxOne))
 			{
diff --git a/Engine/Source/Programs/AutomationTool/Gauntlet/Unreal/BuildSource/Gauntlet.UnrealBuildSource.cs b/Engine/Source/Programs/AutomationTool/Gauntlet/Unreal/BuildSource/Gauntlet.UnrealBuildSource.cs
--- a/Engine/Source/Programs/AutomationTool/Gauntlet/Unreal/BuildSource/Gauntlet.UnrealBuildSource.cs
+++ b/Engine/Source/Programs/AutomationTool/Gauntlet/Unreal/BuildSource/Gauntlet.UnrealBuildSource.cs
@@ -608,6 +608,10 @@ namespace Gauntlet
                     {
                         Flags |= BuildFlags.Bulk;
                     }
+					else
+					{
+						Flags |= BuildFlags.NotBulk;
+					}
 
                     var Build = GetMatchingBuilds(TargetType, TargetPlatform, TargetConfiguration, Flags).FirstOrDefault();
 
