Fixed issue in BuildTarget script with attempting to parse platforms as enums.

Added fallback for projects that don't define Client/Game target scripts

diff --git a/Engine/Source/Programs/AutomationTool/Scripts/BuildTarget.Automation.cs b/Engine/Source/Programs/AutomationTool/Scripts/BuildTarget.Automation.cs
--- a/Engine/Source/Programs/AutomationTool/Scripts/BuildTarget.Automation.cs
+++ b/Engine/Source/Programs/AutomationTool/Scripts/BuildTarget.Automation.cs
@@ -54,12 +54,38 @@ namespace AutomationTool
 			bool NoTools = ParseParam("notools");
 
 			IEnumerable<string> TargetList = Targets.Split(new char[] { '+' }, StringSplitOptions.RemoveEmptyEntries);
-			IEnumerable<UnrealTargetConfiguration> ConfigurationList = Configurations.Split(new char[] { '+' }, StringSplitOptions.RemoveEmptyEntries)
-																		.Select(C => (UnrealTargetConfiguration)Enum.Parse(typeof(UnrealTargetConfiguration), C, true));
+			IEnumerable<UnrealTargetConfiguration> ConfigurationList = null;
+			IEnumerable<UnrealTargetPlatform> PlatformList = null;
 
-			IEnumerable<UnrealTargetPlatform> PlatformList = Platforms.Split(new char[] { '+' }, StringSplitOptions.RemoveEmptyEntries)
-																		.Select(C => (UnrealTargetPlatform)Enum.Parse(typeof(UnrealTargetPlatform), C, true));
+			try
+			{
+				ConfigurationList = Configurations.Split(new char[] { '+' }, StringSplitOptions.RemoveEmptyEntries)
+																		.Select(C => (UnrealTargetConfiguration)Enum.Parse(typeof(UnrealTargetConfiguration), C, true)).ToArray();
+			}
+			catch (Exception Ex)
+			{
+				LogError("Failed to parse configuration string. {0}", Ex.Message);
+				return ExitCode.Error_Arguments;
+			}
 
+			try
+			{
+				PlatformList = Platforms.Split(new char[] { '+' }, StringSplitOptions.RemoveEmptyEntries)
+																		.Select(C =>
+																		{
+																			UnrealTargetPlatform Platform;
+																			if (!UnrealTargetPlatform.TryParse(C, out Platform))
+																			{
+																				throw new AutomationException("No such platform {0}", C);
+																			}
+																			return Platform;
+																		}).ToArray();
+			}
+			catch (Exception Ex)
+			{
+				LogError("Failed to parse configuration string. {0}", Ex.Message);
+				return ExitCode.Error_Arguments;
+			}
 
 			if (String.IsNullOrEmpty(ProjectName))
 			{
@@ -82,6 +108,17 @@ namespace AutomationTool
 				{
 					string TargetScript = TargetScripts.Where(S => S.IndexOf(TargetName, StringComparison.OrdinalIgnoreCase) >= 0).FirstOrDefault();
 
+					if (TargetScript == null && (
+							!TargetName.Equals("Client", StringComparison.OrdinalIgnoreCase) ||
+							!TargetName.Equals("Game", StringComparison.OrdinalIgnoreCase)
+							)
+						)
+					{
+						// if there's no ProjectGame.Target.cs or ProjectClient.Target.cs then
+						// fallback to Project.Target.cs
+						TargetScript = TargetScripts.Where(S => S.IndexOf(ProjectName + ".", StringComparison.OrdinalIgnoreCase) >= 0).FirstOrDefault();
+					}
+
 					if (TargetScript == null)
 					{
 						throw new AutomationException("No Target.cs file for target {0} in project {1}", TargetName, ProjectName);
