Filter supported service device platforms
#rnx

diff --git a/Engine/Source/Programs/AutomationTool/Gauntlet/Framework/Devices/Gauntlet.DevicePool.cs b/Engine/Source/Programs/AutomationTool/Gauntlet/Framework/Devices/Gauntlet.DevicePool.cs
--- a/Engine/Source/Programs/AutomationTool/Gauntlet/Framework/Devices/Gauntlet.DevicePool.cs
+++ b/Engine/Source/Programs/AutomationTool/Gauntlet/Framework/Devices/Gauntlet.DevicePool.cs
@@ -539,7 +539,7 @@ namespace Gauntlet
 				{ UnrealTargetPlatform.XboxOne , "XboxOne-DevKit" },
 				{ UnrealTargetPlatform.Android , "Android" },
 				{ UnrealTargetPlatform.Switch , "Switch" }
-		};
+			};
 
 			List<string> Devices = new List<string>();
 
@@ -979,10 +979,23 @@ namespace Gauntlet
 			var TooFewTotalDevices = RequiredDevices.Where(KP => TotalDeviceTypes[KP.Key] < RequiredDevices[KP.Key]).Select(KP => KP.Key);
 			var TooFewCurrentDevices = RequiredDevices.Where(KP => AvailableDeviceTypes[KP.Key] < RequiredDevices[KP.Key]).Select(KP => KP.Key);
 
-			// Request devices from the service if we need them
-			if (UseServiceDevices && !String.IsNullOrEmpty(DeviceURL) && (TooFewTotalDevices.Count() > 0 || TooFewCurrentDevices.Count() > 0))
+			List<UnrealTargetPlatform> ServicePlatforms = new List<UnrealTargetPlatform>()
+			{
+				UnrealTargetPlatform.PS4, UnrealTargetPlatform.XboxOne, UnrealTargetPlatform.Switch
+			};
+
+			// support Android over wifi, though not on workers
+			if (!Globals.IsWorker)
 			{
-				var Devices = TooFewTotalDevices.Concat(TooFewCurrentDevices);
+				ServicePlatforms.Add(UnrealTargetPlatform.Android);
+			}
+
+			var Devices = TooFewTotalDevices.Concat(TooFewCurrentDevices);
+			var UnsupportedPlatforms = Devices.Where(D => !ServicePlatforms.Contains(D.Platform));
+
+			// Request devices from the service if we need them
+			if (UseServiceDevices && !String.IsNullOrEmpty(DeviceURL) && UnsupportedPlatforms.Count() == 0 && (TooFewTotalDevices.Count() > 0 || TooFewCurrentDevices.Count() > 0))
+			{				
 
 				Dictionary<UnrealTargetConstraint, int> DeviceCounts = new Dictionary<UnrealTargetConstraint, int>();
 
