Fix UAT compile error.

#rnx

diff --git a/Engine/Source/Programs/AutomationTool/Gauntlet/Framework/Devices/Gauntlet.DevicePool.cs b/Engine/Source/Programs/AutomationTool/Gauntlet/Framework/Devices/Gauntlet.DevicePool.cs
--- a/Engine/Source/Programs/AutomationTool/Gauntlet/Framework/Devices/Gauntlet.DevicePool.cs
+++ b/Engine/Source/Programs/AutomationTool/Gauntlet/Framework/Devices/Gauntlet.DevicePool.cs
@@ -991,7 +991,7 @@ namespace Gauntlet
 			}
 
 			var Devices = TooFewTotalDevices.Concat(TooFewCurrentDevices);
-			var UnsupportedPlatforms = Devices.Where(D => !ServicePlatforms.Contains(D.Platform));
+			var UnsupportedPlatforms = Devices.Where(D => !ServicePlatforms.Contains(D.Platform.Value));
 
 			// Request devices from the service if we need them
 			if (UseServiceDevices && !String.IsNullOrEmpty(DeviceURL) && UnsupportedPlatforms.Count() == 0 && (TooFewTotalDevices.Count() > 0 || TooFewCurrentDevices.Count() > 0))
