[NO NON-PII COMMIT MESSAGE]

diff --git a/Engine/Source/Programs/AutomationTool/ShooterGame/Properties/AssemblyInfo.cs b/Engine/Source/Programs/AutomationTool/ShooterGame/Properties/AssemblyInfo.cs
new file mode 100644
--- /dev/null
+++ b/Engine/Source/Programs/AutomationTool/ShooterGame/Properties/AssemblyInfo.cs
@@ -0,0 +1,38 @@
+// Copyright 1998-2019 Epic Games, Inc. All Rights Reserved.
+
+using System.Reflection;
+using System.Runtime.CompilerServices;
+using System.Runtime.InteropServices;
+
+// General Information about an assembly is controlled through the following 
+// set of attributes. Change these attribute values to modify the information
+// associated with an assembly.
+[assembly: AssemblyTitle("ShooterGame.Automation")]
+[assembly: AssemblyDescription("")]
+[assembly: AssemblyConfiguration("")]
+[assembly: AssemblyCompany("")]
+[assembly: AssemblyProduct("ShooterGame.Automation")]
+[assembly: AssemblyCopyright("Copyright 1998-2019 Epic Games, Inc.")]
+[assembly: AssemblyTrademark("")]
+[assembly: AssemblyCulture("")]
+
+// Setting ComVisible to false makes the types in this assembly not visible 
+// to COM components.  If you need to access a type in this assembly from 
+// COM, set the ComVisible attribute to true on that type.
+[assembly: ComVisible(false)]
+
+// The following GUID is for the ID of the typelib if this project is exposed to COM
+[assembly: Guid("be36e4c1-ba0c-4c50-97e0-2994189bcad7")]
+
+// Version information for an assembly consists of the following four values:
+//
+//      Major Version
+//      Minor Version 
+//      Build Number
+//      Revision
+//
+// You can specify all the values or you can default the Build and Revision Numbers 
+// by using the '*' as shown below:
+// [assembly: AssemblyVersion("1.0.*")]
+[assembly: AssemblyVersion("1.0.0.0")]
+[assembly: AssemblyFileVersion("1.0.0.0")]
diff --git a/Engine/Source/Programs/AutomationTool/ShooterGame/RunShooterTests.cs b/Engine/Source/Programs/AutomationTool/ShooterGame/RunShooterTests.cs
new file mode 100644
--- /dev/null
+++ b/Engine/Source/Programs/AutomationTool/ShooterGame/RunShooterTests.cs
@@ -0,0 +1,14 @@
+// Copyright 1998-2019 Epic Games, Inc.All Rights Reserved.
+using System;
+using System.Collections.Generic;
+using System.Linq;
+using System.Text;
+using System.Threading.Tasks;
+using Gauntlet;
+
+namespace ShooterTest
+{
+	class RunShooterTests : Gauntlet.RunUnreal
+	{
+	}
+}
diff --git a/Engine/Source/Programs/AutomationTool/ShooterGame/ShooterGame.Automation.csproj b/Engine/Source/Programs/AutomationTool/ShooterGame/ShooterGame.Automation.csproj
new file mode 100644
--- /dev/null
+++ b/Engine/Source/Programs/AutomationTool/ShooterGame/ShooterGame.Automation.csproj
@@ -0,0 +1,133 @@
+﻿<?xml version="1.0" encoding="utf-8"?>
+<Project ToolsVersion="12.0" DefaultTargets="Build" xmlns="http://schemas.microsoft.com/developer/msbuild/2003">
+  <Import Project="$(MSBuildExtensionsPath)\$(MSBuildToolsVersion)\Microsoft.Common.props" Condition="Exists('$(MSBuildExtensionsPath)\$(MSBuildToolsVersion)\Microsoft.Common.props')" />
+  <PropertyGroup>
+    <Configuration Condition=" '$(Configuration)' == '' ">Debug</Configuration>
+    <Platform Condition=" '$(Platform)' == '' ">AnyCPU</Platform>
+    <ProjectGuid>{D6C6AD55-30E8-4CF1-B84D-489C61E7C8EC}</ProjectGuid>
+    <OutputType>Library</OutputType>
+    <AppDesignerFolder>Properties</AppDesignerFolder>
+    <RootNamespace>ShooterGame.Automation</RootNamespace>
+    <AssemblyName>ShooterGame.Automation</AssemblyName>
+    <TargetFrameworkVersion>v4.6.2</TargetFrameworkVersion>
+    <FileAlignment>512</FileAlignment>
+    <WarningsNotAsErrors>612,618</WarningsNotAsErrors>
+    <TargetFrameworkProfile />
+  </PropertyGroup>
+  <PropertyGroup Condition=" '$(Configuration)|$(Platform)' == 'Debug|AnyCPU' ">
+    <DebugSymbols>true</DebugSymbols>
+    <DebugType>full</DebugType>
+    <Optimize>false</Optimize>
+    <OutputPath>..\..\..\..\Binaries\DotNET\AutomationScripts\</OutputPath>
+    <DefineConstants>DEBUG;TRACE</DefineConstants>
+    <ErrorReport>prompt</ErrorReport>
+    <WarningLevel>4</WarningLevel>
+    <PlatformTarget>AnyCPU</PlatformTarget>
+    <TreatWarningsAsErrors>false</TreatWarningsAsErrors>
+  </PropertyGroup>
+  <PropertyGroup Condition=" '$(Configuration)|$(Platform)' == 'Development|AnyCPU' ">
+    <DebugType>pdbonly</DebugType>
+    <Optimize>false</Optimize>
+    <OutputPath>..\..\..\..\..\Engine\Binaries\DotNET\AutomationScripts\</OutputPath>
+    <DefineConstants>TRACE</DefineConstants>
+    <ErrorReport>prompt</ErrorReport>
+    <WarningLevel>4</WarningLevel>
+    <PlatformTarget>AnyCPU</PlatformTarget>
+    <TreatWarningsAsErrors>false</TreatWarningsAsErrors>
+  </PropertyGroup>
+  <ItemGroup>
+    <Reference Include="System" />
+    <Reference Include="System.Core" />
+    <Reference Include="System.Drawing" />
+    <Reference Include="System.IO.Compression" />
+    <Reference Include="System.IO.Compression.FileSystem" />
+    <Reference Include="System.Management" />
+    <Reference Include="System.Runtime.Serialization" />
+    <Reference Include="System.ServiceModel" />
+    <Reference Include="System.Web" />
+    <Reference Include="System.Web.Extensions" />
+    <Reference Include="System.Windows.Forms" />
+    <Reference Include="System.Xml.Linq" />
+    <Reference Include="System.Data.DataSetExtensions" />
+    <Reference Include="Microsoft.CSharp" />
+    <Reference Include="System.Data" />
+    <Reference Include="System.Xml" />
+  </ItemGroup>
+  <ItemGroup>
+    <Compile Include="Properties\AssemblyInfo.cs" />
+    <Compile Include="RunShooterTests.cs" />
+    <Compile Include="Tests\ShooterTest.ListenServerQuickMatchTest.cs" />
+    <Compile Include="Tests\ShooterTest.ListenServerTest.cs" />
+    <Compile Include="Tests\ShooterTest.BootTest.cs" />
+    <Compile Include="Tests\ShooterTest.BasicDedicatedServerTest.cs" />
+    <Compile Include="Tests\ShooterTest.DedicatedServerTest.cs" />
+    <Compile Include="Tests\ShooterTest.TestConfig.cs" />
+  </ItemGroup>
+  <ItemGroup>
+    <ProjectReference Include="..\..\..\..\..\Engine\Source\Programs\DotNETCommon\DotNETUtilities\DotNETUtilities.csproj">
+      <Project>{5d7d66e8-8c76-4af9-b3ec-2ef03421d730}</Project>
+      <Name>DotNETUtilities</Name>
+    </ProjectReference>
+    <ProjectReference Include="..\..\..\..\..\Engine\Source\Programs\AutomationTool\Android\Android.Automation.csproj">
+      <Project>{32e3a47e-359a-4f78-ae6c-703b87de0b97}</Project>
+      <Name>Android.Automation</Name>
+      <Private>False</Private>
+    </ProjectReference>
+    <ProjectReference Include="..\..\..\..\..\Engine\Source\Programs\AutomationTool\AutomationUtils\AutomationUtils.Automation.csproj">
+      <Project>{2c96a7f2-b1a3-4258-8e0a-e588ff41a53e}</Project>
+      <Name>AutomationUtils.Automation</Name>
+      <Private>False</Private>
+    </ProjectReference>
+    <ProjectReference Include="..\..\..\..\..\Engine\Source\Programs\AutomationTool\Gauntlet\Gauntlet.Automation.csproj">
+      <Project>{767b4f85-ab56-4b00-a033-04c7600acc3d}</Project>
+      <Name>Gauntlet.Automation</Name>
+    </ProjectReference>
+    <ProjectReference Include="..\..\..\..\..\Engine\Source\Programs\AutomationTool\Localization\Localization.Automation.csproj">
+      <Project>{c2f959b1-1727-4598-9c32-3d9acce72750}</Project>
+      <Name>Localization.Automation</Name>
+    </ProjectReference>
+    <ProjectReference Include="..\..\..\..\..\Engine\Source\Programs\AutomationTool\XLocLocalization\XLocLocalization.Automation.csproj">
+      <Project>{7e2d6784-c9b7-40c6-a89b-0a8afddce60d}</Project>
+      <Name>XLocLocalization.Automation</Name>
+      <Private>False</Private>
+    </ProjectReference>
+    <ProjectReference Include="..\..\..\..\..\Engine\Source\Programs\AutomationTool\OneSkyLocalization\OneSkyLocalization.Automation.csproj">
+      <Project>{39369368-5603-46da-baae-766bb3cb2b51}</Project>
+      <Name>OneSkyLocalization.Automation</Name>
+    </ProjectReference>
+    <ProjectReference Include="..\..\..\..\..\Engine\Source\Programs\AutomationTool\Scripts\AutomationScripts.Automation.csproj">
+      <Project>{8aa00d65-0954-4a27-ac0d-fb8b1106120f}</Project>
+      <Name>AutomationScripts.Automation</Name>
+      <Private>False</Private>
+    </ProjectReference>
+    <ProjectReference Include="..\..\..\..\..\Engine\Source\Programs\UnrealBuildTool\UnrealBuildTool.csproj">
+      <Project>{fd7c5e1a-cfe4-4fd5-a525-1eb1599a39ac}</Project>
+      <Name>UnrealBuildTool</Name>
+      <Private>False</Private>
+    </ProjectReference>
+  </ItemGroup>
+  <ItemGroup />
+  <ItemGroup />
+  <Choose>
+    <When Condition="Exists('Platform\Android\Gauntlet.TargetDeviceAndroid.cs')" />
+  </Choose>
+  <Choose>
+    <When Condition="Exists('Platform\IOS\Gauntlet.TargetDeviceIOS.cs')" />
+  </Choose>
+  <Choose>
+    <When Condition="Exists('Platform\Mac\Gauntlet.TargetDeviceMac.cs')" />
+  </Choose>
+  <Choose>
+    <When Condition="Exists('Platform\PS4\Gauntlet.TargetDevicePS4.cs')" />
+  </Choose>
+  <Choose>
+    <When Condition="Exists('Platform\Switch\Gauntlet.TargetDeviceSwitch.cs')" />
+  </Choose>
+  <Choose>
+    <When Condition="Exists('Platform\Windows\Gauntlet.TargetDeviceWindows.cs')" />
+  </Choose>
+  <Choose>
+    <When Condition="Exists('Platform\XboxOne\Gauntlet.TargetDeviceXboxOne.cs')" />
+  </Choose>
+  <Import Project="$(MSBuildToolsPath)\Microsoft.CSharp.targets" />
+</Project>
\ No newline at end of file
diff --git a/Engine/Source/Programs/AutomationTool/ShooterGame/Tests/ShooterTest.BasicDedicatedServerTest.cs b/Engine/Source/Programs/AutomationTool/ShooterGame/Tests/ShooterTest.BasicDedicatedServerTest.cs
new file mode 100644
--- /dev/null
+++ b/Engine/Source/Programs/AutomationTool/ShooterGame/Tests/ShooterTest.BasicDedicatedServerTest.cs
@@ -0,0 +1,35 @@
+// Copyright 1998-2019 Epic Games, Inc.All Rights Reserved.
+using System;
+using System.Collections.Generic;
+using System.Linq;
+using System.Text;
+using System.Threading.Tasks;
+using EpicGame;
+using Gauntlet;
+
+namespace ShooterTest
+{
+	/// <summary>
+	/// This tests a dedicated server and a client that opens strait into the server.  
+	/// </summary>
+	public class BasicDedicatedServerTest : UnrealTestNode<ShooterTestConfig>
+	{
+		public BasicDedicatedServerTest(UnrealTestContext InContext) : base(InContext)
+		{
+		}
+
+		public override ShooterTestConfig GetConfiguration()
+		{
+			ShooterTestConfig Config = base.GetConfiguration();
+			Config.PreAssignAccount = false;
+			Config.NoMCP = true;
+
+			UnrealTestRole Client = Config.RequireRole(UnrealTargetRole.Client);
+			UnrealTestRole Server = Config.RequireRole(UnrealTargetRole.Server);
+
+			Client.Controllers.Add("BasicDedicatedServerTest");
+
+			return Config;
+		}
+	}
+}
diff --git a/Engine/Source/Programs/AutomationTool/ShooterGame/Tests/ShooterTest.BootTest.cs b/Engine/Source/Programs/AutomationTool/ShooterGame/Tests/ShooterTest.BootTest.cs
new file mode 100644
--- /dev/null
+++ b/Engine/Source/Programs/AutomationTool/ShooterGame/Tests/ShooterTest.BootTest.cs
@@ -0,0 +1,30 @@
+// Copyright 1998-2019 Epic Games, Inc.All Rights Reserved.
+using System;
+using System.Collections.Generic;
+using System.Linq;
+using System.Text;
+using System.Threading.Tasks;
+using EpicGame;
+using Gauntlet;
+
+namespace ShooterTest
+{
+	/// <summary>
+	/// Basic Boot Test
+	/// </summary>
+	public class BootTest : EpicGameTestNode
+	{
+		public BootTest(UnrealTestContext InContext) : base (InContext)
+		{
+		}
+
+		public override EpicGameTestConfig GetConfiguration()
+		{
+			EpicGameTestConfig Config = base.GetConfiguration();
+			UnrealTestRole Client = Config.RequireRole(UnrealTargetRole.Client);
+			Client.Controllers.Add("BootTest");
+
+			return Config;
+		}
+	}
+}
diff --git a/Engine/Source/Programs/AutomationTool/ShooterGame/Tests/ShooterTest.DedicatedServerTest.cs b/Engine/Source/Programs/AutomationTool/ShooterGame/Tests/ShooterTest.DedicatedServerTest.cs
new file mode 100644
--- /dev/null
+++ b/Engine/Source/Programs/AutomationTool/ShooterGame/Tests/ShooterTest.DedicatedServerTest.cs
@@ -0,0 +1,36 @@
+// Copyright 1998-2019 Epic Games, Inc.All Rights Reserved.
+using System;
+using System.Collections.Generic;
+using System.Linq;
+using System.Text;
+using System.Threading.Tasks;
+using EpicGame;
+using Gauntlet;
+
+namespace ShooterTest
+{
+	/// <summary>
+	/// This tests a dedicated server and a client that searches for the server and enters.  
+	/// </summary>
+	public class DedicatedServerTest : UnrealTestNode<ShooterTestConfig>
+	{
+		public DedicatedServerTest(UnrealTestContext InContext) : base(InContext)
+		{
+		}
+
+		public override ShooterTestConfig GetConfiguration()
+		{
+			ShooterTestConfig Config = base.GetConfiguration();
+			Config.PreAssignAccount = false;
+			Config.NoMCP = true;
+
+			UnrealTestRole Client = Config.RequireRole(UnrealTargetRole.Client);
+			UnrealTestRole Server = Config.RequireRole(UnrealTargetRole.Server);
+
+			Client.Controllers.Add("DedicatedServerTest");
+
+			return Config;
+		}
+	}
+}
+
diff --git a/Engine/Source/Programs/AutomationTool/ShooterGame/Tests/ShooterTest.ListenServerQuickMatchTest.cs b/Engine/Source/Programs/AutomationTool/ShooterGame/Tests/ShooterTest.ListenServerQuickMatchTest.cs
new file mode 100644
--- /dev/null
+++ b/Engine/Source/Programs/AutomationTool/ShooterGame/Tests/ShooterTest.ListenServerQuickMatchTest.cs
@@ -0,0 +1,34 @@
+// Copyright 1998-2019 Epic Games, Inc.All Rights Reserved.
+using System;
+using System.Collections.Generic;
+using System.Linq;
+using System.Text;
+using System.Threading.Tasks;
+using EpicGame;
+using Gauntlet;
+
+namespace ShooterTest
+{
+	/// <summary>
+	/// This tests a listen server and a client that enters the server using quick match.  
+	/// </summary>
+	public class ListenServerQuickMatchTest : UnrealTestNode<ShooterTestConfig>
+	{
+		public ListenServerQuickMatchTest(UnrealTestContext InContext) : base(InContext)
+		{
+		}
+
+		public override ShooterTestConfig GetConfiguration()
+		{
+			ShooterTestConfig Config = base.GetConfiguration();
+			Config.PreAssignAccount = false;
+			Config.NoMCP = true;
+
+			IEnumerable<UnrealTestRole> Clients = Config.RequireRoles(UnrealTargetRole.Client, 2);
+			Clients.ElementAt(0).Controllers.Add("ListenServerHost");
+			Clients.ElementAt(1).Controllers.Add("ListenServerQuickMatchClient");
+
+			return Config;
+		}
+	}
+}
diff --git a/Engine/Source/Programs/AutomationTool/ShooterGame/Tests/ShooterTest.ListenServerTest.cs b/Engine/Source/Programs/AutomationTool/ShooterGame/Tests/ShooterTest.ListenServerTest.cs
new file mode 100644
--- /dev/null
+++ b/Engine/Source/Programs/AutomationTool/ShooterGame/Tests/ShooterTest.ListenServerTest.cs
@@ -0,0 +1,34 @@
+// Copyright 1998-2019 Epic Games, Inc.All Rights Reserved.
+using System;
+using System.Collections.Generic;
+using System.Linq;
+using System.Text;
+using System.Threading.Tasks;
+using EpicGame;
+using Gauntlet;
+
+namespace ShooterTest
+{
+	/// <summary>
+	/// This tests a listen server and a client that opens a client that searches the server and enters.
+	/// </summary>
+	public class ListenServerTest : UnrealTestNode<ShooterTestConfig>
+	{
+		public ListenServerTest(UnrealTestContext InContext) : base(InContext)
+		{
+		}
+
+		public override ShooterTestConfig GetConfiguration()
+		{
+			ShooterTestConfig Config = base.GetConfiguration();
+			Config.PreAssignAccount = false;
+			Config.NoMCP = true;
+
+			IEnumerable<UnrealTestRole> Clients = Config.RequireRoles(UnrealTargetRole.Client, 2);
+			Clients.ElementAt(0).Controllers.Add("ListenServerHost");
+			Clients.ElementAt(1).Controllers.Add("ListenServerClient");
+
+			return Config;
+		}
+	}
+}
diff --git a/Engine/Source/Programs/AutomationTool/ShooterGame/Tests/ShooterTest.TestConfig.cs b/Engine/Source/Programs/AutomationTool/ShooterGame/Tests/ShooterTest.TestConfig.cs
new file mode 100644
--- /dev/null
+++ b/Engine/Source/Programs/AutomationTool/ShooterGame/Tests/ShooterTest.TestConfig.cs
@@ -0,0 +1,55 @@
+// Copyright 1998-2019 Epic Games, Inc.All Rights Reserved.
+using System;
+using System.Collections.Generic;
+using System.Linq;
+using System.Text;
+using System.Threading.Tasks;
+using EpicGame;
+using Gauntlet;
+
+namespace ShooterTest
+{
+	public class ShooterTestConfig : EpicGameTestConfig
+	{
+		[AutoParam]
+		public bool NoSteam = false;
+
+		[AutoParam]
+		public bool WithPacketHandlerEncryption = false;
+
+		[AutoParam]
+		public bool NoSeamlessTravel = false;
+
+		[AutoParam]
+		public int TargetNumOfCycledMatches = 2;
+
+		public override void ApplyToConfig(UnrealAppConfig AppConfig, UnrealSessionRole ConfigRole, IEnumerable<UnrealSessionRole> OtherRoles)
+		{
+			base.ApplyToConfig(AppConfig, ConfigRole, OtherRoles);
+
+			if (NoSteam)
+			{
+				AppConfig.CommandLine += " -nosteam";
+			}
+
+			if (WithPacketHandlerEncryption)
+			{
+				AppConfig.CommandLine += "-ini:Engine:[PacketHandlerComponents]:EncryptionComponent=AESHandlerComponent -execcmds=\"ShooterGame.TestEncryption 1\"";
+			}
+
+			if (NoSeamlessTravel)
+			{
+				AppConfig.CommandLine += " -noseamlesstravel";
+			}
+
+			if (AppConfig.ProcessType.IsClient())
+			{
+				AppConfig.CommandLine += string.Format(" -TargetNumOfCycledMatches={0}", TargetNumOfCycledMatches);
+			}
+
+			const float InitTime = 120.0f;
+			const float MatchTime = 300.0f;
+			MaxDuration = InitTime + (MatchTime * TargetNumOfCycledMatches);
+		}
+	}
+}
