Moved EngineTest Gauntlet script out of Gauntlet project and into EngineTest project

diff --git a/Engine/Source/Programs/AutomationTool/Gauntlet/EngineTest/RunEngineTests.cs b/Engine/Source/Programs/AutomationTool/Gauntlet/EngineTest/RunEngineTests.cs
deleted file mode 100644
diff --git a/Engine/Source/Programs/AutomationTool/Gauntlet/EngineTest/Tests/EngineTest.BootTest.cs b/Engine/Source/Programs/AutomationTool/Gauntlet/EngineTest/Tests/EngineTest.BootTest.cs
deleted file mode 100644
diff --git a/Engine/Source/Programs/AutomationTool/Gauntlet/EngineTest/Tests/EngineTest.RunAutomatedTests.cs b/Engine/Source/Programs/AutomationTool/Gauntlet/EngineTest/Tests/EngineTest.RunAutomatedTests.cs
deleted file mode 100644
diff --git a/Engine/Source/Programs/AutomationTool/Gauntlet/Gauntlet.Automation.csproj b/Engine/Source/Programs/AutomationTool/Gauntlet/Gauntlet.Automation.csproj
--- a/Engine/Source/Programs/AutomationTool/Gauntlet/Gauntlet.Automation.csproj
+++ b/Engine/Source/Programs/AutomationTool/Gauntlet/Gauntlet.Automation.csproj
@@ -68,8 +68,6 @@
   <ItemGroup>
     <Compile Include="Editor\EditorTest.BootTest.cs" />
     <Compile Include="Editor\RunEditorTests.cs" />
-    <Compile Include="EngineTest\RunEngineTests.cs" />
-    <Compile Include="EngineTest\Tests\EngineTest.RunAutomatedTests.cs" />
     <Compile Include="Framework\Account\Gauntlet.AccountPool.cs" />
     <Compile Include="Framework\Base\Gauntlet.AppConfig.cs" />
     <Compile Include="Framework\Base\Gauntlet.AppInstall.cs" />
