Merging 7325797 using 4.23_to_Dev-VR

Now requiring a certificate to package HL2
Now requiring "CN=" in the company name for creating a certificate

diff --git a/Engine/Plugins/Runtime/AR/Microsoft/HoloLensAR/Source/HoloLensPlatformEditor/Private/HoloLensTargetSettingsCustomization.cpp b/Engine/Plugins/Runtime/AR/Microsoft/HoloLensAR/Source/HoloLensPlatformEditor/Private/HoloLensTargetSettingsCustomization.cpp
--- a/Engine/Plugins/Runtime/AR/Microsoft/HoloLensAR/Source/HoloLensPlatformEditor/Private/HoloLensTargetSettingsCustomization.cpp
+++ b/Engine/Plugins/Runtime/AR/Microsoft/HoloLensAR/Source/HoloLensPlatformEditor/Private/HoloLensTargetSettingsCustomization.cpp
@@ -530,6 +530,17 @@ FReply FHoloLensTargetSettingsCustomization::GenerateSigningCertificate()
 		}
 	}
 
+	//make sure the company name has required syntax
+	FString CompanyDistinguishedName = GetDefault<UGeneralProjectSettings>()->CompanyDistinguishedName;
+	if (!CompanyDistinguishedName.Contains(TEXT("CN=")))
+	{
+		FNotificationInfo Info(LOCTEXT("SigningCertificateFailed_MalformedCompanyName", "Company Distinguished Name must contain 'CN=YourCompany' syntax"));
+		Info.ExpireDuration = 5.0f;
+		FSlateNotificationManager::Get().AddNotification(Info);
+
+		return FReply::Handled();
+	}
+
 	FString WinSDKVersion;
 	if (!FPlatformMisc::QueryRegKey(HKEY_CURRENT_USER, *WinSDKSubKey, TEXT("ProductVersion"), WinSDKVersion))
 	{
diff --git a/Engine/Source/Programs/AutomationTool/HoloLens/HoloLensPlatform.Automation.cs b/Engine/Source/Programs/AutomationTool/HoloLens/HoloLensPlatform.Automation.cs
--- a/Engine/Source/Programs/AutomationTool/HoloLens/HoloLensPlatform.Automation.cs
+++ b/Engine/Source/Programs/AutomationTool/HoloLens/HoloLensPlatform.Automation.cs
@@ -732,26 +732,7 @@ LogWarning("PackagePakFiles intermediate dir {0}", IntermediateDirectory);
 			{
 				if (!IsBuildMachine && !Params.Unattended)
 				{
-					string Name;
-					string Publisher;
-					if (!ConfigFile.GetString("/Script/EngineSettings.GeneralProjectSettings", "ProjectName", out Name))
-					{
-						Name = "DefaultUE4Project";
-					}
-					if (!ConfigFile.GetString("/Script/EngineSettings.GeneralProjectSettings", "CompanyDistinguishedName", out Publisher))
-					{
-						Publisher = "CN=NoPublisher";
-					}
-
-					if (!string.IsNullOrEmpty(Publisher))
-					{
-						LogWarning("No certificate found at {0}.  Generating temporary self-signed certificate for {1}.", SigningCertificatePath, Publisher);
-						GenerateSigningCertificate(SigningCertificatePath, Publisher);
-					}
-					else
-					{
-						LogWarning("No certificate found at {0} and temporary certificate cannot be generated (missing publisher name).  Check your Company Distinguished Name setting.  Signing will probably fail.", SigningCertificatePath);
-					}
+					LogError("Certificate is required.  Please go to Project Settings > HoloLens > Create Signing Certificate");
 				}
 				else
 				{
