When packaging for Mac, edit the executable to remove all RPATH entries that point to folders that do not exist in the final build to satisfy Gatekeeper requirements.

diff --git a/Engine/Source/Programs/AutomationTool/Mac/MacPlatform.Automation.cs b/Engine/Source/Programs/AutomationTool/Mac/MacPlatform.Automation.cs
--- a/Engine/Source/Programs/AutomationTool/Mac/MacPlatform.Automation.cs
+++ b/Engine/Source/Programs/AutomationTool/Mac/MacPlatform.Automation.cs
@@ -3,6 +3,7 @@ using System;
 using System.Collections.Generic;
 using System.Linq;
 using System.Text;
+using System.Text.RegularExpressions;
 using System.IO;
 using AutomationTool;
 using UnrealBuildTool;
@@ -271,10 +272,48 @@ public class MacPlatform : Platform
 		}
 	}
 
+	private void RemoveExtraRPaths(ProjectParams Params, DeploymentContext SC)
+	{
+		// When we link the executable we add RPATH entries for all possible places where dylibs can be loaded from, so that the same executable can be used from Binaries/Mac
+		// as well as in a packaged, self-contained application. In recent versions of macOS, Gatekeeper doesn't allow RPATHs pointing to folders that don't exist,
+		// so we remove these based on the type of packaging (Params.CreateAppBundle).
+		List<FileReference> Exes = GetExecutableNames(SC);
+		foreach (var ExePath in Exes)
+		{
+			IProcessResult CommandResult = Run("otool", "-l \"" + ExePath + "\"", null, ERunOptions.None);
+			if (CommandResult.ExitCode == 0)
+			{
+				StringReader Reader = new StringReader(CommandResult.Output);
+				Regex RPathPattern = new Regex(@"^\s+path (?<rpath>\S+)\s.+$");
+				string ToRemovePattern = Params.CreateAppBundle ? "/../../../" : "@loader_path/../UE4/";
+
+				string OutputLine;
+				while ((OutputLine = Reader.ReadLine()) != null)
+				{
+					if (OutputLine.EndsWith("cmd LC_RPATH"))
+					{
+						OutputLine = Reader.ReadLine();
+						OutputLine = Reader.ReadLine();
+						Match RPathMatch = RPathPattern.Match(OutputLine);
+						if (RPathMatch.Success)
+						{
+							string RPath = RPathMatch.Groups["rpath"].Value;
+							if (RPath.Contains(ToRemovePattern))
+							{
+								Run("xcrun", "install_name_tool -delete_rpath " + RPath + " \"" + ExePath + "\"", null, ERunOptions.None);
+							}
+						}
+					}
+				}
+			}
+		}
+	}
+
 	public override void Package(ProjectParams Params, DeploymentContext SC, int WorkingCL)
 	{
 		// package up the program, potentially with an installer for Mac
 		PrintRunTime();
+		RemoveExtraRPaths(Params, SC);
 	}
 
 	public override void ProcessArchivedProject(ProjectParams Params, DeploymentContext SC)
