Fixed issue with resolving project path

diff --git a/Engine/Source/Programs/AutomationTool/Gauntlet/Unreal/Base/Gauntlet.UnrealTestContext.cs b/Engine/Source/Programs/AutomationTool/Gauntlet/Unreal/Base/Gauntlet.UnrealTestContext.cs
--- a/Engine/Source/Programs/AutomationTool/Gauntlet/Unreal/Base/Gauntlet.UnrealTestContext.cs
+++ b/Engine/Source/Programs/AutomationTool/Gauntlet/Unreal/Base/Gauntlet.UnrealTestContext.cs
@@ -233,13 +233,16 @@ namespace Gauntlet
 			}
 			else
 			{
-				ProjectPath = ProjectUtils.FindProjectFileFromName(Project);
-
-				if (ProjectPath == null)
+				if (!string.IsNullOrEmpty(Project))
 				{
-					throw new AutomationException("Could not find project file for {0}", Project);
+					ProjectPath = ProjectUtils.FindProjectFileFromName(Project);
+
+					if (ProjectPath == null)
+					{
+						throw new AutomationException("Could not find project file for {0}", Project);
+					}
+					Project = ProjectPath.GetFileNameWithoutExtension();
 				}
-				Project = ProjectPath.GetFileNameWithoutExtension();
 			}
 
 			if (string.IsNullOrEmpty(Sandbox))
