Sign binaries with a SHA1 and SHA256 certificate.

diff --git a/Engine/Source/Programs/AutomationTool/AutomationUtils/CommandUtils.cs b/Engine/Source/Programs/AutomationTool/AutomationUtils/CommandUtils.cs
--- a/Engine/Source/Programs/AutomationTool/AutomationUtils/CommandUtils.cs
+++ b/Engine/Source/Programs/AutomationTool/AutomationUtils/CommandUtils.cs
@@ -2697,32 +2697,108 @@ namespace AutomationTool
         public string Uploaded { get; set; }
     }
 
-	/// <summary>
-	/// Code signing
-	/// </summary>
-	[Help("NoSign", "Skips signing of code/content files.")]
-	public class CodeSign
+	static class CodeSignWindows
 	{
 		/// <summary>
-		/// If so, what is the signing identity to search for?
+		/// The signing identity to find a certificate for
 		/// </summary>
-		public static string SigningIdentity = "Epic Games";
+		const string SigningIdentity = "Epic Games";
 
 		/// <summary>
-		/// Should we use the machine store?
+		/// Whether to look in the machine store for signing certificates, rather than the user store.
 		/// </summary>
-		public static bool bUseMachineStoreInsteadOfUserStore = false;
+		static bool bUseMachineStoreForCertificates = false;
 
-		/// <summary>
-		/// How long to keep re-trying code signing for
-		/// </summary>
-		public static TimeSpan CodeSignTimeOut = new TimeSpan(0, 3, 0); // Keep trying to sign one file for up to 3 minutes
+		public enum SignatureType
+		{
+			SHA1,
+			SHA256
+		};
+
+		static readonly string[] TimestampServersSHA1 =
+		{
+			"http://timestamp.verisign.com/scripts/timestamp.dll",
+			"http://timestamp.globalsign.com/scripts/timstamp.dll",
+			"http://timestamp.comodoca.com/authenticode"
+		};
+
+		static readonly string[] TimestampServersSHA256 =
+		{
+			"http://sha256timestamp.ws.symantec.com/sha256/timestamp",
+			"http://rfc3161timestamp.globalsign.com/advanced",
+			"http://timestamp.comodoca.com/?td=sha256"
+		};
+
+		public static void Sign(FileReference File, SignatureType SignatureType)
+		{
+			List<FileReference> Files = new List<FileReference> { File };
+			Sign(Files, SignatureType);
+		}
+
+		public static void Sign(List<FileReference> Files, SignatureType SignatureType)
+		{
+			string SignToolPath = GetSignToolPath();
+			string SpecificStoreArg = bUseMachineStoreForCertificates ? " /sm" : "";
+
+			for(int FileIdx = 0; FileIdx < Files.Count; )
+			{
+				Stopwatch Timer = Stopwatch.StartNew();
+
+				int NumAttempts = 0;
+				for(;;)
+				{
+					//@TODO: Verbosity choosing
+					//  /v will spew lots of info
+					//  /q does nothing on success and minimal output on failure
+					StringBuilder CommandLine = new StringBuilder();
+					if(SignatureType == SignatureType.SHA1)
+					{
+						CommandLine.AppendFormat("sign{0} /a /n \"{1}\" /t {2} /v", SpecificStoreArg, SigningIdentity, TimestampServersSHA1[NumAttempts % TimestampServersSHA1.Length]);
+					}
+					else if(SignatureType == SignatureType.SHA256)
+					{
+						CommandLine.AppendFormat("sign{0} /a /fd sha256 /td sha256 /as /n \"{1}\" /tr {2}", SpecificStoreArg, SigningIdentity, TimestampServersSHA256[NumAttempts % TimestampServersSHA256.Length]);
+					}
+					else
+					{
+						throw new ArgumentException(String.Format("Invalid signature type type ({0})", SignatureType));
+					}
+
+					// Append the files for this batch
+					int NextFileIdx = FileIdx;
+					while(NextFileIdx < Files.Count && CommandLine.Length + Files[NextFileIdx].FullName.Length < 2000)
+					{
+						CommandLine.AppendFormat(" \"{0}\"", Files[NextFileIdx]);
+						NextFileIdx++;
+					}
+
+					IProcessResult Result = CommandUtils.Run(SignToolPath, CommandLine.ToString(), null, CommandUtils.ERunOptions.AllowSpew);
+					NumAttempts++;
+
+					if (Result.ExitCode != 1)
+					{
+						if (Result.ExitCode == 2)
+						{
+							CommandUtils.LogError(String.Format("Signtool returned a warning."));
+						}
+						// Success!
+						FileIdx = NextFileIdx;
+						break;
+					}
+
+					if (Timer.Elapsed.TotalMinutes > 3.0)
+					{
+						throw new AutomationException("Failed to sign files {0} times over a period of {1}", NumAttempts, Timer.Elapsed);
+					}
+				}
+			}
+		}
 
 		/// <summary>
 		/// Finds the path to SignTool.exe, or throws an exception.
 		/// </summary>
 		/// <returns>Path to signtool.exe</returns>
-		public static string GetSignToolPath()
+		static string GetSignToolPath()
 		{
 			string[] PossibleSignToolNames =
 			{
@@ -2740,7 +2816,14 @@ namespace AutomationTool
 
 			throw new AutomationException("SignTool not found at '{0}' (are you missing the Windows SDK?)", String.Join("' or '", PossibleSignToolNames));
 		}
+	}
 
+	/// <summary>
+	/// Code signing
+	/// </summary>
+	[Help("NoSign", "Skips signing of code/content files.")]
+	public class CodeSign
+	{
 		/// <summary>
 		/// Code signs the specified file
 		/// </summary>
@@ -2779,60 +2862,10 @@ namespace AutomationTool
 				return;
 			}
 
-			string SignToolName = GetSignToolPath();
-
 			TargetFileInfo.IsReadOnly = false;
 
-			// Code sign the executable
-			string[] TimestampServer = { "http://timestamp.verisign.com/scripts/timestamp.dll",
-									     "http://timestamp.globalsign.com/scripts/timstamp.dll",
-										 "http://timestamp.comodoca.com/authenticode",
-										 "http://www.startssl.com/timestamp"
-									   };
-			int TimestampServerIndex = 0;
-
-			string SpecificStoreArg = bUseMachineStoreInsteadOfUserStore ? " /sm" : "";
-
-			DateTime StartTime = DateTime.Now;
-
-			int NumTrials = 0;
-			for (; ; )
-			{
-				//@TODO: Verbosity choosing
-				//  /v will spew lots of info
-				//  /q does nothing on success and minimal output on failure
-				string CodeSignArgs = String.Format("sign{0} /a /n \"{1}\" /t {2} /d \"{3}\" /v \"{4}\"", SpecificStoreArg, SigningIdentity, TimestampServer[TimestampServerIndex], TargetFileInfo.Name, TargetFileInfo.FullName);
-
-				IProcessResult Result = CommandUtils.Run(SignToolName, CodeSignArgs, null, CommandUtils.ERunOptions.AllowSpew);
-				++NumTrials;
-
-				if (Result.ExitCode != 1)
-				{
-					if (Result.ExitCode == 2)
-					{
-						CommandUtils.LogError(String.Format("Signtool returned a warning."));
-					}
-					// Success!
-					break;
-				}
-				else
-				{
-					// try another timestamp server on the next iteration
-					TimestampServerIndex++;
-					if (TimestampServerIndex >= TimestampServer.Count())
-					{
-						// loop back to the first timestamp server
-						TimestampServerIndex = 0;
-					}
-
-					// Keep retrying until we run out of time
-					TimeSpan RunTime = DateTime.Now - StartTime;
-					if (RunTime > CodeSignTimeOut)
-					{
-						throw new AutomationException("Failed to sign executable '{0}' {1} times over a period of {2}", TargetFileInfo.FullName, NumTrials, RunTime);
-					}
-				}
-			}
+			CodeSignWindows.Sign(new FileReference(TargetFileInfo), CodeSignWindows.SignatureType.SHA1);
+			CodeSignWindows.Sign(new FileReference(TargetFileInfo), CodeSignWindows.SignatureType.SHA256);
 		}
 
 		/// <summary>
@@ -2840,6 +2873,8 @@ namespace AutomationTool
 		/// </summary>
 		public static void SignMacFileOrFolder(string InPath, bool bIgnoreExtension = false)
 		{
+			TimeSpan CodeSignTimeOut = new TimeSpan(0, 3, 0); // Keep trying to sign one file for up to 3 minutes
+
 			bool bExists = CommandUtils.FileExists(InPath) || CommandUtils.DirectoryExists(InPath);
 			if (!bExists)
 			{
@@ -2942,59 +2977,6 @@ namespace AutomationTool
 			}
 		}
 
-		public static void SignListFilesIfEXEOrDLL(string FilesToSign)
-		{
-			string SignToolName = GetSignToolPath();
-
-			// nothing to sign
-			if (String.IsNullOrEmpty(FilesToSign))
-			{
-				return;
-			}
-
-			// Code sign the executable
-			string[] TimestampServer = { "http://timestamp.verisign.com/scripts/timestamp.dll",
-									     "http://timestamp.globalsign.com/scripts/timstamp.dll",
-										 "http://timestamp.comodoca.com/authenticode",
-										 "http://www.startssl.com/timestamp"
-									   };
-
-			string SpecificStoreArg = bUseMachineStoreInsteadOfUserStore ? " /sm" : "";	
-			
-			Stopwatch Stopwatch = Stopwatch.StartNew();
-
-			int NumTrials = 0;
-			for (; ; )
-			{
-				//@TODO: Verbosity choosing
-				//  /v will spew lots of info
-				//  /q does nothing on success and minimal output on failure
-				string CodeSignArgs = String.Format("sign{0} /a /n \"{1}\" /t {2} /debug {3}", SpecificStoreArg, SigningIdentity, TimestampServer[NumTrials % TimestampServer.Length], FilesToSign);
-
-				IProcessResult Result = CommandUtils.Run(SignToolName, CodeSignArgs, null, CommandUtils.ERunOptions.AllowSpew);
-				++NumTrials;
-
-				if (Result.ExitCode != 1)
-				{
-					if (Result.ExitCode == 2)
-					{
-						CommandUtils.LogError(String.Format("Signtool returned a warning."));
-					}
-					// Success!
-					break;
-				}
-				else
-				{
-					// Keep retrying until we run out of time
-					TimeSpan RunTime = Stopwatch.Elapsed;
-					if (RunTime > CodeSignTimeOut && NumTrials >= TimestampServer.Length)
-					{
-						throw new AutomationException("Failed to sign executables {0} times over a period of {1}", NumTrials, RunTime);
-					}
-				}
-			}
-		}
-
 		public static void SignMultipleFilesIfEXEOrDLL(List<FileReference> Files, bool bIgnoreExtension = false)
 		{
 			if (UnrealBuildTool.Utils.IsRunningOnMono)
@@ -3002,7 +2984,7 @@ namespace AutomationTool
 				CommandUtils.LogLog(String.Format("Can't sign we are running under mono."));
 				return;
 			}
-			List<string> FinalFiles = new List<string>();
+			List<FileReference> FinalFiles = new List<FileReference>();
 			foreach (string Filename in Files.Select(x => x.FullName))
 			{
 				FileInfo TargetFileInfo = new FileInfo(Filename);
@@ -3028,28 +3010,11 @@ namespace AutomationTool
 				}
 				if (IsExecutable && CommandUtils.FileExists(Filename))
 				{
-					FinalFiles.Add(Filename);
+					FinalFiles.Add(new FileReference(TargetFileInfo));
 				}
 			}			
-
-			StringBuilder FilesToSignBuilder = new StringBuilder();
-			List<string> FinalListSignStrings = new List<string>();
-			foreach(string File in FinalFiles)
-			{
-				FilesToSignBuilder.Append("\"" + File + "\" ");				
-				if(FilesToSignBuilder.Length > 1900)
-				{
-					string AddFilesToFinalList = FilesToSignBuilder.ToString();
-					FinalListSignStrings.Add(AddFilesToFinalList);
-					FilesToSignBuilder.Clear();
-				}
-			}
-			FinalListSignStrings.Add(FilesToSignBuilder.ToString());
-			foreach(string FilesToSign in FinalListSignStrings)
-			{
-				SignListFilesIfEXEOrDLL(FilesToSign);
-			}
-
+			CodeSignWindows.Sign(FinalFiles, CodeSignWindows.SignatureType.SHA1);
+			CodeSignWindows.Sign(FinalFiles, CodeSignWindows.SignatureType.SHA256);
 		}
 	}
 
