Only use SHA256 certs to sign MSI files, since they don't support multiple signatures.

#rnx

diff --git a/Engine/Source/Programs/AutomationTool/AutomationUtils/CommandUtils.cs b/Engine/Source/Programs/AutomationTool/AutomationUtils/CommandUtils.cs
--- a/Engine/Source/Programs/AutomationTool/AutomationUtils/CommandUtils.cs
+++ b/Engine/Source/Programs/AutomationTool/AutomationUtils/CommandUtils.cs
@@ -3013,7 +3013,7 @@ namespace AutomationTool
 					FinalFiles.Add(new FileReference(TargetFileInfo));
 				}
 			}			
-			CodeSignWindows.Sign(FinalFiles, CodeSignWindows.SignatureType.SHA1);
+			CodeSignWindows.Sign(FinalFiles.Where(x => !x.HasExtension(".msi")).ToList(), CodeSignWindows.SignatureType.SHA1); // MSI files can only have one signature; prefer SHA256
 			CodeSignWindows.Sign(FinalFiles, CodeSignWindows.SignatureType.SHA256);
 		}
 	}
