UAT: Propagate the XGE autorecover attribute when building through UAT.

#rnx

diff --git a/Engine/Source/Programs/AutomationTool/AutomationUtils/UE4Build.cs b/Engine/Source/Programs/AutomationTool/AutomationUtils/UE4Build.cs
--- a/Engine/Source/Programs/AutomationTool/AutomationUtils/UE4Build.cs
+++ b/Engine/Source/Programs/AutomationTool/AutomationUtils/UE4Build.cs
@@ -1,4 +1,4 @@
-﻿// Copyright 1998-2019 Epic Games, Inc. All Rights Reserved.
+// Copyright 1998-2019 Epic Games, Inc. All Rights Reserved.
 using System;
 using System.Collections.Generic;
 using System.Text;
@@ -842,6 +842,10 @@ namespace AutomationTool
 									ToolElement.SetAttribute("AutoReserveMemory", Element.Attributes["AutoReserveMemory"].Value);
 								}
 								ToolElement.SetAttribute("OutputFileMasks", Element.Attributes["OutputFileMasks"].Value);
+								if(Element.HasAttribute("AutoRecover"))
+								{
+									ToolElement.SetAttribute("AutoRecover", Element.Attributes["AutoRecover"].Value);
+								}
 								//ToolElement.SetAttribute("AllowRestartOnLocal", "false");  //vs2012 can't really restart, so we will go with this for now
 								if (Element.Attributes["OutputFileMasks"].Value == "PCLaunch.rc.res")
 								{
