- Add architecture flavors for HoloLens
- Force Win10SDK version to 10.0.18362.0
- Fix bug in Finalize that wasn't checking for architecture flavors in platforms not Android

diff --git a/Engine/Build/InstalledEngineBuild.xml b/Engine/Build/InstalledEngineBuild.xml
--- a/Engine/Build/InstalledEngineBuild.xml
+++ b/Engine/Build/InstalledEngineBuild.xml
@@ -82,6 +82,9 @@
 
 	<!-- Architectures that we build for Android -->
 	<Property Name="AndroidArchitectures" Value="armv7+arm64"/>
+	
+	<!-- Hololens Architectures -->
+	<Property Name="HoloLensArchitectures" Value="x64+arm64+multi"/>
 
 	<Property Name="AndroidGPUArchitectures" Value="es2"/>
 	<Property Name="LuminGPUArchitectures" Value="es2"/>
@@ -352,7 +355,7 @@
 
 		</Node>
 	</Agent>
-	<Agent Name="Target HoloLens" Type="CompileWin64;Win64">
+	<Agent Name="Target HoloLens" Type="HoloLens">
 		<Node Name="Compile UE4Game HoloLens" Requires="Compile UnrealHeaderTool Win64" Produces="#UE4Game HoloLens">
 
 			<ForEach Name="Configuration" Values="$(GameConfigurations)">
@@ -361,7 +364,6 @@
 				<Compile Target="UE4Game" Platform="HoloLens" Configuration="$(Configuration)" Tag="#UE4Game HoloLens" Arguments="-precompile -allmodules -nolink -architecture=x64 $(TargetDebugInfoArg)"/>
 				<Compile Target="UE4Game" Platform="HoloLens" Configuration="$(Configuration)" Tag="#UE4Game HoloLens" Arguments="-precompile -architecture=x64 $(TargetDebugInfoArg)" Clean="false"/>
 			</ForEach>
-
 		</Node>
 	</Agent>
 	<Agent Name="Target Lumin" Type="CompileWin64;Win64">
@@ -867,6 +869,7 @@
 			<!-- Create command line arguments for any platforms that were built with additional architectures -->
 			<Property Name="ArchitectureCommands" Value=""/>
 			<Property Name="ArchitectureCommands" Value="$(ArchitectureCommands) -AndroidArchitectures=$(AndroidArchitectures) -AndroidGPUArchitectures=$(AndroidGPUArchitectures)" If="'$(WithAndroid)' == true"/>
+			<Property Name="ArchitectureCommands" Value="$(ArchitectureCommands) -HoloLensArchitectures=$(HoloLensArchitectures)" If="'$(WithHoloLens)' == true"/>
 
 			<!-- Finalize the build and then tag the text file that this process creates -->
 			<Command Name="FinalizeInstalledBuild" Arguments="-OutputDir=&quot;$(LocalInstalledDir)&quot; -ContentOnlyPlatforms=$(ContentOnlyPlatforms) $(ArchitectureCommands) -AnalyticsTypeOverride=$(AnalyticsTypeOverride)"/>
diff --git a/Engine/Config/HoloLens/HoloLensEngine.ini b/Engine/Config/HoloLens/HoloLensEngine.ini
--- a/Engine/Config/HoloLens/HoloLensEngine.ini
+++ b/Engine/Config/HoloLens/HoloLensEngine.ini
@@ -14,4 +14,7 @@ PoolSizeVRAMPercentage=70
 ConfigClass=/Engine/DeviceProfiles/SetDeviceProfileWindows.SetDeviceProfileWindows_C
 
 [SystemSettings]
-r.setres=1280x720
\ No newline at end of file
+r.setres=1280x720
+
+[/Script/HoloLensPlatformEditor.HoloLensTargetSettings]
+Windows10SDKVersion=10.0.18362.0
\ No newline at end of file
diff --git a/Engine/Source/Programs/AutomationTool/Scripts/FinalizeInstalledBuild.cs b/Engine/Source/Programs/AutomationTool/Scripts/FinalizeInstalledBuild.cs
--- a/Engine/Source/Programs/AutomationTool/Scripts/FinalizeInstalledBuild.cs
+++ b/Engine/Source/Programs/AutomationTool/Scripts/FinalizeInstalledBuild.cs
@@ -96,9 +96,10 @@ namespace AutomationTool
 					{
 						AllArchNames = Architectures.Split('+').ToList();
 					}
+					// if there aren't any, use the default
 					else
 					{
-						AllArchNames = new List<string>();
+						AllArchNames = new List<string>() { Architecture };
 					}
 
 					// Check whether this platform should only be used for content based projects
@@ -113,25 +114,42 @@ namespace AutomationTool
 						{
 							// Need to check for development receipt as we use that for the Engine code in DebugGame
 							UnrealTargetConfiguration EngineConfiguration = (CodeTargetConfiguration == UnrealTargetConfiguration.DebugGame) ? UnrealTargetConfiguration.Development : CodeTargetConfiguration;
-							FileReference ReceiptFileName = TargetReceipt.GetDefaultPath(new DirectoryReference(OutputEnginePath), TargetName, CodeTargetPlatform, EngineConfiguration, Architecture);
 
-							if (FileReference.Exists(ReceiptFileName))
+							// Android has multiple architecture flavors built without receipts, so use the default arch target instead
+							if (CodeTargetPlatform == UnrealTargetPlatform.Android)
 							{
-								// Strip the output folder so that this can be used on any machine
-								string RelativeReceiptFileName = ReceiptFileName.MakeRelativeTo(new DirectoryReference(OutputDir));
-
-								// If we have pre-compiled architectures for this platform then add an entry for each of these -
-								// there isn't a receipt for each architecture like some other platforms
-								if (AllArchNames.Count > 0)
+								FileReference ReceiptFileName = TargetReceipt.GetDefaultPath(new DirectoryReference(OutputEnginePath), TargetName, CodeTargetPlatform, EngineConfiguration, Architecture);
+								if (FileReference.Exists(ReceiptFileName))
 								{
-									foreach (string Arch in AllArchNames)
+									// Strip the output folder so that this can be used on any machine
+									string RelativeReceiptFileName = ReceiptFileName.MakeRelativeTo(new DirectoryReference(OutputDir));
+
+									// Blindly append all of the architecture names
+									if (AllArchNames.Count > 0)
 									{
-										InstalledConfigs.Add(new InstalledPlatformInfo.InstalledPlatformConfiguration(CodeTargetConfiguration, CodeTargetPlatform, TargetType.Game, Arch, RelativeReceiptFileName, ProjectType, bCanBeDisplayed));
+										foreach (string Arch in AllArchNames)
+										{
+											InstalledConfigs.Add(new InstalledPlatformInfo.InstalledPlatformConfiguration(CodeTargetConfiguration, CodeTargetPlatform, TargetType.Game, Arch, RelativeReceiptFileName, ProjectType, bCanBeDisplayed));
+										}
+									}
+									// if for some reason we didn't specify any flavors, just add the default one.
+									else
+									{
+										InstalledConfigs.Add(new InstalledPlatformInfo.InstalledPlatformConfiguration(CodeTargetConfiguration, CodeTargetPlatform, TargetType.Game, Architecture, RelativeReceiptFileName, ProjectType, bCanBeDisplayed));
 									}
 								}
-								else
+							}
+							// If we're not Android, check the existence of the target receipts for each architecture specified.
+							else
+							{
+								foreach (string Arch in AllArchNames)
 								{
-									InstalledConfigs.Add(new InstalledPlatformInfo.InstalledPlatformConfiguration(CodeTargetConfiguration, CodeTargetPlatform, TargetType.Game, Architecture, RelativeReceiptFileName, ProjectType, bCanBeDisplayed));
+									FileReference ReceiptFileName = TargetReceipt.GetDefaultPath(new DirectoryReference(OutputEnginePath), TargetName, CodeTargetPlatform, EngineConfiguration, Arch);
+									if (FileReference.Exists(ReceiptFileName))
+									{
+										string RelativeReceiptFileName = ReceiptFileName.MakeRelativeTo(new DirectoryReference(OutputDir));
+										InstalledConfigs.Add(new InstalledPlatformInfo.InstalledPlatformConfiguration(CodeTargetConfiguration, CodeTargetPlatform, TargetType.Game, Arch, RelativeReceiptFileName, ProjectType, bCanBeDisplayed));
+									}
 								}
 							}
 						}
