Change to SHA1 signing for MSI files. Doesn't seem to like being signed with SHA256.

#rnx

diff --git a/Engine/Source/Programs/AutomationTool/AutomationUtils/CommandUtils.cs b/Engine/Source/Programs/AutomationTool/AutomationUtils/CommandUtils.cs
--- a/Engine/Source/Programs/AutomationTool/AutomationUtils/CommandUtils.cs
+++ b/Engine/Source/Programs/AutomationTool/AutomationUtils/CommandUtils.cs
@@ -3013,8 +3013,8 @@ namespace AutomationTool
 					FinalFiles.Add(new FileReference(TargetFileInfo));
 				}
 			}			
-			CodeSignWindows.Sign(FinalFiles.Where(x => !x.HasExtension(".msi")).ToList(), CodeSignWindows.SignatureType.SHA1); // MSI files can only have one signature; prefer SHA256
-			CodeSignWindows.Sign(FinalFiles, CodeSignWindows.SignatureType.SHA256);
+			CodeSignWindows.Sign(FinalFiles, CodeSignWindows.SignatureType.SHA1);
+			CodeSignWindows.Sign(FinalFiles.Where(x => !x.HasExtension(".msi")).ToList(), CodeSignWindows.SignatureType.SHA256); // MSI files can only have one signature; prefer SHA1 for compatibility
 		}
 	}
 
