Fix error compiling editor targets from installed builds.

#rnx

diff --git a/Engine/Source/Programs/AutomationTool/Scripts/FinalizeInstalledBuild.cs b/Engine/Source/Programs/AutomationTool/Scripts/FinalizeInstalledBuild.cs
--- a/Engine/Source/Programs/AutomationTool/Scripts/FinalizeInstalledBuild.cs
+++ b/Engine/Source/Programs/AutomationTool/Scripts/FinalizeInstalledBuild.cs
@@ -71,7 +71,9 @@ namespace AutomationTool
 			List<InstalledPlatformInfo.InstalledPlatformConfiguration> InstalledConfigs = new List<InstalledPlatformInfo.InstalledPlatformConfiguration>();
 
 			// Add the editor platform, otherwise we'll never be able to run UAT
-			InstalledConfigs.Add(new InstalledPlatformInfo.InstalledPlatformConfiguration(UnrealTargetConfiguration.Development, HostPlatform.Current.HostEditorPlatform, TargetRules.TargetType.Editor, "", "", EProjectType.Unknown, false));
+			string EditorArchitecture = PlatformExports.GetDefaultArchitecture(HostPlatform.Current.HostEditorPlatform, null);
+			InstalledConfigs.Add(new InstalledPlatformInfo.InstalledPlatformConfiguration(UnrealTargetConfiguration.Development, HostPlatform.Current.HostEditorPlatform, TargetRules.TargetType.Editor, EditorArchitecture, "", EProjectType.Unknown, false));
+			InstalledConfigs.Add(new InstalledPlatformInfo.InstalledPlatformConfiguration(UnrealTargetConfiguration.DebugGame, HostPlatform.Current.HostEditorPlatform, TargetRules.TargetType.Editor, EditorArchitecture, "", EProjectType.Unknown, false));
 
 			foreach (UnrealTargetPlatform CodeTargetPlatform in UnrealTargetPlatform.GetValidPlatforms())
 			{
diff --git a/Engine/Source/Programs/UnrealBuildTool/System/InstalledPlatformInfo.cs b/Engine/Source/Programs/UnrealBuildTool/System/InstalledPlatformInfo.cs
--- a/Engine/Source/Programs/UnrealBuildTool/System/InstalledPlatformInfo.cs
+++ b/Engine/Source/Programs/UnrealBuildTool/System/InstalledPlatformInfo.cs
@@ -183,7 +183,7 @@ namespace UnrealBuildTool
 
 			string PlatformTypeName;
 			TargetType PlatformType = TargetType.Game;
-			if (ParseSubValue(PlatformConfiguration, "PlatformTypeName=", out PlatformTypeName))
+			if (ParseSubValue(PlatformConfiguration, "PlatformType=", out PlatformTypeName))
 			{
 				if (!Enum.TryParse(PlatformTypeName, out PlatformType))
 				{
