Fixed issue with project paths that contain a - or +.

Stop tests timing out due to shader compiles
Added error if no tests are actually run

diff --git a/Engine/Source/Programs/AutomationTool/Gauntlet/Unreal/Base/Gauntlet.UnrealTestNode.cs b/Engine/Source/Programs/AutomationTool/Gauntlet/Unreal/Base/Gauntlet.UnrealTestNode.cs
--- a/Engine/Source/Programs/AutomationTool/Gauntlet/Unreal/Base/Gauntlet.UnrealTestNode.cs
+++ b/Engine/Source/Programs/AutomationTool/Gauntlet/Unreal/Base/Gauntlet.UnrealTestNode.cs
@@ -1157,6 +1157,7 @@ namespace Gauntlet
 			MB.Paragraph(string.Format("Context: {0}", Context.ToString()));
 			MB.Paragraph(string.Format("FatalErrors: {0}, Ensures: {1}, Errors: {2}, Warnings: {3}", FatalErrors, Ensures, Errors, Warnings));
 			MB.Paragraph(string.Format("ResultHash: {0}", GetTestResultHash()));
+			MB.Paragraph(string.Format("Result: {0}", GetTestResult()));
 			//MB.Paragraph(string.Format("Artifacts: {0}", CachedArtifactPath));
 
 			return MB.ToString();
diff --git a/Engine/Source/Programs/AutomationTool/Gauntlet/Unreal/BuildSource/Gauntlet.UnrealBuildSource.cs b/Engine/Source/Programs/AutomationTool/Gauntlet/Unreal/BuildSource/Gauntlet.UnrealBuildSource.cs
--- a/Engine/Source/Programs/AutomationTool/Gauntlet/Unreal/BuildSource/Gauntlet.UnrealBuildSource.cs
+++ b/Engine/Source/Programs/AutomationTool/Gauntlet/Unreal/BuildSource/Gauntlet.UnrealBuildSource.cs
@@ -420,6 +420,10 @@ namespace Gauntlet
 				Config.CommandLine += " " + Role.CommandLine;
 			}
 
+			// Cleanup the commandline
+			Config.CommandLine = GenerateProcessedCommandLine(Config.CommandLine);
+
+			// Now add the project (the above code doesn't handle arguments without a leading - so do this last
 			bool IsContentOnlyProject = (Config.Build.Flags & BuildFlags.ContentOnlyProject) == BuildFlags.ContentOnlyProject;
 
 			// Add in editor - TODO, should this be in the editor build?
@@ -439,11 +443,11 @@ namespace Gauntlet
 				// add in -game or -server
 				if (Role.RoleType.IsClient())
 				{
-					Config.CommandLine += " -game";
+					Config.CommandLine = "-game " + Config.CommandLine;
 				}
 				else if (Role.RoleType.IsServer())
 				{
-					Config.CommandLine += " -server";
+					Config.CommandLine = "-server " + Config.CommandLine;
 				}
 			}
 
@@ -451,7 +455,7 @@ namespace Gauntlet
             {
                 Config.FilesToCopy = Role.FilesToCopy;
             }
-			Config.CommandLine = GenerateProcessedCommandLine(Config.CommandLine);
+			
 			return Config;
 		}
 
@@ -466,7 +470,7 @@ namespace Gauntlet
 			// Break down Commandline into individual tokens 
 			Dictionary<string, string> CommandlineTokens = new Dictionary<string, string>(StringComparer.OrdinalIgnoreCase);
 			// turn Name(p1,etc) into a collection of Name|(p1,etc) groups
-			MatchCollection Matches = Regex.Matches(InCommandLine, "(?<option>-?[\\w\\d.:\\[\\]\\/\\\\]+)(=(?<value>(\"([^\"]*)\")|(\\S+)))?");
+			MatchCollection Matches = Regex.Matches(InCommandLine, "(?<option>\\-?[\\w\\d.:\\[\\]\\/\\\\]+)(=(?<value>(\"([^\"]*)\")|(\\S+)))?");
 
 			foreach (Match M in Matches)
 			{
