Made -configuration a fallback to -config in BuildCookRun for consistency with UBT and Gauntlet

diff --git a/Engine/Source/Programs/AutomationTool/AutomationUtils/ProjectParams.cs b/Engine/Source/Programs/AutomationTool/AutomationUtils/ProjectParams.cs
--- a/Engine/Source/Programs/AutomationTool/AutomationUtils/ProjectParams.cs
+++ b/Engine/Source/Programs/AutomationTool/AutomationUtils/ProjectParams.cs
@@ -829,7 +829,10 @@ namespace AutomationTool
                     if (ClientConfig == null)
                         ClientConfig = Command.ParseParamValue("config");
 
-                    if (ClientConfig != null)
+					if (ClientConfig == null)
+						ClientConfig = Command.ParseParamValue("configuration");
+
+					if (ClientConfig != null)
 					{
 						this.ClientConfigsToBuild = new List<UnrealTargetConfiguration>();
 						ParamList<string> Configs = new ParamList<string>(ClientConfig.Split(new char[] { '+' }, StringSplitOptions.RemoveEmptyEntries));
@@ -992,7 +995,10 @@ namespace AutomationTool
                     if (ServerConfig == null)
                         ServerConfig = Command.ParseParamValue("config");
 
-                    if (ServerConfig != null)
+					if (ServerConfig == null)
+						ServerConfig = Command.ParseParamValue("configuration");
+
+					if (ServerConfig != null)
 					{
 						this.ServerConfigsToBuild = new List<UnrealTargetConfiguration>();
 						ParamList<string> Configs = new ParamList<string>(ServerConfig.Split(new char[] { '+' }, StringSplitOptions.RemoveEmptyEntries));
