Don't remove extra RPATH entries from Mac executable if we're only staging the app bundle and won't be archiving it

diff --git a/Engine/Source/Programs/AutomationTool/Mac/MacPlatform.Automation.cs b/Engine/Source/Programs/AutomationTool/Mac/MacPlatform.Automation.cs
--- a/Engine/Source/Programs/AutomationTool/Mac/MacPlatform.Automation.cs
+++ b/Engine/Source/Programs/AutomationTool/Mac/MacPlatform.Automation.cs
@@ -313,7 +313,12 @@ public class MacPlatform : Platform
 	{
 		// package up the program, potentially with an installer for Mac
 		PrintRunTime();
-		RemoveExtraRPaths(Params, SC);
+
+		if (Params.Archive)
+		{
+			// Remove extra RPATHs if we will be archiving the project
+			RemoveExtraRPaths(Params, SC);
+		}
 	}
 
 	public override void ProcessArchivedProject(ProjectParams Params, DeploymentContext SC)
