Now deal with situation where mono is >= 5.0 but msbuild does not exist. For reference it looks like mono that's bundled with UE was upgraded to 5.16 at some point but msbuild was not added.

diff --git a/Engine/Source/Programs/AutomationTool/AutomationUtils/MacHostPlatform.cs b/Engine/Source/Programs/AutomationTool/AutomationUtils/MacHostPlatform.cs
--- a/Engine/Source/Programs/AutomationTool/AutomationUtils/MacHostPlatform.cs
+++ b/Engine/Source/Programs/AutomationTool/AutomationUtils/MacHostPlatform.cs
@@ -31,10 +31,22 @@ namespace AutomationTool
 
 				if (M.Success)
 				{
+					Log.TraceVerbose("Detected mono version as {0}.", VersionString);
 					try
 					{
 						Version Ver = Version.Parse(M.Groups[1].ToString());
 						CanUseMsBuild = Ver.Major >= 5;
+
+						if (Ver.Major >= 5)
+						{
+							// double check msbuild is in the path
+							CanUseMsBuild = string.IsNullOrEmpty(CommandUtils.WhichApp("msbuild")) == false;
+
+							if (!CanUseMsBuild)
+							{
+								Log.TraceWarning("mono 5.0 or greater exists but msbuild not found");
+							}
+						}
 					}
 					catch
 					{
diff --git a/Engine/Source/Programs/AutomationTool/AutomationUtils/ProcessUtils.cs b/Engine/Source/Programs/AutomationTool/AutomationUtils/ProcessUtils.cs
--- a/Engine/Source/Programs/AutomationTool/AutomationUtils/ProcessUtils.cs
+++ b/Engine/Source/Programs/AutomationTool/AutomationUtils/ProcessUtils.cs
@@ -761,6 +761,46 @@ namespace AutomationTool
 			Default = AllowSpew | AppMustExist,
 		}
 
+		/// <summary>
+		/// Resolves the passed in name using the path environment
+		/// </summary>
+		/// <param name="App"></param>
+		/// <returns></returns>
+		public static string WhichApp(string App, bool Quiet=true)
+		{
+			if (FileExists(Quiet, App))
+			{
+				return App;
+			}
+
+			string ResolvedPath = null;
+
+			if (!App.Contains(Path.DirectorySeparatorChar) && !App.Contains(Path.AltDirectorySeparatorChar))
+			{
+				string[] PathDirectories = Environment.GetEnvironmentVariable("PATH").Split(Path.PathSeparator);
+				foreach (string PathDirectory in PathDirectories)
+				{
+					string TryApp = Path.Combine(PathDirectory, App);
+					if (FileExists(Quiet, TryApp))
+					{
+						ResolvedPath = TryApp;
+						break;
+					}
+				}
+			}
+
+			if (ResolvedPath != null)
+			{
+				Log.TraceVeryVerbose("Resolved {0} to {1}", App, ResolvedPath);
+			}
+			else
+			{
+				Log.TraceVerbose("Could not resolve app {0}", App);
+			}
+
+			return ResolvedPath;
+		}
+
 		/// <summary>
 		/// Runs external program.
 		/// </summary>
@@ -780,29 +820,18 @@ namespace AutomationTool
 			{
 				Options &= ~ERunOptions.AppMustExist;
 			}
-
+			
 			// Check if the application exists, including the PATH directories.
 			if (Options.HasFlag(ERunOptions.AppMustExist) && !FileExists(Options.HasFlag(ERunOptions.NoLoggingOfRunCommand) ? true : false, App))
 			{
-				bool bExistsInPath = false;
-				if(!App.Contains(Path.DirectorySeparatorChar) && !App.Contains(Path.AltDirectorySeparatorChar))
-				{
-					string[] PathDirectories = Environment.GetEnvironmentVariable("PATH").Split(Path.PathSeparator);
-					foreach(string PathDirectory in PathDirectories)
-					{
-						string TryApp = Path.Combine(PathDirectory, App);
-						if(FileExists(Options.HasFlag(ERunOptions.NoLoggingOfRunCommand), TryApp))
-						{
-							App = TryApp;
-							bExistsInPath = true;
-							break;
-						}
-					}
-				}
-				if(!bExistsInPath)
+				string ResolvedPath = WhichApp(App);
+
+				if(string.IsNullOrEmpty(ResolvedPath))
 				{
 					throw new AutomationException("BUILD FAILED: Couldn't find the executable to run: {0}", App);
 				}
+
+				App = ResolvedPath;
 			}
 			var StartTime = DateTime.UtcNow;
 
