Missed file from 7828884

diff --git a/Engine/Build/BatchFiles/Mac/SetupMono.sh b/Engine/Build/BatchFiles/Mac/SetupMono.sh
--- a/Engine/Build/BatchFiles/Mac/SetupMono.sh
+++ b/Engine/Build/BatchFiles/Mac/SetupMono.sh
@@ -27,7 +27,7 @@ if [ ! $MONO_VERSION_PATH == "" ] && [ -f $MONO_VERSION_PATH ]; then
 	MONO_VERSION=(`echo ${MONO_VERSION:MONO_VERSION_PREFIX_LEN} |tr '.' ' '`)
 	if [ ${MONO_VERSION[0]} -ge 5 ]; then # Allow any Mono 5.x and up
 		IS_MONO_INSTALLED=1
-		IS_MS_BUILD_AVAILABLE=0
+		IS_MS_BUILD_AVAILABLE=1
 	fi
 fi
 
diff --git a/Engine/Source/Programs/AutomationTool/Scripts/BuildTarget.Automation.cs b/Engine/Source/Programs/AutomationTool/Scripts/BuildTarget.Automation.cs
--- a/Engine/Source/Programs/AutomationTool/Scripts/BuildTarget.Automation.cs
+++ b/Engine/Source/Programs/AutomationTool/Scripts/BuildTarget.Automation.cs
@@ -117,8 +117,8 @@ namespace AutomationTool
 					string TargetScript = TargetScripts.Where(S => S.IndexOf(TargetName, StringComparison.OrdinalIgnoreCase) >= 0).FirstOrDefault();
 
 					if (TargetScript == null && (
-							!TargetName.Equals("Client", StringComparison.OrdinalIgnoreCase) ||
-							!TargetName.Equals("Game", StringComparison.OrdinalIgnoreCase)
+							TargetName.Equals("Client", StringComparison.OrdinalIgnoreCase) ||
+							TargetName.Equals("Game", StringComparison.OrdinalIgnoreCase)
 							)
 						)
 					{
