Fix for packaging warning on Mac. No longer check mono version, just see if msbuild exists. TraceInfo if not,

diff --git a/Engine/Source/Programs/AutomationTool/AutomationUtils/MacHostPlatform.cs b/Engine/Source/Programs/AutomationTool/AutomationUtils/MacHostPlatform.cs
--- a/Engine/Source/Programs/AutomationTool/AutomationUtils/MacHostPlatform.cs
+++ b/Engine/Source/Programs/AutomationTool/AutomationUtils/MacHostPlatform.cs
@@ -18,49 +18,18 @@ namespace AutomationTool
 
 		public override string GetMsBuildExe()
 		{
-			// As of 5.0 mono comes with msbuild which performs better. If that's installed then we'll
-			// use it, if not then fallback to xbuild and watn
+			// As of 5.0 mono comes with msbuild which performs better. If that's installed then use it
 			if (string.IsNullOrEmpty(CachedMsBuildTool))
 			{
-				bool CanUseMsBuild = false;
-			
-				string VersionString = Utils.RunLocalProcessAndReturnStdOut("mono", "--version");
+				bool CanUseMsBuild = string.IsNullOrEmpty(CommandUtils.WhichApp("msbuild")) == false;
 
-				// Mono JIT compiler version 5.18.1.89 (2018-08/223ea7ef92e Tue May 21 12:03:39 EDT 2019) etc
-				Match M = Regex.Match(VersionString, @"version\s+(\d+\..+?)\s");
-
-				if (M.Success)
-				{
-					Log.TraceVerbose("Detected mono version as {0}.", VersionString);
-					try
-					{
-						Version Ver = Version.Parse(M.Groups[1].ToString());
-						CanUseMsBuild = Ver.Major >= 5;
-
-						if (Ver.Major >= 5)
-						{
-							// double check msbuild is in the path
-							CanUseMsBuild = string.IsNullOrEmpty(CommandUtils.WhichApp("msbuild")) == false;
-
-							if (!CanUseMsBuild)
-							{
-								Log.TraceWarning("mono 5.0 or greater exists but msbuild not found");
-							}
-						}
-					}
-					catch
-					{
-						Log.TraceWarning("Unable to parse version string {0}", M.Groups[1]);
-					}
-				}
-	
 				if (CanUseMsBuild)
 				{
 					CachedMsBuildTool = "msbuild";
 				}
 				else
 				{
-					Log.TraceWarning("Using xbuild. Install Mono 5.0 or greater for faster builds!");
+					Log.TraceInformation("Using xbuild. Install Mono 5.0 or greater for faster builds!");
 					CachedMsBuildTool = "xbuild";
 				}
 			}
