Resolve issue running packaged builds when no project file is available

diff --git a/Engine/Source/Programs/AutomationTool/Gauntlet/Platform/Mac/Gauntlet.SelfTest.TestUnrealInstallAndRunMac.cs b/Engine/Source/Programs/AutomationTool/Gauntlet/Platform/Mac/Gauntlet.SelfTest.TestUnrealInstallAndRunMac.cs
--- a/Engine/Source/Programs/AutomationTool/Gauntlet/Platform/Mac/Gauntlet.SelfTest.TestUnrealInstallAndRunMac.cs
+++ b/Engine/Source/Programs/AutomationTool/Gauntlet/Platform/Mac/Gauntlet.SelfTest.TestUnrealInstallAndRunMac.cs
@@ -23,7 +23,7 @@ namespace Gauntlet.SelfTest
 		public override void TickTest()
 		{
 			// create a new build
-			UnrealBuildSource Build = new UnrealBuildSource(this.ProjectFile, this.UnrealPath, this.UsesSharedBuildType, this.BuildPath, new string[] { "" });
+			UnrealBuildSource Build = new UnrealBuildSource(this.ProjectName, this.ProjectFile, this.UnrealPath, this.UsesSharedBuildType, this.BuildPath, new string[] { "" });
 
 			// check it's valid
 			if (!CheckResult(Build.BuildCount > 0, "staged build was invalid"))
diff --git a/Engine/Source/Programs/AutomationTool/Gauntlet/SelfTest/Unreal/Gauntlet.SelfTest.TestShortSoloGame.cs b/Engine/Source/Programs/AutomationTool/Gauntlet/SelfTest/Unreal/Gauntlet.SelfTest.TestShortSoloGame.cs
--- a/Engine/Source/Programs/AutomationTool/Gauntlet/SelfTest/Unreal/Gauntlet.SelfTest.TestShortSoloGame.cs
+++ b/Engine/Source/Programs/AutomationTool/Gauntlet/SelfTest/Unreal/Gauntlet.SelfTest.TestShortSoloGame.cs
@@ -50,7 +50,7 @@ namespace Gauntlet.SelfTest
 			}
 
 			// create a new build
-			UnrealBuildSource Build = new UnrealBuildSource(ProjectFile, this.UnrealPath, UsesSharedBuildType, BuildPath);
+			UnrealBuildSource Build = new UnrealBuildSource(this.ProjectName, ProjectFile, this.UnrealPath, UsesSharedBuildType, BuildPath);
 
 			// check it's valid
 			if (!CheckResult(Build.BuildCount > 0, "staged build was invalid"))
diff --git a/Engine/Source/Programs/AutomationTool/Gauntlet/SelfTest/Unreal/Gauntlet.SelfTest.TestUnrealBuildSource.cs b/Engine/Source/Programs/AutomationTool/Gauntlet/SelfTest/Unreal/Gauntlet.SelfTest.TestUnrealBuildSource.cs
--- a/Engine/Source/Programs/AutomationTool/Gauntlet/SelfTest/Unreal/Gauntlet.SelfTest.TestUnrealBuildSource.cs
+++ b/Engine/Source/Programs/AutomationTool/Gauntlet/SelfTest/Unreal/Gauntlet.SelfTest.TestUnrealBuildSource.cs
@@ -30,7 +30,7 @@ namespace Gauntlet.SelfTest
 		public override void TickTest()
 		{
 			// create the build source
-			UnrealBuildSource BuildSource = new UnrealBuildSource(ProjectFile, this.UnrealPath, UsesSharedBuildType, BuildPath, new string[] { "" });
+			UnrealBuildSource BuildSource = new UnrealBuildSource(this.ProjectName, ProjectFile, this.UnrealPath, UsesSharedBuildType, BuildPath, new string[] { "" });
 
 			// check editor and statged info is valid
 			CheckResult(BuildSource.EditorValid, "Editor build was invalid");
diff --git a/Engine/Source/Programs/AutomationTool/Gauntlet/SelfTest/Unreal/Gauntlet.SelfTest.TestUnrealOptions.cs b/Engine/Source/Programs/AutomationTool/Gauntlet/SelfTest/Unreal/Gauntlet.SelfTest.TestUnrealOptions.cs
--- a/Engine/Source/Programs/AutomationTool/Gauntlet/SelfTest/Unreal/Gauntlet.SelfTest.TestUnrealOptions.cs
+++ b/Engine/Source/Programs/AutomationTool/Gauntlet/SelfTest/Unreal/Gauntlet.SelfTest.TestUnrealOptions.cs
@@ -22,7 +22,7 @@ namespace Gauntlet.SelfTest
 		public override void TickTest()
 		{
 			// Grab the most recent Release build of Orion
-			UnrealBuildSource Build = new UnrealBuildSource(this.ProjectFile, this.UnrealPath, this.UsesSharedBuildType, this.BuildPath);
+			UnrealBuildSource Build = new UnrealBuildSource(this.ProjectName, this.ProjectFile, this.UnrealPath, this.UsesSharedBuildType, this.BuildPath);
 
 			// create client and server riles
 			UnrealSessionRole ClientRole = new UnrealSessionRole(UnrealTargetRole.Client, UnrealTargetPlatform.Win64, UnrealTargetConfiguration.Development);
diff --git a/Engine/Source/Programs/AutomationTool/Gauntlet/SelfTest/Unreal/Gauntlet.SelfTest.TestUnrealSession.cs b/Engine/Source/Programs/AutomationTool/Gauntlet/SelfTest/Unreal/Gauntlet.SelfTest.TestUnrealSession.cs
--- a/Engine/Source/Programs/AutomationTool/Gauntlet/SelfTest/Unreal/Gauntlet.SelfTest.TestUnrealSession.cs
+++ b/Engine/Source/Programs/AutomationTool/Gauntlet/SelfTest/Unreal/Gauntlet.SelfTest.TestUnrealSession.cs
@@ -33,7 +33,7 @@ namespace Gauntlet.SelfTest
 			});
 			
 			// Create a new build (params come from our base class will be similar to "OrionGame" and "p:\builds\orion\branch-cl")
-			UnrealBuildSource Build = new UnrealBuildSource(this.ProjectFile, this.UnrealPath, this.UsesSharedBuildType, this.BuildPath, new string[] { "" });
+			UnrealBuildSource Build = new UnrealBuildSource(this.ProjectName, this.ProjectFile, this.UnrealPath, this.UsesSharedBuildType, this.BuildPath, new string[] { "" });
 
 			// create a new options structure
 			UnrealOptions Options = new UnrealOptions();
diff --git a/Engine/Source/Programs/AutomationTool/Gauntlet/Unreal/RunUnreal.cs b/Engine/Source/Programs/AutomationTool/Gauntlet/Unreal/RunUnreal.cs
--- a/Engine/Source/Programs/AutomationTool/Gauntlet/Unreal/RunUnreal.cs
+++ b/Engine/Source/Programs/AutomationTool/Gauntlet/Unreal/RunUnreal.cs
@@ -123,7 +123,7 @@ namespace Gauntlet
 			DirectoryReference UnrealPath = new DirectoryReference(Environment.CurrentDirectory);
 					
 			// todo, pass this in as a BuildSource and remove the COntextOption params specific to finding builds
-			UnrealBuildSource BuildInfo = (UnrealBuildSource)Activator.CreateInstance(ContextOptions.BuildSourceType, new object[] { ContextOptions.ProjectPath, UnrealPath, ContextOptions.UsesSharedBuildType, ContextOptions.Build, ContextOptions.SearchPaths });
+			UnrealBuildSource BuildInfo = (UnrealBuildSource)Activator.CreateInstance(ContextOptions.BuildSourceType, new object[] { ContextOptions.Project, ContextOptions.ProjectPath, UnrealPath, ContextOptions.UsesSharedBuildType, ContextOptions.Build, ContextOptions.SearchPaths });
 
 			// Setup accounts
 			SetupAccounts();
