missed file from last cl

diff --git a/Engine/Source/Programs/AutomationTool/Gauntlet/Unreal/BuildSource/Gauntlet.UnrealBuildSource.cs b/Engine/Source/Programs/AutomationTool/Gauntlet/Unreal/BuildSource/Gauntlet.UnrealBuildSource.cs
--- a/Engine/Source/Programs/AutomationTool/Gauntlet/Unreal/BuildSource/Gauntlet.UnrealBuildSource.cs
+++ b/Engine/Source/Programs/AutomationTool/Gauntlet/Unreal/BuildSource/Gauntlet.UnrealBuildSource.cs
@@ -35,19 +35,19 @@ namespace Gauntlet
 
 		public int BuildCount { get { return DiscoveredBuilds.Count; } }
 
-		public UnrealBuildSource(FileReference InProjectPath, DirectoryReference InUnrealPath, bool InUsesSharedBuildType, string BuildReference) 
+		public UnrealBuildSource(string InProjectName, FileReference InProjectPath, DirectoryReference InUnrealPath, bool InUsesSharedBuildType, string BuildReference) 
 		{
-			InitBuildSource(InProjectPath, InUnrealPath, InUsesSharedBuildType, BuildReference, null);
+			InitBuildSource(InProjectName, InProjectPath, InUnrealPath, InUsesSharedBuildType, BuildReference, null);
 		}
 
-		public UnrealBuildSource(FileReference InProjectPath, DirectoryReference InUnrealPath, bool InUsesSharedBuildType, string BuildReference, Func<string, string> ResolutionDelegate)
+		public UnrealBuildSource(string InProjectName, FileReference InProjectPath, DirectoryReference InUnrealPath, bool InUsesSharedBuildType, string BuildReference, Func<string, string> ResolutionDelegate)
 		{
-			InitBuildSource(InProjectPath, InUnrealPath, InUsesSharedBuildType, BuildReference, ResolutionDelegate);
+			InitBuildSource(InProjectName, InProjectPath, InUnrealPath, InUsesSharedBuildType, BuildReference, ResolutionDelegate);
 		}
 
-		public UnrealBuildSource(FileReference InProjectPath, DirectoryReference InUnrealPath, bool InUsesSharedBuildType, string BuildReference, IEnumerable<string> InSearchPaths)
+		public UnrealBuildSource(string InProjectName, FileReference InProjectPath, DirectoryReference InUnrealPath, bool InUsesSharedBuildType, string BuildReference, IEnumerable<string> InSearchPaths)
 		{
-			InitBuildSource(InProjectPath, InUnrealPath, InUsesSharedBuildType, BuildReference, (string BuildRef) =>
+			InitBuildSource(InProjectName, InProjectPath, InUnrealPath, InUsesSharedBuildType, BuildReference, (string BuildRef) =>
 			{
 				foreach (string SearchPath in InSearchPaths)
 				{
@@ -74,13 +74,13 @@ namespace Gauntlet
 		}
 
 
-		protected void InitBuildSource(FileReference InProjectPath, DirectoryReference InUnrealPath, bool InUsesSharedBuildType, string InBuildArgument, Func<string, string> ResolutionDelegate)
+		protected void InitBuildSource(string InProjectName, FileReference InProjectPath, DirectoryReference InUnrealPath, bool InUsesSharedBuildType, string InBuildArgument, Func<string, string> ResolutionDelegate)
 		{
 			UnrealPath = InUnrealPath;
 			UsesSharedBuildType = InUsesSharedBuildType;
 
 			ProjectPath = InProjectPath;
-			ProjectName = InProjectPath.GetFileNameWithoutExtension();
+			ProjectName = InProjectName;
 
 			// Resolve the build argument into something meaningful
 			string ResolvedBuildName;
