Fixing boolean check on platform parse
#rnx

diff --git a/Engine/Source/Programs/AutomationTool/Gauntlet/Framework/Devices/Gauntlet.DevicePool.cs b/Engine/Source/Programs/AutomationTool/Gauntlet/Framework/Devices/Gauntlet.DevicePool.cs
--- a/Engine/Source/Programs/AutomationTool/Gauntlet/Framework/Devices/Gauntlet.DevicePool.cs
+++ b/Engine/Source/Programs/AutomationTool/Gauntlet/Framework/Devices/Gauntlet.DevicePool.cs
@@ -771,7 +771,7 @@ namespace Gauntlet
 
 							if (M.Success)
 							{
-								if (UnrealTargetPlatform.TryParse(M.Groups[1].ToString(), out DevicePlatform))
+								if (!UnrealTargetPlatform.TryParse(M.Groups[1].ToString(), out DevicePlatform))
 								{
 									throw new AutomationException("platform {0} is not a recognized device type", M.Groups[1].ToString());
 								}
