#Gauntlet Fix parameter order

diff --git a/Engine/Source/Programs/AutomationTool/Gauntlet/Unreal/BuildSource/Gauntlet.UnrealBuildSource.cs b/Engine/Source/Programs/AutomationTool/Gauntlet/Unreal/BuildSource/Gauntlet.UnrealBuildSource.cs
--- a/Engine/Source/Programs/AutomationTool/Gauntlet/Unreal/BuildSource/Gauntlet.UnrealBuildSource.cs
+++ b/Engine/Source/Programs/AutomationTool/Gauntlet/Unreal/BuildSource/Gauntlet.UnrealBuildSource.cs
@@ -443,11 +443,11 @@ namespace Gauntlet
 				// add in -game or -server
 				if (Role.RoleType.IsClient())
 				{
-					Config.CommandLine = "-game " + Config.CommandLine;
+					Config.CommandLine = Config.CommandLine + " -game";
 				}
 				else if (Role.RoleType.IsServer())
 				{
-					Config.CommandLine = "-server " + Config.CommandLine;
+					Config.CommandLine = Config.CommandLine + " -server";
 				}
 			}
 
@@ -544,7 +544,7 @@ namespace Gauntlet
 
 			if (TargetType.UsesEditor())
 			{
-				ExePath = string.Format("Engine/Binaries/{0}/UE4Editor{1}", BuildHostPlatform.Current.Platform, Platform.GetExeExtension(TargetPlatform));
+				ExePath = string.Format("Engine/Binaries/{0}/UE4Editor-Win64-DebugGame{1}", BuildHostPlatform.Current.Platform, Platform.GetExeExtension(TargetPlatform));
 			}
 			else
 			{
