#Gauntlet Remove a debugging line I left in

diff --git a/Engine/Source/Programs/AutomationTool/Gauntlet/Unreal/BuildSource/Gauntlet.UnrealBuildSource.cs b/Engine/Source/Programs/AutomationTool/Gauntlet/Unreal/BuildSource/Gauntlet.UnrealBuildSource.cs
--- a/Engine/Source/Programs/AutomationTool/Gauntlet/Unreal/BuildSource/Gauntlet.UnrealBuildSource.cs
+++ b/Engine/Source/Programs/AutomationTool/Gauntlet/Unreal/BuildSource/Gauntlet.UnrealBuildSource.cs
@@ -544,7 +544,7 @@ namespace Gauntlet
 
 			if (TargetType.UsesEditor())
 			{
-				ExePath = string.Format("Engine/Binaries/{0}/UE4Editor-Win64-DebugGame{1}", BuildHostPlatform.Current.Platform, Platform.GetExeExtension(TargetPlatform));
+				ExePath = string.Format("Engine/Binaries/{0}/UE4Editor{1}", BuildHostPlatform.Current.Platform, Platform.GetExeExtension(TargetPlatform));
 			}
 			else
 			{
