Fix Gauntlet always recopying obb's to target device.  This seem to only happen when running locally and was caused by whitespace not being trimmed from the file names after parsing 'ls' results from ADB.  Since this hasn't been seen on the test farm, I'm think this has to be a device specific or maybe a Windows/Linux isse.

diff --git a/Engine/Source/Programs/AutomationTool/Gauntlet/Platform/Android/Gauntlet.TargetDeviceAndroid.cs b/Engine/Source/Programs/AutomationTool/Gauntlet/Platform/Android/Gauntlet.TargetDeviceAndroid.cs
--- a/Engine/Source/Programs/AutomationTool/Gauntlet/Platform/Android/Gauntlet.TargetDeviceAndroid.cs
+++ b/Engine/Source/Programs/AutomationTool/Gauntlet/Platform/Android/Gauntlet.TargetDeviceAndroid.cs
@@ -965,7 +965,13 @@ namespace Gauntlet
 				// if != 0 then no folder exists
 				if (AdbResult.ExitCode == 0)
 				{
-					IEnumerable<string> CurrentRemoteFileList = AdbResult.Output.Replace("\r\n", "\n").Split('\n');
+					string[] Delimiters = { "\r\n", "\n" };
+					string[] CurrentRemoteFileList = AdbResult.Output.Split(Delimiters, StringSplitOptions.RemoveEmptyEntries);
+					for (int i = 0; i < CurrentRemoteFileList.Length; ++i)
+					{
+						CurrentRemoteFileList[i] = CurrentRemoteFileList[i].Trim();
+					}
+
 					IEnumerable<string> NewRemoteFileList = FilesToInstall.Values.Select(F => Path.GetFileName(F));
 
 					// delete any files that should not be there
