Fix android artifacts being saved to the wrong directory.

diff --git a/Engine/Source/Programs/AutomationTool/Gauntlet/Platform/Android/Gauntlet.TargetDeviceAndroid.cs b/Engine/Source/Programs/AutomationTool/Gauntlet/Platform/Android/Gauntlet.TargetDeviceAndroid.cs
--- a/Engine/Source/Programs/AutomationTool/Gauntlet/Platform/Android/Gauntlet.TargetDeviceAndroid.cs
+++ b/Engine/Source/Programs/AutomationTool/Gauntlet/Platform/Android/Gauntlet.TargetDeviceAndroid.cs
@@ -234,7 +234,7 @@ namespace Gauntlet
 			Directory.CreateDirectory(LocalSaved);
 
 			// pull all the artifacts
-			string ArtifactPullCommand = string.Format("pull {0} {1}", Install.AndroidDevice.DeviceArtifactPath, Install.AndroidDevice.LocalCachePath);
+			string ArtifactPullCommand = string.Format("pull {0} {1}", Install.AndroidDevice.DeviceArtifactPath, LocalSaved);
 			IProcessResult PullCmd = Install.AndroidDevice.RunAdbDeviceCommand(ArtifactPullCommand);
 
 			if (PullCmd.ExitCode != 0)
@@ -831,14 +831,7 @@ namespace Gauntlet
 				string StorageLocation = StorageQueryResult.Output.Trim(); // "/mnt/sdcard";
 
 				// remote dir used to save things
-				string RemoteDir = StorageLocation + "/UE4Game/" + AppConfig.ProjectName;
-
-				// if not a bulk/dev build, remote dir will be under /{StorageLocation}/Android/data/{PackageName}
-				if ((Build.Flags & ( BuildFlags.Bulk | BuildFlags.CanReplaceExecutable)) == 0)
-				{
-					RemoteDir = StorageLocation + "/Android/data/" + Build.AndroidPackageName + "/files/UE4Game/" + AppConfig.ProjectName;
-				}
-				
+				string RemoteDir = StorageLocation + "/UE4Game/" + AppConfig.ProjectName;				
 				string DependencyDir = RemoteDir + "/deps";
 
 				// device artifact path, always clear between runs
