Gauntlet: Add support for specifying client, server, editor exec cmds
#rnx

diff --git a/Engine/Source/Programs/AutomationTool/Gauntlet/Unreal/RunUnreal.cs b/Engine/Source/Programs/AutomationTool/Gauntlet/Unreal/RunUnreal.cs
--- a/Engine/Source/Programs/AutomationTool/Gauntlet/Unreal/RunUnreal.cs
+++ b/Engine/Source/Programs/AutomationTool/Gauntlet/Unreal/RunUnreal.cs
@@ -182,6 +182,14 @@ namespace Gauntlet
 
 					// look for -clientargs= and -editorclient etc
 					Role.ExtraArgs = Globals.Params.ParseValue(Type.ToString() + "Args", "");
+
+					// look for -clientexeccmds=, -editorexeccmds= etc, these are separate from clientargs for sanity
+					string ExecCmds = Globals.Params.ParseValue(Type.ToString() + "ExecCmds", "");
+					if (!string.IsNullOrEmpty(ExecCmds))
+					{
+						Role.ExtraArgs += string.Format(" -ExecCmds=\"{0}\"", ExecCmds);
+					}
+
 					bool UsesEditor = EditorForAllRoles || Globals.Params.ParseParam("Editor" + Type.ToString());
 
 					if (UsesEditor)
