Fixed issue with running automation tests using the editor as a client or server

diff --git a/Engine/Source/Programs/AutomationTool/Gauntlet/Unreal/BuildSource/Gauntlet.UnrealBuildSource.cs b/Engine/Source/Programs/AutomationTool/Gauntlet/Unreal/BuildSource/Gauntlet.UnrealBuildSource.cs
--- a/Engine/Source/Programs/AutomationTool/Gauntlet/Unreal/BuildSource/Gauntlet.UnrealBuildSource.cs
+++ b/Engine/Source/Programs/AutomationTool/Gauntlet/Unreal/BuildSource/Gauntlet.UnrealBuildSource.cs
@@ -429,6 +429,16 @@ namespace Gauntlet
 			// Add in editor - TODO, should this be in the editor build?
 			if (Role.RoleType.UsesEditor() || IsContentOnlyProject)
 			{
+				// add in -game or -server
+				if (Role.RoleType.IsClient())
+				{
+					Config.CommandLine = "-game " + Config.CommandLine;
+				}
+				else if (Role.RoleType.IsServer())
+				{
+					Config.CommandLine = "-server " + Config.CommandLine;
+				}
+
 				string ProjectParam = ProjectPath.FullName;
 
 				// if content only we need to provide a relative path to the uproject.
@@ -439,16 +449,6 @@ namespace Gauntlet
 
 				// project must be first
 				Config.CommandLine = ProjectParam + " " + Config.CommandLine;
-
-				// add in -game or -server
-				if (Role.RoleType.IsClient())
-				{
-					Config.CommandLine = Config.CommandLine + " -game";
-				}
-				else if (Role.RoleType.IsServer())
-				{
-					Config.CommandLine = Config.CommandLine + " -server";
-				}
 			}
 
             if (Role.FilesToCopy != null)
