Change distribution install batch files to use mkdir/cp instead of mv to intall OBB files to get around new restrictions in Android 9

#ue4
#android

diff --git a/Engine/Source/Programs/AutomationTool/Android/AndroidPlatform.Automation.cs b/Engine/Source/Programs/AutomationTool/Android/AndroidPlatform.Automation.cs
--- a/Engine/Source/Programs/AutomationTool/Android/AndroidPlatform.Automation.cs
+++ b/Engine/Source/Programs/AutomationTool/Android/AndroidPlatform.Automation.cs
@@ -683,13 +683,16 @@ public class AndroidPlatform : Platform
 						"\t$ADB $DEVICE shell 'rm -r $EXTERNAL_STORAGE/UE4Game/UE4CommandLine.txt" + "'",
 						"\t$ADB $DEVICE shell 'rm -r $EXTERNAL_STORAGE/" + TargetAndroidLocation + PackageName + "'",
 						"\t$ADB $DEVICE shell 'rm -r $EXTERNAL_STORAGE/Android/" + TargetAndroidLocation + PackageName + "'",
+						"\t$ADB $DEVICE shell 'rm -r $EXTERNAL_STORAGE/Download/" + TargetAndroidLocation + PackageName + "'",
 						bPackageDataInsideApk ? "" : "\techo",
 						bPackageDataInsideApk ? "" : "\techo Installing new data. Failures here indicate storage problems \\(missing SD card or bad permissions\\) and are fatal.",
 						bPackageDataInsideApk ? "" : "\tSTORAGE=$(echo \"`$ADB $DEVICE shell 'echo $EXTERNAL_STORAGE'`\" | cat -v | tr -d '^M')",
 						bPackageDataInsideApk ? "" : "\t$ADB $DEVICE " + OBBInstallCommand,
 						bPackageDataInsideApk ? "if [ 1 ]; then" : "\tif [ $? -eq 0 ]; then",
 						!bHavePatch ? "" : (bPackageDataInsideApk ? "" : "\t$ADB $DEVICE " + PatchInstallCommand),
-						bDontMoveOBB ? "" : "\t\t$ADB $DEVICE shell mv $STORAGE/Download/obb/" + PackageName + " $STORAGE/Android/obb/" + PackageName,
+						bDontMoveOBB ? "" : "\t\t$ADB $DEVICE shell mkdir $STORAGE/Android/" + TargetAndroidLocation + PackageName, // don't check for error since installing may create the obb directory
+						bDontMoveOBB ? "" : "\t\t$ADB $DEVICE shell cp $STORAGE/Download/" + TargetAndroidLocation + PackageName + "/* $STORAGE/Android/" + TargetAndroidLocation + PackageName,
+						bDontMoveOBB ? "" : "\t\t$ADB $DEVICE shell rm -r $STORAGE/Download/" + TargetAndroidLocation + PackageName,
 						"\t\techo",
 						"\t\techo Installation successful",
 						"\t\texit 0",
@@ -731,14 +734,17 @@ public class AndroidPlatform : Platform
 						"%ADB% %DEVICE% shell rm -r %STORAGE%/UE4Game/UE4CommandLine.txt", // we need to delete the commandline in UE4Game or it will mess up loading
 						"%ADB% %DEVICE% shell rm -r %STORAGE%/" + TargetAndroidLocation + PackageName,
 						"%ADB% %DEVICE% shell rm -r %STORAGE%/Android/" + TargetAndroidLocation + PackageName,
+						"%ADB% %DEVICE% shell rm -r %STORAGE%/Download/" + TargetAndroidLocation + PackageName,
 						bPackageDataInsideApk ? "" : "@echo.",
 						bPackageDataInsideApk ? "" : "@echo Installing new data. Failures here indicate storage problems (missing SD card or bad permissions) and are fatal.",
 						bPackageDataInsideApk ? "" : "%ADB% %DEVICE% " + OBBInstallCommand,
 						bPackageDataInsideApk ? "" : "if \"%ERRORLEVEL%\" NEQ \"0\" goto Error",
 						!bHavePatch ? "" : (bPackageDataInsideApk ? "" : "%ADB% %DEVICE% " + PatchInstallCommand),
 						!bHavePatch ? "" : (bPackageDataInsideApk ? "" : "if \"%ERRORLEVEL%\" NEQ \"0\" goto Error"),
-						bDontMoveOBB ? "" : "%ADB% %DEVICE% shell mv %STORAGE%/Download/obb/" + PackageName + " %STORAGE%/Android/obb/" + PackageName,
+						bDontMoveOBB ? "" : "%ADB% %DEVICE% shell mkdir %STORAGE%/Android/" + TargetAndroidLocation + PackageName, // don't check for error since installing may create the obb directory
+						bDontMoveOBB ? "" : "%ADB% %DEVICE% shell cp %STORAGE%/Download/" + TargetAndroidLocation + PackageName + "/* %STORAGE%/Android/" + TargetAndroidLocation + PackageName,
 						bDontMoveOBB ? "" : "if \"%ERRORLEVEL%\" NEQ \"0\" goto Error",
+						bDontMoveOBB ? "" : "%ADB% %DEVICE% shell rm -r %STORAGE%/Download/" + TargetAndroidLocation + PackageName,
 						"@echo.",
 						bNeedGrantStoragePermission ? "@echo Grant READ_EXTERNAL_STORAGE and WRITE_EXTERNAL_STORAGE to the apk for reading OBB file or game file in external storage." : "",
 						bNeedGrantStoragePermission ? "%ADB% %DEVICE% " + ReadPermissionGrantCommand : "",
