Exec when starting the process from the script produced by LinuxPlatform.Automation.cs (Thanks lotodore! PR 6118)

diff --git a/Engine/Source/Programs/AutomationTool/Linux/LinuxPlatform.Automation.cs b/Engine/Source/Programs/AutomationTool/Linux/LinuxPlatform.Automation.cs
--- a/Engine/Source/Programs/AutomationTool/Linux/LinuxPlatform.Automation.cs
+++ b/Engine/Source/Programs/AutomationTool/Linux/LinuxPlatform.Automation.cs
@@ -122,7 +122,7 @@ public abstract class BaseLinuxPlatform : Platform
 		Script.AppendFormat("UE4_TRUE_SCRIPT_NAME=$(echo \\\"$0\\\" | xargs readlink -f)" + EOL);
 		Script.AppendFormat("UE4_PROJECT_ROOT=$(dirname \"$UE4_TRUE_SCRIPT_NAME\")" + EOL);
 		Script.AppendFormat("chmod +x \"$UE4_PROJECT_ROOT/{0}\"" + EOL, StagedRelativeTargetPath);
-		Script.AppendFormat("\"$UE4_PROJECT_ROOT/{0}\" {1} $@ " + EOL, StagedRelativeTargetPath, StagedArguments);
+		Script.AppendFormat("exec \"$UE4_PROJECT_ROOT/{0}\" {1} $@ " + EOL, StagedRelativeTargetPath, StagedArguments);
 
 		// write out the 
 		FileReference.WriteAllText(IntermediateFile, Script.ToString());
