Add store version as a comment to main and patch OBB files so they are unique for uploading to Play Store

#ue4
#android

diff --git a/Engine/Source/Programs/AutomationTool/Android/AndroidPlatform.Automation.cs b/Engine/Source/Programs/AutomationTool/Android/AndroidPlatform.Automation.cs
--- a/Engine/Source/Programs/AutomationTool/Android/AndroidPlatform.Automation.cs
+++ b/Engine/Source/Programs/AutomationTool/Android/AndroidPlatform.Automation.cs
@@ -319,6 +319,7 @@ public class AndroidPlatform : Platform
 			ObbFile.CompressionMethod = CompressionMethod.None;
 			ObbFile.CompressionLevel = Ionic.Zlib.CompressionLevel.None;
 			ObbFile.UseZip64WhenSaving = Ionic.Zip.Zip64Option.Never;
+			ObbFile.Comment = String.Format("{0,10}", "1");
 
 			int ObbFileCount = 0;
 			ObbFile.AddProgress +=
@@ -352,6 +353,26 @@ public class AndroidPlatform : Platform
 		return true;
 	}
 
+	private bool UpdateObbStoreVersion(string Filename)
+	{
+		string Version = Path.GetFileNameWithoutExtension(Filename).Split('.')[1];
+
+		using (ZipFile ObbFile = ZipFile.Read(Filename))
+		{
+			// Add the store version from the filename as a comment
+			ObbFile.Comment = String.Format("{0,10}", Version);
+			try
+			{
+				ObbFile.Save();
+			}
+			catch (Exception)
+			{
+				return false;
+			}
+		}
+		return true;
+	}
+
 	public override void Package(ProjectParams Params, DeploymentContext SC, int WorkingCL)
 	{
 		if (SC.StageTargetConfigurations.Count != 1)
@@ -433,8 +454,8 @@ public class AndroidPlatform : Platform
 				FilesToObb = new List<FileReference>();
 
 				// Collect the filesize and place into Obb or Patch list
-				Int64 MainObbSize = 22 + 4096;		// EOCD without comment + padding
-				Int64 PatchObbSize = 22 + 4096;		// EOCD without comment + padding
+				Int64 MainObbSize = 22 + 10 + 4096;		// EOCD wit comment (store version) + padding
+				Int64 PatchObbSize = 22 + 10 + 4096;	// EOCD wit comment (store version) + padding
 				foreach (FileReference FileRef in FilesForObb)
 				{
 					FileInfo LocalFileInfo = new FileInfo(FileRef.FullName);
@@ -558,11 +579,17 @@ public class AndroidPlatform : Platform
 				    ObbName = GetFinalObbName(ApkName, SC);
 					CopyFile(LocalObbName, ObbName);
 
+					// apply store version to OBB to make it unique for PlayStore upload
+					UpdateObbStoreVersion(ObbName);
+
 					if (File.Exists(LocalPatchName))
 					{
 						DevicePatchName = GetDevicePatchName(ApkName, SC);
 						PatchName = GetFinalPatchName(ApkName, SC);
 						CopyFile(LocalPatchName, PatchName);
+
+						// apply store version to OBB to make it unique for PlayStore upload
+						UpdateObbStoreVersion(PatchName);
 					}
 				}
 
