Gauntlet: Fixing issue with server/beacon offset
#rnx

diff --git a/Engine/Source/Programs/AutomationTool/Gauntlet/Unreal/Game/EpicGame.TestConfig.cs b/Engine/Source/Programs/AutomationTool/Gauntlet/Unreal/Game/EpicGame.TestConfig.cs
--- a/Engine/Source/Programs/AutomationTool/Gauntlet/Unreal/Game/EpicGame.TestConfig.cs
+++ b/Engine/Source/Programs/AutomationTool/Gauntlet/Unreal/Game/EpicGame.TestConfig.cs
@@ -47,16 +47,16 @@ namespace EpicGame
 		/// <summary>
 		/// Unique server port to avoid matchmaking collisions
 		/// </summary>
-		[AutoParam(ServerPortStart)]
-		protected int ServerPort;
-		const int ServerPortStart = 7777;
+		[AutoParam]
+		int ServerPortStart = 7777;
+		protected int ServerPort;		
 
 		/// <summary>
 		/// Unique server beacon port to avoid matchmaking collisions
 		/// </summary>
-		[AutoParam(BeaconPortStart)]
+		[AutoParam]
+		int BeaconPortStart = 15000;
 		protected int BeaconPort;
-		const int BeaconPortStart = 15000;
 
 		/// <summary>
 		/// Make sure the client gets -logpso when we are collecting them
@@ -97,7 +97,7 @@ namespace EpicGame
 
 			// techinically this doesn't matter for mcp because the server will pick a free port and tell the backend what its using, but
 			// nomcp requires us to know the port and thus we need to make sure ones we pick haven't been previously assigned or grabbed
-			if (NumberOfConfigsCreated > 1)
+			if (NumberOfConfigsCreated > 1 || Globals.IsWorker)
 			{
 				BuildIDOverride += string.Format("{0}", NumberOfConfigsCreated);
 				ServerPort = (ServerPortStart + NumberOfConfigsCreated);
