Back out changelist 8061173. Possible cause of EdTest timeout

diff --git a/Engine/Source/Programs/AutomationTool/Gauntlet/Unreal/Game/EpicGame.TestConfig.cs b/Engine/Source/Programs/AutomationTool/Gauntlet/Unreal/Game/EpicGame.TestConfig.cs
--- a/Engine/Source/Programs/AutomationTool/Gauntlet/Unreal/Game/EpicGame.TestConfig.cs
+++ b/Engine/Source/Programs/AutomationTool/Gauntlet/Unreal/Game/EpicGame.TestConfig.cs
@@ -47,16 +47,16 @@ namespace EpicGame
 		/// <summary>
 		/// Unique server port to avoid matchmaking collisions
 		/// </summary>
-		[AutoParam]
-		int ServerPortStart = 7777;
-		protected int ServerPort;		
+		[AutoParam(ServerPortStart)]
+		protected int ServerPort;
+		const int ServerPortStart = 7777;
 
 		/// <summary>
 		/// Unique server beacon port to avoid matchmaking collisions
 		/// </summary>
-		[AutoParam]
-		int BeaconPortStart = 15000;
+		[AutoParam(BeaconPortStart)]
 		protected int BeaconPort;
+		const int BeaconPortStart = 15000;
 
 		/// <summary>
 		/// Make sure the client gets -logpso when we are collecting them
@@ -97,7 +97,7 @@ namespace EpicGame
 
 			// techinically this doesn't matter for mcp because the server will pick a free port and tell the backend what its using, but
 			// nomcp requires us to know the port and thus we need to make sure ones we pick haven't been previously assigned or grabbed
-			if (NumberOfConfigsCreated > 1 || Globals.IsWorker)
+			if (NumberOfConfigsCreated > 1)
 			{
 				BuildIDOverride += string.Format("{0}", NumberOfConfigsCreated);
 				ServerPort = (ServerPortStart + NumberOfConfigsCreated);
