Gauntlet: avoid port/beacon collisions when running multiple instances of UAT
#rnx

diff --git a/Engine/Source/Programs/AutomationTool/Gauntlet/Unreal/Game/EpicGame.TestConfig.cs b/Engine/Source/Programs/AutomationTool/Gauntlet/Unreal/Game/EpicGame.TestConfig.cs
--- a/Engine/Source/Programs/AutomationTool/Gauntlet/Unreal/Game/EpicGame.TestConfig.cs
+++ b/Engine/Source/Programs/AutomationTool/Gauntlet/Unreal/Game/EpicGame.TestConfig.cs
@@ -47,17 +47,20 @@ namespace EpicGame
 		/// <summary>
 		/// Unique server port to avoid matchmaking collisions
 		/// </summary>
-		[AutoParam(ServerPortStart)]
-		protected int ServerPort;
-		const int ServerPortStart = 7777;
+		[AutoParam]
+		int ServerPortStart = 7777;
+
+		protected int ServerPort { get; private set; }
+
 
 		/// <summary>
 		/// Unique server beacon port to avoid matchmaking collisions
 		/// </summary>
-		[AutoParam(BeaconPortStart)]
-		protected int BeaconPort;
-		const int BeaconPortStart = 15000;
+		[AutoParam]
+		int BeaconPortStart = 15000;
 
+		protected int BeaconPort { get; private set; }
+		
 		/// <summary>
 		/// Make sure the client gets -logpso when we are collecting them
 		/// </summary>
@@ -95,6 +98,9 @@ namespace EpicGame
 				BuildIDOverride = LocalIP.Substring(LocalIP.Length - 4);
 			}
 
+			ServerPort = ServerPortStart;
+			BeaconPort = BeaconPortStart;
+
 			// techinically this doesn't matter for mcp because the server will pick a free port and tell the backend what its using, but
 			// nomcp requires us to know the port and thus we need to make sure ones we pick haven't been previously assigned or grabbed
 			if (NumberOfConfigsCreated > 1)
