Fixed issue running ElementalDemo on Mac under Gauntlet

diff --git a/Engine/Source/Programs/AutomationTool/Gauntlet/Unreal/BuildSource/Gauntlet.StagedBuildSource.cs b/Engine/Source/Programs/AutomationTool/Gauntlet/Unreal/BuildSource/Gauntlet.StagedBuildSource.cs
--- a/Engine/Source/Programs/AutomationTool/Gauntlet/Unreal/BuildSource/Gauntlet.StagedBuildSource.cs
+++ b/Engine/Source/Programs/AutomationTool/Gauntlet/Unreal/BuildSource/Gauntlet.StagedBuildSource.cs
@@ -89,7 +89,7 @@ namespace Gauntlet
 				string[] ExecutablePaths = new string[]
 				{
 					Path.Combine(InPath, InProjectName, "Binaries", InPlatform.ToString()),
-					Path.Combine(InPath),
+					//Path.Combine(InPath),
 					Path.Combine(InPath, "Engine", "Binaries", InPlatform.ToString()),
 				};
 
