Fix issue with missing -IOS on Code only buildversions inside the manifest.
Leigh.Swift

diff --git a/Engine/Source/Programs/AutomationTool/AutomationUtils/MCPPublic.cs b/Engine/Source/Programs/AutomationTool/AutomationUtils/MCPPublic.cs
--- a/Engine/Source/Programs/AutomationTool/AutomationUtils/MCPPublic.cs
+++ b/Engine/Source/Programs/AutomationTool/AutomationUtils/MCPPublic.cs
@@ -697,6 +697,10 @@ namespace EpicGames.MCP.Automation
 			/// </summary>
 			public string BuildVersion;
 			/// <summary>
+			/// Used as part of the build version string.
+			/// </summary>
+			public MCPPlatform Platform;
+			/// <summary>
 			/// Optional. The set of files that should be kept from ManifestA.
 			/// </summary>
 			public HashSet<string> FilesToKeepFromA;
@@ -704,6 +708,7 @@ namespace EpicGames.MCP.Automation
 			/// Optional. The set of files that should be kept from ManifestB.
 			/// </summary>
 			public HashSet<string> FilesToKeepFromB;
+			
 		}
 
 		public class ManifestDiffOptions
