Formalize test report generation data path
#rnx

diff --git a/Engine/Source/Programs/AutomationTool/Gauntlet/Framework/Gauntlet.Utils.cs b/Engine/Source/Programs/AutomationTool/Gauntlet/Framework/Gauntlet.Utils.cs
--- a/Engine/Source/Programs/AutomationTool/Gauntlet/Framework/Gauntlet.Utils.cs
+++ b/Engine/Source/Programs/AutomationTool/Gauntlet/Framework/Gauntlet.Utils.cs
@@ -17,7 +17,7 @@ using System.Text;
 
 namespace Gauntlet
 {
-	public class Globals
+	public static class Globals
 	{
 		static Params InnerParams = new Params(Environment.GetCommandLineArgs());
 
