Fix Gauntlet launching projects that have a space in the path.

#rnx

diff --git a/Engine/Source/Programs/AutomationTool/Gauntlet/Platform/Mac/Gauntlet.TargetDeviceMac.cs b/Engine/Source/Programs/AutomationTool/Gauntlet/Platform/Mac/Gauntlet.TargetDeviceMac.cs
--- a/Engine/Source/Programs/AutomationTool/Gauntlet/Platform/Mac/Gauntlet.TargetDeviceMac.cs
+++ b/Engine/Source/Programs/AutomationTool/Gauntlet/Platform/Mac/Gauntlet.TargetDeviceMac.cs
@@ -201,7 +201,7 @@ namespace Gauntlet
 
 			// Mac always forces this to stop logs and other artifacts going to different places
 			// Mac always forces this to stop logs and other artifacts going to different places
-			MacApp.CommandArguments += string.Format(" -userdir={0}", UserDir);
+			MacApp.CommandArguments += string.Format(" -userdir=\"{0}\"", UserDir);
 			MacApp.ArtifactPath = Path.Combine(UserDir, @"Saved");
 
 			// temp - Mac doesn't support -userdir?
@@ -280,7 +280,7 @@ namespace Gauntlet
 			MacApp.RunOptions = RunOptions;
 
 			// Mac always forces this to stop logs and other artifacts going to different places
-			MacApp.CommandArguments += string.Format(" -userdir={0}", UserDir);
+			MacApp.CommandArguments += string.Format(" -userdir=\"{0}\"", UserDir);
 			MacApp.ArtifactPath = Path.Combine(UserDir, @"Saved");
 
 			// now turn the Foo.app into Foo/Content/MacOS/Foo
diff --git a/Engine/Source/Programs/AutomationTool/Gauntlet/Platform/Windows/Gauntlet.TargetDeviceWindows.cs b/Engine/Source/Programs/AutomationTool/Gauntlet/Platform/Windows/Gauntlet.TargetDeviceWindows.cs
--- a/Engine/Source/Programs/AutomationTool/Gauntlet/Platform/Windows/Gauntlet.TargetDeviceWindows.cs
+++ b/Engine/Source/Programs/AutomationTool/Gauntlet/Platform/Windows/Gauntlet.TargetDeviceWindows.cs
@@ -385,7 +385,7 @@ namespace Gauntlet
 
 			if (string.IsNullOrEmpty(UserDir) == false)
 			{
-				WinApp.CommandArguments += string.Format(" -userdir={0}", UserDir);
+				WinApp.CommandArguments += string.Format(" -userdir=\"{0}\"", UserDir);
 				WinApp.ArtifactPath = Path.Combine(UserDir, @"Saved");
 
 				Utils.SystemHelpers.MarkDirectoryForCleanup(UserDir);
@@ -500,7 +500,7 @@ namespace Gauntlet
 			WinApp.RunOptions = RunOptions;
 	
 			// Force this to stop logs and other artifacts going to different places
-			WinApp.CommandArguments = AppConfig.CommandLine + string.Format(" -userdir={0}", UserDir);
+			WinApp.CommandArguments = AppConfig.CommandLine + string.Format(" -userdir=\"{0}\"", UserDir);
 			WinApp.ArtifactPath = Path.Combine(UserDir, @"Saved");
 			WinApp.ExecutablePath = EditorBuild.ExecutablePath;
 
diff --git a/Engine/Source/Programs/AutomationTool/Gauntlet/Unreal/BuildSource/Gauntlet.UnrealBuildSource.cs b/Engine/Source/Programs/AutomationTool/Gauntlet/Unreal/BuildSource/Gauntlet.UnrealBuildSource.cs
--- a/Engine/Source/Programs/AutomationTool/Gauntlet/Unreal/BuildSource/Gauntlet.UnrealBuildSource.cs
+++ b/Engine/Source/Programs/AutomationTool/Gauntlet/Unreal/BuildSource/Gauntlet.UnrealBuildSource.cs
@@ -448,7 +448,7 @@ namespace Gauntlet
 				}
 
 				// project must be first
-				Config.CommandLine = ProjectParam + " " + Config.CommandLine;
+				Config.CommandLine = String.Format("\"{0}\"", ProjectParam) + " " + Config.CommandLine;
 			}
 
             if (Role.FilesToCopy != null)
