Fix build error when using the BuildCook process.
#RNX

diff --git a/Engine/Source/Programs/AutomationTool/Scripts/BuildProjectCommand.Automation.cs b/Engine/Source/Programs/AutomationTool/Scripts/BuildProjectCommand.Automation.cs
--- a/Engine/Source/Programs/AutomationTool/Scripts/BuildProjectCommand.Automation.cs
+++ b/Engine/Source/Programs/AutomationTool/Scripts/BuildProjectCommand.Automation.cs
@@ -88,7 +88,7 @@ public partial class Project : CommandUtils
 				CrashReportPlatforms.Add(EditorPlatform);
 				if (Params.EditorTargets.Contains("UnrealHeaderTool") == false)
 				{
-					Agenda.AddTargets(new string[] { "UnrealHeaderTool" }, EditorPlatform, EditorConfiguration);
+					Agenda.AddTargets(new string[] { "UnrealHeaderTool" }, EditorPlatform, EditorConfiguration, Params.CodeBasedUprojectPath);
 				}
 				if (Params.EditorTargets.Contains("ShaderCompileWorker") == false)
 				{
