Fixed issue with missing arguments to UHT

diff --git a/Engine/Source/Programs/AutomationTool/Scripts/BuildTarget.Automation.cs b/Engine/Source/Programs/AutomationTool/Scripts/BuildTarget.Automation.cs
--- a/Engine/Source/Programs/AutomationTool/Scripts/BuildTarget.Automation.cs
+++ b/Engine/Source/Programs/AutomationTool/Scripts/BuildTarget.Automation.cs
@@ -161,7 +161,7 @@ namespace AutomationTool
 
 			if (!NoTools)
 			{
-				Agenda.AddTarget("UnrealHeaderTool", CurrentPlatform, UnrealTargetConfiguration.Development);
+				Agenda.AddTarget("UnrealHeaderTool", CurrentPlatform, UnrealTargetConfiguration.Development, ProjectFile);
 			}
 
 			if (string.IsNullOrEmpty(EditorTarget) == false)
@@ -172,9 +172,9 @@ namespace AutomationTool
 
 				if (!NoTools)
 				{
+					Agenda.AddTarget("UnrealPak", CurrentPlatform, UnrealTargetConfiguration.Development, ProjectFile);
 					Agenda.AddTarget("ShaderCompileWorker", CurrentPlatform, UnrealTargetConfiguration.Development);
 					Agenda.AddTarget("UnrealLightmass", CurrentPlatform, UnrealTargetConfiguration.Development);
-					Agenda.AddTarget("UnrealPak", CurrentPlatform, UnrealTargetConfiguration.Development);
 					Agenda.AddTarget("CrashReportClient", CurrentPlatform, UnrealTargetConfiguration.Shipping);
 				}
 			}
