Optional paks are now contained in their own install bundles instead of using optional tag sets.

diff --git a/Engine/Source/Programs/AutomationTool/AutomationUtils/BundleUtils.cs b/Engine/Source/Programs/AutomationTool/AutomationUtils/BundleUtils.cs
--- a/Engine/Source/Programs/AutomationTool/AutomationUtils/BundleUtils.cs
+++ b/Engine/Source/Programs/AutomationTool/AutomationUtils/BundleUtils.cs
@@ -15,27 +15,14 @@ namespace AutomationUtils.Automation
 		public interface IReadOnlyBundleSettings
 		{
 			string Name { get; }
-			string OverrideName { get; } // name used for ini section name in the BackgroundDownload.ini file
-			string ParentName { get;  }
 			List<string> Tags { get; }
-			List<BundleSettings> Children { get; }
-			bool bFoundParent { get; }
 			bool bContainsShaderLibrary { get; }
 		}
 
 
 		public class BundleSettings : IReadOnlyBundleSettings
 		{
-			public BundleSettings()
-			{
-				Children = new List<BundleSettings>();
-				bFoundParent = false;
-				bContainsShaderLibrary = false;
-			}
-
 			public string Name { get; set; }
-			public string OverrideName { get; set; }
-			public string ParentName { get; set; }
 			public List<string> Tags { get; set; }
 			public List<BundleSettings> Children { get; set; }
 			public List<string> FileRegex { get; set; }
@@ -46,7 +33,6 @@ namespace AutomationUtils.Automation
 		};
 		public static bool LoadBundleConfig(string BundleIniFile, ref Dictionary<string, BundleSettings> Bundles)
 		{
-			
 			if ( System.IO.File.Exists(BundleIniFile) == false )
 			{
 				CommandUtils.LogWarning("Unable to find bundle config ini file  {0}", BundleIniFile);
@@ -63,22 +49,6 @@ namespace AutomationUtils.Automation
 				Bundle.Name = SectionName;
 				Bundle.Priority = PriorityCounter;
 				++PriorityCounter;
-				{
-					string OverrideName;
-					if (BundleConfig.GetString(SectionName, "OverrideName", out OverrideName))
-					{
-						Bundle.OverrideName = OverrideName;
-					}
-					else
-					{
-						Bundle.OverrideName = Bundle.Name;
-					}
-				}
-				{
-					string ParentName;
-					BundleConfig.GetString(SectionName, "Parent", out ParentName);
-					Bundle.ParentName = ParentName;
-				}
 				{
 					string ExecFileName;
 					BundleConfig.GetString(SectionName, "ExecFileName", out ExecFileName);
@@ -105,24 +75,9 @@ namespace AutomationUtils.Automation
 				}
 
 				Bundles.Add(SectionName, Bundle);
-				//BundleConfig.GetArray(SectionName, "OptionalTags", out Bundle.OptionalTags);
 			}
 
-
-
-			foreach (var BundleIt in Bundles)
-			{
-				BundleSettings Bundle = BundleIt.Value;
-				if (Bundles.ContainsKey(Bundle.ParentName))
-				{
-					BundleSettings ParentBundle = Bundles[Bundle.ParentName];
-					ParentBundle.Children.Add(Bundle);
-					Bundle.bFoundParent = true;
-				}
-			}
 			return true;
 		}
-
-
 	}
 }
diff --git a/Engine/Source/Runtime/Core/Public/GenericPlatform/GenericPlatformMemory.h b/Engine/Source/Runtime/Core/Public/GenericPlatform/GenericPlatformMemory.h
--- a/Engine/Source/Runtime/Core/Public/GenericPlatform/GenericPlatformMemory.h
+++ b/Engine/Source/Runtime/Core/Public/GenericPlatform/GenericPlatformMemory.h
@@ -21,19 +21,34 @@ struct FGenericMemoryStats;
 * Platform-dependent "bucket" for memory size, where Default is the normal, or possibly the largest.
 * This is generally used for texture LOD settings for how to fit in smaller memory devices
 */
+#define PLATFORM_MEMORY_SIZE_BUCKET_LIST(XBUCKET) \
+	/* not used with texture LODs (you can't use bigger textures than what is cooked out, which is what Default should map to) */ \
+	XBUCKET(Largest)	\
+	XBUCKET(Larger)		\
+	/* these are used by texture LODs */ \
+	XBUCKET(Default)	\
+	XBUCKET(Smaller)	\
+	XBUCKET(Smallest)	\
+	XBUCKET(Tiniest)	\
+
+#define PLATFORM_MEMORY_SIZE_BUCKET_ENUM(Name) Name,
 enum class EPlatformMemorySizeBucket
 {
-	// not used with texture LODs (you can't use bigger textures than what is cooked out, which is what Default should map to)
-	Largest,
-	Larger,
-
-	// these are used by texture LODs
-	Default,
-	Smaller,
-	Smallest,
-    Tiniest
+	PLATFORM_MEMORY_SIZE_BUCKET_LIST(PLATFORM_MEMORY_SIZE_BUCKET_ENUM)
 };
+#undef PLATFORM_MEMORY_SIZE_BUCKET_ENUM
 
+inline const TCHAR* LexToString(EPlatformMemorySizeBucket Bucket)
+{
+#define PLATFORM_MEMORY_SIZE_BUCKET_LEXTOSTRING(Name) case EPlatformMemorySizeBucket::Name: return TEXT(#Name);
+	switch (Bucket)
+	{
+		PLATFORM_MEMORY_SIZE_BUCKET_LIST(PLATFORM_MEMORY_SIZE_BUCKET_LEXTOSTRING)
+	}
+#undef PLATFORM_MEMORY_SIZE_BUCKET_LEXTOSTRING
+
+	return TEXT("Unknown");
+}
 
 /** 
  * Struct used to hold common memory constants for all platforms.
diff --git a/Engine/Source/Runtime/InstallBundleManager/Private/BundlePrereqCombinedStatusHelper.cpp b/Engine/Source/Runtime/InstallBundleManager/Private/BundlePrereqCombinedStatusHelper.cpp
--- a/Engine/Source/Runtime/InstallBundleManager/Private/BundlePrereqCombinedStatusHelper.cpp
+++ b/Engine/Source/Runtime/InstallBundleManager/Private/BundlePrereqCombinedStatusHelper.cpp
@@ -98,7 +98,7 @@ void FBundlePrereqCombinedStatusHelper::CleanUpDelegates()
 	}
 }
 
-void FBundlePrereqCombinedStatusHelper::SetBundlesToTrackFromContentState(FInstallBundleContentState& BundleContentState)
+void FBundlePrereqCombinedStatusHelper::SetBundlesToTrackFromContentState(const FInstallBundleContentState& BundleContentState)
 {
 	RequiredBundleNames.Empty();
 	CachedBundleWeights.Empty();
@@ -113,7 +113,7 @@ void FBundlePrereqCombinedStatusHelper::SetBundlesToTrackFromContentState(FInsta
 	CurrentCombinedStatus.bBundleRequiresUpdate = bBundleNeedsUpdate;
 	
 	//Save required bundles and their weights
-	for (TPair<FName, float>& BundleStatePair : BundleContentState.IndividualBundleWeights)
+	for (const TPair<FName, float>& BundleStatePair : BundleContentState.IndividualBundleWeights)
 	{
 		RequiredBundleNames.Add(BundleStatePair.Key);
 		CachedBundleWeights.FindOrAdd(BundleStatePair.Key) = BundleStatePair.Value;
diff --git a/Engine/Source/Runtime/InstallBundleManager/Public/BundlePrereqCombinedStatusHelper.h b/Engine/Source/Runtime/InstallBundleManager/Public/BundlePrereqCombinedStatusHelper.h
--- a/Engine/Source/Runtime/InstallBundleManager/Public/BundlePrereqCombinedStatusHelper.h
+++ b/Engine/Source/Runtime/InstallBundleManager/Public/BundlePrereqCombinedStatusHelper.h
@@ -50,7 +50,7 @@ public:
 	FBundlePrereqCombinedStatusHelper& operator=(FBundlePrereqCombinedStatusHelper&& Other);
 	
 	//Setup tracking for all bundles required in the supplied BundleContentState
-	void SetBundlesToTrackFromContentState(FInstallBundleContentState& BundleContentState);
+	void SetBundlesToTrackFromContentState(const FInstallBundleContentState& BundleContentState);
 	
 	//Get current CombinedBundleStatus for everything setup to track
 	const FCombinedBundleStatus& GetCurrentCombinedState() const;
diff --git a/Engine/Source/Runtime/InstallBundleManager/Public/InstallBundleManagerInterface.h b/Engine/Source/Runtime/InstallBundleManager/Public/InstallBundleManagerInterface.h
--- a/Engine/Source/Runtime/InstallBundleManager/Public/InstallBundleManagerInterface.h
+++ b/Engine/Source/Runtime/InstallBundleManager/Public/InstallBundleManagerInterface.h
@@ -231,6 +231,7 @@ struct FInstallBundleContentState
 	TMap<FName, float> IndividualBundleWeights;
 	uint64 DownloadSize = 0;
 	uint64 InstallSize = 0;
+	uint64 InstallOverheadSize = 0;
 	uint64 FreeSpace = 0;
 };
 
@@ -311,11 +312,13 @@ public:
 
     virtual void CancelAllGetContentStateRequestsForTag(FName RequestTag) = 0;
     
-	virtual FInstallBundleRequestInfo RequestRemoveBundle(FName BundleName) = 0;
+	virtual FInstallBundleRequestInfo RequestRemoveContent(FName BundleName) = 0;
 
-	virtual void RequestRemoveBundleOnNextInit(FName BundleName) = 0;
+	virtual void RequestRemoveContentOnNextInit(FName BundleName) = 0;
+	virtual void RequestRemoveContentOnNextInit(TArrayView<FName> BundleNames) = 0;
 
-	virtual void CancelRequestRemoveBundleOnNextInit(FName BundleName) = 0;
+	virtual void CancelRequestRemoveContentOnNextInit(FName BundleName) = 0;
+	virtual void CancelRequestRemoveContentOnNextInit(TArrayView<FName> BundleName) = 0;
 
 	virtual void CancelBundle(FName BundleName, EInstallBundleCancelFlags Flags) = 0;
 
diff --git a/Engine/Source/Runtime/NullInstallBundleManager/Source/NullInstallBundleManager.cpp b/Engine/Source/Runtime/NullInstallBundleManager/Source/NullInstallBundleManager.cpp
--- a/Engine/Source/Runtime/NullInstallBundleManager/Source/NullInstallBundleManager.cpp
+++ b/Engine/Source/Runtime/NullInstallBundleManager/Source/NullInstallBundleManager.cpp
@@ -45,7 +45,7 @@ class FNullInstallBundleManager : public IInstallBundleManager
 		return RetInfo;
 	}
 
-	virtual FInstallBundleRequestInfo RequestRemoveBundle(FName BundleName) override
+	virtual FInstallBundleRequestInfo RequestRemoveContent(FName BundleName) override
 	{
 		FInstallBundleRequestInfo RetInfo;
 		return RetInfo;
@@ -65,19 +65,23 @@ class FNullInstallBundleManager : public IInstallBundleManager
 		Callback.ExecuteIfBound(State);
 	}
 
-    virtual void CancelAllGetContentStateRequestsForTag(FName RequestTag) override
-    {
-        
-    }
-    
-	virtual void RequestRemoveBundleOnNextInit(FName BundleName) override
+	virtual void CancelAllGetContentStateRequestsForTag(FName RequestTag) override
 	{
+	}
 
+	virtual void RequestRemoveContentOnNextInit(FName BundleName) override
+	{
+	}
+	virtual void RequestRemoveContentOnNextInit(TArrayView<FName> BundleNames) override
+	{
 	}
 
-	virtual void CancelRequestRemoveBundleOnNextInit(FName BundleName) override
+	virtual void CancelRequestRemoveContentOnNextInit(FName BundleName) override
 	{
 
+	}
+	virtual void CancelRequestRemoveContentOnNextInit(TArrayView<FName> BundleNames) override
+	{
 	}
 
 	virtual void CancelBundle(FName BundleName, EInstallBundleCancelFlags Flags) override
