BuildGraph: Fix incorrect error handling when a BuildGraph script contains a parse error.

diff --git a/Engine/Source/Programs/AutomationTool/BuildGraph/Script.cs b/Engine/Source/Programs/AutomationTool/BuildGraph/Script.cs
--- a/Engine/Source/Programs/AutomationTool/BuildGraph/Script.cs
+++ b/Engine/Source/Programs/AutomationTool/BuildGraph/Script.cs
@@ -81,7 +81,7 @@ namespace AutomationTool
 				{
 					if (!Document.bHasErrors)
 					{
-						Log.TraceError(File.FullName, Ex.LineNumber, "{0}", Ex.Message);
+						Log.TraceError(File, Ex.LineNumber, "{0}", Ex.Message);
 						Document.bHasErrors = true;
 					}
 				}
