Allow addition of robomerge tags to the HLOD UAT commands

diff --git a/Engine/Source/Programs/AutomationTool/Scripts/RebuildHLODCommand.Automation.cs b/Engine/Source/Programs/AutomationTool/Scripts/RebuildHLODCommand.Automation.cs
--- a/Engine/Source/Programs/AutomationTool/Scripts/RebuildHLODCommand.Automation.cs
+++ b/Engine/Source/Programs/AutomationTool/Scripts/RebuildHLODCommand.Automation.cs
@@ -375,7 +375,14 @@ namespace AutomationScripts.Automation
 					// RoboMerge default behaviour is to ignore buildmachine changes 
 					return "";
 				}
-				return (Command == "deadend" ? "#ROBOMERGE #DEADEND" : "#ROBOMERGE[ALL]") + " #DisregardExcludedAuthors";
+				else if (Command == "all")
+				{
+					return "#ROBOMERGE[ALL] #DisregardExcludedAuthors";
+				}
+				else
+				{
+					return "#ROBOMERGE " + Command + " #DisregardExcludedAuthors";
+				}
 			}
 		}
 
