UAT: Build both flavors of UHT (with project and without) when packaging, since we don't know whether any UHT plugins are enabled in the project.

#rnx

diff --git a/Engine/Source/Programs/AutomationTool/Scripts/BuildProjectCommand.Automation.cs b/Engine/Source/Programs/AutomationTool/Scripts/BuildProjectCommand.Automation.cs
--- a/Engine/Source/Programs/AutomationTool/Scripts/BuildProjectCommand.Automation.cs
+++ b/Engine/Source/Programs/AutomationTool/Scripts/BuildProjectCommand.Automation.cs
@@ -88,6 +88,7 @@ public partial class Project : CommandUtils
 				CrashReportPlatforms.Add(EditorPlatform);
 				if (Params.EditorTargets.Contains("UnrealHeaderTool") == false)
 				{
+					Agenda.AddTargets(new string[] { "UnrealHeaderTool" }, EditorPlatform, EditorConfiguration);
 					Agenda.AddTargets(new string[] { "UnrealHeaderTool" }, EditorPlatform, EditorConfiguration, Params.CodeBasedUprojectPath);
 				}
 				if (Params.EditorTargets.Contains("ShaderCompileWorker") == false)
