Fixed Mac UAT code for deleting extra RPATHs to properly handle paths containing space character

diff --git a/Engine/Source/Programs/AutomationTool/Mac/MacPlatform.Automation.cs b/Engine/Source/Programs/AutomationTool/Mac/MacPlatform.Automation.cs
--- a/Engine/Source/Programs/AutomationTool/Mac/MacPlatform.Automation.cs
+++ b/Engine/Source/Programs/AutomationTool/Mac/MacPlatform.Automation.cs
@@ -284,7 +284,7 @@ public class MacPlatform : Platform
 			if (CommandResult.ExitCode == 0)
 			{
 				StringReader Reader = new StringReader(CommandResult.Output);
-				Regex RPathPattern = new Regex(@"^\s+path (?<rpath>\S+)\s.+$");
+				Regex RPathPattern = new Regex(@"^\s+path (?<rpath>.+)\s\(offset");
 				string ToRemovePattern = Params.CreateAppBundle ? "/../../../" : "@loader_path/../UE4/";
 
 				string OutputLine;
@@ -300,7 +300,7 @@ public class MacPlatform : Platform
 							string RPath = RPathMatch.Groups["rpath"].Value;
 							if (RPath.Contains(ToRemovePattern))
 							{
-								Run("xcrun", "install_name_tool -delete_rpath " + RPath + " \"" + ExePath + "\"", null, ERunOptions.None);
+								Run("xcrun", "install_name_tool -delete_rpath \"" + RPath + "\" \"" + ExePath + "\"", null, ERunOptions.None);
 							}
 						}
 					}
