BuildGraph: Add a <ModifyConfig> task which allows modifying a config file value from BuildGraph scripts.

diff --git a/Engine/Source/Programs/AutomationTool/BuildGraph/BuildGraph.Automation.csproj b/Engine/Source/Programs/AutomationTool/BuildGraph/BuildGraph.Automation.csproj
--- a/Engine/Source/Programs/AutomationTool/BuildGraph/BuildGraph.Automation.csproj
+++ b/Engine/Source/Programs/AutomationTool/BuildGraph/BuildGraph.Automation.csproj
@@ -102,6 +102,7 @@
     <Compile Include="Tasks\UnzipTask.cs" />
     <Compile Include="Tasks\ZipTask.cs" />
     <Compile Include="TempStorage.cs" />
+    <Compile Include="Tasks\ModifyConfigTask.cs" />
   </ItemGroup>
   <ItemGroup>
     <ProjectReference Include="..\..\DotNETCommon\DotNETUtilities\DotNETUtilities.csproj">
diff --git a/Engine/Source/Programs/AutomationTool/BuildGraph/Tasks/ModifyConfigTask.cs b/Engine/Source/Programs/AutomationTool/BuildGraph/Tasks/ModifyConfigTask.cs
new file mode 100644
--- /dev/null
+++ b/Engine/Source/Programs/AutomationTool/BuildGraph/Tasks/ModifyConfigTask.cs
@@ -0,0 +1,134 @@
+// Copyright 1998-2019 Epic Games, Inc. All Rights Reserved.
+
+using AutomationTool;
+using System;
+using System.Collections.Generic;
+using System.Linq;
+using System.Text;
+using System.Threading.Tasks;
+using System.Xml;
+using Tools.DotNETCommon;
+using UnrealBuildTool;
+
+namespace BuildGraph.Tasks
+{
+	/// <summary>
+	/// Parameters for a ModifyConfig task
+	/// </summary>
+	public class ModifyConfigTaskParameters
+	{
+		/// <summary>
+		/// Path to the config file
+		/// </summary>
+		[TaskParameter(ValidationType = TaskParameterValidationType.FileSpec)]
+		public string File;
+
+		/// <summary>
+		/// The section name to modify
+		/// </summary>
+		[TaskParameter]
+		public string Section;
+
+		/// <summary>
+		/// The property name to set
+		/// </summary>
+		[TaskParameter]
+		public string Key;
+
+		/// <summary>
+		/// The property value to set
+		/// </summary>
+		[TaskParameter]
+		public string Value;
+
+		/// <summary>
+		/// Tag to be applied to the extracted files
+		/// </summary>
+		[TaskParameter(Optional = true, ValidationType = TaskParameterValidationType.TagList)]
+		public string Tag;
+	}
+
+	/// <summary>
+	/// Modifies a config file
+	/// </summary>
+	[TaskElement("ModifyConfig", typeof(ModifyConfigTaskParameters))]
+	public class ModifyConfigTask : CustomTask
+	{
+		/// <summary>
+		/// Parameters for this task
+		/// </summary>
+		ModifyConfigTaskParameters Parameters;
+
+		/// <summary>
+		/// Constructor
+		/// </summary>
+		/// <param name="InParameters">Parameters for this task</param>
+		public ModifyConfigTask(ModifyConfigTaskParameters InParameters)
+		{
+			Parameters = InParameters;
+		}
+
+		/// <summary>
+		/// Execute the task.
+		/// </summary>
+		/// <param name="Job">Information about the current job</param>
+		/// <param name="BuildProducts">Set of build products produced by this node.</param>
+		/// <param name="TagNameToFileSet">Mapping from tag names to the set of files they include</param>
+		public override void Execute(JobContext Job, HashSet<FileReference> BuildProducts, Dictionary<string, HashSet<FileReference>> TagNameToFileSet)
+		{
+			FileReference ConfigFileLocation = ResolveFile(Parameters.File);
+
+			ConfigFile ConfigFile;
+			if(FileReference.Exists(ConfigFileLocation))
+			{
+				ConfigFile = new ConfigFile(ConfigFileLocation);
+			}
+			else
+			{
+				ConfigFile = new ConfigFile();
+			}
+
+			ConfigFileSection Section = ConfigFile.FindOrAddSection(Parameters.Section);
+			Section.Lines.RemoveAll(x => String.Compare(x.Key, Parameters.Key, StringComparison.OrdinalIgnoreCase) == 0);
+			Section.Lines.Add(new ConfigLine(ConfigLineAction.Set, Parameters.Key, Parameters.Value));
+
+			FileReference.MakeWriteable(ConfigFileLocation);
+			ConfigFile.Write(ConfigFileLocation);
+
+			// Apply the optional tag to the produced archive
+			foreach (string TagName in FindTagNamesFromList(Parameters.Tag))
+			{
+				FindOrAddTagSet(TagNameToFileSet, TagName).Add(ConfigFileLocation);
+			}
+
+			// Add the archive to the set of build products
+			BuildProducts.Add(ConfigFileLocation);
+		}
+
+		/// <summary>
+		/// Output this task out to an XML writer.
+		/// </summary>
+		public override void Write(XmlWriter Writer)
+		{
+			Write(Writer, Parameters);
+		}
+
+		/// <summary>
+		/// Find all the tags which are used as inputs to this task
+		/// </summary>
+		/// <returns>The tag names which are read by this task</returns>
+		public override IEnumerable<string> FindConsumedTagNames()
+		{
+			return FindTagNamesFromFilespec(Parameters.File);
+		}
+
+		/// <summary>
+		/// Find all the tags which are modified by this task
+		/// </summary>
+		/// <returns>The tag names which are modified by this task</returns>
+		public override IEnumerable<string> FindProducedTagNames()
+		{
+			return FindTagNamesFromList(Parameters.Tag);
+		}
+	}
+}
diff --git a/Engine/Source/Programs/UnrealBuildTool/System/ConfigFile.cs b/Engine/Source/Programs/UnrealBuildTool/System/ConfigFile.cs
--- a/Engine/Source/Programs/UnrealBuildTool/System/ConfigFile.cs
+++ b/Engine/Source/Programs/UnrealBuildTool/System/ConfigFile.cs
@@ -1,4 +1,4 @@
-﻿// Copyright 1998-2019 Epic Games, Inc. All Rights Reserved.
+// Copyright 1998-2019 Epic Games, Inc. All Rights Reserved.
 
 using System;
 using System.Collections.Generic;
@@ -137,6 +137,14 @@ namespace UnrealBuildTool
 		/// </summary>
 		Dictionary<string, ConfigFileSection> Sections = new Dictionary<string, ConfigFileSection>(StringComparer.InvariantCultureIgnoreCase);
 
+		/// <summary>
+		/// Constructs a new, empty config file
+		/// </summary>
+		/// <param name="DefaultAction">The default action to take when encountering arrays without a '+' prefix</param>
+		public ConfigFile(ConfigLineAction DefaultAction = ConfigLineAction.Set)
+		{
+		}
+
 		/// <summary>
 		/// Reads config data from the given file.
 		/// </summary>
@@ -343,6 +351,22 @@ namespace UnrealBuildTool
 			get { return Sections.Keys; }
 		}
 
+		/// <summary>
+		/// Tries to get a config section by name, or creates one if it doesn't exist
+		/// </summary>
+		/// <param name="SectionName">Name of the section to look for</param>
+		/// <returns>The config section</returns>
+		public ConfigFileSection FindOrAddSection(string SectionName)
+		{
+			ConfigFileSection Section;
+			if(!Sections.TryGetValue(SectionName, out Section))
+			{
+				Section = new ConfigFileSection(SectionName);
+				Sections.Add(SectionName, Section);
+			}
+			return Section;
+		}
+
 		/// <summary>
 		/// Tries to get a config section by name
 		/// </summary>
