Tidy up ProjectUtils.RequiresTempTarget.

#rnx

diff --git a/Engine/Source/Programs/AutomationTool/AutomationUtils/ProjectUtils.cs b/Engine/Source/Programs/AutomationTool/AutomationUtils/ProjectUtils.cs
--- a/Engine/Source/Programs/AutomationTool/AutomationUtils/ProjectUtils.cs
+++ b/Engine/Source/Programs/AutomationTool/AutomationUtils/ProjectUtils.cs
@@ -124,10 +124,10 @@ namespace AutomationTool
 			return ProjectClientBinariesPath;
 		}
 
-		private static bool RequiresTempTarget(FileReference RawProjectPath, List<UnrealTargetPlatform> ClientTargetPlatforms, List<UnrealTargetConfiguration> ClientTargetConfigurations, bool AssetNativizationRequested)
+		private static bool RequiresTempTarget(FileReference RawProjectPath, List<UnrealTargetPlatform> Platforms, List<UnrealTargetConfiguration> Configurations, bool AssetNativizationRequested)
 		{
 			string Reason;
-			if(RequiresTempTarget(RawProjectPath, ClientTargetPlatforms, ClientTargetConfigurations, AssetNativizationRequested, out Reason))
+			if(RequiresTempTarget(RawProjectPath, Platforms, Configurations, AssetNativizationRequested, out Reason))
 			{
 				Log.TraceInformation("{0} requires a temporary target.cs to be generated ({1})", RawProjectPath.GetFileName(), Reason);
 				return true;
@@ -135,7 +135,7 @@ namespace AutomationTool
 			return false;
 		}
 
-		private static bool RequiresTempTarget(FileReference RawProjectPath, List<UnrealTargetPlatform> ClientTargetPlatforms, List<UnrealTargetConfiguration> ClientTargetConfigurations, bool AssetNativizationRequested, out string Reason)
+		private static bool RequiresTempTarget(FileReference RawProjectPath, List<UnrealTargetPlatform> Platforms, List<UnrealTargetConfiguration> Configurations, bool AssetNativizationRequested, out string Reason)
 		{
 			// check to see if we already have a Target.cs file
 			if (File.Exists (Path.Combine (Path.GetDirectoryName (RawProjectPath.FullName), "Source", RawProjectPath.GetFileNameWithoutExtension() + ".Target.cs")))
@@ -168,42 +168,14 @@ namespace AutomationTool
                 return true;
             }
 
-			if (ClientTargetPlatforms != null)
+			// Check if encryption or signing is enabled
+			foreach (UnrealTargetPlatform Platform in Platforms)
 			{
-				foreach (UnrealTargetPlatform ClientPlatform in ClientTargetPlatforms)
+				EncryptionAndSigning.CryptoSettings Settings = EncryptionAndSigning.ParseCryptoSettings(RawProjectPath.Directory, Platform);
+				if (Settings.IsAnyEncryptionEnabled() || Settings.IsPakSigningEnabled())
 				{
-					EncryptionAndSigning.CryptoSettings Settings = EncryptionAndSigning.ParseCryptoSettings(RawProjectPath.Directory, ClientPlatform);
-					if (Settings.IsAnyEncryptionEnabled() || Settings.IsPakSigningEnabled())
-					{
-						Reason = "encryption/signing is enabled";
-						return true;
-					}
-				}
-			}
-
-			// no Target file, now check to see if build settings have changed
-			List<UnrealTargetPlatform> TargetPlatforms = ClientTargetPlatforms;
-			if (ClientTargetPlatforms == null || ClientTargetPlatforms.Count < 1)
-			{
-				// No client target platforms, add all in
-				TargetPlatforms = new List<UnrealTargetPlatform>();
-				foreach (UnrealTargetPlatform TargetPlatformType in UnrealTargetPlatform.GetValidPlatforms())
-				{
-					TargetPlatforms.Add(TargetPlatformType);
-				}
-			}
-
-			List<UnrealTargetConfiguration> TargetConfigurations = ClientTargetConfigurations;
-			if (TargetConfigurations == null || TargetConfigurations.Count < 1)
-			{
-				// No client target configurations, add all in
-				TargetConfigurations = new List<UnrealTargetConfiguration>();
-				foreach (UnrealTargetConfiguration TargetConfigurationType in Enum.GetValues(typeof(UnrealTargetConfiguration)))
-				{
-					if (TargetConfigurationType != UnrealTargetConfiguration.Unknown)
-					{
-						TargetConfigurations.Add(TargetConfigurationType);
-					}
+					Reason = "encryption/signing is enabled";
+					return true;
 				}
 			}
 
@@ -212,14 +184,14 @@ namespace AutomationTool
 			List<PluginInfo> AvailablePlugins = Plugins.ReadAvailablePlugins(CommandUtils.EngineDirectory, RawProjectPath, Project.AdditionalPluginDirectories);
 
 			// check the target platforms for any differences in build settings or additional plugins
-			foreach (UnrealTargetPlatform TargetPlatformType in TargetPlatforms)
+			foreach (UnrealTargetPlatform Platform in Platforms)
 			{
-				if(!CommandUtils.IsEngineInstalled() && !PlatformExports.HasDefaultBuildConfig(RawProjectPath, TargetPlatformType))
+				if(!CommandUtils.IsEngineInstalled() && !PlatformExports.HasDefaultBuildConfig(RawProjectPath, Platform))
 				{
 					Reason = "project has non-default build configuration";
 					return true;
 				}
-				if(PlatformExports.RequiresBuild(RawProjectPath, TargetPlatformType))
+				if(PlatformExports.RequiresBuild(RawProjectPath, Platform))
 				{
 					Reason = "overriden by target platform";
 					return true;
@@ -229,17 +201,17 @@ namespace AutomationTool
 				foreach(PluginInfo Plugin in AvailablePlugins)
 				{
 					bool bPluginEnabledForTarget = false;
-					foreach (UnrealTargetConfiguration TargetConfigType in TargetConfigurations)
+					foreach (UnrealTargetConfiguration Configuration in Configurations)
 					{
-						bPluginEnabledForTarget |= Plugins.IsPluginEnabledForProject(Plugin, Project, TargetPlatformType, TargetConfigType, TargetRules.TargetType.Game);
+						bPluginEnabledForTarget |= Plugins.IsPluginEnabledForProject(Plugin, Project, Platform, Configuration, TargetRules.TargetType.Game);
 					}
 
 					bool bPluginEnabledForBaseTarget = false;
 					if(!Plugin.Descriptor.bInstalled)
 					{
-						foreach (UnrealTargetConfiguration TargetConfigType in TargetConfigurations)
+						foreach (UnrealTargetConfiguration Configuration in Configurations)
 						{
-							bPluginEnabledForBaseTarget |= Plugins.IsPluginEnabledForProject(Plugin, null, TargetPlatformType, TargetConfigType, TargetRules.TargetType.Game);
+							bPluginEnabledForBaseTarget |= Plugins.IsPluginEnabledForProject(Plugin, null, Platform, Configuration, TargetRules.TargetType.Game);
 						}
 					}
 
@@ -336,8 +308,34 @@ namespace AutomationTool
 			List<string> ExtraSearchPaths = null;
 			if (RawProjectPath != null)
 			{
-                string TempTargetDir = CommandUtils.CombinePaths(Path.GetDirectoryName(RawProjectPath.FullName), "Intermediate", "Source");
-                if (RequiresTempTarget(RawProjectPath, ClientTargetPlatforms, ClientTargetConfigurations, AssetNativizationRequested))
+				// no Target file, now check to see if build settings have changed
+				List<UnrealTargetPlatform> TargetPlatforms = ClientTargetPlatforms;
+				if (ClientTargetPlatforms == null || ClientTargetPlatforms.Count < 1)
+				{
+					// No client target platforms, add all in
+					TargetPlatforms = new List<UnrealTargetPlatform>();
+					foreach (UnrealTargetPlatform TargetPlatformType in UnrealTargetPlatform.GetValidPlatforms())
+					{
+						TargetPlatforms.Add(TargetPlatformType);
+					}
+				}
+
+				List<UnrealTargetConfiguration> TargetConfigurations = ClientTargetConfigurations;
+				if (TargetConfigurations == null || TargetConfigurations.Count < 1)
+				{
+					// No client target configurations, add all in
+					TargetConfigurations = new List<UnrealTargetConfiguration>();
+					foreach (UnrealTargetConfiguration TargetConfigurationType in Enum.GetValues(typeof(UnrealTargetConfiguration)))
+					{
+						if (TargetConfigurationType != UnrealTargetConfiguration.Unknown)
+						{
+							TargetConfigurations.Add(TargetConfigurationType);
+						}
+					}
+				}
+
+				string TempTargetDir = CommandUtils.CombinePaths(Path.GetDirectoryName(RawProjectPath.FullName), "Intermediate", "Source");
+                if (RequiresTempTarget(RawProjectPath, TargetPlatforms, TargetConfigurations, AssetNativizationRequested))
 				{
 					GenerateTempTarget(RawProjectPath);
 					Properties.bWasGenerated = true;
