Remove development asset registry from regioncn plugin builds.

diff --git a/Engine/Source/Programs/AutomationTool/Scripts/CopyBuildToStagingDirectory.Automation.cs b/Engine/Source/Programs/AutomationTool/Scripts/CopyBuildToStagingDirectory.Automation.cs
--- a/Engine/Source/Programs/AutomationTool/Scripts/CopyBuildToStagingDirectory.Automation.cs
+++ b/Engine/Source/Programs/AutomationTool/Scripts/CopyBuildToStagingDirectory.Automation.cs
@@ -442,7 +442,8 @@ public partial class Project : CommandUtils
 			// Put all of the cooked dir into the staged dir
 			SC.PlatformCookDir = String.IsNullOrEmpty(Params.CookOutputDir) ? DirectoryReference.Combine(DLCRoot, "Saved", "Cooked", SC.CookPlatform) : DirectoryReference.Combine(new DirectoryReference(Params.CookOutputDir), SC.CookPlatform);
 			DirectoryReference PlatformEngineDir = DirectoryReference.Combine(SC.PlatformCookDir, "Engine");
-			SC.MetadataDir = DirectoryReference.Combine(SC.PlatformCookDir, SC.ShortProjectName, "Metadata");
+			string RelativeDLCRootPath = DLCRoot.MakeRelativeTo(SC.LocalRoot);
+			SC.MetadataDir = DirectoryReference.Combine(SC.PlatformCookDir, RelativeDLCRootPath, "Metadata");
 
 			// Put the config files into the staged dir
 			DirectoryReference ConfigDir = DirectoryReference.Combine(DLCRoot, "Config");
