BuildGraph: Print out the number of files being added to a zip file from the <Zip> BuildGraph task.

diff --git a/Engine/Source/Programs/AutomationTool/BuildGraph/Tasks/ZipTask.cs b/Engine/Source/Programs/AutomationTool/BuildGraph/Tasks/ZipTask.cs
--- a/Engine/Source/Programs/AutomationTool/BuildGraph/Tasks/ZipTask.cs
+++ b/Engine/Source/Programs/AutomationTool/BuildGraph/Tasks/ZipTask.cs
@@ -71,17 +71,18 @@ namespace BuildGraph.Tasks
 		public override void Execute(JobContext Job, HashSet<FileReference> BuildProducts, Dictionary<string, HashSet<FileReference>> TagNameToFileSet)
 		{
 			// Find all the input files
-			IEnumerable<FileReference> Files;
+			List<FileReference> Files;
 			if(Parameters.Files == null)
 			{
-				Files = DirectoryReference.EnumerateFiles(Parameters.FromDir, "*", System.IO.SearchOption.AllDirectories);
+				Files = DirectoryReference.EnumerateFiles(Parameters.FromDir, "*", System.IO.SearchOption.AllDirectories).ToList();
 			}
 			else
 			{
-				Files = ResolveFilespec(Parameters.FromDir, Parameters.Files, TagNameToFileSet);
+				Files = ResolveFilespec(Parameters.FromDir, Parameters.Files, TagNameToFileSet).ToList();
 			}
 
 			// Create the zip file
+			Log.TraceInformation("Adding {0} files to {1}...", Files.Count, Parameters.ZipFile);
 			CommandUtils.ZipFiles(Parameters.ZipFile, Parameters.FromDir, Files);
 
 			// Apply the optional tag to the produced archive
