Fix editor-only plugins triggering recompilation of code projects.

Also rework logic for determining whether a project needs to be compiled, to more closely match equivalent code in UAT. ITargetPlatform::NeedsTempTarget() now allows querying whether a project needs a temporary target to be generated.

diff --git a/Engine/Plugins/Editor/CryptoKeys/Source/CryptoKeys/Private/CryptoKeys.cpp b/Engine/Plugins/Editor/CryptoKeys/Source/CryptoKeys/Private/CryptoKeys.cpp
--- a/Engine/Plugins/Editor/CryptoKeys/Source/CryptoKeys/Private/CryptoKeys.cpp
+++ b/Engine/Plugins/Editor/CryptoKeys/Source/CryptoKeys/Private/CryptoKeys.cpp
@@ -24,7 +24,7 @@ public:
 		FPropertyEditorModule& PropertyModule = FModuleManager::GetModuleChecked<FPropertyEditorModule>("PropertyEditor");
 		PropertyModule.RegisterCustomClassLayout(UCryptoKeysSettings::StaticClass()->GetFName(), FOnGetDetailCustomizationInstance::CreateStatic(&FCryptoKeysSettingsDetails::MakeInstance));
 
-		IModularFeatures::Get().RegisterModularFeature(FProjectBuildMutatorFeature::GetFeatureName(), &ProjectBuildMutator);
+		IModularFeatures::Get().RegisterModularFeature(PROJECT_BUILD_MUTATOR_FEATURE, &ProjectBuildMutator);
 	}
 
 	virtual void ShutdownModule() override
@@ -36,7 +36,7 @@ public:
 			FPropertyEditorModule& PropertyModule = FModuleManager::GetModuleChecked<FPropertyEditorModule>("PropertyEditor");
 			PropertyModule.UnregisterCustomClassLayout(UCryptoKeysSettings::StaticClass()->GetFName());
 
-			IModularFeatures::Get().UnregisterModularFeature(FProjectBuildMutatorFeature::GetFeatureName(), &ProjectBuildMutator);
+			IModularFeatures::Get().UnregisterModularFeature(PROJECT_BUILD_MUTATOR_FEATURE, &ProjectBuildMutator);
 		}
 	}
 
diff --git a/Engine/Plugins/Editor/CryptoKeys/Source/CryptoKeys/Private/CryptoKeysProjectBuildMutatorFeature.cpp b/Engine/Plugins/Editor/CryptoKeys/Source/CryptoKeys/Private/CryptoKeysProjectBuildMutatorFeature.cpp
--- a/Engine/Plugins/Editor/CryptoKeys/Source/CryptoKeys/Private/CryptoKeysProjectBuildMutatorFeature.cpp
+++ b/Engine/Plugins/Editor/CryptoKeys/Source/CryptoKeys/Private/CryptoKeysProjectBuildMutatorFeature.cpp
@@ -2,8 +2,17 @@
 #include "CryptoKeysProjectBuildMutatorFeature.h"
 #include "CryptoKeysSettings.h"
 
-bool FCryptoKeysProjectBuildMutatorFeature ::RequiresProjectBuild(FName InPlatformInfoName) const
+#define LOCTEXT_NAMESPACE "CryptoKeys"
+
+bool FCryptoKeysProjectBuildMutatorFeature ::RequiresProjectBuild(const FName& InPlatformInfoName, FText& OutReason) const
 {
 	UCryptoKeysSettings* Settings = GetMutableDefault<UCryptoKeysSettings>();
-	return Settings->IsEncryptionEnabled() || Settings->IsSigningEnabled();
-}
\ No newline at end of file
+	if (Settings->IsEncryptionEnabled() || Settings->IsSigningEnabled())
+	{
+		OutReason = LOCTEXT("BuildMutator_BuildRequired", "encryption/signing enabled");
+		return true;
+	}
+	return false;
+}
+
+#undef LOCTEXT_NAMESPACE
diff --git a/Engine/Plugins/Editor/CryptoKeys/Source/CryptoKeys/Private/CryptoKeysProjectBuildMutatorFeature.h b/Engine/Plugins/Editor/CryptoKeys/Source/CryptoKeys/Private/CryptoKeysProjectBuildMutatorFeature.h
--- a/Engine/Plugins/Editor/CryptoKeys/Source/CryptoKeys/Private/CryptoKeysProjectBuildMutatorFeature.h
+++ b/Engine/Plugins/Editor/CryptoKeys/Source/CryptoKeys/Private/CryptoKeysProjectBuildMutatorFeature.h
@@ -2,13 +2,10 @@
 
 #pragma once
 
-#include "ProjectBuildMutatorFeature.h"
+#include "Interfaces/IProjectBuildMutatorFeature.h"
 
-class FCryptoKeysProjectBuildMutatorFeature : public FProjectBuildMutatorFeature
+class FCryptoKeysProjectBuildMutatorFeature : public IProjectBuildMutatorFeature
 {
 public:
-
-	virtual bool RequiresProjectBuild(FName InPlatformInfoName) const override;
-
-private:
+	virtual bool RequiresProjectBuild(const FName& InPlatformInfoName, FText& OutReason) const override;
 };
\ No newline at end of file
diff --git a/Engine/Plugins/Runtime/AR/Microsoft/HoloLensAR/Source/HoloLensTargetPlatform/Private/HoloLensTargetPlatform.cpp b/Engine/Plugins/Runtime/AR/Microsoft/HoloLensAR/Source/HoloLensTargetPlatform/Private/HoloLensTargetPlatform.cpp
--- a/Engine/Plugins/Runtime/AR/Microsoft/HoloLensAR/Source/HoloLensTargetPlatform/Private/HoloLensTargetPlatform.cpp
+++ b/Engine/Plugins/Runtime/AR/Microsoft/HoloLensAR/Source/HoloLensTargetPlatform/Private/HoloLensTargetPlatform.cpp
@@ -171,7 +171,7 @@ bool FHoloLensTargetPlatform::IsSdkInstalled(bool bProjectHasCode, FString& OutD
 	return SDKVersions.Num() > 0;
 }
 
-int32 FHoloLensTargetPlatform::CheckRequirements(const FString& ProjectPath, bool bProjectHasCode, FString& OutTutorialPath, FString& OutDocumentationPath, FText& CustomizedLogMessage) const
+int32 FHoloLensTargetPlatform::CheckRequirements(bool bProjectHasCode, EBuildConfiguration Configuration, bool bRequiresAssetNativization, FString& OutTutorialPath, FString& OutDocumentationPath, FText& CustomizedLogMessage) const
 {
 	OutDocumentationPath = TEXT("Platforms/HoloLens/GettingStarted");
 	FString LocalErrors;
diff --git a/Engine/Plugins/Runtime/AR/Microsoft/HoloLensAR/Source/HoloLensTargetPlatform/Private/HoloLensTargetPlatform.h b/Engine/Plugins/Runtime/AR/Microsoft/HoloLensAR/Source/HoloLensTargetPlatform/Private/HoloLensTargetPlatform.h
--- a/Engine/Plugins/Runtime/AR/Microsoft/HoloLensAR/Source/HoloLensTargetPlatform/Private/HoloLensTargetPlatform.h
+++ b/Engine/Plugins/Runtime/AR/Microsoft/HoloLensAR/Source/HoloLensTargetPlatform/Private/HoloLensTargetPlatform.h
@@ -134,7 +134,7 @@ public:
 	}
 
 	virtual bool IsSdkInstalled(bool bProjectHasCode, FString& OutDocumentationPath) const override;
-	virtual int32 CheckRequirements(const FString& ProjectPath, bool bProjectHasCode, FString& OutTutorialPath, FString& OutDocumentationPath, FText& CustomizedLogMessage) const override;
+	virtual int32 CheckRequirements(bool bProjectHasCode, EBuildConfiguration Configuration, bool bRequiresAssetNativization, FString& OutTutorialPath, FString& OutDocumentationPath, FText& CustomizedLogMessage) const override;
 	//~ End ITargetPlatform Interface
 
 protected:
diff --git a/Engine/Source/Developer/AllDesktopTargetPlatform/Private/AllDesktopTargetPlatform.h b/Engine/Source/Developer/AllDesktopTargetPlatform/Private/AllDesktopTargetPlatform.h
--- a/Engine/Source/Developer/AllDesktopTargetPlatform/Private/AllDesktopTargetPlatform.h
+++ b/Engine/Source/Developer/AllDesktopTargetPlatform/Private/AllDesktopTargetPlatform.h
@@ -112,7 +112,7 @@ public:
 		return true;
 	}
 
-	virtual int32 CheckRequirements(const FString& ProjectPath, bool bProjectHasCode, FString& OutTutorialPath, FString& OutDocumenationPath, FText& CustomizedLogMessage) const override
+	virtual int32 CheckRequirements(bool bProjectHasCode, EBuildConfiguration Configuration, bool bRequiresAssetNativization, FString& OutTutorialPath, FString& OutDocumentationPath, FText& CustomizedLogMessage) const override
 	{
 		return ETargetPlatformReadyStatus::Ready;
 	}
diff --git a/Engine/Source/Developer/Android/AndroidTargetPlatform/Private/AndroidTargetPlatform.cpp b/Engine/Source/Developer/Android/AndroidTargetPlatform/Private/AndroidTargetPlatform.cpp
--- a/Engine/Source/Developer/Android/AndroidTargetPlatform/Private/AndroidTargetPlatform.cpp
+++ b/Engine/Source/Developer/Android/AndroidTargetPlatform/Private/AndroidTargetPlatform.cpp
@@ -307,7 +307,7 @@ bool FAndroidTargetPlatform::IsSdkInstalled(bool bProjectHasCode, FString& OutDo
 	return true;
 }
 
-int32 FAndroidTargetPlatform::CheckRequirements(const FString& ProjectPath, bool bProjectHasCode, FString& OutTutorialPath, FString& OutDocumentationPath, FText& CustomizedLogMessage) const
+int32 FAndroidTargetPlatform::CheckRequirements(bool bProjectHasCode, EBuildConfiguration Configuration, bool bRequiresAssetNativization, FString& OutTutorialPath, FString& OutDocumentationPath, FText& CustomizedLogMessage) const
 {
 	OutDocumentationPath = TEXT("Platforms/Android/GettingStarted");
 
diff --git a/Engine/Source/Developer/Android/AndroidTargetPlatform/Private/AndroidTargetPlatform.h b/Engine/Source/Developer/Android/AndroidTargetPlatform/Private/AndroidTargetPlatform.h
--- a/Engine/Source/Developer/Android/AndroidTargetPlatform/Private/AndroidTargetPlatform.h
+++ b/Engine/Source/Developer/Android/AndroidTargetPlatform/Private/AndroidTargetPlatform.h
@@ -180,7 +180,7 @@ public:
 
 	virtual bool IsSdkInstalled(bool bProjectHasCode, FString& OutDocumentationPath) const override;
 
-	virtual int32 CheckRequirements(const FString& ProjectPath, bool bProjectHasCode, FString& OutTutorialPath, FString& OutDocumentationPath, FText& CustomizedLogMessage) const override;
+	virtual int32 CheckRequirements(bool bProjectHasCode, EBuildConfiguration Configuration, bool bRequiresAssetNativization, FString& OutTutorialPath, FString& OutDocumentationPath, FText& CustomizedLogMessage) const override;
 
 	virtual bool SupportsFeature( ETargetPlatformFeatures Feature ) const override;
 
diff --git a/Engine/Source/Developer/IOS/IOSPlatformEditor/Private/IOSCustomIconProjectBuildMutatorFeature.cpp b/Engine/Source/Developer/IOS/IOSPlatformEditor/Private/IOSCustomIconProjectBuildMutatorFeature.cpp
--- a/Engine/Source/Developer/IOS/IOSPlatformEditor/Private/IOSCustomIconProjectBuildMutatorFeature.cpp
+++ b/Engine/Source/Developer/IOS/IOSPlatformEditor/Private/IOSCustomIconProjectBuildMutatorFeature.cpp
@@ -45,7 +45,7 @@ static bool RequiresBuild()
 	return false;
 }
 
-bool FIOSCustomIconProjectBuildMutatorFeature ::RequiresProjectBuild(FName InPlatformInfoName) const
+bool FIOSCustomIconProjectBuildMutatorFeature ::RequiresProjectBuild(const FName& InPlatformInfoName, FText& OutReason) const
 {
 	const PlatformInfo::FPlatformInfo* const PlatInfo = PlatformInfo::FindPlatformInfo(InPlatformInfoName);
 	check(PlatInfo);
@@ -57,7 +57,12 @@ bool FIOSCustomIconProjectBuildMutatorFeature ::RequiresProjectBuild(FName InPla
 		{
 			if (InPlatformInfoName.ToString() == TEXT("IOS"))
 			{
-				return RequiresBuild();
+				bool bResult = RequiresBuild();
+				if (bResult)
+				{
+					OutReason = NSLOCTEXT("IOSPlatformEditor", "RequiresBuildDueToCustomIcon", "custom icon for IOS");
+				}
+				return bResult;
 			}
 		}
 	}
diff --git a/Engine/Source/Developer/IOS/IOSPlatformEditor/Private/IOSCustomIconProjectBuildMutatorFeature.h b/Engine/Source/Developer/IOS/IOSPlatformEditor/Private/IOSCustomIconProjectBuildMutatorFeature.h
--- a/Engine/Source/Developer/IOS/IOSPlatformEditor/Private/IOSCustomIconProjectBuildMutatorFeature.h
+++ b/Engine/Source/Developer/IOS/IOSPlatformEditor/Private/IOSCustomIconProjectBuildMutatorFeature.h
@@ -2,13 +2,10 @@
 
 #pragma once
 
-#include "ProjectBuildMutatorFeature.h"
+#include "Interfaces/IProjectBuildMutatorFeature.h"
 
-class FIOSCustomIconProjectBuildMutatorFeature : public FProjectBuildMutatorFeature
+class FIOSCustomIconProjectBuildMutatorFeature : public IProjectBuildMutatorFeature
 {
 public:
-
-	virtual bool RequiresProjectBuild(FName InPlatformInfoName) const override;
-
-private:
+	virtual bool RequiresProjectBuild(const FName& InPlatformInfoName, FText& OutReason) const override;
 };
\ No newline at end of file
diff --git a/Engine/Source/Developer/IOS/IOSPlatformEditor/Private/IOSPlatformEditorModule.cpp b/Engine/Source/Developer/IOS/IOSPlatformEditor/Private/IOSPlatformEditorModule.cpp
--- a/Engine/Source/Developer/IOS/IOSPlatformEditor/Private/IOSPlatformEditorModule.cpp
+++ b/Engine/Source/Developer/IOS/IOSPlatformEditor/Private/IOSPlatformEditorModule.cpp
@@ -66,7 +66,7 @@ void FIOSPlatformEditorModule::StartupModule()
 		}
 	}
 
-	IModularFeatures::Get().RegisterModularFeature(FProjectBuildMutatorFeature::GetFeatureName(), &ProjectBuildMutator);
+	IModularFeatures::Get().RegisterModularFeature(PROJECT_BUILD_MUTATOR_FEATURE, &ProjectBuildMutator);
 }
 
 void FIOSPlatformEditorModule::ShutdownModule()
@@ -79,7 +79,7 @@ void FIOSPlatformEditorModule::ShutdownModule()
 		SettingsModule->UnregisterSettings("Project", "Platforms", "iOSMetalQuality");
 	}
 
-	IModularFeatures::Get().UnregisterModularFeature(FProjectBuildMutatorFeature::GetFeatureName(), &ProjectBuildMutator);
+	IModularFeatures::Get().UnregisterModularFeature(PROJECT_BUILD_MUTATOR_FEATURE, &ProjectBuildMutator);
 }
 
 void FIOSPlatformEditorModule::HandleSelectIOSSection()
diff --git a/Engine/Source/Developer/IOS/IOSTargetPlatform/Private/IOSTargetPlatform.cpp b/Engine/Source/Developer/IOS/IOSTargetPlatform/Private/IOSTargetPlatform.cpp
--- a/Engine/Source/Developer/IOS/IOSTargetPlatform/Private/IOSTargetPlatform.cpp
+++ b/Engine/Source/Developer/IOS/IOSTargetPlatform/Private/IOSTargetPlatform.cpp
@@ -209,7 +209,7 @@ bool FIOSTargetPlatform::IsSdkInstalled(bool bProjectHasCode, FString& OutTutori
 	return biOSSDKInstalled;
 }
 
-int32 FIOSTargetPlatform::CheckRequirements(const FString& ProjectPath, bool bProjectHasCode, FString& OutTutorialPath, FString& OutDocumentationPath, FText& CustomizedLogMessage) const
+int32 FIOSTargetPlatform::CheckRequirements(bool bProjectHasCode, EBuildConfiguration Configuration, bool bRequiresAssetNativization, FString& OutTutorialPath, FString& OutDocumentationPath, FText& CustomizedLogMessage) const
 {
 	OutDocumentationPath = TEXT("Platforms/iOS/QuickStart/6");
 
@@ -229,7 +229,9 @@ int32 FIOSTargetPlatform::CheckRequirements(const FString& ProjectPath, bool bPr
 			OutTutorialPath = FString("/Engine/Tutorial/Mobile/iOSonPCRestrictions.iOSonPCRestrictions");
 			bReadyToBuild |= ETargetPlatformReadyStatus::CodeUnsupported;
 		}
-		if (!IProjectManager::Get().HasDefaultPluginSettings())
+
+		FText Reason;
+		if (RequiresTempTarget(bProjectHasCode, Configuration, bRequiresAssetNativization, Reason))
 		{
 			OutTutorialPath = FString("/Engine/Tutorial/Mobile/iOSonPCValidPlugins.iOSonPCValidPlugins");
 			bReadyToBuild |= ETargetPlatformReadyStatus::PluginsUnsupported;
@@ -253,6 +255,7 @@ int32 FIOSTargetPlatform::CheckRequirements(const FString& ProjectPath, bool bPr
 	FString TeamID;
 	GConfig->GetString(TEXT("/Script/IOSRuntimeSettings.IOSRuntimeSettings"), TEXT("IOSTeamID"), TeamID, GEngineIni);
 
+	FString ProjectPath = FPaths::ConvertRelativePathToFull(FPaths::GetProjectFilePath());
 #if PLATFORM_MAC
     FString CmdExe = TEXT("/bin/sh");
     FString ScriptPath = FPaths::ConvertRelativePathToFull(FPaths::EngineDir() / TEXT("Build/BatchFiles/Mac/RunMono.sh"));
diff --git a/Engine/Source/Developer/IOS/IOSTargetPlatform/Private/IOSTargetPlatform.h b/Engine/Source/Developer/IOS/IOSTargetPlatform/Private/IOSTargetPlatform.h
--- a/Engine/Source/Developer/IOS/IOSTargetPlatform/Private/IOSTargetPlatform.h
+++ b/Engine/Source/Developer/IOS/IOSTargetPlatform/Private/IOSTargetPlatform.h
@@ -101,7 +101,7 @@ public:
 	virtual bool CanSupportXGEShaderCompile() const override;
 
 	virtual bool IsSdkInstalled(bool bProjectHasCode, FString& OutTutorialPath) const override;
-	virtual int32 CheckRequirements(const FString& ProjectPath, bool bProjectHasCode, FString& OutTutorialPath, FString& OutDocumenationPath, FText& CustomizedLogMessage) const override;
+	virtual int32 CheckRequirements(bool bProjectHasCode, EBuildConfiguration Configuration, bool bRequiresAssetNativization, FString& OutTutorialPath, FString& OutDocumentationPath, FText& CustomizedLogMessage) const override;
 
 
 #if WITH_ENGINE
diff --git a/Engine/Source/Developer/Linux/LinuxTargetPlatform/Private/LinuxTargetPlatform.h b/Engine/Source/Developer/Linux/LinuxTargetPlatform/Private/LinuxTargetPlatform.h
--- a/Engine/Source/Developer/Linux/LinuxTargetPlatform/Private/LinuxTargetPlatform.h
+++ b/Engine/Source/Developer/Linux/LinuxTargetPlatform/Private/LinuxTargetPlatform.h
@@ -202,9 +202,9 @@ public:
 		return true;
 	}
 
-	virtual int32 CheckRequirements(const FString& ProjectPath, bool bProjectHasCode, FString& OutTutorialPath, FString& OutDocumentationPath, FText& CustomizedLogMessage) const override
+	virtual int32 CheckRequirements(bool bProjectHasCode, EBuildConfiguration Configuration, bool bRequiresAssetNativization, FString& OutTutorialPath, FString& OutDocumentationPath, FText& CustomizedLogMessage) const override
 	{
-		int32 ReadyToBuild = TSuper::CheckRequirements(ProjectPath, bProjectHasCode, OutTutorialPath, OutDocumentationPath, CustomizedLogMessage);
+		int32 ReadyToBuild = TSuper::CheckRequirements(bProjectHasCode, Configuration, bRequiresAssetNativization, OutTutorialPath, OutDocumentationPath, CustomizedLogMessage);
 
 		// do not support code/plugins in Installed builds if the required libs aren't bundled (on Windows/Mac)
 		if (!PLATFORM_LINUX && !FInstalledPlatformInfo::Get().IsValidPlatform(TSuper::GetPlatformInfo().BinaryFolderName, EProjectType::Code))
@@ -214,7 +214,8 @@ public:
 				ReadyToBuild |= ETargetPlatformReadyStatus::CodeUnsupported;
 			}
 
-			if (!IProjectManager::Get().HasDefaultPluginSettings())
+			FText Reason;
+			if (RequiresTempTarget(bProjectHasCode, Configuration, bRequiresAssetNativization, Reason))
 			{
 				ReadyToBuild |= ETargetPlatformReadyStatus::PluginsUnsupported;
 			}
diff --git a/Engine/Source/Developer/Lumin/LuminTargetPlatform/Private/LuminTargetPlatform.cpp b/Engine/Source/Developer/Lumin/LuminTargetPlatform/Private/LuminTargetPlatform.cpp
--- a/Engine/Source/Developer/Lumin/LuminTargetPlatform/Private/LuminTargetPlatform.cpp
+++ b/Engine/Source/Developer/Lumin/LuminTargetPlatform/Private/LuminTargetPlatform.cpp
@@ -47,7 +47,7 @@ bool FLuminTargetPlatform::IsSdkInstalled(bool bProjectHasCode, FString& OutDocu
 	return true;
 }
 
-int32 FLuminTargetPlatform::CheckRequirements(const FString& ProjectPath, bool bProjectHasCode, FString& OutTutorialPath, FString& OutDocumentationPath, FText& CustomizedLogMessage) const
+int32 FLuminTargetPlatform::CheckRequirements(bool bProjectHasCode, EBuildConfiguration Configuration, bool bRequiresAssetNativization, FString& OutTutorialPath, FString& OutDocumentationPath, FText& CustomizedLogMessage) const
 {
 	OutDocumentationPath = TEXT("Platforms/Android/GettingStarted");
 
diff --git a/Engine/Source/Developer/Lumin/LuminTargetPlatform/Private/LuminTargetPlatform.h b/Engine/Source/Developer/Lumin/LuminTargetPlatform/Private/LuminTargetPlatform.h
--- a/Engine/Source/Developer/Lumin/LuminTargetPlatform/Private/LuminTargetPlatform.h
+++ b/Engine/Source/Developer/Lumin/LuminTargetPlatform/Private/LuminTargetPlatform.h
@@ -61,7 +61,7 @@ public:
 
 	virtual bool IsSdkInstalled(bool bProjectHasCode, FString& OutDocumentationPath) const override;
 
-	virtual int32 CheckRequirements(const FString& ProjectPath, bool bProjectHasCode, FString& OutTutorialPath, FString& OutDocumentationPath, FText& CustomizedLogMessage) const override;
+	virtual int32 CheckRequirements(bool bProjectHasCode, EBuildConfiguration Configuration, bool bRequiresAssetNativization, FString& OutTutorialPath, FString& OutDocumentationPath, FText& CustomizedLogMessage) const override;
 
 #if WITH_ENGINE
 
diff --git a/Engine/Source/Developer/TargetPlatform/Private/TargetPlatformBase.cpp b/Engine/Source/Developer/TargetPlatform/Private/TargetPlatformBase.cpp
--- a/Engine/Source/Developer/TargetPlatform/Private/TargetPlatformBase.cpp
+++ b/Engine/Source/Developer/TargetPlatform/Private/TargetPlatformBase.cpp
@@ -3,6 +3,15 @@
 #include "Common/TargetPlatformBase.h"
 #include "HAL/IConsoleManager.h"
 #include "DeviceBrowserDefaultPlatformWidgetCreator.h"
+#include "Interfaces/IProjectBuildMutatorFeature.h"
+#include "Features/IModularFeatures.h"
+#include "Interfaces/IProjectManager.h"
+#include "Interfaces/IPluginManager.h"
+#include "ProjectDescriptor.h"
+#include "Misc/App.h"
+#include "Misc/ConfigCacheIni.h"
+
+#define LOCTEXT_NAMESPACE "TargetPlatform"
 
 bool FTargetPlatformBase::UsesForwardShading() const
 {
@@ -38,8 +47,197 @@ float FTargetPlatformBase::GetDownSampleMeshDistanceFieldDivider() const
 	return 1.0f;
 }
 
+static bool IsPluginEnabledForTarget(const IPlugin& Plugin, const FProjectDescriptor* Project, const FString& Platform, EBuildConfiguration Configuration, EBuildTargetType TargetType)
+{
+	if (!Plugin.GetDescriptor().SupportsTargetPlatform(Platform))
+	{
+		return false;
+	}
+
+	bool bEnabledForProject = Plugin.IsEnabledByDefault();
+	if (Project != nullptr)
+	{
+		for(const FPluginReferenceDescriptor& PluginReference : Project->Plugins)
+		{
+			if (PluginReference.Name == Plugin.GetName() && !PluginReference.bOptional)
+			{
+				bEnabledForProject = PluginReference.IsEnabledForPlatform(Platform) && PluginReference.IsEnabledForTargetConfiguration(Configuration) && PluginReference.IsEnabledForTarget(TargetType);
+			}
+		}
+	}
+	return bEnabledForProject;
+}
+
+static bool IsPluginCompiledForTarget(const IPlugin& Plugin, const FProjectDescriptor* Project, const FString& Platform, EBuildConfiguration Configuration, EBuildTargetType TargetType, bool bRequiresCookedData)
+{
+	bool bCompiledForTarget = false;
+	if (IsPluginEnabledForTarget(Plugin, Project, Platform, Configuration, TargetType))
+	{
+		bool bBuildDeveloperTools = (TargetType == EBuildTargetType::Editor || TargetType == EBuildTargetType::Program || (Configuration != EBuildConfiguration::Test && Configuration != EBuildConfiguration::Shipping));
+		for (const FModuleDescriptor& Module : Plugin.GetDescriptor().Modules)
+		{
+			if (Module.IsCompiledInConfiguration(Platform, Configuration, TEXT(""), TargetType, bBuildDeveloperTools, bRequiresCookedData))
+			{
+				bCompiledForTarget = true;
+				break;
+			}
+		}
+	}
+	return bCompiledForTarget;
+}
+
+bool FTargetPlatformBase::RequiresTempTarget(bool bProjectHasCode, EBuildConfiguration Configuration, bool bRequiresAssetNativization, FText& OutReason) const
+{
+	// check to see if we already have a Target.cs file
+	if (bProjectHasCode)
+	{
+		return false;
+	}
+
+	// check if asset nativization is enabled
+	if (bRequiresAssetNativization)
+    {
+		OutReason = LOCTEXT("TempTarget_Nativization", "asset nativization is enabled");
+        return true;
+    }
+
+	// check to see if any projectmutator modular features are available
+	for (IProjectBuildMutatorFeature* Feature : IModularFeatures::Get().GetModularFeatureImplementations<IProjectBuildMutatorFeature>(PROJECT_BUILD_MUTATOR_FEATURE))
+	{
+		if (Feature->RequiresProjectBuild(PlatformInfo->PlatformInfoName, OutReason))
+		{
+			return true;
+		}
+	}
+
+	// check the target platforms for any differences in build settings or additional plugins
+	const FProjectDescriptor* Project = IProjectManager::Get().GetCurrentProject();
+	if (!FApp::IsEngineInstalled() && !HasDefaultBuildSettings())
+	{
+		OutReason = LOCTEXT("TempTarget_NonDefaultBuildConfig", "project has non-default build configuration");
+		return true;
+	}
+
+	// find if there are any plugins enabled or disabled which differ from the default
+	for(const TSharedRef<IPlugin>& Plugin: IPluginManager::Get().GetDiscoveredPlugins())
+	{
+		const FString Platform = PlatformInfo->UBTTargetId.ToString();
+
+		bool bCompiledForTarget = IsPluginCompiledForTarget(*Plugin, Project, Platform, Configuration, PlatformInfo->PlatformType, RequiresCookedData());
+
+		bool bCompiledForBaseTarget = false;
+		if (!Plugin->GetDescriptor().bInstalled)
+		{
+			bCompiledForBaseTarget |= IsPluginCompiledForTarget(*Plugin, nullptr, Platform, Configuration, PlatformInfo->PlatformType, RequiresCookedData());
+		}
+
+		if (bCompiledForTarget != bCompiledForBaseTarget)
+		{
+			if (bCompiledForTarget)
+			{
+				OutReason = FText::Format(LOCTEXT("TempTarget_PluginEnabled", "{0} plugin is enabled"), FText::FromString(Plugin->GetName()));
+				return true;
+			}
+			else
+			{
+				OutReason = FText::Format(LOCTEXT("TempTarget_PluginDisabled", "{0} plugin is disabled"), FText::FromString(Plugin->GetName()));
+				return true;
+			}
+		}
+	}
+	return false;
+}
+
 TSharedPtr<IDeviceManagerCustomPlatformWidgetCreator> FTargetPlatformBase::GetCustomWidgetCreator() const
 {
 	static TSharedPtr<FDeviceBrowserDefaultPlatformWidgetCreator> DefaultWidgetCreator = MakeShared<FDeviceBrowserDefaultPlatformWidgetCreator>();
 	return DefaultWidgetCreator;
 }
+
+bool FTargetPlatformBase::HasDefaultBuildSettings() const
+{
+	// first check default build settings for all platforms
+	TArray<FString> BoolKeys, IntKeys, StringKeys, BuildKeys;
+	BuildKeys.Add(TEXT("bCompileApex")); 
+	BuildKeys.Add(TEXT("bCompileICU"));
+	BuildKeys.Add(TEXT("bCompileSimplygon")); 
+	BuildKeys.Add(TEXT("bCompileSimplygonSSF"));
+	BuildKeys.Add(TEXT("bCompileRecast")); 
+	BuildKeys.Add(TEXT("bCompileSpeedTree"));
+	BuildKeys.Add(TEXT("bCompileWithPluginSupport")); 
+	BuildKeys.Add(TEXT("bCompilePhysXVehicle")); 
+	BuildKeys.Add(TEXT("bCompileFreeType"));
+	BuildKeys.Add(TEXT("bCompileForSize"));	
+	BuildKeys.Add(TEXT("bCompileCEF3")); 
+	BuildKeys.Add(TEXT("bCompileCustomSQLitePlatform"));
+
+	const PlatformInfo::FPlatformInfo& PlatInfo = GetPlatformInfo();
+	if (!DoProjectSettingsMatchDefault(PlatInfo.TargetPlatformName.ToString(), TEXT("/Script/BuildSettings.BuildSettings"), &BuildKeys, nullptr, nullptr))
+	{
+		return false;
+	}
+
+	FString PlatformSection;
+	GetBuildProjectSettingKeys(PlatformSection, BoolKeys, IntKeys, StringKeys);
+
+	if(!DoProjectSettingsMatchDefault(PlatInfo.TargetPlatformName.ToString(), PlatformSection, &BoolKeys, &IntKeys, &StringKeys))
+	{
+		return false;
+	}
+
+	return true;
+}
+
+bool FTargetPlatformBase::DoProjectSettingsMatchDefault(const FString& InPlatformName, const FString& InSection, const TArray<FString>* InBoolKeys, const TArray<FString>* InIntKeys, const TArray<FString>* InStringKeys)
+{
+	FConfigFile ProjIni;
+	FConfigFile DefaultIni;
+	FConfigCacheIni::LoadLocalIniFile(ProjIni, TEXT("Engine"), true, *InPlatformName, true);
+	FConfigCacheIni::LoadExternalIniFile(DefaultIni, TEXT("Engine"), *FPaths::EngineConfigDir(), *FPaths::EngineConfigDir(), true, NULL, true);
+
+	if (InBoolKeys != NULL)
+	{
+		for (int Index = 0; Index < InBoolKeys->Num(); ++Index)
+		{
+			FString Default(TEXT("False")), Project(TEXT("False"));
+			DefaultIni.GetString(*InSection, *((*InBoolKeys)[Index]), Default);
+			ProjIni.GetString(*InSection, *((*InBoolKeys)[Index]), Project);
+			if (Default.Compare(Project, ESearchCase::IgnoreCase))
+			{
+				return false;
+			}
+		}
+	}
+
+	if (InIntKeys != NULL)
+	{
+		for (int Index = 0; Index < InIntKeys->Num(); ++Index)
+		{
+			int64 Default(0), Project(0);
+			DefaultIni.GetInt64(*InSection, *((*InIntKeys)[Index]), Default);
+			ProjIni.GetInt64(*InSection, *((*InIntKeys)[Index]), Project);
+			if (Default != Project)
+			{
+				return false;
+			}
+		}
+	}
+
+	if (InStringKeys != NULL)
+	{
+		for (int Index = 0; Index < InStringKeys->Num(); ++Index)
+		{
+			FString Default(TEXT("False")), Project(TEXT("False"));
+			DefaultIni.GetString(*InSection, *((*InStringKeys)[Index]), Default);
+			ProjIni.GetString(*InSection, *((*InStringKeys)[Index]), Project);
+			if (Default.Compare(Project, ESearchCase::IgnoreCase))
+			{
+				return false;
+			}
+		}
+	}
+
+	return true;
+}
+
+#undef LOCTEXT_NAMESPACE
diff --git a/Engine/Source/Developer/TargetPlatform/Public/Common/TargetPlatformBase.h b/Engine/Source/Developer/TargetPlatform/Public/Common/TargetPlatformBase.h
--- a/Engine/Source/Developer/TargetPlatform/Public/Common/TargetPlatformBase.h
+++ b/Engine/Source/Developer/TargetPlatform/Public/Common/TargetPlatformBase.h
@@ -73,7 +73,7 @@ public:
 		return true;
 	}
 
-	virtual int32 CheckRequirements(const FString& ProjectPath, bool bProjectHasCode, FString& OutTutorialPath, FString& OutDocumentationPath, FText& CustomizedLogMessage) const override
+	virtual int32 CheckRequirements(bool bProjectHasCode, EBuildConfiguration Configuration, bool bRequiresAssetNativization, FString& OutTutorialPath, FString& OutDocumentationPath, FText& CustomizedLogMessage) const override
 	{
 		int32 bReadyToBuild = ETargetPlatformReadyStatus::Ready; // @todo How do we check that the iOS SDK is installed when building from Windows? Is that even possible?
 		if (!IsSdkInstalled(bProjectHasCode, OutTutorialPath))
@@ -83,6 +83,8 @@ public:
 		return bReadyToBuild;
 	}
 
+	TARGETPLATFORM_API virtual bool RequiresTempTarget(bool bProjectHasCode, EBuildConfiguration Configuration, bool bRequiresAssetNativization, FText& OutReason) const override;
+
 	virtual bool SupportsVariants() const override
 	{
 		return false;
@@ -137,6 +139,10 @@ protected:
 	/** Information about this platform */
 	const PlatformInfo::FPlatformInfo *PlatformInfo;
 	int32 PlatformOrdinal;
+
+private:
+	bool HasDefaultBuildSettings() const;
+	static bool DoProjectSettingsMatchDefault(const FString& InPlatformName, const FString& InSection, const TArray<FString>* InBoolKeys, const TArray<FString>* InIntKeys, const TArray<FString>* InStringKeys);
 };
 
 
diff --git a/Engine/Source/Developer/TargetPlatform/Public/Interfaces/IProjectBuildMutatorFeature.h b/Engine/Source/Developer/TargetPlatform/Public/Interfaces/IProjectBuildMutatorFeature.h
new file mode 100644
--- /dev/null
+++ b/Engine/Source/Developer/TargetPlatform/Public/Interfaces/IProjectBuildMutatorFeature.h
@@ -0,0 +1,14 @@
+// Copyright 1998-2019 Epic Games, Inc. All Rights Reserved.
+
+#pragma once
+
+#include "CoreMinimal.h"
+#include "Features/IModularFeature.h"
+
+#define PROJECT_BUILD_MUTATOR_FEATURE "ProjectBuildMutatorFeature"
+
+class IProjectBuildMutatorFeature : public IModularFeature
+{
+public:
+	virtual bool RequiresProjectBuild(const FName& InPlatformInfoName, FText& OutReason) const = 0;
+};
diff --git a/Engine/Source/Developer/TargetPlatform/Public/Interfaces/ITargetPlatform.h b/Engine/Source/Developer/TargetPlatform/Public/Interfaces/ITargetPlatform.h
--- a/Engine/Source/Developer/TargetPlatform/Public/Interfaces/ITargetPlatform.h
+++ b/Engine/Source/Developer/TargetPlatform/Public/Interfaces/ITargetPlatform.h
@@ -166,14 +166,26 @@ public:
 	/**
 	 * Checks whether the platform's build requirements are met so that we can do things like package for the platform.
 	 *
-	 * @param ProjectPath Path to the project.
 	 * @param bProjectHasCode true if the project has code, and therefore any compilation based SDK requirements should be checked.
+	 * @param Configuration The configuration being built
+	 * @param bRequiresAssetNativization Whether asset nativization is required
 	 * @param OutTutorialPath Let's the platform tell the editor a path to show some information about how to fix any problem.
 	 * @param OutDocumentationPath Let's the platform tell the editor a documentation path.
 	 * @param CustomizedLogMessage Let's the platform return a customized log message instead of the default for the returned status.
 	 * @return A mask of ETargetPlatformReadyStatus flags to indicate missing requirements, or 0 if all requirements are met.
 	 */
-	virtual int32 CheckRequirements(const FString& ProjectPath, bool bProjectHasCode, FString& OutTutorialPath, FString& OutDocumentationPath, FText& CustomizedLogMessage) const = 0;
+	virtual int32 CheckRequirements(bool bProjectHasCode, EBuildConfiguration Configuration, bool bRequiresAssetNativization, FString& OutTutorialPath, FString& OutDocumentationPath, FText& CustomizedLogMessage) const = 0;
+
+	/**
+	 * Checks whether the current project needs a temporary .target.cs file to be packaged as a code project.
+	 *
+	 * @param bProjectHasCode Whether the project has code
+	 * @param Configuration The configuration being built
+	 * @param bRequiresAssetNativization Whether asset nativization is enabled
+	 * @param OutReason On success, includes a description of the reason why a target is required
+	 * @return True if a temporary target is required
+	 */
+	virtual bool RequiresTempTarget(bool bProjectHasCode, EBuildConfiguration Configuration, bool bRequiresAssetNativization, FText& OutReason) const = 0;
 
 	/**
 	 * Returns the information about this platform
diff --git a/Engine/Source/Developer/TargetPlatform/TargetPlatform.Build.cs b/Engine/Source/Developer/TargetPlatform/TargetPlatform.Build.cs
--- a/Engine/Source/Developer/TargetPlatform/TargetPlatform.Build.cs
+++ b/Engine/Source/Developer/TargetPlatform/TargetPlatform.Build.cs
@@ -11,6 +11,7 @@ public class TargetPlatform : ModuleRules
 		PrivateDependencyModuleNames.Add("SlateCore");
 		PrivateDependencyModuleNames.Add("Slate");
 		PrivateDependencyModuleNames.Add("EditorStyle");
+		PrivateDependencyModuleNames.Add("Projects");
 		PublicDependencyModuleNames.Add("AudioPlatformConfiguration");
 		PublicDependencyModuleNames.Add("DesktopPlatform");
 		PublicDependencyModuleNames.Add("LauncherPlatform");
diff --git a/Engine/Source/Editor/GameProjectGeneration/Private/GameProjectGenerationModule.cpp b/Engine/Source/Editor/GameProjectGeneration/Private/GameProjectGenerationModule.cpp
--- a/Engine/Source/Editor/GameProjectGeneration/Private/GameProjectGenerationModule.cpp
+++ b/Engine/Source/Editor/GameProjectGeneration/Private/GameProjectGenerationModule.cpp
@@ -115,11 +115,6 @@ bool FGameProjectGenerationModule::ProjectHasCodeFiles()
 	return GameProjectUtils::ProjectHasCodeFiles();
 }
 
-bool FGameProjectGenerationModule::ProjectRequiresBuild(const FName InPlatformName)
-{
-	return GameProjectUtils::ProjectRequiresBuild(InPlatformName);
-}
-
 FString FGameProjectGenerationModule::DetermineModuleIncludePath(const FModuleContextInfo& ModuleInfo, const FString& FileRelativeTo)
 {
 	return GameProjectUtils::DetermineModuleIncludePath(ModuleInfo, FileRelativeTo);
diff --git a/Engine/Source/Editor/GameProjectGeneration/Private/GameProjectUtils.cpp b/Engine/Source/Editor/GameProjectGeneration/Private/GameProjectUtils.cpp
--- a/Engine/Source/Editor/GameProjectGeneration/Private/GameProjectUtils.cpp
+++ b/Engine/Source/Editor/GameProjectGeneration/Private/GameProjectUtils.cpp
@@ -78,8 +78,6 @@
 #include "Blueprint/BlueprintSupport.h"
 #include "Settings/ProjectPackagingSettings.h"
 
-#include "ProjectBuildMutatorFeature.h"
-
 #if WITH_LIVE_CODING
 #include "ILiveCodingModule.h"
 #endif
@@ -3665,124 +3663,6 @@ bool GameProjectUtils::ProjectHasCodeFiles()
 	return FileNames.Num() > 0;
 }
 
-PROJECTS_API bool HasDefaultPluginSettings(const FString& Platform);
-
-bool GameProjectUtils::ProjectRequiresBuild(const FName InPlatformInfoName)
-{
-	//  early out on projects with code files
-	if (ProjectHasCodeFiles())
-	{
-		return true;
-	}
-
-	bool bRequiresBuild = false;
-
-	if (!FApp::IsEngineInstalled())
-	{
-		// check to see if the default build settings have changed
-		bRequiresBuild |= !HasDefaultBuildSettings(InPlatformInfoName);
-	}
-
-	// check to see if any plugins beyond the defaults have been enabled
-	const PlatformInfo::FPlatformInfo* PlatformInfo = PlatformInfo::FindPlatformInfo(InPlatformInfoName);
-	FName PlatformName = (PlatformInfo != nullptr) ? PlatformInfo->UBTTargetId : InPlatformInfoName;
-	bRequiresBuild |= !HasDefaultPluginSettings(PlatformName.ToString());
-
-	// check to see if Blueprint nativization is enabled in the Project settings
-	bRequiresBuild |= GetDefault<UProjectPackagingSettings>()->BlueprintNativizationMethod != EProjectPackagingBlueprintNativizationMethod::Disabled;
-
-	// check to see if any projectmutator modular features are available
-	for (FProjectBuildMutatorFeature* Feature : IModularFeatures::Get().GetModularFeatureImplementations<FProjectBuildMutatorFeature>(FProjectBuildMutatorFeature::GetFeatureName()))
-	{
-		bRequiresBuild |= Feature->RequiresProjectBuild(InPlatformInfoName);
-	}
-
-	return bRequiresBuild;
-}
-
-bool GameProjectUtils::DoProjectSettingsMatchDefault(const FString& InPlatformName, const FString& InSection, const TArray<FString>* InBoolKeys, const TArray<FString>* InIntKeys, const TArray<FString>* InStringKeys)
-{
-	FConfigFile ProjIni;
-	FConfigFile DefaultIni;
-	FConfigCacheIni::LoadLocalIniFile(ProjIni, TEXT("Engine"), true, *InPlatformName, true);
-	FConfigCacheIni::LoadExternalIniFile(DefaultIni, TEXT("Engine"), *FPaths::EngineConfigDir(), *FPaths::EngineConfigDir(), true, NULL, true);
-
-	if (InBoolKeys != NULL)
-	{
-		for (int Index = 0; Index < InBoolKeys->Num(); ++Index)
-		{
-			FString Default(TEXT("False")), Project(TEXT("False"));
-			DefaultIni.GetString(*InSection, *((*InBoolKeys)[Index]), Default);
-			ProjIni.GetString(*InSection, *((*InBoolKeys)[Index]), Project);
-			if (Default.Compare(Project, ESearchCase::IgnoreCase))
-			{
-				return false;
-			}
-		}
-	}
-
-	if (InIntKeys != NULL)
-	{
-		for (int Index = 0; Index < InIntKeys->Num(); ++Index)
-		{
-			int64 Default(0), Project(0);
-			DefaultIni.GetInt64(*InSection, *((*InIntKeys)[Index]), Default);
-			ProjIni.GetInt64(*InSection, *((*InIntKeys)[Index]), Project);
-			if (Default != Project)
-			{
-				return false;
-			}
-		}
-	}
-
-	if (InStringKeys != NULL)
-	{
-		for (int Index = 0; Index < InStringKeys->Num(); ++Index)
-		{
-			FString Default(TEXT("False")), Project(TEXT("False"));
-			DefaultIni.GetString(*InSection, *((*InStringKeys)[Index]), Default);
-			ProjIni.GetString(*InSection, *((*InStringKeys)[Index]), Project);
-			if (Default.Compare(Project, ESearchCase::IgnoreCase))
-			{
-				return false;
-			}
-		}
-	}
-
-	return true;
-}
-
-bool GameProjectUtils::HasDefaultBuildSettings(const FName InPlatformInfoName)
-{
-	// first check default build settings for all platforms
-	TArray<FString> BoolKeys, IntKeys, StringKeys, BuildKeys;
-	BuildKeys.Add(TEXT("bCompileApex")); BuildKeys.Add(TEXT("bCompileICU"));
-	BuildKeys.Add(TEXT("bCompileSimplygon")); BuildKeys.Add(TEXT("bCompileSimplygonSSF"));
-	BuildKeys.Add(TEXT("bCompileRecast")); BuildKeys.Add(TEXT("bCompileSpeedTree"));
-	BuildKeys.Add(TEXT("bCompileWithPluginSupport")); BuildKeys.Add(TEXT("bCompilePhysXVehicle")); BuildKeys.Add(TEXT("bCompileFreeType"));
-	BuildKeys.Add(TEXT("bCompileForSize"));	BuildKeys.Add(TEXT("bCompileCEF3")); BuildKeys.Add(TEXT("bCompileCustomSQLitePlatform"));
-
-	const PlatformInfo::FPlatformInfo* const PlatInfo = PlatformInfo::FindPlatformInfo(InPlatformInfoName);
-	check(PlatInfo);
-
-	if (!DoProjectSettingsMatchDefault(PlatInfo->TargetPlatformName.ToString(), TEXT("/Script/BuildSettings.BuildSettings"), &BuildKeys))
-	{
-		return false;
-	}
-
-	if (PlatInfo->SDKStatus == PlatformInfo::EPlatformSDKStatus::Installed)
-	{
-		const ITargetPlatform* const Platform = GetTargetPlatformManager()->FindTargetPlatform(PlatInfo->TargetPlatformName.ToString());
-		if (Platform)
-		{
-			FString PlatformSection;
-			Platform->GetBuildProjectSettingKeys(PlatformSection, BoolKeys, IntKeys, StringKeys);
-			return DoProjectSettingsMatchDefault(PlatInfo->TargetPlatformName.ToString(), PlatformSection, &BoolKeys, &IntKeys, &StringKeys);
-		}
-	}
-	return true;
-}
-
 TArray<FString> GameProjectUtils::GetRequiredAdditionalDependencies(const FNewClassInfo& ClassInfo)
 {
 	TArray<FString> Out;
diff --git a/Engine/Source/Editor/GameProjectGeneration/Public/GameProjectGenerationModule.h b/Engine/Source/Editor/GameProjectGeneration/Public/GameProjectGenerationModule.h
--- a/Engine/Source/Editor/GameProjectGeneration/Public/GameProjectGenerationModule.h
+++ b/Engine/Source/Editor/GameProjectGeneration/Public/GameProjectGenerationModule.h
@@ -115,9 +115,6 @@ public:
 	/** Clear the list of supported target platforms */
 	virtual void ClearSupportedTargetPlatforms();
 
-	/** Check to see if the project requires a code build */
-	virtual bool ProjectRequiresBuild(const FName InPlatformInfoName);
-
 public:
 
 	/** (Un)register a new type of template category to be shown on the new project page */
diff --git a/Engine/Source/Editor/GameProjectGeneration/Public/GameProjectUtils.h b/Engine/Source/Editor/GameProjectGeneration/Public/GameProjectUtils.h
--- a/Engine/Source/Editor/GameProjectGeneration/Public/GameProjectUtils.h
+++ b/Engine/Source/Editor/GameProjectGeneration/Public/GameProjectUtils.h
@@ -262,9 +262,6 @@ public:
 	/** Returns a list of #include lines formed from InList */
 	static FString MakeIncludeList(const TArray<FString>& InList);
 
-	/** Returns true if the currently loaded project requires a code build */
-	static bool ProjectRequiresBuild(const FName InPlatformInfoName);
-
 	/** Deletes the specified list of files that were created during file creation */
 	static void DeleteCreatedFiles(const FString& RootFolder, const TArray<FString>& CreatedFiles);
 
@@ -491,10 +488,6 @@ private:
 	 */
 	static bool UpdateRequiredAdditionalDependencies(FProjectDescriptor& Descriptor, TArray<FString>& RequiredDependencies, const FString& ModuleName);
 
-	/** checks the project ini file against the default ini file to determine whether or not the build settings have changed from default */
-	static bool HasDefaultBuildSettings(const FName InPlatformInfoName);
-	static bool DoProjectSettingsMatchDefault(const FString& InPlatformnName, const FString& InSection, const TArray<FString>* InBoolKeys, const TArray<FString>* InIntKeys = NULL, const TArray<FString>* InStringKeys = NULL);
-
 private:
 	/**
 	 * Updates the projects and modifies FProjectDescriptor accordingly to given modifier.
diff --git a/Engine/Source/Editor/GameProjectGeneration/Public/ProjectBuildMutatorFeature.h b/Engine/Source/Editor/GameProjectGeneration/Public/ProjectBuildMutatorFeature.h
deleted file mode 100644
diff --git a/Engine/Source/Editor/MainFrame/Private/Frame/MainFrameActions.cpp b/Engine/Source/Editor/MainFrame/Private/Frame/MainFrameActions.cpp
--- a/Engine/Source/Editor/MainFrame/Private/Frame/MainFrameActions.cpp
+++ b/Engine/Source/Editor/MainFrame/Private/Frame/MainFrameActions.cpp
@@ -450,11 +450,11 @@ void FMainFrameActionCallbacks::PackageProject( const FName InPlatformInfoName )
 {
 	GUnrealEd->CancelPlayingViaLauncher();
 	SaveAll();
-	
+
 	// does the project have any code?
 	FGameProjectGenerationModule& GameProjectModule = FModuleManager::LoadModuleChecked<FGameProjectGenerationModule>(TEXT("GameProjectGeneration"));
-	bool bProjectHasCode = GameProjectModule.Get().ProjectRequiresBuild(InPlatformInfoName);
-
+	bool bProjectHasCode = GameProjectModule.Get().ProjectHasCodeFiles();
+	
 	const PlatformInfo::FPlatformInfo* const PlatformInfo = PlatformInfo::FindPlatformInfo(InPlatformInfoName);
 	check(PlatformInfo);
 
@@ -484,16 +484,18 @@ void FMainFrameActionCallbacks::PackageProject( const FName InPlatformInfoName )
 	}
 
 	UProjectPackagingSettings* PackagingSettings = Cast<UProjectPackagingSettings>(UProjectPackagingSettings::StaticClass()->GetDefaultObject());
-
+	const UProjectPackagingSettings::FConfigurationInfo& ConfigurationInfo = UProjectPackagingSettings::ConfigurationInfo[PackagingSettings->BuildConfiguration];
+	bool bAssetNativizationEnabled = (PackagingSettings->BlueprintNativizationMethod != EProjectPackagingBlueprintNativizationMethod::Disabled);
+	
+	const ITargetPlatform* const Platform = GetTargetPlatformManager()->FindTargetPlatform(PlatformInfo->TargetPlatformName.ToString());
 	{
-		const ITargetPlatform* const Platform = GetTargetPlatformManager()->FindTargetPlatform(PlatformInfo->TargetPlatformName.ToString());
 		if (Platform)
 		{
 			FString NotInstalledTutorialLink;
 			FString DocumentationLink;
 			FText CustomizedLogMessage;
-			FString ProjectPath = FPaths::IsProjectFilePathSet() ? FPaths::ConvertRelativePathToFull(FPaths::GetProjectFilePath()) : FPaths::RootDir() / FApp::GetProjectName() / FApp::GetProjectName() + TEXT(".uproject");
-			int32 Result = Platform->CheckRequirements(ProjectPath, bProjectHasCode, NotInstalledTutorialLink, DocumentationLink, CustomizedLogMessage);
+
+			int32 Result = Platform->CheckRequirements(bProjectHasCode, ConfigurationInfo.Configuration, bAssetNativizationEnabled, NotInstalledTutorialLink, DocumentationLink, CustomizedLogMessage);
 
 			// report to analytics
 			FEditorAnalytics::ReportBuildRequirementsFailure(TEXT("Editor.Package.Failed"), PlatformInfo->TargetPlatformName.ToString(), bProjectHasCode, Result);
@@ -706,7 +708,8 @@ void FMainFrameActionCallbacks::PackageProject( const FName InPlatformInfoName )
 	}
 	else if(PackagingSettings->Build == EProjectPackagingBuild::IfProjectHasCode)
 	{
-		bBuild = bProjectHasCode || !FApp::GetEngineIsPromotedBuild();
+		FText Reason;
+		bBuild = bProjectHasCode || !FApp::GetEngineIsPromotedBuild() || (Platform == nullptr || Platform->RequiresTempTarget(bProjectHasCode, ConfigurationInfo.Configuration, bAssetNativizationEnabled, Reason));
 	}
 	else if(PackagingSettings->Build == EProjectPackagingBuild::IfEditorWasBuiltLocally)
 	{
@@ -734,12 +737,11 @@ void FMainFrameActionCallbacks::PackageProject( const FName InPlatformInfoName )
 		OptionalParams += FString::Printf(TEXT(" -NumCookersToSpawn=%d"), NumCookers); 
 	}
 
-	const UProjectPackagingSettings::FConfigurationInfo& Info = UProjectPackagingSettings::ConfigurationInfo[PackagingSettings->BuildConfiguration];
-	if (Info.TargetType == EBuildTargetType::Client)
+	if (ConfigurationInfo.TargetType == EBuildTargetType::Client)
 	{
 		OptionalParams += TEXT(" -client");
 	}
-	else if (Info.TargetType == EBuildTargetType::Server)
+	else if (ConfigurationInfo.TargetType == EBuildTargetType::Server)
 	{
 		OptionalParams += TEXT(" -server");
 	}
@@ -751,7 +753,7 @@ void FMainFrameActionCallbacks::PackageProject( const FName InPlatformInfoName )
 		FApp::IsEngineInstalled() ? TEXT(" -installed") : TEXT(""),
 		*ProjectPath,
 		*PackagingSettings->StagingDirectory.Path,
-		LexToString(Info.Configuration),
+		LexToString(ConfigurationInfo.Configuration),
 		*FUnrealEdMisc::Get().GetExecutableForCommandlets(),
 		*OptionalParams
 	);
diff --git a/Engine/Source/Editor/UnrealEd/Classes/Settings/LevelEditorPlaySettings.h b/Engine/Source/Editor/UnrealEd/Classes/Settings/LevelEditorPlaySettings.h
--- a/Engine/Source/Editor/UnrealEd/Classes/Settings/LevelEditorPlaySettings.h
+++ b/Engine/Source/Editor/UnrealEd/Classes/Settings/LevelEditorPlaySettings.h
@@ -414,6 +414,23 @@ public:
 	EVisibility GetClientWindowSizeVisibility() const { return (RunUnderOneProcess ? EVisibility::Hidden : EVisibility::Visible); }
 	bool IsCreateAudioDeviceForEveryPlayer() const { return CreateAudioDeviceForEveryPlayer; }
 
+	EBuildConfiguration GetLaunchBuildConfiguration() const
+	{
+		switch (LaunchConfiguration)
+		{
+		case LaunchConfig_Debug:
+			return EBuildConfiguration::Debug;
+		case LaunchConfig_Development:
+			return EBuildConfiguration::Development;
+		case LaunchConfig_Test:
+			return EBuildConfiguration::Test;
+		case LaunchConfig_Shipping:
+			return EBuildConfiguration::Shipping;
+		default:
+			return FApp::GetBuildConfiguration();
+		}
+	}
+
 public:
 
 	/** The last used height for multiple instance windows (in pixels). */
diff --git a/Engine/Source/Editor/UnrealEd/Private/Kismet2/DebuggerCommands.cpp b/Engine/Source/Editor/UnrealEd/Private/Kismet2/DebuggerCommands.cpp
--- a/Engine/Source/Editor/UnrealEd/Private/Kismet2/DebuggerCommands.cpp
+++ b/Engine/Source/Editor/UnrealEd/Private/Kismet2/DebuggerCommands.cpp
@@ -2128,7 +2128,7 @@ bool FInternalPlayWorldCommandCallbacks::IsReadyToLaunchOnDevice(FString DeviceI
 	checkf(PlatformInfo, TEXT("Unable to find PlatformInfo for %s"), *PlatformName);
 
 	FGameProjectGenerationModule& GameProjectModule = FModuleManager::LoadModuleChecked<FGameProjectGenerationModule>(TEXT("GameProjectGeneration"));
-	bool bHasCode = GameProjectModule.Get().ProjectRequiresBuild(FName(*PlatformName));
+	bool bHasCode = GameProjectModule.Get().ProjectHasCodeFiles();
 
 	if (PlatformInfo->SDKStatus == PlatformInfo::EPlatformSDKStatus::NotInstalled)
 	{
@@ -2146,8 +2146,10 @@ bool FInternalPlayWorldCommandCallbacks::IsReadyToLaunchOnDevice(FString DeviceI
 		FString NotInstalledTutorialLink;
 		FString DocumentationLink;
 		FText CustomizedLogMessage;
-		FString ProjectPath = FPaths::IsProjectFilePathSet() ? FPaths::ConvertRelativePathToFull(FPaths::GetProjectFilePath()) : FPaths::RootDir() / FApp::GetProjectName() / FApp::GetProjectName() + TEXT(".uproject");
-		int32 Result = Platform->CheckRequirements(ProjectPath, bHasCode, NotInstalledTutorialLink, DocumentationLink, CustomizedLogMessage);
+
+	    EBuildConfiguration BuildConfiguration = GetDefault<ULevelEditorPlaySettings>()->GetLaunchBuildConfiguration();
+	    bool bEnableAssetNativization = false;
+		int32 Result = Platform->CheckRequirements(bHasCode, BuildConfiguration, bEnableAssetNativization, NotInstalledTutorialLink, DocumentationLink, CustomizedLogMessage);
 		
 		// report to analytics
 		FEditorAnalytics::ReportBuildRequirementsFailure(TEXT("Editor.LaunchOn.Failed"), PlatformName, bHasCode, Result);
diff --git a/Engine/Source/Editor/UnrealEd/Private/PlayLevel.cpp b/Engine/Source/Editor/UnrealEd/Private/PlayLevel.cpp
--- a/Engine/Source/Editor/UnrealEd/Private/PlayLevel.cpp
+++ b/Engine/Source/Editor/UnrealEd/Private/PlayLevel.cpp
@@ -38,6 +38,7 @@
 #include "Engine/GameInstance.h"
 #include "Engine/RendererSettings.h"
 #include "Engine/World.h"
+#include "DesktopPlatformModule.h"
 #include "Settings/LevelEditorPlaySettings.h"
 #include "AI/NavigationSystemBase.h"
 #include "Editor/EditorEngine.h"
@@ -100,6 +101,7 @@
 #include "Engine/LevelStreaming.h"
 #include "Components/ModelComponent.h"
 #include "GameDelegates.h"
+#include "PlatformInfo.h"
 #include "Net/OnlineEngineInterface.h"
 #include "Kismet2/DebuggerCommands.h"
 #include "Misc/ScopeExit.h"
@@ -1991,10 +1993,10 @@ void UEditorEngine::PlayUsingLauncher()
 				return;
 			}
 		}
-		
+
 		// does the project have any code?
 		FGameProjectGenerationModule& GameProjectModule = FModuleManager::LoadModuleChecked<FGameProjectGenerationModule>(TEXT("GameProjectGeneration"));
-		bPlayUsingLauncherHasCode = GameProjectModule.Get().ProjectRequiresBuild(FName(*LaunchPlatformName));
+		bPlayUsingLauncherHasCode = GameProjectModule.Get().ProjectHasCodeFiles();
 
 		const ULevelEditorPlaySettings* PlayInSettings = GetDefault<ULevelEditorPlaySettings>();
 		// Setup launch profile, keep the setting here to a minimum.
diff --git a/Engine/Source/Programs/AutomationTool/AutomationUtils/ProjectUtils.cs b/Engine/Source/Programs/AutomationTool/AutomationUtils/ProjectUtils.cs
--- a/Engine/Source/Programs/AutomationTool/AutomationUtils/ProjectUtils.cs
+++ b/Engine/Source/Programs/AutomationTool/AutomationUtils/ProjectUtils.cs
@@ -124,114 +124,111 @@ namespace AutomationTool
 			return ProjectClientBinariesPath;
 		}
 
-		private static bool RequiresTempTarget(FileReference RawProjectPath, List<UnrealTargetPlatform> Platforms, List<UnrealTargetConfiguration> Configurations, bool AssetNativizationRequested)
-		{
-			string Reason;
-			if(RequiresTempTarget(RawProjectPath, Platforms, Configurations, AssetNativizationRequested, out Reason))
-			{
-				Log.TraceInformation("{0} requires a temporary target.cs to be generated ({1})", RawProjectPath.GetFileName(), Reason);
-				return true;
-			}
-			return false;
-		}
-
-		private static bool RequiresTempTarget(FileReference RawProjectPath, List<UnrealTargetPlatform> Platforms, List<UnrealTargetConfiguration> Configurations, bool AssetNativizationRequested, out string Reason)
+		private static bool ProjectHasCode(FileReference RawProjectPath)
 		{
 			// check to see if we already have a Target.cs file
-			if (File.Exists (Path.Combine (Path.GetDirectoryName (RawProjectPath.FullName), "Source", RawProjectPath.GetFileNameWithoutExtension() + ".Target.cs")))
+			if (File.Exists(Path.Combine(Path.GetDirectoryName(RawProjectPath.FullName), "Source", RawProjectPath.GetFileNameWithoutExtension() + ".Target.cs")))
 			{
-				Reason = null;
-				return false;
+				return true;
 			}
 			else if (Directory.Exists(Path.Combine(Path.GetDirectoryName(RawProjectPath.FullName), "Source")))
 			{
 				// wasn't one in the main Source directory, let's check all sub-directories
 				//@todo: may want to read each target.cs to see if it has a target corresponding to the project name as a final check
-				FileInfo[] Files = (new DirectoryInfo( Path.Combine (Path.GetDirectoryName (RawProjectPath.FullName), "Source")).GetFiles ("*.Target.cs", SearchOption.AllDirectories));
+				FileInfo[] Files = (new DirectoryInfo(Path.Combine(Path.GetDirectoryName(RawProjectPath.FullName), "Source")).GetFiles("*.Target.cs", SearchOption.AllDirectories));
 				if (Files.Length > 0)
 				{
-					Reason = null;
-					return false;
+					return true;
 				}
 			}
+			return false;
+		}
 
-            //
-            // once we reach this point, we can surmise that this is an asset-
-            // only (code free) project
+		private static bool RequiresTempTarget(FileReference RawProjectPath, List<UnrealTargetPlatform> Platforms, List<UnrealTargetConfiguration> Configurations, bool AssetNativizationRequested)
+		{
+			bool bHasCode = ProjectHasCode(RawProjectPath);
+			foreach (UnrealTargetPlatform Platform in Platforms)
+			{
+				foreach(UnrealTargetConfiguration Configuration in Configurations)
+				{
+					string Reason;
+					if(RequiresTempTarget(RawProjectPath, bHasCode, Platform, Configuration, TargetType.Game, AssetNativizationRequested, true, out Reason))
+					{
+						Log.TraceInformation("{0} requires a temporary target.cs to be generated ({1})", RawProjectPath.GetFileName(), Reason);
+						return true;
+					}
+				}
+			}
+			return false;
+		}
 
-            if (AssetNativizationRequested)
+		private static bool RequiresTempTarget(FileReference RawProjectPath, bool bProjectHasCode, UnrealTargetPlatform Platform, UnrealTargetConfiguration Configuration, TargetType TargetType, bool bRequiresAssetNativization, bool bRequiresCookedData, out string OutReason)
+		{
+			// check to see if we already have a Target.cs file
+			if (bProjectHasCode)
+			{
+				OutReason = null;
+				return false;
+			}
+
+			// check if asset nativization is enabled
+			if (bRequiresAssetNativization)
             {
-                // we're going to be converting some of the project's assets 
-                // into native code, so we require a distinct target (executable) 
-                // be generated for this project
-				Reason = "asset nativization is enabled";
+				OutReason = "asset nativization is enabled";
                 return true;
             }
 
 			// Check if encryption or signing is enabled
-			foreach (UnrealTargetPlatform Platform in Platforms)
+			EncryptionAndSigning.CryptoSettings Settings = EncryptionAndSigning.ParseCryptoSettings(RawProjectPath.Directory, Platform);
+			if (Settings.IsAnyEncryptionEnabled() || Settings.IsPakSigningEnabled())
 			{
-				EncryptionAndSigning.CryptoSettings Settings = EncryptionAndSigning.ParseCryptoSettings(RawProjectPath.Directory, Platform);
-				if (Settings.IsAnyEncryptionEnabled() || Settings.IsPakSigningEnabled())
-				{
-					Reason = "encryption/signing is enabled";
-					return true;
-				}
+				OutReason = "encryption/signing is enabled";
+				return true;
+			}
+
+			// check the target platforms for any differences in build settings or additional plugins
+			if(!CommandUtils.IsEngineInstalled() && !PlatformExports.HasDefaultBuildConfig(RawProjectPath, Platform))
+			{
+				OutReason = "project has non-default build configuration";
+				return true;
+			}
+			if(PlatformExports.RequiresBuild(RawProjectPath, Platform))
+			{
+				OutReason = "overriden by target platform";
+				return true;
 			}
 
 			// Read the project descriptor, and find all the plugins available to this project
 			ProjectDescriptor Project = ProjectDescriptor.FromFile(RawProjectPath);
-			List<PluginInfo> AvailablePlugins = Plugins.ReadAvailablePlugins(CommandUtils.EngineDirectory, RawProjectPath, Project.AdditionalPluginDirectories);
 
-			// check the target platforms for any differences in build settings or additional plugins
-			foreach (UnrealTargetPlatform Platform in Platforms)
+			// find if there are any plugins enabled or disabled which differ from the default
+			List<PluginInfo> Plugins = UnrealBuildTool.Plugins.ReadAvailablePlugins(CommandUtils.EngineDirectory, RawProjectPath, Project.AdditionalPluginDirectories);
+			foreach (PluginInfo Plugin in Plugins)
 			{
-				if(!CommandUtils.IsEngineInstalled() && !PlatformExports.HasDefaultBuildConfig(RawProjectPath, Platform))
-				{
-					Reason = "project has non-default build configuration";
-					return true;
-				}
-				if(PlatformExports.RequiresBuild(RawProjectPath, Platform))
+				bool bPluginEnabledForTarget = UnrealBuildTool.Plugins.IsPluginCompiledForTarget(Plugin, Project, Platform, Configuration, TargetType, bRequiresCookedData);
+
+				bool bPluginEnabledForBaseTarget = false;
+				if(!Plugin.Descriptor.bInstalled)
 				{
-					Reason = "overriden by target platform";
-					return true;
+					bPluginEnabledForBaseTarget |= UnrealBuildTool.Plugins.IsPluginCompiledForTarget(Plugin, null, Platform, Configuration, TargetType, bRequiresCookedData);
 				}
 
-				// find if there are any plugins enabled or disabled which differ from the default
-				foreach(PluginInfo Plugin in AvailablePlugins)
+				if (bPluginEnabledForTarget != bPluginEnabledForBaseTarget)
 				{
-					bool bPluginEnabledForTarget = false;
-					foreach (UnrealTargetConfiguration Configuration in Configurations)
+					if(bPluginEnabledForTarget)
 					{
-						bPluginEnabledForTarget |= Plugins.IsPluginEnabledForProject(Plugin, Project, Platform, Configuration, TargetRules.TargetType.Game);
+						OutReason = String.Format("{0} plugin is enabled", Plugin.Name);
+						return true;
 					}
-
-					bool bPluginEnabledForBaseTarget = false;
-					if(!Plugin.Descriptor.bInstalled)
-					{
-						foreach (UnrealTargetConfiguration Configuration in Configurations)
-						{
-							bPluginEnabledForBaseTarget |= Plugins.IsPluginEnabledForProject(Plugin, null, Platform, Configuration, TargetRules.TargetType.Game);
-						}
-					}
-
-					if (bPluginEnabledForTarget != bPluginEnabledForBaseTarget)
+					else
 					{
-						if(bPluginEnabledForTarget)
-						{
-							Reason = String.Format("{0} plugin is enabled", Plugin.Name);
-							return true;
-						}
-						else
-						{
-							Reason = String.Format("{0} plugin is disabled", Plugin.Name);
-							return true;
-						}
+						OutReason = String.Format("{0} plugin is disabled", Plugin.Name);
+						return true;
 					}
 				}
 			}
 
-			Reason = null;
+			OutReason = null;
 			return false;
 		}
 
diff --git a/Engine/Source/Programs/UnrealBuildTool/ProjectFiles/Xcode/XcodeProject.cs b/Engine/Source/Programs/UnrealBuildTool/ProjectFiles/Xcode/XcodeProject.cs
--- a/Engine/Source/Programs/UnrealBuildTool/ProjectFiles/Xcode/XcodeProject.cs
+++ b/Engine/Source/Programs/UnrealBuildTool/ProjectFiles/Xcode/XcodeProject.cs
@@ -748,7 +748,7 @@ namespace UnrealBuildTool
 			// @todo: Cache this kind of things since every target will re-do this work!
 			foreach (PluginInfo PI in AvailablePlugins)
 			{
-				if (!Plugins.IsPluginEnabledForProject(PI, Project, UnrealTargetPlatform.IOS, UnrealTargetConfiguration.Development, TargetRules.TargetType.Game))
+				if (!Plugins.IsPluginEnabledForTarget(PI, Project, UnrealTargetPlatform.IOS, UnrealTargetConfiguration.Development, TargetRules.TargetType.Game))
 				{
 					continue;
 				}
diff --git a/Engine/Source/Programs/UnrealBuildTool/System/Plugins.cs b/Engine/Source/Programs/UnrealBuildTool/System/Plugins.cs
--- a/Engine/Source/Programs/UnrealBuildTool/System/Plugins.cs
+++ b/Engine/Source/Programs/UnrealBuildTool/System/Plugins.cs
@@ -512,10 +512,10 @@ namespace UnrealBuildTool
 		/// <param name="Project">The project to check. May be null.</param>
 		/// <param name="Plugin">Information about the plugin</param>
 		/// <param name="Platform">The target platform</param>
-		/// <param name="TargetConfiguration">The target configuration</param>
-		/// <param name="Target"></param>
+		/// <param name="Configuration">The target configuration</param>
+		/// <param name="TargetType">The type of target being built</param>
 		/// <returns>True if the plugin should be enabled for this project</returns>
-		public static bool IsPluginEnabledForProject(PluginInfo Plugin, ProjectDescriptor Project, UnrealTargetPlatform Platform, UnrealTargetConfiguration TargetConfiguration, TargetType Target)
+		public static bool IsPluginEnabledForTarget(PluginInfo Plugin, ProjectDescriptor Project, UnrealTargetPlatform Platform, UnrealTargetConfiguration Configuration, TargetType TargetType)
 		{
 			if (!Plugin.Descriptor.SupportsTargetPlatform(Platform))
 			{
@@ -529,11 +529,39 @@ namespace UnrealBuildTool
 				{
 					if (String.Compare(PluginReference.Name, Plugin.Name, true) == 0 && !PluginReference.bOptional)
 					{
-						bEnabled = PluginReference.IsEnabledForPlatform(Platform) && PluginReference.IsEnabledForTargetConfiguration(TargetConfiguration) && PluginReference.IsEnabledForTarget(Target);
+						bEnabled = PluginReference.IsEnabledForPlatform(Platform) && PluginReference.IsEnabledForTargetConfiguration(Configuration) && PluginReference.IsEnabledForTarget(TargetType);
 					}
 				}
 			}
 			return bEnabled;
 		}
+
+		/// <summary>
+		/// Determine if a plugin is enabled for a given project
+		/// </summary>
+		/// <param name="Project">The project to check. May be null.</param>
+		/// <param name="Plugin">Information about the plugin</param>
+		/// <param name="Platform">The target platform</param>
+		/// <param name="Configuration">The target configuration</param>
+		/// <param name="TargetType">The type of target being built</param>
+		/// <param name="bRequiresCookedData">Whether the target requires cooked data</param>
+		/// <returns>True if the plugin should be enabled for this project</returns>
+		public static bool IsPluginCompiledForTarget(PluginInfo Plugin, ProjectDescriptor Project, UnrealTargetPlatform Platform, UnrealTargetConfiguration Configuration, TargetType TargetType, bool bRequiresCookedData)
+		{
+			bool bCompiledForTarget = false;
+			if (IsPluginEnabledForTarget(Plugin, Project, Platform, Configuration, TargetType))
+			{
+				bool bBuildDeveloperTools = (TargetType == TargetType.Editor || TargetType == TargetType.Program || (Configuration != UnrealTargetConfiguration.Test && Configuration != UnrealTargetConfiguration.Shipping));
+				foreach (ModuleDescriptor Module in Plugin.Descriptor.Modules)
+				{
+					if (Module.IsCompiledInConfiguration(Platform, Configuration, "", TargetType, bBuildDeveloperTools, bRequiresCookedData))
+					{
+						bCompiledForTarget = true;
+						break;
+					}
+				}
+			}
+			return bCompiledForTarget;
+		}
 	}
 }
