Allow selecting the target to package through the project packaging menu, rather than showing client/server in the configuration menu. This allows support for multiple game/client/server targets.

diff --git a/Engine/Source/Editor/MainFrame/Private/Frame/MainFrameActions.cpp b/Engine/Source/Editor/MainFrame/Private/Frame/MainFrameActions.cpp
--- a/Engine/Source/Editor/MainFrame/Private/Frame/MainFrameActions.cpp
+++ b/Engine/Source/Editor/MainFrame/Private/Frame/MainFrameActions.cpp
@@ -531,6 +531,18 @@ bool FMainFrameActionCallbacks::PackageBuildConfigurationIsChecked( EProjectPack
 	return (GetDefault<UProjectPackagingSettings>()->BuildConfiguration == BuildConfiguration);
 }
 
+void FMainFrameActionCallbacks::PackageBuildTarget( FString TargetName )
+{
+	UProjectPackagingSettings* PackagingSettings = GetMutableDefault<UProjectPackagingSettings>();
+	PackagingSettings->BuildTarget = TargetName;
+}
+
+bool FMainFrameActionCallbacks::PackageBuildTargetIsChecked( FString TargetName )
+{
+	const FTargetInfo* Target = GetDefault<UProjectPackagingSettings>()->GetBuildTargetInfo();
+	return (Target != nullptr && Target->Name == TargetName);
+}
+
 void FMainFrameActionCallbacks::PackageProject( const FName InPlatformInfoName )
 {
 	GUnrealEd->CancelPlayingViaLauncher();
@@ -571,7 +583,7 @@ void FMainFrameActionCallbacks::PackageProject( const FName InPlatformInfoName )
 	UProjectPackagingSettings* PackagingSettings = Cast<UProjectPackagingSettings>(UProjectPackagingSettings::StaticClass()->GetDefaultObject());
 	const UProjectPackagingSettings::FConfigurationInfo& ConfigurationInfo = UProjectPackagingSettings::ConfigurationInfo[PackagingSettings->BuildConfiguration];
 	bool bAssetNativizationEnabled = (PackagingSettings->BlueprintNativizationMethod != EProjectPackagingBlueprintNativizationMethod::Disabled);
-	
+
 	const ITargetPlatform* const Platform = GetTargetPlatformManager()->FindTargetPlatform(PlatformInfo->TargetPlatformName.ToString());
 	{
 		if (Platform)
@@ -822,13 +834,18 @@ void FMainFrameActionCallbacks::PackageProject( const FName InPlatformInfoName )
 		OptionalParams += FString::Printf(TEXT(" -NumCookersToSpawn=%d"), NumCookers); 
 	}
 
-	if (ConfigurationInfo.bClientAndServer)
+	const FTargetInfo* Target = PackagingSettings->GetBuildTargetInfo();
+	if (Target == nullptr)
+	{
+		OptionalParams += FString::Printf(TEXT(" -clientconfig=%s"), LexToString(ConfigurationInfo.Configuration));
+	}
+	else if(Target->Type == EBuildTargetType::Server)
 	{
-		OptionalParams += FString::Printf(TEXT(" -client -clientconfig=%s -server -serverconfig=%s"), LexToString(ConfigurationInfo.Configuration), LexToString(ConfigurationInfo.Configuration));
+		OptionalParams += FString::Printf(TEXT(" -target=%s -serverconfig=%s"), *Target->Name, LexToString(ConfigurationInfo.Configuration));
 	}
 	else
 	{
-		OptionalParams += FString::Printf(TEXT(" -clientconfig=%s"), LexToString(ConfigurationInfo.Configuration));
+		OptionalParams += FString::Printf(TEXT(" -target=%s -clientconfig=%s"), *Target->Name, LexToString(ConfigurationInfo.Configuration));
 	}
 
 	FString ProjectPath = FPaths::IsProjectFilePathSet() ? FPaths::ConvertRelativePathToFull(FPaths::GetProjectFilePath()) : FPaths::RootDir() / FApp::GetProjectName() / FApp::GetProjectName() + TEXT(".uproject");
diff --git a/Engine/Source/Editor/MainFrame/Private/Frame/MainFrameActions.h b/Engine/Source/Editor/MainFrame/Private/Frame/MainFrameActions.h
--- a/Engine/Source/Editor/MainFrame/Private/Frame/MainFrameActions.h
+++ b/Engine/Source/Editor/MainFrame/Private/Frame/MainFrameActions.h
@@ -206,6 +206,12 @@ public:
 	/** Determines whether the specified build configuration option is checked. */
 	static bool PackageBuildConfigurationIsChecked( EProjectPackagingBuildConfigurations BuildConfiguration );
 
+	/** Sets the project packaging build configuration. */
+	static void PackageBuildTarget( FString TargetName );
+
+	/** Determines whether the specified build configuration option is checked. */
+	static bool PackageBuildTargetIsChecked( FString TargetName );
+
 	/** Packages the project for the specified platform. */
 	static void PackageProject( const FName InPlatformInfoName );
 
diff --git a/Engine/Source/Editor/MainFrame/Private/Menus/PackageProjectMenu.h b/Engine/Source/Editor/MainFrame/Private/Menus/PackageProjectMenu.h
--- a/Engine/Source/Editor/MainFrame/Private/Menus/PackageProjectMenu.h
+++ b/Engine/Source/Editor/MainFrame/Private/Menus/PackageProjectMenu.h
@@ -90,6 +90,11 @@ public:
 			LOCTEXT("PackageProjectBuildConfigurationSubMenuToolTip", "Select the build configuration to package the project with"),
 			FNewMenuDelegate::CreateStatic(&FPackageProjectMenu::MakeBuildConfigurationsMenu)
 		);
+		MenuBuilder.AddSubMenu(
+			LOCTEXT("PackageProjectBuildTargetSubMenuLabel", "Build Target"),
+			LOCTEXT("PackageProjectBuildTargetSubMenuToolTip", "Select the build target to package"),
+			FNewMenuDelegate::CreateStatic(&FPackageProjectMenu::MakeBuildTargetsMenu)
+		);
 
 		MenuBuilder.AddMenuSeparator();
 		MenuBuilder.AddMenuEntry(FMainFrameCommands::Get().PackagingSettings);
@@ -187,6 +192,38 @@ protected:
 			);
 		}
 	}
+
+	/**
+	 * Creates a build configuration sub-menu.
+	 *
+	 * @param MenuBuilder The builder for the menu that owns this menu.
+	 */
+	static void MakeBuildTargetsMenu(FMenuBuilder& MenuBuilder)
+	{
+		IDesktopPlatform* DesktopPlatform = FDesktopPlatformModule::Get();
+
+		TArray<FTargetInfo> Targets = DesktopPlatform->GetTargetsForCurrentProject();
+		Targets.Sort([](const FTargetInfo& A, const FTargetInfo& B){ return A.Name < B.Name; });
+
+		for (const FTargetInfo& Target : Targets)
+		{
+			if (Target.Type == EBuildTargetType::Game || Target.Type == EBuildTargetType::Client || Target.Type == EBuildTargetType::Server)
+			{
+				MenuBuilder.AddMenuEntry(
+					FText::FromString(Target.Name),
+					FText::Format(LOCTEXT("PackageTargetName", "Package the '{0}' target."), FText::FromString(Target.Name)),
+					FSlateIcon(),
+					FUIAction(
+						FExecuteAction::CreateStatic(&FMainFrameActionCallbacks::PackageBuildTarget, Target.Name),
+						FCanExecuteAction(),
+						FIsActionChecked::CreateStatic(&FMainFrameActionCallbacks::PackageBuildTargetIsChecked, Target.Name)
+					),
+					NAME_None,
+					EUserInterfaceActionType::RadioButton
+				);
+			}
+		}
+	}
 };
 
 
diff --git a/Engine/Source/Editor/UnrealEd/Classes/Settings/ProjectPackagingSettings.h b/Engine/Source/Editor/UnrealEd/Classes/Settings/ProjectPackagingSettings.h
--- a/Engine/Source/Editor/UnrealEd/Classes/Settings/ProjectPackagingSettings.h
+++ b/Engine/Source/Editor/UnrealEd/Classes/Settings/ProjectPackagingSettings.h
@@ -8,6 +8,8 @@
 #include "Engine/EngineTypes.h"
 #include "ProjectPackagingSettings.generated.h"
 
+struct FTargetInfo;
+
 /**
  * Enumerates the available build configurations for project packaging.
  */
@@ -17,33 +19,18 @@ enum EProjectPackagingBuildConfigurations
 	/** Debug configuration. */
 	PPBC_Debug UMETA(DisplayName="Debug"),
 
-	/** Debug client & server configuration. */
-	PPBC_DebugClientAndServer UMETA(DisplayName = "Debug (Client & Server)"),
-
 	/** DebugGame configuration. */
 	PPBC_DebugGame UMETA(DisplayName="DebugGame"),
 
-	/** DebugGame client & server configuration. */
-	PPBC_DebugGameClientAndServer UMETA(DisplayName = "DebugGame (Client & Server)"),
-
 	/** Development configuration. */
 	PPBC_Development UMETA(DisplayName="Development"),
 
-	/** Development client & server configuration. */
-	PPBC_DevelopmentClientAndServer UMETA(DisplayName = "Development (Client & Server)"),
-
 	/** Test configuration. */
 	PPBC_Test UMETA(DisplayName="Test"),
 
-	/** Test client & server configuration. */
-	PPBC_TestClientAndServer UMETA(DisplayName = "Test (Client & Server)"),
-
 	/** Shipping configuration. */
 	PPBC_Shipping UMETA(DisplayName="Shipping"),
 
-	/** Shipping client & server configuration. */
-	PPBC_ShippingClientAndServer UMETA(DisplayName = "Shipping (Client & Server)"),
-
 	/** Number of entries in the enum. */
 	PPBC_MAX
 };
@@ -121,7 +108,6 @@ public:
 	struct FConfigurationInfo
 	{
 		EBuildConfiguration Configuration;
-		bool bClientAndServer;
 		FText Name;
 		FText ToolTip;
 	};
@@ -139,6 +125,10 @@ public:
 	UPROPERTY(config, EditAnywhere, Category=Project)
 	TEnumAsByte<EProjectPackagingBuildConfigurations> BuildConfiguration;
 
+	/** Name of the target to build */
+	UPROPERTY(config, EditAnywhere, Category=Project)
+	FString BuildTarget;
+
 	/** The directory to which the packaged project will be copied. */
 	UPROPERTY(config, EditAnywhere, Category=Project)
 	FDirectoryPath StagingDirectory;
@@ -454,6 +444,9 @@ public:
 	/** Gets a list of all valid packaging configurations for the current project */
 	static TArray<EProjectPackagingBuildConfigurations> GetValidPackageConfigurations();
 
+	/** Gets the current build target, checking that it's valid, and the default build target if it is not */
+	const FTargetInfo* GetBuildTargetInfo() const;
+
 private:
 	/** Returns the index of the specified Blueprint in the exclusive nativization list (otherwise INDEX_NONE) */
 	int32 FindBlueprintInNativizationList(const UBlueprint* InBlueprint) const;
diff --git a/Engine/Source/Editor/UnrealEd/Private/Settings/SettingsClasses.cpp b/Engine/Source/Editor/UnrealEd/Private/Settings/SettingsClasses.cpp
--- a/Engine/Source/Editor/UnrealEd/Private/Settings/SettingsClasses.cpp
+++ b/Engine/Source/Editor/UnrealEd/Private/Settings/SettingsClasses.cpp
@@ -813,20 +813,11 @@ void ULevelEditorViewportSettings::PostEditChangeProperty( struct FPropertyChang
 
 const UProjectPackagingSettings::FConfigurationInfo UProjectPackagingSettings::ConfigurationInfo[PPBC_MAX] = 
 {
-	/* PPBC_Debug */                      { EBuildConfiguration::Debug, false, LOCTEXT("DebugConfiguration", "Debug"), LOCTEXT("DebugConfigurationTooltip", "Package the game in Debug configuration") },
-	/* PPBC_DebugClientAndServer */       { EBuildConfiguration::Debug, true, LOCTEXT("DebugClientConfiguration", "Debug (Client & Server)"), LOCTEXT("DebugClientConfigurationTooltip", "Package a client and server in Debug configuration") },
-
-	/* PPBC_DebugGame */                  { EBuildConfiguration::DebugGame, false, LOCTEXT("DebugGameConfiguration", "DebugGame"), LOCTEXT("DebugGameConfigurationTooltip", "Package the game in DebugGame configuration") },
-	/* PPBC_DebugGameClientAndServer */   { EBuildConfiguration::DebugGame, true, LOCTEXT("DebugGameClientConfiguration", "DebugGame (Client & Server)"), LOCTEXT("DebugGameClientConfigurationTooltip", "Package a client and server in DebugGame configuration") },
-
-	/* PPBC_Development */                { EBuildConfiguration::Development, false, LOCTEXT("DevelopmentConfiguration", "Development"), LOCTEXT("DevelopmentConfigurationTooltip", "Package the game in Development configuration") },
-	/* PPBC_DevelopmentClientAndServer */ { EBuildConfiguration::Development, true, LOCTEXT("DevelopmentClientConfiguration", "Development (Client & Server)"), LOCTEXT("DevelopmentClientConfigurationTooltip", "Package a client and server in Development configuration") },
-
-	/* PPBC_Test */                       { EBuildConfiguration::Test, false, LOCTEXT("TestConfiguration", "Test"), LOCTEXT("TestConfigurationTooltip", "Package the game in Test configuration") },
-	/* PPBC_TestClientAndServer */        { EBuildConfiguration::Test, true, LOCTEXT("TestClientConfiguration", "Test (Client & Server)"), LOCTEXT("TestClientConfigurationTooltip", "Package a client and server in Test configuration") },
-
-	/* PPBC_Shipping */                   { EBuildConfiguration::Shipping, false, LOCTEXT("ShippingConfiguration", "Shipping"), LOCTEXT("ShippingConfigurationTooltip", "Package the game in Shipping configuration") },
-	/* PPBC_ShippingClientAndServer */    { EBuildConfiguration::Shipping, true, LOCTEXT("ShippingClientConfiguration", "Shipping (Client & Server)"), LOCTEXT("ShippingClientConfigurationTooltip", "Package a client and server in Shipping configuration") },
+	/* PPBC_Debug */         { EBuildConfiguration::Debug, LOCTEXT("DebugConfiguration", "Debug"), LOCTEXT("DebugConfigurationTooltip", "Package the game in Debug configuration") },
+	/* PPBC_DebugGame */     { EBuildConfiguration::DebugGame, LOCTEXT("DebugGameConfiguration", "DebugGame"), LOCTEXT("DebugGameConfigurationTooltip", "Package the game in DebugGame configuration") },
+	/* PPBC_Development */   { EBuildConfiguration::Development, LOCTEXT("DevelopmentConfiguration", "Development"), LOCTEXT("DevelopmentConfigurationTooltip", "Package the game in Development configuration") },
+	/* PPBC_Test */          { EBuildConfiguration::Test, LOCTEXT("TestConfiguration", "Test"), LOCTEXT("TestConfigurationTooltip", "Package the game in Test configuration") },
+	/* PPBC_Shipping */      { EBuildConfiguration::Shipping, LOCTEXT("ShippingConfiguration", "Shipping"), LOCTEXT("ShippingConfigurationTooltip", "Package the game in Shipping configuration") },
 };
 
 UProjectPackagingSettings::UProjectPackagingSettings( const FObjectInitializer& ObjectInitializer )
@@ -905,7 +896,7 @@ void UProjectPackagingSettings::PostEditChangeProperty( FPropertyChangedEvent& P
 	}
 	else if (Name == FName(TEXT("ForDistribution")))
 	{
-		if (ForDistribution && BuildConfiguration != EProjectPackagingBuildConfigurations::PPBC_Shipping && BuildConfiguration != EProjectPackagingBuildConfigurations::PPBC_ShippingClientAndServer)
+		if (ForDistribution && BuildConfiguration != EProjectPackagingBuildConfigurations::PPBC_Shipping)
 		{
 			BuildConfiguration = EProjectPackagingBuildConfigurations::PPBC_Shipping;
 			// force serialization for "Build COnfiguration"
@@ -1137,62 +1128,39 @@ TArray<EProjectPackagingBuildConfigurations> UProjectPackagingSettings::GetValid
 	for (int32 Idx = 0; Idx < PPBC_MAX; Idx++)
 	{
 		EProjectPackagingBuildConfigurations PackagingConfiguration = (EProjectPackagingBuildConfigurations)Idx;
-		const UProjectPackagingSettings::FConfigurationInfo& Info = UProjectPackagingSettings::ConfigurationInfo[Idx];
 
-		// Make sure the selected configuration is supported
-		EProjectType ProjectType = bHasCode? EProjectType::Code : EProjectType::Content;
-		if (Info.bClientAndServer)
+		// Check the target type is valid
+		const UProjectPackagingSettings::FConfigurationInfo& Info = UProjectPackagingSettings::ConfigurationInfo[Idx];
+		if(!bHasCode && Info.Configuration == EBuildConfiguration::DebugGame)
 		{
-			if(!FInstalledPlatformInfo::Get().IsValid(EBuildTargetType::Client, TOptional<FString>(), Info.Configuration, ProjectType, EInstalledPlatformState::Downloaded))
-			{
-				continue;
-			}
-			if(!FInstalledPlatformInfo::Get().IsValid(EBuildTargetType::Server, TOptional<FString>(), Info.Configuration, ProjectType, EInstalledPlatformState::Downloaded))
-			{
-				continue;
-			}
+			continue;
 		}
-		else
+
+		Configurations.Add(PackagingConfiguration);
+	}
+	return Configurations;
+}
+
+const FTargetInfo* UProjectPackagingSettings::GetBuildTargetInfo() const
+{
+	const FTargetInfo* DefaultGameTarget = nullptr;
+	const FTargetInfo* DefaultClientTarget = nullptr;
+	for (const FTargetInfo& Target : FDesktopPlatformModule::Get()->GetTargetsForCurrentProject())
+	{
+		if (Target.Name == BuildTarget)
 		{
-			if(!FInstalledPlatformInfo::Get().IsValid(EBuildTargetType::Game, TOptional<FString>(), Info.Configuration, ProjectType, EInstalledPlatformState::Downloaded))
-			{
-				continue;
-			}
+			return &Target;
 		}
-
-		// Check the target type is valid
-		if(bHasCode)
+		else if (Target.Type == EBuildTargetType::Game && (DefaultGameTarget == nullptr || Target.Name < DefaultGameTarget->Name))
 		{
-			if (Info.bClientAndServer)
-			{
-				if (!Targets->ContainsByPredicate([](const FTargetInfo& Target) -> bool { return Target.Type == EBuildTargetType::Client; }))
-				{
-					continue;
-				}
-				if (!Targets->ContainsByPredicate([](const FTargetInfo& Target) -> bool { return Target.Type == EBuildTargetType::Server; }))
-				{
-					continue;
-				}
-			}
-			else
-			{
-				if (!Targets->ContainsByPredicate([](const FTargetInfo& Target) -> bool { return Target.Type == EBuildTargetType::Game; }))
-				{
-					continue;
-				}
-			}
+			DefaultGameTarget = &Target;
 		}
-		else
+		else if (Target.Type == EBuildTargetType::Client && (DefaultClientTarget == nullptr || Target.Name < DefaultClientTarget->Name))
 		{
-			if(Info.Configuration == EBuildConfiguration::DebugGame)
-			{
-				continue;
-			}
+			DefaultClientTarget = &Target;
 		}
-
-		Configurations.Add(PackagingConfiguration);
 	}
-	return Configurations;
+	return (DefaultGameTarget != nullptr)? DefaultGameTarget : DefaultClientTarget;
 }
 
 int32 UProjectPackagingSettings::FindBlueprintInNativizationList(const UBlueprint* InBlueprint) const
diff --git a/Engine/Source/Programs/AutomationTool/AutomationUtils/ProjectParams.cs b/Engine/Source/Programs/AutomationTool/AutomationUtils/ProjectParams.cs
--- a/Engine/Source/Programs/AutomationTool/AutomationUtils/ProjectParams.cs
+++ b/Engine/Source/Programs/AutomationTool/AutomationUtils/ProjectParams.cs
@@ -258,8 +258,7 @@ namespace AutomationTool
 			this.ServerCookedTargets = InParams.ServerCookedTargets;
 			this.EditorTargets = InParams.EditorTargets;
 			this.ProgramTargets = InParams.ProgramTargets;
-			this.SpecifiedClientTarget = InParams.SpecifiedClientTarget;
-			this.SpecifiedServerTarget = InParams.SpecifiedServerTarget;
+			this.TargetNames = new List<string>(InParams.TargetNames);
 			this.ClientTargetPlatforms = InParams.ClientTargetPlatforms;
             this.ClientDependentPlatformMap = InParams.ClientDependentPlatformMap;
 			this.ServerTargetPlatforms = InParams.ServerTargetPlatforms;
@@ -678,8 +677,19 @@ namespace AutomationTool
 				this.DedicatedServer = true;
 			}*/
 			this.NoClient = GetParamValueIfNotSpecified(Command, NoClient, this.NoClient, "noclient");
-			this.SpecifiedClientTarget = ParseParamValueIfNotSpecified(Command, SpecifiedClientTarget, "client", "", true);
-			this.SpecifiedServerTarget = ParseParamValueIfNotSpecified(Command, SpecifiedServerTarget, "server", "", true);
+
+			if(Command != null)
+			{
+				if(TargetNames == null)
+				{
+					TargetNames = new List<string>();
+				}
+				foreach(string TargetParam in Command.ParseParamValues("target"))
+				{
+					TargetNames.AddRange(TargetParam.Split('+'));
+				}
+			}
+
 			this.LogWindow = GetParamValueIfNotSpecified(Command, LogWindow, this.LogWindow, "logwindow");
 			this.Stage = GetParamValueIfNotSpecified(Command, Stage, this.Stage, "stage");
 			this.SkipStage = GetParamValueIfNotSpecified(Command, SkipStage, this.SkipStage, "skipstage");
@@ -1448,14 +1458,9 @@ namespace AutomationTool
 		}
 
 		/// <summary>
-		/// Build: If more than one client target is found, build this one.
+		/// Build: Specifies the names of targets to build
 		/// </summary>
-		private string SpecifiedClientTarget;
-
-		/// <summary>
-		/// Build: If more than one server target is found, build this one.
-		/// </summary>
-		private string SpecifiedServerTarget;
+		private List<string> TargetNames;
 
 		#endregion
 
@@ -1990,30 +1995,16 @@ namespace AutomationTool
 			return DetectedTargets.FindAll(Target => Target.Rules.Type == DesiredType).ConvertAll(Target => Target.TargetName);
 		}
 
-		private String ChooseTarget(List<String> Targets, TargetType Type, string SpecifiedTargetName)
+		private String ChooseTarget(List<String> Targets, TargetType Type)
 		{
-			if (String.IsNullOrEmpty(SpecifiedTargetName))
+			switch (Targets.Count)
 			{
-				switch (Targets.Count)
-				{
-					case 1:
-						return Targets.First();
-					case 0:
-						throw new AutomationException("{0} target not found!", Type);
-					default:
-						throw new AutomationException("More than one {0} target found. Specify which one to use with the -{1}= option.", Type, Type);
-				}
-			}
-			else
-			{
-				if (Targets.Contains(SpecifiedTargetName))
-				{
-					return SpecifiedTargetName;
-				}
-				else
-				{
-					throw new AutomationException("Could not find {0} target named {1}.", Type, SpecifiedTargetName);
-				}
+				case 1:
+					return Targets.First();
+				case 0:
+					throw new AutomationException("{0} target not found!", Type);
+				default:
+					throw new AutomationException("More than one {0} target found. Specify which one to use with the -{1}= option.", Type, Type);
 			}
 		}
 
@@ -2050,12 +2041,82 @@ namespace AutomationTool
 				EditorTarget = "UE4Editor";
 				ServerTarget = "UE4Server";
 			}
-			else if (!CommandUtils.CmdEnv.HasCapabilityToCompile)
+			else if (TargetNames.Count > 0)
 			{
-				var ShortName = ProjectUtils.GetShortProjectName(RawProjectPath);
-				GameTarget = Client ? (ShortName + "Client") : ShortName;
-				EditorTarget = ShortName + "Editor";
-				ServerTarget = ShortName + "Server";
+				// Resolve all the targets that need to be built
+				List<SingleTargetProperties> Targets = new List<SingleTargetProperties>();
+				foreach (string TargetName in TargetNames)
+				{
+					SingleTargetProperties Target = DetectedTargets.FirstOrDefault(x => String.Equals(x.TargetName, TargetName, StringComparison.OrdinalIgnoreCase));
+					if(Target == null)
+					{
+						throw new AutomationException("Unable to find target '{0}'", TargetName);
+					}
+					Targets.Add(Target);
+				}
+
+				// Make sure we haven't specified game and clients together
+				if (Targets.Any(x => x.Rules.Type == TargetType.Client) && Targets.Any(x => x.Rules.Type == TargetType.Game))
+				{
+					throw new AutomationException("Cannot specify client ang game targets to be built together");
+				}
+
+				// Create the lists to receive all the target types
+				if (ClientCookedTargetsList == null)
+				{
+					ClientCookedTargetsList = new ParamList<string>();
+				}
+				if (ServerCookedTargetsList == null)
+				{
+					ServerCookedTargetsList = new ParamList<string>();
+				}
+				if (ProgramTargetsList == null)
+				{
+					ProgramTargetsList = new ParamList<string>();
+				}
+
+				// Add them to the appropriate lists
+				bool bHasGameTarget = false;
+				foreach (SingleTargetProperties Target in Targets)
+				{
+					if (Target.Rules.Type == TargetType.Game)
+					{
+						ClientCookedTargetsList.Add(Target.TargetName);
+						bHasGameTarget = true;
+					}
+					else if (Target.Rules.Type == TargetType.Client)
+					{
+						ClientCookedTargetsList.Add(Target.TargetName);
+						Client = true;
+					}
+					else if (Target.Rules.Type == TargetType.Server)
+					{
+						ServerCookedTargetsList.Add(Target.TargetName);
+						DedicatedServer = true;
+					}
+					else
+					{
+						ProgramTargetsList.Add(Target.TargetName);
+						ProjectType = TargetType.Program;
+					}
+				}
+
+				// Validate all the settings
+				if (Client && bHasGameTarget)
+				{
+					throw new AutomationException("Cannot mix game and client targets");
+				}
+
+				// Find the editor target name
+				List<SingleTargetProperties> EditorTargets = Properties.Targets.Where(x => x.Rules.Type == TargetType.Editor).ToList();
+				if (EditorTargets.Count == 1)
+				{
+					EditorTarget = EditorTargets[0].TargetName;
+				}
+				else if (EditorTargets.Count > 1)
+				{
+					throw new AutomationException("There can be only one Editor target per project.");
+				}
 			}
 			else if (!CommandUtils.IsNullOrEmpty(Properties.Targets))
 			{
@@ -2069,7 +2130,7 @@ namespace AutomationTool
 
 				if (Client)
 				{
-					GameTarget = ChooseTarget(AvailableClientTargets, TargetType.Client, SpecifiedClientTarget);
+					GameTarget = ChooseTarget(AvailableClientTargets, TargetType.Client);
 					ProjectType = TargetType.Client;
 				}
 				else if (AvailableGameTargets.Count > 0)
@@ -2094,7 +2155,7 @@ namespace AutomationTool
 
 				if (AvailableServerTargets.Count > 0 && (DedicatedServer || Cook || CookOnTheFly)) // only if server is needed
 				{
-					ServerTarget = ChooseTarget(AvailableServerTargets, TargetType.Server, SpecifiedServerTarget);
+					ServerTarget = ChooseTarget(AvailableServerTargets, TargetType.Server);
 				}
 			}
 			else if (!CommandUtils.IsNullOrEmpty(Properties.Programs))
@@ -2688,14 +2749,7 @@ namespace AutomationTool
 				CommandUtils.LogLog("SkipBuildEditor={0}", SkipBuildEditor);
 				CommandUtils.LogLog("Cook={0}", Cook);
 				CommandUtils.LogLog("Clean={0}", Clean);
-				if (String.IsNullOrEmpty(SpecifiedClientTarget))
-				{
-					CommandUtils.LogLog("Client={0}", Client);
-				}
-				else
-				{
-					CommandUtils.LogLog("Client='{0}'", SpecifiedClientTarget);
-				}
+				CommandUtils.LogLog("Client={0}", Client);
 				CommandUtils.LogLog("ClientConfigsToBuild={0}", string.Join(",", ClientConfigsToBuild));
 				CommandUtils.LogLog("ClientCookedTargets={0}", ClientCookedTargets.ToString());
 				CommandUtils.LogLog("ClientTargetPlatform={0}", string.Join(",", ClientTargetPlatforms));
@@ -2718,14 +2772,7 @@ namespace AutomationTool
 				CommandUtils.LogLog("DLCPakPluginFile={0}", DLCPakPluginFile);
                 CommandUtils.LogLog("DiffCookedContentPath={0}", DiffCookedContentPath);
                 CommandUtils.LogLog("AdditionalCookerOptions={0}", AdditionalCookerOptions);
-				if (String.IsNullOrEmpty(SpecifiedServerTarget))
-				{
-					CommandUtils.LogLog("DedicatedServer={0}", DedicatedServer);
-				}
-				else
-				{
-					CommandUtils.LogLog("DedicatedServer='{0}'", SpecifiedServerTarget);
-				}
+				CommandUtils.LogLog("DedicatedServer={0}", DedicatedServer);
 				CommandUtils.LogLog("DirectoriesToCook={0}", DirectoriesToCook.ToString());
                 CommandUtils.LogLog("CulturesToCook={0}", CommandUtils.IsNullOrEmpty(CulturesToCook) ? "<Not Specified> (Use Defaults)" : CulturesToCook.ToString());
 				CommandUtils.LogLog("EditorTargets={0}", EditorTargets.ToString());
diff --git a/Engine/Source/Programs/AutomationTool/AutomationUtils/ProjectUtils.cs b/Engine/Source/Programs/AutomationTool/AutomationUtils/ProjectUtils.cs
--- a/Engine/Source/Programs/AutomationTool/AutomationUtils/ProjectUtils.cs
+++ b/Engine/Source/Programs/AutomationTool/AutomationUtils/ProjectUtils.cs
@@ -12,7 +12,7 @@ using System.Reflection;
 
 namespace AutomationTool
 {
-	public struct SingleTargetProperties
+	public class SingleTargetProperties
 	{
 		public string TargetName;
 		public TargetRules Rules;
@@ -568,7 +568,7 @@ namespace AutomationTool
 					TargetRules Rules = Activator.CreateInstance(TargetType, DummyTargetInfo) as TargetRules;
 					CommandUtils.LogVerbose("Adding target: {0} ({1})", TargetType.Name, Rules.Type);
 
-					SingleTargetProperties TargetData;
+					SingleTargetProperties TargetData = new SingleTargetProperties();
 					TargetData.TargetName = GetTargetName(TargetType);
 					TargetData.Rules = Rules;
 					if (Rules.Type == global::UnrealBuildTool.TargetType.Program)
@@ -890,7 +890,7 @@ namespace AutomationTool
 				}
 			}
 
-            SingleTargetProperties Result;
+            SingleTargetProperties Result = new SingleTargetProperties();
             Result.TargetName = ProgramName;
             Result.Rules = null;
             return Result;
