Remove #region blocks.

#rnx

diff --git a/Engine/Source/Programs/AutomationTool/AutomationUtils/Automation.cs b/Engine/Source/Programs/AutomationTool/AutomationUtils/Automation.cs
--- a/Engine/Source/Programs/AutomationTool/AutomationUtils/Automation.cs
+++ b/Engine/Source/Programs/AutomationTool/AutomationUtils/Automation.cs
@@ -11,8 +11,6 @@ using Tools.DotNETCommon;
 
 namespace AutomationTool
 {
-	#region Command info
-
 	/// <summary>
 	/// Command to execute info.
 	/// </summary>
@@ -38,10 +36,6 @@ namespace AutomationTool
 		}
 	}
 
-	#endregion
-
-	#region Command Line helpers
-
 	/// <summary>
 	/// Helper class for automatically registering command line params.
 	/// </summary>
@@ -132,8 +126,6 @@ namespace AutomationTool
 		}
 	}
 
-	#endregion
-
 	[Help(
 @"Executes scripted commands
 
@@ -153,8 +145,6 @@ AutomationTool.exe [-verbose] [-compileonly] [-p4] Command0 [-Arg0 -Arg1 -Arg2 .
     [Help("UseLocalBuildStorage", @"Allows you to use local storage for your root build storage dir (default of P:\Builds (on PC) is changed to Engine\Saved\LocalBuilds). Used for local testing.")]
     public static class Automation
 	{
-		#region Command line parsing
-
 		/// <summary>
 		/// Parses command line parameter.
 		/// </summary>
@@ -417,10 +407,6 @@ AutomationTool.exe [-verbose] [-compileonly] [-p4] Command0 [-Arg0 -Arg1 -Arg2 .
 			}
 		}
 
-		#endregion
-
-		#region Main Program
-
 		/// <summary>
 		/// Compiler with all scripts
 		/// </summary>
@@ -571,10 +557,6 @@ AutomationTool.exe [-verbose] [-compileonly] [-p4] Command0 [-Arg0 -Arg1 -Arg2 .
 			return ExitCode.Success;
 		}
 
-		#endregion
-
-		#region Help
-
 		/// <summary>
 		/// Display help for the specified commands (to execute)
 		/// </summary>
@@ -620,10 +602,6 @@ AutomationTool.exe [-verbose] [-compileonly] [-p4] Command0 [-Arg0 -Arg1 -Arg2 .
 			CommandUtils.LogInformation(Message);
 		}
 
-		#endregion
-
-		#region HelperFunctions
-
 		/// <summary>
 		/// Returns true if AutomationTool is running using installed Engine components
 		/// </summary>
@@ -634,10 +612,6 @@ AutomationTool.exe [-verbose] [-compileonly] [-p4] Command0 [-Arg0 -Arg1 -Arg2 .
 			return CommandUtils.IsEngineInstalled();
 		}
 
-		#endregion
-
-		#region Properties
-
 		/// <summary>
 		/// True if this process is running on a build machine, false if locally.
 		/// </summary>
@@ -678,7 +652,6 @@ AutomationTool.exe [-verbose] [-compileonly] [-p4] Command0 [-Arg0 -Arg1 -Arg2 .
 			}
 		}
 		private static bool? bShouldKillProcesses;
-		#endregion
 	}
 
 }
diff --git a/Engine/Source/Programs/AutomationTool/AutomationUtils/CommandEnvironment.cs b/Engine/Source/Programs/AutomationTool/AutomationUtils/CommandEnvironment.cs
--- a/Engine/Source/Programs/AutomationTool/AutomationUtils/CommandEnvironment.cs
+++ b/Engine/Source/Programs/AutomationTool/AutomationUtils/CommandEnvironment.cs
@@ -50,8 +50,6 @@ namespace AutomationTool
 		/// </summary>
 		public static readonly string KnownFileRelativeToRoot = @"Engine/Config/BaseEngine.ini";
 
-		#region Command Environment properties
-
 		public string LocalRoot { get; protected set; }
 		public string EngineSavedFolder { get; protected set; }
 		public string LogFolder { get; protected set; }
@@ -67,8 +65,6 @@ namespace AutomationTool
 		public string MallocNanoZone { get; protected set; }
 		public bool IsChildInstance { get; protected set; }
 
-		#endregion
-
 		/// <summary>
 		/// Initializes the environment.
 		/// </summary>
diff --git a/Engine/Source/Programs/AutomationTool/AutomationUtils/CommandUtils.cs b/Engine/Source/Programs/AutomationTool/AutomationUtils/CommandUtils.cs
--- a/Engine/Source/Programs/AutomationTool/AutomationUtils/CommandUtils.cs
+++ b/Engine/Source/Programs/AutomationTool/AutomationUtils/CommandUtils.cs
@@ -20,8 +20,6 @@ using System.Xml.Serialization;
 
 namespace AutomationTool
 {
-	#region ParamList
-
 	/// <summary>
 	/// Wrapper around List with support for multi parameter constructor, i.e:
 	///   var Maps = new ParamList<string>("Map1", "Map2");
@@ -54,10 +52,6 @@ namespace AutomationTool
 		}
 	}
 
-	#endregion
-
-	#region PathSeparator
-
 	public enum PathSeparator
 	{
 		Default = 0,
@@ -67,15 +61,11 @@ namespace AutomationTool
 		Local
 	}
 
-	#endregion
-
 	/// <summary>
 	/// Base utility function for script commands.
 	/// </summary>
 	public partial class CommandUtils
 	{
-		#region Environment Setup
-
 		static private CommandEnvironment CmdEnvironment;
 
 		/// <summary>
@@ -104,8 +94,6 @@ namespace AutomationTool
 			CmdEnvironment = new CommandEnvironment();
 		}
 
-		#endregion
-
 		/// <summary>
 		/// Returns true if AutomationTool is running using installed Engine components
 		/// </summary>
@@ -121,8 +109,6 @@ namespace AutomationTool
 
 		static private bool? bIsEngineInstalled;
 
-		#region Logging
-
 		/// <summary>
 		/// Writes formatted text to log (with LogEventType.Console).
 		/// </summary>
@@ -283,10 +269,6 @@ namespace AutomationTool
             Tools.DotNETCommon.Log.WriteLine(1, Verbosity, LogUtils.FormatException(Ex));
 		}
 
-		#endregion
-
-		#region Progress Logging
-
 		public static void LogPushProgress(bool bShowProgress, int Numerator, int Denominator)
 		{
 			if(bShowProgress)
@@ -327,10 +309,6 @@ namespace AutomationTool
 			}
 		}
 
-		#endregion
-
-		#region IO
-
 		/// <summary>
 		/// Finds files in specified paths. 
 		/// </summary>
@@ -1619,10 +1597,6 @@ namespace AutomationTool
 			CloneDirectoryRecursiveWorker(SourcePath, TargetPath, ClonedFiles, bIncremental: true);
 		}
 
-		#endregion
-
-		#region Threaded Copy
-
         /// <summary>
 		/// Copies files using multiple threads
 		/// </summary>
@@ -1743,10 +1717,6 @@ namespace AutomationTool
 			}
 		}
 
-		#endregion
-
-		#region Environment variables
-
 		/// <summary>
 		/// Gets environment variable value.
 		/// </summary>
@@ -1800,10 +1770,6 @@ namespace AutomationTool
 			}
 		}
 
-		#endregion
-
-		#region CommandLine
-
 		/// <summary>
 		/// Converts a list of arguments to a string where each argument is separated with a space character.
 		/// </summary>
@@ -1934,10 +1900,6 @@ namespace AutomationTool
 			return UnrealBuildTool.Utils.MakePathSafeToUseWithCommandLine(InPath);
 		}
 
-		#endregion
-
-		#region Other
-
 		public static string EscapePath(string InPath)
 		{
 			return InPath.Replace(":", "").Replace("/", "+").Replace("\\", "+").Replace(" ", "+");
@@ -1953,10 +1915,6 @@ namespace AutomationTool
 			return Collection == null || Collection.Count == 0;
 		}
 
-	    #endregion
-
-		#region Properties
-
 		/// <summary>
 		/// Checks if this command is running on a build machine.
 		/// </summary>
@@ -2010,8 +1968,6 @@ namespace AutomationTool
 		/// </summary>
 		public static TelemetryData Telemetry = new TelemetryData();
 
-		#endregion
-
         /// <summary>
         /// Cached location of the build root storage because the logic to compute it is a little non-trivial.
         /// </summary>
diff --git a/Engine/Source/Programs/AutomationTool/AutomationUtils/HelpUtils.cs b/Engine/Source/Programs/AutomationTool/AutomationUtils/HelpUtils.cs
--- a/Engine/Source/Programs/AutomationTool/AutomationUtils/HelpUtils.cs
+++ b/Engine/Source/Programs/AutomationTool/AutomationUtils/HelpUtils.cs
@@ -1,4 +1,4 @@
-﻿// Copyright 1998-2019 Epic Games, Inc. All Rights Reserved.
+// Copyright 1998-2019 Epic Games, Inc. All Rights Reserved.
 using System;
 using System.Collections;
 using System.Collections.Generic;
@@ -16,8 +16,6 @@ using System.Reflection;
 
 namespace AutomationTool
 {
-	#region Help
-
 	/// <summary>
 	/// Help Attribute.
 	/// </summary>
@@ -74,15 +72,11 @@ namespace AutomationTool
 		}
 	}
 
-	#endregion
-
 	/// <summary>
 	/// Base utility function for script commands.
 	/// </summary>
 	public partial class CommandUtils
 	{
-		#region Help
-
 		/// <summary>
 		/// Displays help for the specified command.
 		/// <param name="Command">Command type.</param>
@@ -377,7 +371,5 @@ namespace AutomationTool
 			}
 			return Members;
 		}
-
-		#endregion
 	}
 }
diff --git a/Engine/Source/Programs/AutomationTool/AutomationUtils/P4Utils.cs b/Engine/Source/Programs/AutomationTool/AutomationUtils/P4Utils.cs
--- a/Engine/Source/Programs/AutomationTool/AutomationUtils/P4Utils.cs
+++ b/Engine/Source/Programs/AutomationTool/AutomationUtils/P4Utils.cs
@@ -930,8 +930,6 @@ namespace AutomationTool
 
 	public partial class CommandUtils
 	{
-		#region Environment Setup
-
 		static private P4Connection PerforceConnection;
 		static private P4Environment PerforceEnvironment;
 
@@ -985,8 +983,6 @@ namespace AutomationTool
 			PerforceConnection = new P4Connection(User: P4Env.User, Client: P4Env.Client, ServerAndPort: P4Env.ServerAndPort);
 		}
 
-		#endregion
-
 		/// <summary>
 		/// Check if P4 is supported.
 		/// </summary>		
@@ -3966,8 +3962,6 @@ namespace AutomationTool
 			return null;
 		}
 
-		#region Utilities
-
 		private static object[] OldStyleBinaryFlags = new object[]
 		{
 			P4FileAttributes.Uncompressed,
@@ -4093,7 +4087,5 @@ namespace AutomationTool
 			}
 			return Text;
 		}
-
-		#endregion
 	}
 }
\ No newline at end of file
diff --git a/Engine/Source/Programs/AutomationTool/AutomationUtils/Platform.cs b/Engine/Source/Programs/AutomationTool/AutomationUtils/Platform.cs
--- a/Engine/Source/Programs/AutomationTool/AutomationUtils/Platform.cs
+++ b/Engine/Source/Programs/AutomationTool/AutomationUtils/Platform.cs
@@ -41,8 +41,6 @@ namespace AutomationTool
 	/// </summary>
 	public class Platform : CommandUtils
 	{
-		#region Intialization
-
 		private static Dictionary<TargetPlatformDescriptor, Platform> AllPlatforms = new Dictionary<TargetPlatformDescriptor, Platform>();
 		internal static void InitializePlatforms(Assembly[] AssembliesWithPlatforms = null)
 		{
@@ -120,8 +118,6 @@ namespace AutomationTool
 			}
 		}
 
-		#endregion
-
 		protected UnrealTargetPlatform TargetPlatformType;
 		protected UnrealTargetPlatform TargetIniPlatformType;
 
@@ -559,8 +555,6 @@ namespace AutomationTool
 			get { return true; }
 		}
 
-		#region Hooks
-
 		public virtual void PreBuildAgenda(UE4Build Build, UE4Build.BuildAgenda Agenda, ProjectParams Params)
 		{
 
@@ -617,10 +611,6 @@ namespace AutomationTool
 			return new UnrealTargetPlatform[] { PlatformType };
 		}
 
-		#endregion
-
-		#region Utilities
-
 		// let the platform set the exe extension if it chooses (otherwise, use
 		// the switch statement in GetExeExtension below)
 		protected virtual string GetPlatformExeExtension()
@@ -714,7 +704,5 @@ namespace AutomationTool
 
 			return ValidPlatforms;
 		}
-
-		#endregion
 	}
 }
diff --git a/Engine/Source/Programs/AutomationTool/AutomationUtils/ProcessUtils.cs b/Engine/Source/Programs/AutomationTool/AutomationUtils/ProcessUtils.cs
--- a/Engine/Source/Programs/AutomationTool/AutomationUtils/ProcessUtils.cs
+++ b/Engine/Source/Programs/AutomationTool/AutomationUtils/ProcessUtils.cs
@@ -206,8 +206,6 @@ namespace AutomationTool
 		}
 	}
 
-	#region ProcessResult Helper Class
-
 	public interface IProcessResult : IProcess
 	{
 		void OnProcessExited();
@@ -693,12 +691,8 @@ namespace AutomationTool
 		}
 	}
 
-	#endregion
-
 	public partial class CommandUtils
 	{
-		#region Statistics
-
 		private static Dictionary<string, int> ExeToTimeInMs = new Dictionary<string, int>();
 
 		public static void AddRunTime(string Exe, int TimeInMs)
@@ -729,8 +723,6 @@ namespace AutomationTool
 			}
 		}
 
-		#endregion
-
 		[Flags]
 		public enum ERunOptions
 		{
diff --git a/Engine/Source/Programs/AutomationTool/AutomationUtils/ProjectParams.cs b/Engine/Source/Programs/AutomationTool/AutomationUtils/ProjectParams.cs
--- a/Engine/Source/Programs/AutomationTool/AutomationUtils/ProjectParams.cs
+++ b/Engine/Source/Programs/AutomationTool/AutomationUtils/ProjectParams.cs
@@ -14,8 +14,6 @@ namespace AutomationTool
 	[Help("servertargetplatform=PlatformName", "target platform for building, cooking and deployment of the dedicated server (also -ServerPlatform)")]
 	public class ProjectParams
 	{
-		#region Constructors
-
 		/// <summary>
 		/// Gets a parameter from the command line if it hasn't been specified in the constructor. 
 		/// If the command line is not available, default value will be used.
@@ -1102,10 +1100,6 @@ namespace AutomationTool
 			return ConfigValue;
 		}
 
-		#endregion
-
-		#region Shared
-
 		/// <summary>
 		/// Shared: Full path to the .uproject file
 		/// </summary>
@@ -1349,10 +1343,6 @@ namespace AutomationTool
 		/// </summary>
 		public List<string> ConfigOverrideParams = new List<string>();
 
-        #endregion
-
-        #region Build
-
         /// <summary>
         /// Build: True if build step should be executed, command: -build
         /// </summary>
@@ -1462,10 +1452,6 @@ namespace AutomationTool
 		/// </summary>
 		private List<string> TargetNames;
 
-		#endregion
-
-		#region Cook
-
 		/// <summary>
 		/// Cook: List of maps to cook.
 		/// </summary>
@@ -1662,10 +1648,6 @@ namespace AutomationTool
 		[Help("IgnoreCookErrors", "Ignores cook errors and continues with packaging etc")]
 		public bool IgnoreCookErrors { private set; get; }
 
-		#endregion
-
-		#region Stage
-
 		/// <summary>
 		/// Stage: Commandline: -nodebuginfo
 		/// </summary>
@@ -1749,10 +1731,6 @@ namespace AutomationTool
 		/// </summary>
 		public bool bUseExtraFlavor = false;
 
-		#endregion
-
-		#region Run
-
 		/// <summary>
 		/// Run: True if the Run step should be executed, command: -run
 		/// </summary>
@@ -1918,10 +1896,6 @@ namespace AutomationTool
         /// </summary>
         public string ServerDeviceAddress;
 
-        #endregion
-
-		#region Package
-
 		[Help("package", "package the project for the target platform")]
 		public bool Package { get; set; }
 
@@ -1952,27 +1926,15 @@ namespace AutomationTool
 		[Help("AdditionalPackageOptions", "extra options to pass to the platform's packager")]
 		public string AdditionalPackageOptions { get; set; }
 
-#endregion
-
-#region Deploy
-
-[Help("deploy", "deploy the project for the target platform")]
+		[Help("deploy", "deploy the project for the target platform")]
 		public bool Deploy { get; set; }
 
 		[Help("deploy", "Location to deploy to on the target platform")]
 		public string DeployFolder { get; set; }
 
-		#endregion
-
-		#region GetFile
-
 		[Help("getfile", "download file from target after successful run")]
 		public string GetFile { get; set; }
 
-		#endregion
-
-		#region Misc
-
 		[Help("MapsToRebuildLightMaps", "List of maps that need light maps rebuilding")]
 		public ParamList<string> MapsToRebuildLightMaps = new ParamList<string>();
 
@@ -1982,10 +1944,6 @@ namespace AutomationTool
         [Help("IgnoreLightMapErrors", "Whether Light Map errors should be treated as critical")]
 		public bool IgnoreLightMapErrors { get; set; }
 
-		#endregion
-
-		#region Initialization
-
 		private List<SingleTargetProperties> DetectedTargets;
 		private Dictionary<UnrealTargetPlatform, ConfigHierarchy> LoadedEngineConfigs;
 		private Dictionary<UnrealTargetPlatform, ConfigHierarchy> LoadedGameConfigs;
@@ -2271,10 +2229,6 @@ namespace AutomationTool
 			}
 		}
 
-		#endregion
-
-		#region Utilities
-
 		public bool HasEditorTargets
 		{
 			get { return !CommandUtils.IsNullOrEmpty(EditorTargets); }
@@ -2844,7 +2798,5 @@ namespace AutomationTool
 
 			Validate();
 		}
-
-		#endregion
 	}
 }
diff --git a/Engine/Source/Programs/AutomationTool/AutomationUtils/ScriptCompiler.cs b/Engine/Source/Programs/AutomationTool/AutomationUtils/ScriptCompiler.cs
--- a/Engine/Source/Programs/AutomationTool/AutomationUtils/ScriptCompiler.cs
+++ b/Engine/Source/Programs/AutomationTool/AutomationUtils/ScriptCompiler.cs
@@ -1,4 +1,4 @@
-﻿// Copyright 1998-2019 Epic Games, Inc. All Rights Reserved.
+// Copyright 1998-2019 Epic Games, Inc. All Rights Reserved.
 using System;
 using System.Collections.Generic;
 using System.Linq;
@@ -31,8 +31,6 @@ namespace AutomationTool
 	/// </summary>
 	public class ScriptCompiler
 	{
-		#region Fields
-				
 		private Dictionary<string, Type> ScriptCommands;
 #if DEBUG
 		const string BuildConfig = "Debug";
@@ -41,10 +39,6 @@ namespace AutomationTool
 #endif
 		const string DefaultScriptsDLLName = "AutomationScripts.Automation.dll";
 
-		#endregion
-
-		#region Compilation
-
 		public ScriptCompiler()
 		{
 		}
@@ -491,15 +485,9 @@ namespace AutomationTool
 			return CommandUtils.CombinePaths(CommandUtils.CmdEnv.LocalRoot, "Engine", "Binaries", "DotNET", "AutomationScripts");
 		}
 
-		#endregion
-
-		#region Properties
-
 		public Dictionary<string, Type> Commands
 		{
 			get { return ScriptCommands; }
 		}
-
-		#endregion
 	}
 }
diff --git a/Engine/Source/Programs/AutomationTool/AutomationUtils/Utils.cs b/Engine/Source/Programs/AutomationTool/AutomationUtils/Utils.cs
--- a/Engine/Source/Programs/AutomationTool/AutomationUtils/Utils.cs
+++ b/Engine/Source/Programs/AutomationTool/AutomationUtils/Utils.cs
@@ -13,8 +13,6 @@ using Tools.DotNETCommon;
 
 namespace AutomationTool
 {
-	#region UAT Internal Utils
-
 	/// <summary>
 	/// AutomationTool internal Utilities.
 	/// </summary>
@@ -870,11 +868,6 @@ namespace AutomationTool
 
 	}
 
-	#endregion
-
-
-	#region VersionFileReader
-
 
 	/// <summary>
 	/// This is to ensure that UAT can produce version strings precisely compatible
@@ -1046,10 +1039,6 @@ namespace AutomationTool
         }
     }
 
-    #endregion
-
-    #region VersionFileUpdater
-
     /// <summary>
 	/// VersionFileUpdater.
 	/// </summary>
@@ -1250,6 +1239,4 @@ namespace AutomationTool
 		/// </summary>
 		protected List<string> OriginalLines;
 	}
-
-	#endregion
 }
diff --git a/Engine/Source/Programs/AutomationTool/Scripts/ArchiveCommand.Automation.cs b/Engine/Source/Programs/AutomationTool/Scripts/ArchiveCommand.Automation.cs
--- a/Engine/Source/Programs/AutomationTool/Scripts/ArchiveCommand.Automation.cs
+++ b/Engine/Source/Programs/AutomationTool/Scripts/ArchiveCommand.Automation.cs
@@ -1,4 +1,4 @@
-﻿// Copyright 1998-2019 Epic Games, Inc. All Rights Reserved.
+// Copyright 1998-2019 Epic Games, Inc. All Rights Reserved.
 using System;
 using System.Collections.Generic;
 using System.IO;
@@ -10,9 +10,6 @@ using Tools.DotNETCommon;
 
 public partial class Project : CommandUtils
 {
-
-	#region Utilities
-
 	public static void CreateArchiveManifest(ProjectParams Params, DeploymentContext SC)
 	{
 		if (!Params.Archive)
@@ -46,10 +43,6 @@ public partial class Project : CommandUtils
 		}
 	}
 
-	#endregion
-
-	#region Archive Command
-
 	public static void Archive(ProjectParams Params)
 	{
 		Params.ValidateAndLog();
@@ -86,6 +79,4 @@ public partial class Project : CommandUtils
 		}
 		LogInformation("********** ARCHIVE COMMAND COMPLETED **********");
     }
-
-	#endregion
 }
diff --git a/Engine/Source/Programs/AutomationTool/Scripts/BuildCookRun.Automation.cs b/Engine/Source/Programs/AutomationTool/Scripts/BuildCookRun.Automation.cs
--- a/Engine/Source/Programs/AutomationTool/Scripts/BuildCookRun.Automation.cs
+++ b/Engine/Source/Programs/AutomationTool/Scripts/BuildCookRun.Automation.cs
@@ -1,4 +1,4 @@
-﻿// Copyright 1998-2019 Epic Games, Inc. All Rights Reserved.
+// Copyright 1998-2019 Epic Games, Inc. All Rights Reserved.
 using System;
 using System.Collections.Generic;
 using System.IO;
@@ -21,8 +21,6 @@ If no DefaultMap can be found, the command falls back to /Engine/Maps/Entry.")]
 [Help(typeof(CodeSign))]
 public class BuildCookRun : BuildCommand
 {
-	#region BaseCommand interface
-
 	public override void ExecuteBuild()
 	{
 		// these need to be done first
@@ -41,10 +39,6 @@ public class BuildCookRun : BuildCommand
 		DoBuildCookRun(Params);
 	}
 
-	#endregion
-
-	#region Setup
-
 	protected ProjectParams SetupParams()
 	{
 		LogInformation("Setting up ProjectParams for {0}", ProjectPath);
@@ -191,10 +185,6 @@ public class BuildCookRun : BuildCommand
 		return DefaultMap;
 	}
 
-	#endregion
-
-	#region BuildCookRun
-
 	protected void DoBuildCookRun(ProjectParams Params)
 	{
 		const ProjectBuildTargets ClientTargets = ProjectBuildTargets.ClientCooked | ProjectBuildTargets.ServerCooked;
@@ -346,6 +336,4 @@ public class BuildCookRun : BuildCommand
 			return ProjectFullPath;
 		}
 	}
-
-	#endregion
 }
diff --git a/Engine/Source/Programs/AutomationTool/Scripts/BuildProjectCommand.Automation.cs b/Engine/Source/Programs/AutomationTool/Scripts/BuildProjectCommand.Automation.cs
--- a/Engine/Source/Programs/AutomationTool/Scripts/BuildProjectCommand.Automation.cs
+++ b/Engine/Source/Programs/AutomationTool/Scripts/BuildProjectCommand.Automation.cs
@@ -37,8 +37,6 @@ public enum ProjectBuildTargets
 /// </remarks>
 public partial class Project : CommandUtils
 {
-	#region Build Command
-
 	/// <summary>
 	/// PlatformSupportsCrashReporter
 	/// </summary>
@@ -232,6 +230,4 @@ public partial class Project : CommandUtils
 
 		LogInformation("********** BUILD COMMAND COMPLETED **********");
 	}
-
-	#endregion
 }
diff --git a/Engine/Source/Programs/AutomationTool/Scripts/CookCommand.Automation.cs b/Engine/Source/Programs/AutomationTool/Scripts/CookCommand.Automation.cs
--- a/Engine/Source/Programs/AutomationTool/Scripts/CookCommand.Automation.cs
+++ b/Engine/Source/Programs/AutomationTool/Scripts/CookCommand.Automation.cs
@@ -21,8 +21,6 @@ using Tools.DotNETCommon;
 /// </remarks>
 public partial class Project : CommandUtils
 {
-    #region Cook Command
-
     public static void Cook(ProjectParams Params)
 	{
 		if ((!Params.Cook && !(Params.CookOnTheFly && !Params.SkipServer)) || Params.SkipCook)
@@ -595,6 +593,4 @@ public partial class Project : CommandUtils
 			DeleteDirectory(bQuiet, CleanDir);
 		}
 	}
-
-	#endregion
 }
diff --git a/Engine/Source/Programs/AutomationTool/Scripts/CopyBuildToStagingDirectory.Automation.cs b/Engine/Source/Programs/AutomationTool/Scripts/CopyBuildToStagingDirectory.Automation.cs
--- a/Engine/Source/Programs/AutomationTool/Scripts/CopyBuildToStagingDirectory.Automation.cs
+++ b/Engine/Source/Programs/AutomationTool/Scripts/CopyBuildToStagingDirectory.Automation.cs
@@ -24,8 +24,6 @@ using System.Diagnostics;
 public partial class Project : CommandUtils
 {
 
-	#region Utilities
-
 	private static readonly object SyncLock = new object();
 
 	/// <returns>The path for the BuildPatchTool executable depending on host platform.</returns>
@@ -3167,10 +3165,6 @@ public partial class Project : CommandUtils
 		FileReference.WriteAllLines(ManifestFile, DeltaFiles);
 	}
 
-	#endregion
-
-	#region Stage Command
-
 	//@todo move this
 	public static List<DeploymentContext> CreateDeploymentContext(ProjectParams Params, bool InDedicatedServer, bool DoCleanStage = false)
 	{
@@ -3395,6 +3389,4 @@ public partial class Project : CommandUtils
 			LogInformation("********** STAGE COMMAND COMPLETED **********");
 		}
 	}
-
-	#endregion
 }
diff --git a/Engine/Source/Programs/AutomationTool/Scripts/DeployCommand.Automation.cs b/Engine/Source/Programs/AutomationTool/Scripts/DeployCommand.Automation.cs
--- a/Engine/Source/Programs/AutomationTool/Scripts/DeployCommand.Automation.cs
+++ b/Engine/Source/Programs/AutomationTool/Scripts/DeployCommand.Automation.cs
@@ -1,4 +1,4 @@
-﻿// Copyright 1998-2019 Epic Games, Inc. All Rights Reserved.
+// Copyright 1998-2019 Epic Games, Inc. All Rights Reserved.
 using System;
 using System.Collections.Generic;
 using System.IO;
@@ -9,8 +9,6 @@ using UnrealBuildTool;
 
 public partial class Project : CommandUtils
 {
-	#region Deploy Command
-
 	public static void Deploy(ProjectParams Params)
 	{
 		Params.ValidateAndLog();
@@ -70,6 +68,4 @@ public partial class Project : CommandUtils
 
 		LogInformation("********** DEPLOY COMMAND COMPLETED **********");
 	}
-
-	#endregion
 }
diff --git a/Engine/Source/Programs/AutomationTool/Scripts/GenerateDSYM.Automation.cs b/Engine/Source/Programs/AutomationTool/Scripts/GenerateDSYM.Automation.cs
--- a/Engine/Source/Programs/AutomationTool/Scripts/GenerateDSYM.Automation.cs
+++ b/Engine/Source/Programs/AutomationTool/Scripts/GenerateDSYM.Automation.cs
@@ -12,8 +12,6 @@ using UnrealBuildTool;
 [Help("config=Configuration", @"Project configuration (required), i.e: -config=Development")]
 public class GenerateDSYM : BuildCommand
 {
-	#region BaseCommand interface
-
 	public override void ExecuteBuild()
 	{
 		var ProjectName = ParseParamValue("project");
@@ -23,6 +21,4 @@ public class GenerateDSYM : BuildCommand
 
 		RunAndLog(CmdEnv, IPPExe, "RPC " + ProjectName + " -config " + Config + " GenDSYM");
 	}
-
-	#endregion
 }
\ No newline at end of file
diff --git a/Engine/Source/Programs/AutomationTool/Scripts/GetFileCommand.Automation.cs b/Engine/Source/Programs/AutomationTool/Scripts/GetFileCommand.Automation.cs
--- a/Engine/Source/Programs/AutomationTool/Scripts/GetFileCommand.Automation.cs
+++ b/Engine/Source/Programs/AutomationTool/Scripts/GetFileCommand.Automation.cs
@@ -11,8 +11,6 @@ using Tools.DotNETCommon;
 
 public partial class Project : CommandUtils
 {
-	#region GetFile Command
-
 	public static void GetFile(ProjectParams Params)
 	{
 		Params.ValidateAndLog();
@@ -36,6 +34,4 @@ public partial class Project : CommandUtils
 
 		LogInformation("********** GETFILE COMMAND COMPLETED **********");
 	}
-
-	#endregion
 }
diff --git a/Engine/Source/Programs/AutomationTool/Scripts/PackageCommand.Automation.cs b/Engine/Source/Programs/AutomationTool/Scripts/PackageCommand.Automation.cs
--- a/Engine/Source/Programs/AutomationTool/Scripts/PackageCommand.Automation.cs
+++ b/Engine/Source/Programs/AutomationTool/Scripts/PackageCommand.Automation.cs
@@ -1,4 +1,4 @@
-﻿// Copyright 1998-2019 Epic Games, Inc. All Rights Reserved.
+// Copyright 1998-2019 Epic Games, Inc. All Rights Reserved.
 using System;
 using System.Collections.Generic;
 using System.IO;
@@ -9,8 +9,6 @@ using UnrealBuildTool;
 
 public partial class Project : CommandUtils
 {
-	#region Package Command
-
 	public static void Package(ProjectParams Params, int WorkingCL=-1)
 	{
 		if (!Params.SkipStage || Params.Package)
@@ -42,6 +40,4 @@ public partial class Project : CommandUtils
 			}
 		}
 	}
-
-	#endregion
 }
diff --git a/Engine/Source/Programs/AutomationTool/Scripts/RebuildHLODCommand.Automation.cs b/Engine/Source/Programs/AutomationTool/Scripts/RebuildHLODCommand.Automation.cs
--- a/Engine/Source/Programs/AutomationTool/Scripts/RebuildHLODCommand.Automation.cs
+++ b/Engine/Source/Programs/AutomationTool/Scripts/RebuildHLODCommand.Automation.cs
@@ -28,9 +28,6 @@ namespace AutomationScripts.Automation
 	[RequireP4]
 	public class RebuildHLOD : BuildCommand
 	{
-		// The rebuild hierarchical LODs process
-		#region RebuildHLOD Command
-
 		public override void ExecuteBuild()
 		{
 			LogInformation("********** REBUILD HLODS COMMAND STARTED **********");
@@ -85,11 +82,6 @@ namespace AutomationScripts.Automation
 			LogInformation("********** REBUILD HLODS COMMAND COMPLETED **********");
 		}
 
-		#endregion
-
-		// Broken down steps used to run the process.
-		#region RebuildHLOD Process Steps
-
 		private void BuildNecessaryTargets()
 		{
 			LogInformation("Running Step:- RebuildHLOD::BuildNecessaryTargets");
@@ -215,11 +207,6 @@ namespace AutomationScripts.Automation
 			}
 		}
 
-		#endregion
-
-		// Helper functions and procedure steps necessary for running the commandlet successfully
-		#region RebuildHLOD Helper Functions
-
 		/**
 		 * Parse the P4 output for any errors that we really care about.
 		 * e.g. umaps and assets are exclusive checkout files, if we cant check out a map for this reason
@@ -343,11 +330,6 @@ namespace AutomationScripts.Automation
 		    }
         }
 
-		#endregion
-
-		// Member vars used in multiple steps.
-		#region RebuildHLOD Property Set-up
-
 		// Users to notify if the process fails or succeeds.
 		List<string> StakeholdersEmailAddresses
 		{
@@ -438,7 +420,5 @@ namespace AutomationScripts.Automation
 				return ProjectFullPath;
 			}
 		}
-
-		#endregion
 	}
 }
diff --git a/Engine/Source/Programs/AutomationTool/Scripts/RebuildLightMapsCommand.Automation.cs b/Engine/Source/Programs/AutomationTool/Scripts/RebuildLightMapsCommand.Automation.cs
--- a/Engine/Source/Programs/AutomationTool/Scripts/RebuildLightMapsCommand.Automation.cs
+++ b/Engine/Source/Programs/AutomationTool/Scripts/RebuildLightMapsCommand.Automation.cs
@@ -23,9 +23,6 @@ namespace AutomationScripts.Automation
 	[Help("StakeholdersEmailAddresses", "Users to notify of completion")]
 	public class RebuildLightMaps : BuildCommand
 	{
-		// The rebuild lighting process
-		#region RebuildLightMaps Command
-
 		public override void ExecuteBuild()
 		{
 			LogInformation("********** REBUILD LIGHT MAPS COMMAND STARTED **********");
@@ -66,11 +63,6 @@ namespace AutomationScripts.Automation
 			LogInformation("********** REBUILD LIGHT MAPS COMMAND COMPLETED **********");
 		}
 
-		#endregion
-
-		// Broken down steps used to run the process.
-		#region RebuildLightMaps Process Steps
-
 		private void BuildNecessaryTargets()
 		{
 			LogInformation("Running Step:- RebuildLightMaps::BuildNecessaryTargets");
@@ -187,11 +179,6 @@ namespace AutomationScripts.Automation
 			}
 		}
 
-		#endregion
-
-		// Helper functions and procedure steps necessary for running the commandlet successfully
-		#region RebuildLightMaps Helper Functions
-
 		/**
 		 * Parse the P4 output for any errors that we really care about.
 		 * e.g. umaps and assets are exclusive checkout files, if we cant check out a map for this reason
@@ -303,11 +290,6 @@ namespace AutomationScripts.Automation
 			}
 		}
 
-		#endregion
-
-		// Member vars used in multiple steps.
-		#region RebuildLightMaps Property Set-up
-
 		// Users to notify if the process fails or succeeds.
 		List<String> StakeholdersEmailAddresses
 		{
@@ -377,7 +359,5 @@ namespace AutomationScripts.Automation
 				return ProjectFullPath;
 			}
 		}
-
-		#endregion
 	}
 }
diff --git a/Engine/Source/Programs/AutomationTool/Scripts/ResavePackagesCommand.Automation.cs b/Engine/Source/Programs/AutomationTool/Scripts/ResavePackagesCommand.Automation.cs
--- a/Engine/Source/Programs/AutomationTool/Scripts/ResavePackagesCommand.Automation.cs
+++ b/Engine/Source/Programs/AutomationTool/Scripts/ResavePackagesCommand.Automation.cs
@@ -25,9 +25,6 @@ namespace AutomationScripts.Automation
 	[RequireP4]
 	public class ResavePackages : BuildCommand
 	{
-        // The rebuild lighting process
-        #region ResavePackages Command
-
         public override void ExecuteBuild()
 		{
 			LogInformation("********** RESAVE PACKAGES COMMAND STARTED **********");
@@ -68,11 +65,6 @@ namespace AutomationScripts.Automation
 			LogInformation("********** RESAVE PACKAGES COMMAND COMPLETED **********");
 		}
 
-		#endregion
-
-		// Broken down steps used to run the process.
-		#region RebuildLightMaps Process Steps
-
 		private void BuildNecessaryTargets()
 		{
 			LogInformation("Running Step:- RebuildLightMaps::BuildNecessaryTargets");
@@ -190,11 +182,6 @@ namespace AutomationScripts.Automation
 			}
 		}
 
-        #endregion
-
-        // Helper functions and procedure steps necessary for running the commandlet successfully
-        #region ResavePackages Helper Functions
-
         /**
 		 * Parse the P4 output for any errors that we really care about.
 		 * e.g. umaps and assets are exclusive checkout files, if we cant check out a map for this reason
@@ -298,11 +285,6 @@ namespace AutomationScripts.Automation
 			}
 		}
 
-		#endregion
-
-		// Member vars used in multiple steps.
-		#region RebuildLightMaps Property Set-up
-
 		// Users to notify if the process fails or succeeds.
 		List<String> StakeholdersEmailAddresses
 		{
@@ -372,7 +354,5 @@ namespace AutomationScripts.Automation
 				return ProjectFullPath;
 			}
 		}
-
-		#endregion
 	}
 }
diff --git a/Engine/Source/Programs/AutomationTool/Scripts/RunProjectCommand.Automation.cs b/Engine/Source/Programs/AutomationTool/Scripts/RunProjectCommand.Automation.cs
--- a/Engine/Source/Programs/AutomationTool/Scripts/RunProjectCommand.Automation.cs
+++ b/Engine/Source/Programs/AutomationTool/Scripts/RunProjectCommand.Automation.cs
@@ -1,4 +1,4 @@
-﻿// Copyright 1998-2019 Epic Games, Inc. All Rights Reserved.
+// Copyright 1998-2019 Epic Games, Inc. All Rights Reserved.
 using System;
 using System.Collections.Generic;
 using System.IO;
@@ -26,8 +26,6 @@ using Tools.DotNETCommon;
 /// </remarks>
 public partial class Project : CommandUtils
 {
-	#region Fields
-
 	/// <summary>
 	/// Thread used to read client log file.
 	/// </summary>
@@ -38,10 +36,6 @@ public partial class Project : CommandUtils
 	/// </summary>
 	public static IProcessResult ServerProcess;
 
-	#endregion
-
-	#region Run Command
-
 	// debug commands for the engine to crash
 	public static string[] CrashCommands = 
     {
@@ -184,10 +178,6 @@ public partial class Project : CommandUtils
 		}
 	}
 
-	#endregion
-
-	#region Client
-
 	private static void RunStandaloneClient(List<DeploymentContext> DeployContextList, string ClientLogFile, ERunOptions ClientRunFlags, string ClientApp, string ClientCmdLine, ProjectParams Params)
 	{
 		if (Params.Unattended)
@@ -910,10 +900,6 @@ public partial class Project : CommandUtils
 		}
 	}
 
-	#endregion
-
-	#region Client Thread
-
 	private static void ClientLogReaderProc(object ArgsContainer)
 	{
 		var Args = ArgsContainer as object[];
@@ -929,10 +915,6 @@ public partial class Project : CommandUtils
 		});
 	}
 
-	#endregion
-
-	#region Servers
-
 	private static IProcessResult RunDedicatedServer(ProjectParams Params, string ServerLogFile, string AdditionalCommandLine)
 	{
 		ProjectParams ServerParams = new ProjectParams(Params);
@@ -1089,6 +1071,4 @@ public partial class Project : CommandUtils
 		PopDir();
 		return Result;
 	}
-
-	#endregion
 }
diff --git a/Engine/Source/Programs/AutomationTool/Scripts/Tests.Automation.cs b/Engine/Source/Programs/AutomationTool/Scripts/Tests.Automation.cs
--- a/Engine/Source/Programs/AutomationTool/Scripts/Tests.Automation.cs
+++ b/Engine/Source/Programs/AutomationTool/Scripts/Tests.Automation.cs
@@ -852,7 +852,6 @@ class TestGamePerf : BuildCommand
 
 		string LevelParam = ParseParamValue("level", "");
 
-		#region Arguments
 		switch (LevelParam)
 		{
 			case "PerfShooterGame_Santuary":
@@ -1012,7 +1011,6 @@ class TestGamePerf : BuildCommand
 			default:
 				break;
 		}
-		#endregion
 
 		try
 		{
@@ -1172,7 +1170,6 @@ class TestOSSCommands : BuildCommand
 
 		StreamWriter WriteErrors = new StreamWriter(new FileStream(ClientLogFile, FileMode.Append, FileAccess.Write, FileShare.ReadWrite));
 
-		# region Error List
 		/* 
          * online sub=amazon test friends - Warning: Failed to get friends interface for amazon
          * online sub=amazon test identity - Error: RegisterUser() : OnlineSubsystemAmazon is improperly configured in DefaultEngine.ini
@@ -1184,7 +1181,6 @@ class TestOSSCommands : BuildCommand
          * Code 401 - online sub=mcp test time - Failed to query server time
          * Error code 500 - online sub=mcp test entitlements - profile template not found
         */
-		#endregion
 
 		using (WriteErrors)
 		{
