Prevent an unnecessary "NoEditor" target being generated when packaging a server.

diff --git a/Engine/Source/Programs/AutomationTool/AutomationUtils/ProjectParams.cs b/Engine/Source/Programs/AutomationTool/AutomationUtils/ProjectParams.cs
--- a/Engine/Source/Programs/AutomationTool/AutomationUtils/ProjectParams.cs
+++ b/Engine/Source/Programs/AutomationTool/AutomationUtils/ProjectParams.cs
@@ -2059,6 +2059,12 @@ namespace AutomationTool
 					}
 				}
 
+				// If we don't have any game/client targets, don't stage any client executable
+				if (ClientCookedTargetsList.Count == 0)
+				{
+					NoClient = true;
+				}
+
 				// Validate all the settings
 				if (Client && bHasGameTarget)
 				{
