Fix empty game name when packaging for some platforms.

#rnx

diff --git a/Engine/Source/Programs/AutomationTool/AutomationUtils/ProjectParams.cs b/Engine/Source/Programs/AutomationTool/AutomationUtils/ProjectParams.cs
--- a/Engine/Source/Programs/AutomationTool/AutomationUtils/ProjectParams.cs
+++ b/Engine/Source/Programs/AutomationTool/AutomationUtils/ProjectParams.cs
@@ -2064,6 +2064,10 @@ namespace AutomationTool
 				{
 					NoClient = true;
 				}
+				else
+				{
+					GameTarget = ClientCookedTargetsList[0];
+				}
 
 				// Validate all the settings
 				if (Client && bHasGameTarget)
