Fix code only builds.

diff --git a/Engine/Source/Programs/AutomationTool/Android/AndroidPlatform.Automation.cs b/Engine/Source/Programs/AutomationTool/Android/AndroidPlatform.Automation.cs
--- a/Engine/Source/Programs/AutomationTool/Android/AndroidPlatform.Automation.cs
+++ b/Engine/Source/Programs/AutomationTool/Android/AndroidPlatform.Automation.cs
@@ -414,17 +414,24 @@ public class AndroidPlatform : Platform
 		List<FileReference> FilesForObb = ObbFileFilter.ApplyToDirectory(new DirectoryReference(StageDirectoryPath), true);
 
 		bool OBBNeedsUpdate = false;
-		System.DateTime OBBTimeStamp = File.GetLastWriteTimeUtc(LocalObbName);
-		foreach (FileReference FileToObb in FilesForObb)
-		{
-			System.DateTime FileTimeStamp = File.GetLastWriteTimeUtc(FileToObb.FullName);
-			if(FileTimeStamp > OBBTimeStamp)
+
+		if (File.Exists(LocalObbName))
+		{ 
+			System.DateTime OBBTimeStamp = File.GetLastWriteTimeUtc(LocalObbName);
+			foreach (FileReference FileToObb in FilesForObb)
 			{
-				OBBNeedsUpdate = true;
-				break;
+				System.DateTime FileTimeStamp = File.GetLastWriteTimeUtc(FileToObb.FullName);
+				if (FileTimeStamp > OBBTimeStamp)
+				{
+					OBBNeedsUpdate = true;
+					break;
+				}
 			}
 		}
-
+		else
+		{
+			OBBNeedsUpdate = true;
+		}
 		Int64 OBBSizeAllowed = GetMaxOBBSizeAllowed(SC);
 		string LimitString = (OBBSizeAllowed < MaxOBBSizeAllowed) ? "2 GiB" : "4 GiB";
 
