Added some options to simplify running EngineTests locally via buildgraph

diff --git a/Engine/Source/Programs/AutomationTool/Gauntlet/Unreal/RunUnreal.cs b/Engine/Source/Programs/AutomationTool/Gauntlet/Unreal/RunUnreal.cs
--- a/Engine/Source/Programs/AutomationTool/Gauntlet/Unreal/RunUnreal.cs
+++ b/Engine/Source/Programs/AutomationTool/Gauntlet/Unreal/RunUnreal.cs
@@ -31,6 +31,12 @@ namespace Gauntlet
 	/// </summary>
 	public class RunUnreal : BuildCommand
 	{
+
+		/// <summary>
+		/// Test node to create if none were specified
+		/// </summary>
+		public virtual string DefaultTestName { get { return "DefaultTest"; } }
+
 		/// <summary>
 		/// Main UAT entrance point. Custom games can derive from RunUnrealTests to run custom setup steps or
 		/// directly set params on ContextOptions to remove the need for certain command line params (e.g.
@@ -103,8 +109,8 @@ namespace Gauntlet
 
 			if (ContextOptions.TestList.Count == 0)
 			{
-				Gauntlet.Log.Info("No test specified, creating default test node");
-				ContextOptions.TestList.Add(TestRequest.CreateRequest("DefaultTest"));
+				Gauntlet.Log.Info("No test specified, using default '{0}'", DefaultTestName);
+				ContextOptions.TestList.Add(TestRequest.CreateRequest(DefaultTestName));
 			}
 
 			bool EditorForAllRoles = Globals.Params.ParseParam("editor") || string.Equals(Globals.Params.ParseValue("build", ""), "editor", StringComparison.OrdinalIgnoreCase);
