Fixed issue where tests that timed out weren't written to the log.

Also improved readability of post-test report

diff --git a/Engine/Source/Developer/AutomationController/Private/AutomationControllerManger.cpp b/Engine/Source/Developer/AutomationController/Private/AutomationControllerManger.cpp
--- a/Engine/Source/Developer/AutomationController/Private/AutomationControllerManger.cpp
+++ b/Engine/Source/Developer/AutomationController/Private/AutomationControllerManger.cpp
@@ -27,6 +27,14 @@
 #include "Logging/MessageLog.h"
 #endif
 
+// these strings are parsed by Gauntlet (AutomationLogParser) so make sure changes are replicated there!
+#define AutomationTestStarting		TEXT("Test Started. Name=(%s)")
+#define AutomationSuccessFormat		TEXT("Test Completed. Result=(Passed) Name=(%s) Path=(%s)")
+#define AutomationFailureFormat		TEXT("Test Completed. Result=(Failed) Name=(%s) Path=(%s)")
+
+#define BeginEventsFormat			TEXT("BeginEvents: %s")
+#define EndEventsFormat				TEXT("EndEvents: %s")
+
 DEFINE_LOG_CATEGORY_STATIC(LogAutomationController, Log, All)
 
 FAutomationControllerManager::FAutomationControllerManager()
@@ -485,7 +493,7 @@ void FAutomationControllerManager::ExecuteNextTask( int32 ClusterIndex, OUT bool
 						{
 							FAutomationTestResults TestResults;
 
-							UE_LOG(LogAutomationController, Display, TEXT("Running Automation: '%s' (Class Name: '%s')"), *TestsRunThisPass[AddressIndex]->GetFullTestPath(), *TestsRunThisPass[AddressIndex]->GetCommand());
+							UE_LOG(LogAutomationController, Display, AutomationTestStarting, *TestsRunThisPass[AddressIndex]->GetDisplayName(), *TestsRunThisPass[AddressIndex]->GetCommand());
 							TestResults.State = EAutomationState::InProcess;
 
 							if (CheckpointFile)
@@ -1081,11 +1089,6 @@ void FAutomationControllerManager::HandleRequestTestsReplyCompleteMessage(const
 
 void FAutomationControllerManager::ReportAutomationResult(const TSharedPtr<IAutomationReport> InReport, int32 ClusterIndex, int32 PassIndex)
 {
-	// these strings are parsed by Gauntlet (AutomationLogParser) so make sure changes are replicated there!
-	#define AutomationSuccessFormat		TEXT("Automation Test Succeeded (%s - %s)")
-	#define AutomationFailureFormat		TEXT("Automation Test Failed (%s - %s)")
-	#define BeginEventsFormat			TEXT("BeginEvents: %s")
-	#define EndEventsFormat				TEXT("EndEvents: %s")
 
 	FName CategoryName = "LogAutomationController";
 #if WITH_EDITOR
diff --git a/Engine/Source/Programs/AutomationTool/Gauntlet/Unreal/Base/Gauntlet.UnrealTestNode.cs b/Engine/Source/Programs/AutomationTool/Gauntlet/Unreal/Base/Gauntlet.UnrealTestNode.cs
--- a/Engine/Source/Programs/AutomationTool/Gauntlet/Unreal/Base/Gauntlet.UnrealTestNode.cs
+++ b/Engine/Source/Programs/AutomationTool/Gauntlet/Unreal/Base/Gauntlet.UnrealTestNode.cs
@@ -938,10 +938,8 @@ namespace Gauntlet
 				}
 			}
 
-			MB.Paragraph(string.Format("FatalErrors: {0}, Ensures: {1}, Errors: {2}, Warnings: {3}",
-				FatalErrors, LogSummary.Ensures.Count(), LogSummary.Errors.Count(), LogSummary.Warnings.Count()));
-
-			MB.Paragraph(string.Format("ResultHash: {0}", GetRoleResultHash(InArtifacts)));
+			MB.Paragraph(string.Format("FatalErrors: {0}, Ensures: {1}, Errors: {2}, Warnings: {3}, Hash: {4}",
+				FatalErrors, LogSummary.Ensures.Count(), LogSummary.Errors.Count(), LogSummary.Warnings.Count(), GetRoleResultHash(InArtifacts)));
 
 			if (GetCachedConfiguration().ShowErrorsInSummary && InArtifacts.LogSummary.Errors.Count() > 0)
 			{
@@ -1156,8 +1154,7 @@ namespace Gauntlet
 			}
 			MB.Paragraph(string.Format("Context: {0}", Context.ToString()));
 			MB.Paragraph(string.Format("FatalErrors: {0}, Ensures: {1}, Errors: {2}, Warnings: {3}", FatalErrors, Ensures, Errors, Warnings));
-			MB.Paragraph(string.Format("ResultHash: {0}", GetTestResultHash()));
-			MB.Paragraph(string.Format("Result: {0}", GetTestResult()));
+			MB.Paragraph(string.Format("Result: {0}, ResultHash: {1}", GetTestResult(), GetTestResultHash()));
 			//MB.Paragraph(string.Format("Artifacts: {0}", CachedArtifactPath));
 
 			return MB.ToString();
diff --git a/Engine/Source/Programs/AutomationTool/Gauntlet/Unreal/Utils/Gauntlet.AutomationLogParser.cs b/Engine/Source/Programs/AutomationTool/Gauntlet/Unreal/Utils/Gauntlet.AutomationLogParser.cs
--- a/Engine/Source/Programs/AutomationTool/Gauntlet/Unreal/Utils/Gauntlet.AutomationLogParser.cs
+++ b/Engine/Source/Programs/AutomationTool/Gauntlet/Unreal/Utils/Gauntlet.AutomationLogParser.cs
@@ -24,6 +24,11 @@ namespace Gauntlet
 		/// </summary>
 		public string		TestName;
 
+		/// <summary>
+		/// True if the test completed (pass or fail)
+		/// </summary>
+		public bool			Completed;
+
 		/// <summary>
 		/// True if the test passed
 		/// </summary>
@@ -40,10 +45,11 @@ namespace Gauntlet
 		/// <param name="InDisplayName"></param>
 		/// <param name="InTestName"></param>
 		/// <param name="bInPassed"></param>
-		public AutomationTestResult(string InDisplayName, string InTestName, bool bInPassed)
+		public AutomationTestResult(string InDisplayName, string InTestName, bool InCompleted, bool bInPassed)
 		{
 			DisplayName = InDisplayName;
 			TestName = InTestName;
+			Completed = InCompleted;
 			Passed = bInPassed;
 			Events = new List<string>();
 		}
@@ -80,42 +86,56 @@ namespace Gauntlet
 		/// <returns></returns>
 		public IEnumerable<AutomationTestResult> GetResults()
 		{
+			IEnumerable<Match> TestStarts = Parser.GetAllMatches(@"LogAutomationController.+Test Started. Name=\((.+?)\)");
+
 			// Find all automation results that succeeded/failed
-			// [2019.04.30-18.49.51:329][244]LogAutomationController: Display: Automation Test Succeeded (ST_PR04 - Project.Functional Tests./Game/Tests/Rendering/PlanarReflection.ST_PR04)
+			// [00:10:54.148]   LogAutomationController: Display: Test Started. Name=(ST_PR04)
+			// [2019.04.30-18.49.51:329][244]LogAutomationController: Display: Test Completed With Success. Name=(ST_PR04) Path=(Project.Functional Tests./Game/Tests/Rendering/PlanarReflection.ST_PR04))
 			// [2019.04.30-18.49.51:330] [244] LogAutomationController: BeginEvents: Project.Functional Tests./Game/Tests/Rendering/PlanarReflection.ST_PR04
 			// [2019.04.30 - 18.49.51:331][244] LogAutomationController: Screenshot 'ST_PR04' was similar!  Global Difference = 0.001377, Max Local Difference = 0.037953
 			// [2019.04.30 - 18.49.51:332][244]LogAutomationController: EndEvents: Project.Functional Tests./Game/Tests/Rendering/PlanarReflection.ST_PR04
-			IEnumerable<Match> TestMatches = Parser.GetAllMatches(@"LogAutomationController.+Test\s+(Succeeded|Failed)\s+\((.+)\s+-\s+(.+)\)");
+			IEnumerable<Match> TestResults = Parser.GetAllMatches(@"LogAutomationController.+Test Completed. Result=\((.+?)\)\s+Name=\((.+?)\)\s+Path=\((.+?)\)");
 
 			string[] AutomationChannel = Parser.GetLogChannel("AutomationController").ToArray();
 
 			// Convert these lines into results by parsing out the details and then adding the events
-			IEnumerable<AutomationTestResult> Results = TestMatches.Select(M =>
+			IEnumerable<AutomationTestResult> Results = TestStarts.Select(TestMatch =>
 			{
-				bool TestPassed = M.Groups[1].ToString().ToLower() == "succeeded" ? true : false;
-				string DisplayName = M.Groups[2].ToString();
-				string TestName = M.Groups[3].ToString();
-
-				AutomationTestResult Result = new AutomationTestResult(DisplayName, TestName, TestPassed);
+				string DisplayName = TestMatch.Groups[1].ToString();
+				string LongName = "";
+				bool Completed = false;
+				bool Passed = false;
+				List<string> Events = new List<string>();
+				
+				Match ResultMatch = TestResults.Where(M => M.Groups[2].ToString() == DisplayName).FirstOrDefault();
+
+				if (ResultMatch != null)
+				{
+					Completed = true;
+					Passed = ResultMatch.Groups[1].ToString().ToLower() == "passed" ? true : false;
+					LongName = ResultMatch.Groups[3].ToString();					
 
-				string EventName = string.Format("BeginEvents: {0}", TestName);
-				int EventIndex = Array.FindIndex(AutomationChannel, S => S.Contains(EventName)) + 1;
+					string EventName = string.Format("BeginEvents: {0}", LongName);
+					int EventIndex = Array.FindIndex(AutomationChannel, S => S.Contains(EventName)) + 1;
 
-				if (EventIndex > 0)
-				{
-					while (EventIndex < AutomationChannel.Length)
+					if (EventIndex > 0)
 					{
-						string Event = AutomationChannel[EventIndex++];
-
-						if (Event.Contains("EndEvents"))
+						while (EventIndex < AutomationChannel.Length)
 						{
-							break;
-						}
+							string Event = AutomationChannel[EventIndex++];
 
-						Result.Events.Add(Event);
+							if (Event.Contains("EndEvents"))
+							{
+								break;
+							}
+
+							Events.Add(Event);
+						}
 					}
 				}
 
+				AutomationTestResult Result = new AutomationTestResult(DisplayName, LongName, Completed, Passed);
+				Result.Events = Events;
 				return Result;
 			});
 
