Switched to curly braces for reporting and parsing test results.

diff --git a/Engine/Source/Developer/AutomationController/Private/AutomationControllerManger.cpp b/Engine/Source/Developer/AutomationController/Private/AutomationControllerManger.cpp
--- a/Engine/Source/Developer/AutomationController/Private/AutomationControllerManger.cpp
+++ b/Engine/Source/Developer/AutomationController/Private/AutomationControllerManger.cpp
@@ -28,9 +28,9 @@
 #endif
 
 // these strings are parsed by Gauntlet (AutomationLogParser) so make sure changes are replicated there!
-#define AutomationTestStarting		TEXT("Test Started. Name=(%s)")
-#define AutomationSuccessFormat		TEXT("Test Completed. Result=(Passed) Name=(%s) Path=(%s)")
-#define AutomationFailureFormat		TEXT("Test Completed. Result=(Failed) Name=(%s) Path=(%s)")
+#define AutomationTestStarting		TEXT("Test Started. Name={%s}")
+#define AutomationSuccessFormat		TEXT("Test Completed. Result={Passed} Name={%s} Path={%s}")
+#define AutomationFailureFormat		TEXT("Test Completed. Result={Failed} Name={%s} Path={%s}")
 
 #define BeginEventsFormat			TEXT("BeginEvents: %s")
 #define EndEventsFormat				TEXT("EndEvents: %s")
diff --git a/Engine/Source/Programs/AutomationTool/Gauntlet/Unreal/Utils/Gauntlet.AutomationLogParser.cs b/Engine/Source/Programs/AutomationTool/Gauntlet/Unreal/Utils/Gauntlet.AutomationLogParser.cs
--- a/Engine/Source/Programs/AutomationTool/Gauntlet/Unreal/Utils/Gauntlet.AutomationLogParser.cs
+++ b/Engine/Source/Programs/AutomationTool/Gauntlet/Unreal/Utils/Gauntlet.AutomationLogParser.cs
@@ -86,15 +86,15 @@ namespace Gauntlet
 		/// <returns></returns>
 		public IEnumerable<AutomationTestResult> GetResults()
 		{
-			IEnumerable<Match> TestStarts = Parser.GetAllMatches(@"LogAutomationController.+Test Started. Name=\((.+?)\)");
+			IEnumerable<Match> TestStarts = Parser.GetAllMatches(@"LogAutomationController.+Test Started. Name={(.+?)}");
 
 			// Find all automation results that succeeded/failed
-			// [00:10:54.148]   LogAutomationController: Display: Test Started. Name=(ST_PR04)
-			// [2019.04.30-18.49.51:329][244]LogAutomationController: Display: Test Completed With Success. Name=(ST_PR04) Path=(Project.Functional Tests./Game/Tests/Rendering/PlanarReflection.ST_PR04))
+			// [00:10:54.148]   LogAutomationController: Display: Test Started. Name={ST_PR04}
+			// [2019.04.30-18.49.51:329][244]LogAutomationController: Display: Test Completed With Success. Name={ST_PR04} Path={Project.Functional Tests./Game/Tests/Rendering/PlanarReflection.ST_PR04}
 			// [2019.04.30-18.49.51:330] [244] LogAutomationController: BeginEvents: Project.Functional Tests./Game/Tests/Rendering/PlanarReflection.ST_PR04
 			// [2019.04.30 - 18.49.51:331][244] LogAutomationController: Screenshot 'ST_PR04' was similar!  Global Difference = 0.001377, Max Local Difference = 0.037953
 			// [2019.04.30 - 18.49.51:332][244]LogAutomationController: EndEvents: Project.Functional Tests./Game/Tests/Rendering/PlanarReflection.ST_PR04
-			IEnumerable<Match> TestResults = Parser.GetAllMatches(@"LogAutomationController.+Test Completed. Result=\((.+?)\)\s+Name=\((.+?)\)\s+Path=\((.+?)\)");
+			IEnumerable<Match> TestResults = Parser.GetAllMatches(@"LogAutomationController.+Test Completed. Result={(.+?)}\s+Name={(.+?)}\s+Path={(.+?)}");
 
 			string[] AutomationChannel = Parser.GetLogChannel("AutomationController").ToArray();
 
