Build armv7 and arm64 architectures when packaging Android plugins.

diff --git a/Engine/Build/InstalledEngineBuild.xml b/Engine/Build/InstalledEngineBuild.xml
--- a/Engine/Build/InstalledEngineBuild.xml
+++ b/Engine/Build/InstalledEngineBuild.xml
@@ -80,7 +80,7 @@
 	<!-- Directory for toolchains -->
 	<Property Name="ToolchainDir" Value="$(RootDir)/Engine/Extras/ThirdPartyNotUE/SDKs"/>
 
-	<!-- Architectures that we build for Android -->
+	<!-- Architectures that we build for Android. NOTE: This should match the architectures specified in BuildPluginCommand. -->
 	<Property Name="AndroidArchitectures" Value="armv7+arm64"/>
 	
 	<!-- Hololens Architectures -->
diff --git a/Engine/Source/Programs/AutomationTool/Scripts/BuildPluginCommand.Automation.cs b/Engine/Source/Programs/AutomationTool/Scripts/BuildPluginCommand.Automation.cs
--- a/Engine/Source/Programs/AutomationTool/Scripts/BuildPluginCommand.Automation.cs
+++ b/Engine/Source/Programs/AutomationTool/Scripts/BuildPluginCommand.Automation.cs
@@ -19,6 +19,8 @@ using System.Text;
 [Help("Unversioned", "Do not embed the current engine version into the descriptor")]
 class BuildPlugin : BuildCommand
 {
+	const string AndroidArchitectures = "armv7+arm64";
+
 	public override void ExecuteBuild()
 	{
 		// Get the plugin filename
@@ -197,6 +199,10 @@ class BuildPlugin : BuildCommand
 			ManifestFileNames.Add(ManifestFileName);
 			
 			string Arguments = String.Format("-plugin={0} -iwyu -noubtmakefiles -manifest={1} -nohotreload", CommandUtils.MakePathSafeToUseWithCommandLine(HostProjectPluginFile.FullName), CommandUtils.MakePathSafeToUseWithCommandLine(ManifestFileName.FullName));
+			if (Platform == UnrealTargetPlatform.Android)
+			{
+				Arguments += String.Format(" -architectures={0}", AndroidArchitectures);
+			}
 			if(!String.IsNullOrEmpty(InAdditionalArgs))
 			{
 				Arguments += InAdditionalArgs;
