BuildGraph: Generate a schema to the default location when -Schema is specified without any arguments.

#rnx

diff --git a/Engine/Source/Programs/AutomationTool/BuildGraph/BuildGraph.cs b/Engine/Source/Programs/AutomationTool/BuildGraph/BuildGraph.cs
--- a/Engine/Source/Programs/AutomationTool/BuildGraph/BuildGraph.cs
+++ b/Engine/Source/Programs/AutomationTool/BuildGraph/BuildGraph.cs
@@ -45,6 +45,7 @@ namespace AutomationTool
 	[Help("Tool for creating extensible build processes in UE4 which can be run locally or in parallel across a build farm.")]
 	[Help("Script=<FileName>", "Path to the script describing the graph")]
 	[Help("Target=<Name>", "Name of the node or output tag to be built")]
+	[Help("Schema", "Generates a schema to the default location")]
 	[Help("Schema=<FileName>", "Generate a schema describing valid script documents, including all the known tasks")]
 	[Help("Set:<Property>=<Value>", "Sets a named property to the given value")]
 	[Help("Clean", "Cleans all cached state of completed build nodes before running")]
@@ -103,6 +104,11 @@ namespace AutomationTool
 				PrintOptions |= GraphPrintOptions.ShowNotifications;
 			}
 
+			if (SchemaFileName == null && ParseParam("Schema"))
+			{
+				SchemaFileName = FileReference.Combine(CommandUtils.EngineDirectory, "Build", "Graph", "Schema.xsd").FullName;
+			}
+
 			// Parse any specific nodes to clean
 			List<string> CleanNodes = new List<string>();
 			foreach(string NodeList in ParseParamValues("CleanNode"))
