BuildGraph: Allow macros to contain agent and node definitions.

diff --git a/Engine/Source/Programs/AutomationTool/BuildGraph/Schema.cs b/Engine/Source/Programs/AutomationTool/BuildGraph/Schema.cs
--- a/Engine/Source/Programs/AutomationTool/BuildGraph/Schema.cs
+++ b/Engine/Source/Programs/AutomationTool/BuildGraph/Schema.cs
@@ -160,6 +160,7 @@ namespace AutomationTool
 		Property,
 		Regex,
 		Macro,
+		MacroBody,
 		Expand,
 		Trace,
 		Warning,
@@ -339,6 +340,7 @@ namespace AutomationTool
 			NewSchema.Items.Add(CreatePropertyType());
 			NewSchema.Items.Add(CreateRegexType());
 			NewSchema.Items.Add(CreateMacroType());
+			NewSchema.Items.Add(CreateMacroBodyType(TaskNameToType));
 			NewSchema.Items.Add(CreateExpandType());
 			NewSchema.Items.Add(CreateDiagnosticType(ScriptSchemaStandardType.Trace));
 			NewSchema.Items.Add(CreateDiagnosticType(ScriptSchemaStandardType.Warning));
@@ -458,6 +460,7 @@ namespace AutomationTool
 			GraphChoice.Items.Add(CreateSchemaElement("Trace", ScriptSchemaStandardType.Trace));
 			GraphChoice.Items.Add(CreateSchemaElement("Warning", ScriptSchemaStandardType.Warning));
 			GraphChoice.Items.Add(CreateSchemaElement("Error", ScriptSchemaStandardType.Error));
+			GraphChoice.Items.Add(CreateSchemaElement("Expand", ScriptSchemaStandardType.Expand));
 			GraphChoice.Items.Add(CreateDoElement(ScriptSchemaStandardType.Graph));
 			GraphChoice.Items.Add(CreateSwitchElement(ScriptSchemaStandardType.Graph));
 			GraphChoice.Items.Add(CreateForEachElement(ScriptSchemaStandardType.Graph));
@@ -505,6 +508,7 @@ namespace AutomationTool
 			TriggerChoice.Items.Add(CreateSchemaElement("Trace", ScriptSchemaStandardType.Trace));
 			TriggerChoice.Items.Add(CreateSchemaElement("Warning", ScriptSchemaStandardType.Warning));
 			TriggerChoice.Items.Add(CreateSchemaElement("Error", ScriptSchemaStandardType.Error));
+			TriggerChoice.Items.Add(CreateSchemaElement("Expand", ScriptSchemaStandardType.Expand));
 			TriggerChoice.Items.Add(CreateDoElement(ScriptSchemaStandardType.TriggerBody));
 			TriggerChoice.Items.Add(CreateSwitchElement(ScriptSchemaStandardType.TriggerBody));
 			TriggerChoice.Items.Add(CreateForEachElement(ScriptSchemaStandardType.TriggerBody));
@@ -552,6 +556,7 @@ namespace AutomationTool
 			AgentChoice.Items.Add(CreateSchemaElement("Trace", ScriptSchemaStandardType.Trace));
 			AgentChoice.Items.Add(CreateSchemaElement("Warning", ScriptSchemaStandardType.Warning));
 			AgentChoice.Items.Add(CreateSchemaElement("Error", ScriptSchemaStandardType.Error));
+			AgentChoice.Items.Add(CreateSchemaElement("Expand", ScriptSchemaStandardType.Expand));
 			AgentChoice.Items.Add(CreateDoElement(ScriptSchemaStandardType.AgentBody));
 			AgentChoice.Items.Add(CreateSwitchElement(ScriptSchemaStandardType.AgentBody));
 			AgentChoice.Items.Add(CreateForEachElement(ScriptSchemaStandardType.AgentBody));
@@ -777,7 +782,7 @@ namespace AutomationTool
 		static XmlSchemaType CreateMacroType()
 		{
 			XmlSchemaComplexContentExtension Extension = new XmlSchemaComplexContentExtension();
-			Extension.BaseTypeName = GetQualifiedTypeName(ScriptSchemaStandardType.NodeBody);
+			Extension.BaseTypeName = GetQualifiedTypeName(ScriptSchemaStandardType.MacroBody);
 			Extension.Attributes.Add(CreateSchemaAttribute("Name", ScriptSchemaStandardType.Name, XmlSchemaUse.Required));
 			Extension.Attributes.Add(CreateSchemaAttribute("Arguments", ScriptSchemaStandardType.BalancedString, XmlSchemaUse.Optional));
 			Extension.Attributes.Add(CreateSchemaAttribute("OptionalArguments", ScriptSchemaStandardType.BalancedString, XmlSchemaUse.Optional));
@@ -792,6 +797,52 @@ namespace AutomationTool
 			return ComplexType;
 		}
 
+		/// <summary>
+		/// Creates the schema type representing the macro type
+		/// </summary>
+		/// <returns>Type definition for a node</returns>
+		static XmlSchemaType CreateMacroBodyType(Dictionary<string, XmlSchemaComplexType> TaskNameToType)
+		{
+			XmlSchemaChoice MacroChoice = new XmlSchemaChoice();
+			MacroChoice.MinOccurs = 0;
+			MacroChoice.MaxOccursString = "unbounded";
+
+			// Graph scope
+			MacroChoice.Items.Add(CreateSchemaElement("Include", ScriptSchemaStandardType.Include));
+			MacroChoice.Items.Add(CreateSchemaElement("Option", ScriptSchemaStandardType.Option));
+			MacroChoice.Items.Add(CreateSchemaElement("EnvVar", ScriptSchemaStandardType.EnvVar));
+			MacroChoice.Items.Add(CreateSchemaElement("Property", ScriptSchemaStandardType.Property));
+			MacroChoice.Items.Add(CreateSchemaElement("Regex", ScriptSchemaStandardType.Regex));
+			MacroChoice.Items.Add(CreateSchemaElement("Macro", ScriptSchemaStandardType.Macro));
+			MacroChoice.Items.Add(CreateSchemaElement("Agent", ScriptSchemaStandardType.Agent));
+			MacroChoice.Items.Add(CreateSchemaElement("Trigger", ScriptSchemaStandardType.Trigger));
+			MacroChoice.Items.Add(CreateSchemaElement("Aggregate", ScriptSchemaStandardType.Aggregate));
+			MacroChoice.Items.Add(CreateSchemaElement("Report", ScriptSchemaStandardType.Report));
+			MacroChoice.Items.Add(CreateSchemaElement("Badge", ScriptSchemaStandardType.Badge));
+			MacroChoice.Items.Add(CreateSchemaElement("Notify", ScriptSchemaStandardType.Notify));
+			MacroChoice.Items.Add(CreateSchemaElement("Trace", ScriptSchemaStandardType.Trace));
+			MacroChoice.Items.Add(CreateSchemaElement("Warning", ScriptSchemaStandardType.Warning));
+			MacroChoice.Items.Add(CreateSchemaElement("Error", ScriptSchemaStandardType.Error));
+			MacroChoice.Items.Add(CreateSchemaElement("Expand", ScriptSchemaStandardType.Expand));
+
+			// Agent scope
+			MacroChoice.Items.Add(CreateSchemaElement("Node", ScriptSchemaStandardType.Node));
+
+			// Node scope
+			MacroChoice.Items.Add(CreateDoElement(ScriptSchemaStandardType.NodeBody));
+			MacroChoice.Items.Add(CreateSwitchElement(ScriptSchemaStandardType.NodeBody));
+			MacroChoice.Items.Add(CreateForEachElement(ScriptSchemaStandardType.NodeBody));
+			foreach (KeyValuePair<string, XmlSchemaComplexType> Pair in TaskNameToType.OrderBy(x => x.Key))
+			{
+				MacroChoice.Items.Add(CreateSchemaElement(Pair.Key, new XmlQualifiedName(Pair.Value.Name, NamespaceURI)));
+			}
+
+			XmlSchemaComplexType NodeType = new XmlSchemaComplexType();
+			NodeType.Name = GetTypeName(ScriptSchemaStandardType.MacroBody);
+			NodeType.Particle = MacroChoice;
+			return NodeType;
+		}
+
 		/// <summary>
 		/// Creates the schema type representing a macro expansion
 		/// </summary>
diff --git a/Engine/Source/Programs/AutomationTool/BuildGraph/Script.cs b/Engine/Source/Programs/AutomationTool/BuildGraph/Script.cs
--- a/Engine/Source/Programs/AutomationTool/BuildGraph/Script.cs
+++ b/Engine/Source/Programs/AutomationTool/BuildGraph/Script.cs
@@ -389,6 +389,9 @@ namespace AutomationTool
 					case "ForEach":
 						ReadForEach(ChildElement, x => ReadGraphBody(x, BaseDirectory, Arguments));
 						break;
+					case "Expand":
+						ReadExpand(ChildElement, x => ReadGraphBody(x, BaseDirectory, Arguments));
+						break;
 					default:
 						LogError(ChildElement, "Invalid element '{0}'", ChildElement.Name);
 						break;
@@ -911,6 +914,9 @@ namespace AutomationTool
 					case "ForEach":
 						ReadForEach(ChildElement, x => ReadAgentBody(x, ParentAgent, ControllingTrigger));
 						break;
+					case "Expand":
+						ReadExpand(ChildElement, x => ReadAgentBody(x, ParentAgent, ControllingTrigger));
+						break;
 					default:
 						LogError(ChildElement, "Unexpected element type '{0}'", ChildElement.Name);
 						break;
