Workaround for issue with Zip64 IPA files created on Ionic.Zip with mono
#rnx

diff --git a/Engine/Source/Programs/AutomationTool/Gauntlet/Platform/IOS/Gauntlet.IOSBuildSource.cs b/Engine/Source/Programs/AutomationTool/Gauntlet/Platform/IOS/Gauntlet.IOSBuildSource.cs
--- a/Engine/Source/Programs/AutomationTool/Gauntlet/Platform/IOS/Gauntlet.IOSBuildSource.cs
+++ b/Engine/Source/Programs/AutomationTool/Gauntlet/Platform/IOS/Gauntlet.IOSBuildSource.cs
@@ -44,7 +44,7 @@ namespace Gauntlet
 			return false;
 		}
 
-		public static IProcessResult ExecuteCommand(String Command, String Arguments)
+		internal static IProcessResult ExecuteCommand(String Command, String Arguments)
 		{
 			CommandUtils.ERunOptions RunOptions = CommandUtils.ERunOptions.AppMustExist;
 
@@ -64,18 +64,43 @@ namespace Gauntlet
 			return Result;
 		}
 
-		private static string GetBundleIdentifier(string SourceIPA)
+		// There are issues with IPA Zip64 files being created with Ionic.Zip possibly limited to when running on mono (see IOSPlatform.PackageIPA)
+		// This manifests as header overflow errors, etc in 7zip, Ionic.Zip, System.IO.Compression, and OSX system unzip
+		// This is limited to bulk builds which exceed 4 gigs, this method works around this, though we really do need to fix
+		internal static bool ExecuteIPAZipCommand(String Arguments, out String Output, String ShouldExist = "")
 		{
-			// Get a list of files in the IPA
-			IProcessResult Result = ExecuteCommand("unzip", String.Format("-Z1 {0}", SourceIPA));
+			IProcessResult Result = ExecuteCommand("unzip", Arguments);
+			Output = Result.Output;
 
 			if (Result.ExitCode != 0)
+			{
+				if (!String.IsNullOrEmpty(ShouldExist))
+				{
+					if (!File.Exists(ShouldExist) && !Directory.Exists(ShouldExist))
+					{
+						Log.Error(String.Format("unzip encountered an error or warning procesing IPA, possibly due to Zip64 issue, {0} missing", ShouldExist));
+						return false;
+					}
+				}
+
+				Log.Warning(String.Format("unzip encountered an error or warning procesing IPA, possibly due to Zip64 issue. Future steps may fail."));
+			}
+			
+			return true;
+		}
+
+		private static string GetBundleIdentifier(string SourceIPA)
+		{
+			string Output;
+
+			// Get a list of files in the IPA
+			if (!ExecuteIPAZipCommand(String.Format("-Z1 {0}", SourceIPA), out Output))
 			{
 				Log.Warning(String.Format("Unable to list files for IPA {0}", SourceIPA));
 				return null;
 			}
 
-			string[] Filenames = Regex.Split(Result.Output, "\r\n|\r|\n");
+			string[] Filenames = Regex.Split(Output, "\r\n|\r|\n");
 			string PList = Filenames.Where(F => F.ToLower().Contains("info.plist")).FirstOrDefault();
 
 			if (String.IsNullOrEmpty(PList))
@@ -85,15 +110,13 @@ namespace Gauntlet
 			}
 
 			// Get the plist info
-			Result = ExecuteCommand("unzip", String.Format("-p '{0}' '{1}'", SourceIPA, PList));
-
-			if (Result.ExitCode != 0)
+			if (!ExecuteIPAZipCommand(String.Format("-p '{0}' '{1}'", SourceIPA, PList), out Output))
 			{
 				Log.Warning(String.Format("Unable to extract plist data for IPA {0}", SourceIPA));
 				return null;
 			}
 
-			string PlistInfo = Result.Output;
+			string PlistInfo = Output;
 
 			// todo: plist parsing, could be better
 			string PackageName = null;
diff --git a/Engine/Source/Programs/AutomationTool/Gauntlet/Platform/IOS/Gauntlet.TargetDeviceIOS.cs b/Engine/Source/Programs/AutomationTool/Gauntlet/Platform/IOS/Gauntlet.TargetDeviceIOS.cs
--- a/Engine/Source/Programs/AutomationTool/Gauntlet/Platform/IOS/Gauntlet.TargetDeviceIOS.cs
+++ b/Engine/Source/Programs/AutomationTool/Gauntlet/Platform/IOS/Gauntlet.TargetDeviceIOS.cs
@@ -836,9 +836,9 @@ namespace Gauntlet
 					}
 
 					Log.Verbose("Unzipping IPA {0} to cache at: {1}", Build.SourceIPAPath, GauntletAppCache);
-					
-					IProcessResult Result = IOSBuild.ExecuteCommand("unzip", String.Format("{0} -d {1}", Build.SourceIPAPath, GauntletAppCache));
-					if (Result.ExitCode != 0 || !Directory.Exists(PayloadDir))
+
+					string Output;
+					if (!IOSBuild.ExecuteIPAZipCommand(String.Format("{0} -d {1}", Build.SourceIPAPath, GauntletAppCache), out Output, PayloadDir))
 					{
 						throw new Exception(String.Format("Unable to extract IPA {0}", Build.SourceIPAPath));
 					}
@@ -852,8 +852,7 @@ namespace Gauntlet
 					{						
 						Log.Verbose("Unzipping Symbols {0} to cache at: {1}", SymbolsZipFile, SymbolsDir);
 
-						Result = IOSBuild.ExecuteCommand("unzip", String.Format("{0} -d {1}", SymbolsZipFile, SymbolsDir));
-						if (Result.ExitCode != 0 || !Directory.Exists(SymbolsDir))
+						if (!IOSBuild.ExecuteIPAZipCommand(String.Format("{0} -d {1}", SymbolsZipFile, SymbolsDir), out Output, SymbolsDir))
 						{
 							throw new Exception(String.Format("Unable to extract build symbols {0} -> {1}", SymbolsZipFile, SymbolsDir));
 						}
