Fixing issue with iOS LLDB callstack parsing regex
#rnx

diff --git a/Engine/Source/Programs/AutomationTool/Gauntlet/Platform/IOS/Gauntlet.TargetDeviceIOS.cs b/Engine/Source/Programs/AutomationTool/Gauntlet/Platform/IOS/Gauntlet.TargetDeviceIOS.cs
--- a/Engine/Source/Programs/AutomationTool/Gauntlet/Platform/IOS/Gauntlet.TargetDeviceIOS.cs
+++ b/Engine/Source/Programs/AutomationTool/Gauntlet/Platform/IOS/Gauntlet.TargetDeviceIOS.cs
@@ -1029,8 +1029,8 @@ namespace Gauntlet
 			Regex LogLineRegex = new Regex(@"(?<timestamp>\s\[\d.+\]\[\s*\d+\])(?<log>.*)");
 			Regex TimeRegex = new Regex(@"\[(?<year>\d+)\.(?<month>\d+)\.(?<day>\d+)-(?<hour>\d+)\.(?<minute>\d+)\.(?<second>\d+):(?<millisecond>\d+)\]\[(?<frame>\s*\d+)\]", RegexOptions.IgnoreCase);
 			Regex ThreadRegex = new Regex(@"(thread\s#)(?<threadnum>\d+),?(?<status>.+)");
-			Regex SymbolicatedFrameRegex = new Regex(@"\s#(?<framenum>\d+):\s0x(?<address>[\da-f]+)\s(?<module>.+)\`(?<symbol>.+)(\sat\s)(?<source>.+)\s\[opt\]");
-			Regex UnsymbolicatedFrameRegex = new Regex(@"frame\s#(?<framenum>\d+):\s0x(?<address>[\da-f]+)\s(?<module>.+)\`(?<symbol>.+)\s\+\s(?<offset>\d+)");
+			Regex SymbolicatedFrameRegex = new Regex(@"\*?\s#(?<framenum>\d+):\s0x(?<address>[\da-f]+)\s(?<module>.+)\`(?<symbol>.+)(\sat\s)(?<source>.+)\s\[opt\]");
+			Regex UnsymbolicatedFrameRegex = new Regex(@"\*?frame\s#(?<framenum>\d+):\s0x(?<address>[\da-f]+)\s(?<module>.+)\`(?<symbol>.+)(\s\+\s(?<offset>\d+))?");
 
 			LinkedList<string> CrashLog = new LinkedList<string>(Regex.Split(LogOutput, "\r\n|\r|\n"));
 						
