Gauntlet: Fix issue pulling Android installs from from patch files
#rnx

diff --git a/Engine/Source/Programs/AutomationTool/Gauntlet/Platform/Android/Gauntlet.AndroidBuildSource.cs b/Engine/Source/Programs/AutomationTool/Gauntlet/Platform/Android/Gauntlet.AndroidBuildSource.cs
--- a/Engine/Source/Programs/AutomationTool/Gauntlet/Platform/Android/Gauntlet.AndroidBuildSource.cs
+++ b/Engine/Source/Programs/AutomationTool/Gauntlet/Platform/Android/Gauntlet.AndroidBuildSource.cs
@@ -72,15 +72,17 @@ namespace Gauntlet
 				Match Info = Regex.Match(Fi.Name, RegEx, RegexOptions.IgnoreCase);
 
 				bool TestInstall = Fi.Name.EndsWith("_TEST.bat", StringComparison.OrdinalIgnoreCase);
+				bool PatchInstall = Fi.Name.EndsWith("_Patch.bat", StringComparison.OrdinalIgnoreCase);
 
 				// filter out non-matching or test installation batch files
 				// test installation scripts are intended to be manually invoked
-				if (Info.Success == false || TestInstall)
+				if (Info.Success == false || TestInstall || PatchInstall)
 				{
-					if (TestInstall)
+					if (TestInstall || PatchInstall)
 					{
-						Log.Verbose("Ignoring test installation batch file {0}", Fi.Name);
+						Log.Verbose("Ignoring {0} installation batch file {1}", TestInstall ? "test" : "patch", Fi.Name);
 					}
+
 					continue;
 				}
 
