Gauntlet: Disable crashlytics in automated tests to get symbolicated tombstone callstacks on Android
#rnx

diff --git a/Engine/Source/Programs/AutomationTool/Gauntlet/Unreal/Game/EpicGame.TestConfig.cs b/Engine/Source/Programs/AutomationTool/Gauntlet/Unreal/Game/EpicGame.TestConfig.cs
--- a/Engine/Source/Programs/AutomationTool/Gauntlet/Unreal/Game/EpicGame.TestConfig.cs
+++ b/Engine/Source/Programs/AutomationTool/Gauntlet/Unreal/Game/EpicGame.TestConfig.cs
@@ -246,9 +246,13 @@ namespace EpicGame
 					AppConfig.CommandLine += " -ini:Engine:[VoiceChat.Vivox]:bEnabled=false";
 				}
 
+				// turn off crashlytics so we get symbolicated tombstone crashes on Android
+				if (ConfigRole.Platform == UnrealTargetPlatform.Android)
+				{
+					AppConfig.CommandLine += " -nocrashlytics";
+				}
 			}
-		}
-		
+		}		
 	}
 	
 	/// <summary>
