Fix issue with mobile build builds being used regardless of -bulk being specified, log warning when fail to mark directory for cleanup
#rnx

diff --git a/Engine/Source/Programs/AutomationTool/Gauntlet/Platform/Android/Gauntlet.TargetDeviceAndroid.cs b/Engine/Source/Programs/AutomationTool/Gauntlet/Platform/Android/Gauntlet.TargetDeviceAndroid.cs
--- a/Engine/Source/Programs/AutomationTool/Gauntlet/Platform/Android/Gauntlet.TargetDeviceAndroid.cs
+++ b/Engine/Source/Programs/AutomationTool/Gauntlet/Platform/Android/Gauntlet.TargetDeviceAndroid.cs
@@ -228,7 +228,14 @@ namespace Gauntlet
 			}
 
 			// mark it as a temp dir (will also create it)
-			Utils.SystemHelpers.MarkDirectoryForCleanup(Install.AndroidDevice.LocalCachePath);
+			try
+			{
+				Utils.SystemHelpers.MarkDirectoryForCleanup(Install.AndroidDevice.LocalCachePath);
+			}
+			catch (Exception Ex)
+			{
+				Log.Warning("Exception marking directory for cleanup {0}", Ex.Message);
+			}
 
 			string LocalSaved = Path.Combine(Install.AndroidDevice.LocalCachePath, "Saved");
 			Directory.CreateDirectory(LocalSaved);
diff --git a/Engine/Source/Programs/AutomationTool/Gauntlet/Unreal/Base/Gauntlet.UnrealSession.cs b/Engine/Source/Programs/AutomationTool/Gauntlet/Unreal/Base/Gauntlet.UnrealSession.cs
--- a/Engine/Source/Programs/AutomationTool/Gauntlet/Unreal/Base/Gauntlet.UnrealSession.cs
+++ b/Engine/Source/Programs/AutomationTool/Gauntlet/Unreal/Base/Gauntlet.UnrealSession.cs
@@ -147,13 +147,13 @@ namespace Gauntlet
 			// where we may be running parallel tests on multiple platforms
 			if (InPlatform == UnrealTargetPlatform.Android || InPlatform == UnrealTargetPlatform.IOS)
 			{
-				if (Globals.Params.ParseParam("notbulk"))
+				if (Globals.Params.ParseParam("bulk"))
 				{
-					RequiredBuildFlags |= BuildFlags.NotBulk;
+					RequiredBuildFlags |= BuildFlags.Bulk;
 				}
-				else if (Globals.Params.ParseParam("bulk"))
+				else
 				{
-					RequiredBuildFlags |= BuildFlags.Bulk;
+					RequiredBuildFlags |= BuildFlags.NotBulk;
 				}
 			}
 
