Gauntlet: Fix iOS bundle identifier detection with IPA's which contain multiple info.plist
#rnx

diff --git a/Engine/Source/Programs/AutomationTool/Gauntlet/Platform/IOS/Gauntlet.IOSBuildSource.cs b/Engine/Source/Programs/AutomationTool/Gauntlet/Platform/IOS/Gauntlet.IOSBuildSource.cs
--- a/Engine/Source/Programs/AutomationTool/Gauntlet/Platform/IOS/Gauntlet.IOSBuildSource.cs
+++ b/Engine/Source/Programs/AutomationTool/Gauntlet/Platform/IOS/Gauntlet.IOSBuildSource.cs
@@ -100,8 +100,9 @@ namespace Gauntlet
 				return null;
 			}
 
-			string[] Filenames = Regex.Split(Output, "\r\n|\r|\n");
-			string PList = Filenames.Where(F => F.ToLower().Contains("info.plist")).FirstOrDefault();
+
+			string[] Filenames = Regex.Split(Output, "\r\n|\r|\n");			
+			string PList = Filenames.Where(F => Regex.IsMatch(F.ToLower().Trim(), @"(payload\/)([^\/]+)(\/info\.plist)")).FirstOrDefault();
 
 			if (String.IsNullOrEmpty(PList))
 			{
