Gauntlet: fixing lldb callstack parser
#rnx

diff --git a/Engine/Source/Programs/AutomationTool/Gauntlet/Platform/IOS/Gauntlet.TargetDeviceIOS.cs b/Engine/Source/Programs/AutomationTool/Gauntlet/Platform/IOS/Gauntlet.TargetDeviceIOS.cs
--- a/Engine/Source/Programs/AutomationTool/Gauntlet/Platform/IOS/Gauntlet.TargetDeviceIOS.cs
+++ b/Engine/Source/Programs/AutomationTool/Gauntlet/Platform/IOS/Gauntlet.TargetDeviceIOS.cs
@@ -1035,8 +1035,7 @@ namespace Gauntlet
 			LinkedList<string> CrashLog = new LinkedList<string>(Regex.Split(LogOutput, "\r\n|\r|\n"));
 						
 			List<ThreadInfo> Threads = new List<ThreadInfo>();
-			ThreadInfo Thread = null;
-			ThreadInfo CrashThread = null;
+			ThreadInfo Thread = null;			
 
 			var LineNode = CrashLog.First;			
 			while (LineNode != null)
@@ -1123,42 +1122,39 @@ namespace Gauntlet
 				// Parse thread
 				if (ThreadRegex.IsMatch(Line))
 				{
+
 					GroupCollection ThreadGroups = ThreadRegex.Match(Line).Groups;
-					Thread = new ThreadInfo()
-					{
-						Num = int.Parse(ThreadGroups["threadnum"].Value),
-						Status = ThreadGroups["status"].Value.Trim()
-					};
+					int Num = int.Parse(ThreadGroups["threadnum"].Value);
+					string Status = ThreadGroups["status"].Value.Trim();
 
-					if (Line.StartsWith("*"))
-					{
-						Thread.Current = true;
-					}
+					Thread = Threads.SingleOrDefault(T => T.Num == Num);
 
-					if (CrashThread == null)
-					{
-						CrashThread = Thread;
-					}
-					else
+					if (Thread == null)
 					{
+						Thread = new ThreadInfo()
+						{
+							Num = Num,
+							Status = Status
+						};
+
+						if (Line.Trim().StartsWith("*"))
+						{
+							Thread.Current = true;
+						}
+
 						Threads.Add(Thread);
-					}
+
+					}					
 				}
 
 				LineNode = LineNode.Next;
 			}
 
-			if (CrashThread == null)
-			{
-				return null;
-			}
-
-			Thread = Threads.Single(T => T.Num == CrashThread.Num);
+			Thread = Threads.SingleOrDefault(T => T.Current == true);
 
 			if (Thread == null)
 			{
-				Log.Warning("Unable to parse full crash callstack");
-				Thread = CrashThread;
+				Log.Warning("Unable to parse full crash callstack");				
 			}
 
 			return Thread;
