Only display warning if callstack threads are found
#rnx

diff --git a/Engine/Source/Programs/AutomationTool/Gauntlet/Platform/IOS/Gauntlet.TargetDeviceIOS.cs b/Engine/Source/Programs/AutomationTool/Gauntlet/Platform/IOS/Gauntlet.TargetDeviceIOS.cs
--- a/Engine/Source/Programs/AutomationTool/Gauntlet/Platform/IOS/Gauntlet.TargetDeviceIOS.cs
+++ b/Engine/Source/Programs/AutomationTool/Gauntlet/Platform/IOS/Gauntlet.TargetDeviceIOS.cs
@@ -1152,7 +1152,7 @@ namespace Gauntlet
 
 			Thread = Threads.SingleOrDefault(T => T.Current == true);
 
-			if (Thread == null)
+			if (Threads.Count > 0 && Thread == null)
 			{
 				Log.Warning("Unable to parse full crash callstack");				
 			}
