Gauntlet: support reservation model constraints on gauntlet service API
#rnx

diff --git a/Engine/Source/Programs/AutomationTool/Gauntlet/Framework/Devices/Gauntlet.DevicePool.cs b/Engine/Source/Programs/AutomationTool/Gauntlet/Framework/Devices/Gauntlet.DevicePool.cs
--- a/Engine/Source/Programs/AutomationTool/Gauntlet/Framework/Devices/Gauntlet.DevicePool.cs
+++ b/Engine/Source/Programs/AutomationTool/Gauntlet/Framework/Devices/Gauntlet.DevicePool.cs
@@ -560,15 +560,24 @@ namespace Gauntlet
 
 				if (!string.IsNullOrEmpty(Entry.Key.Model))
 				{
-					// if specific device model, we can't currently reserve it from service
-					Log.Info("Unable to reserve service device of model: {0}", Entry.Key.Model);
-					return false;
+					// if specific device model, we can't currently reserve it from (legacy) service
+					if (DeviceURL.ToLower().Contains("deviceservice.epicgames.net"))
+					{
+						Log.Info("Unable to reserve service device of model: {0} on legacy service", Entry.Key.Model);
+						return false;
+					}
+									
 				}
 
 				for (int i = 0; i < Entry.Value; i++)
 				{
 					// @todo: if any additional reservation requirements, encode constraint into json
-					Devices.Add(DeviceMap[Entry.Key.Platform.Value] + ":" + Entry.Key.PerfSpec);
+					string Constraint = Entry.Key.PerfSpec.ToString();
+					if (!string.IsNullOrEmpty(Entry.Key.Model))
+					{
+						Constraint = Entry.Key.Model;
+					}
+					Devices.Add(DeviceMap[Entry.Key.Platform.Value] + ":" + Constraint);
 				}				
 			}
 
