PR #6235: Installed build was missing Client/Server Target type in BaseEngine.ini InstalledPlatforms (Contributed by [REDACTED_PERSON])

#rnx

diff --git a/Engine/Source/Programs/AutomationTool/Scripts/FinalizeInstalledBuild.cs b/Engine/Source/Programs/AutomationTool/Scripts/FinalizeInstalledBuild.cs
--- a/Engine/Source/Programs/AutomationTool/Scripts/FinalizeInstalledBuild.cs
+++ b/Engine/Source/Programs/AutomationTool/Scripts/FinalizeInstalledBuild.cs
@@ -111,16 +111,23 @@ namespace AutomationTool
 					bool bCanBeDisplayed = ProjectType == EProjectType.Content;
 					foreach (UnrealTargetConfiguration CodeTargetConfiguration in Enum.GetValues(typeof(UnrealTargetConfiguration)))
 					{
-						string[] TargetNames = new string[] { "UE4Game", "UE4Client", "UE4Server" };
-						foreach(string TargetName in TargetNames)
+						Dictionary<String, TargetType> Targets = new Dictionary<string, TargetType>() {
+							{ "UE4Game", TargetType.Game },
+							{ "UE4Client", TargetType.Client },
+							{ "UE4Server", TargetType.Server }
+						};
+						foreach (KeyValuePair<string, TargetType> Target in Targets)
 						{
+							string CurrentTargetName = Target.Key;
+							TargetType CurrentTargetType = Target.Value;
+
 							// Need to check for development receipt as we use that for the Engine code in DebugGame
 							UnrealTargetConfiguration EngineConfiguration = (CodeTargetConfiguration == UnrealTargetConfiguration.DebugGame) ? UnrealTargetConfiguration.Development : CodeTargetConfiguration;
 
 							// Android has multiple architecture flavors built without receipts, so use the default arch target instead
 							if (CodeTargetPlatform == UnrealTargetPlatform.Android)
 							{
-								FileReference ReceiptFileName = TargetReceipt.GetDefaultPath(new DirectoryReference(OutputEnginePath), TargetName, CodeTargetPlatform, EngineConfiguration, Architecture);
+								FileReference ReceiptFileName = TargetReceipt.GetDefaultPath(new DirectoryReference(OutputEnginePath), CurrentTargetName, CodeTargetPlatform, EngineConfiguration, Architecture);
 								if (FileReference.Exists(ReceiptFileName))
 								{
 									// Strip the output folder so that this can be used on any machine
@@ -131,13 +138,13 @@ namespace AutomationTool
 									{
 										foreach (string Arch in AllArchNames)
 										{
-											InstalledConfigs.Add(new InstalledPlatformInfo.InstalledPlatformConfiguration(CodeTargetConfiguration, CodeTargetPlatform, TargetType.Game, Arch, RelativeReceiptFileName, ProjectType, bCanBeDisplayed));
+											InstalledConfigs.Add(new InstalledPlatformInfo.InstalledPlatformConfiguration(CodeTargetConfiguration, CodeTargetPlatform, CurrentTargetType, Arch, RelativeReceiptFileName, ProjectType, bCanBeDisplayed));
 										}
 									}
 									// if for some reason we didn't specify any flavors, just add the default one.
 									else
 									{
-										InstalledConfigs.Add(new InstalledPlatformInfo.InstalledPlatformConfiguration(CodeTargetConfiguration, CodeTargetPlatform, TargetType.Game, Architecture, RelativeReceiptFileName, ProjectType, bCanBeDisplayed));
+										InstalledConfigs.Add(new InstalledPlatformInfo.InstalledPlatformConfiguration(CodeTargetConfiguration, CodeTargetPlatform, CurrentTargetType, Architecture, RelativeReceiptFileName, ProjectType, bCanBeDisplayed));
 									}
 								}
 							}
@@ -146,11 +153,11 @@ namespace AutomationTool
 							{
 								foreach (string Arch in AllArchNames)
 								{
-									FileReference ReceiptFileName = TargetReceipt.GetDefaultPath(new DirectoryReference(OutputEnginePath), TargetName, CodeTargetPlatform, EngineConfiguration, Arch);
+									FileReference ReceiptFileName = TargetReceipt.GetDefaultPath(new DirectoryReference(OutputEnginePath), CurrentTargetName, CodeTargetPlatform, EngineConfiguration, Arch);
 									if (FileReference.Exists(ReceiptFileName))
 									{
 										string RelativeReceiptFileName = ReceiptFileName.MakeRelativeTo(new DirectoryReference(OutputDir));
-										InstalledConfigs.Add(new InstalledPlatformInfo.InstalledPlatformConfiguration(CodeTargetConfiguration, CodeTargetPlatform, TargetType.Game, Arch, RelativeReceiptFileName, ProjectType, bCanBeDisplayed));
+										InstalledConfigs.Add(new InstalledPlatformInfo.InstalledPlatformConfiguration(CodeTargetConfiguration, CodeTargetPlatform, CurrentTargetType, Arch, RelativeReceiptFileName, ProjectType, bCanBeDisplayed));
 									}
 								}
 							}
