UAT: Move functionality for creating attribute-driven help pages into DotNETCommon.

#rnx

diff --git a/Engine/Source/Programs/AutomationTool/AutomationUtils/HelpUtils.cs b/Engine/Source/Programs/AutomationTool/AutomationUtils/HelpUtils.cs
--- a/Engine/Source/Programs/AutomationTool/AutomationUtils/HelpUtils.cs
+++ b/Engine/Source/Programs/AutomationTool/AutomationUtils/HelpUtils.cs
@@ -10,9 +10,10 @@ using System.Runtime.Serialization;
 using System.Runtime.Serialization.Json;
 using System.Text;
 using System.Threading;
-
 using System.ComponentModel;
 using System.Reflection;
+using Tools.DotNETCommon;
+using System.Linq;
 
 namespace AutomationTool
 {
@@ -108,58 +109,6 @@ namespace AutomationTool
 		/// <param name="Description">Command Desscription</param>
 		/// <param name="Params">Command Parameters</param>
 		public static void LogHelp(Type Command, string Description, List<string> Params)
-		{
-			string HelpMessage = Environment.NewLine;
-
-			bool bFirstLine = true;
-			if (Command != null)
-			{
-				HelpMessage += String.Format("{0} Help:{1}", Command.Name, Environment.NewLine);
-				bFirstLine = false;
-			}
-
-			if (String.IsNullOrEmpty(Description) == false)
-			{
-				if (!bFirstLine)
-				{
-					HelpMessage += Environment.NewLine;
-				}
-				HelpMessage += Description + Environment.NewLine;
-				bFirstLine = false;
-			}
-			
-			if (IsNullOrEmpty(Params) == false)
-			{
-				if (!bFirstLine)
-				{
-					HelpMessage += Environment.NewLine;
-				}
-				HelpMessage += "Parameters: " + Environment.NewLine;
-				HelpMessage += Environment.NewLine;
-
-				foreach(string Line in FormatParams(Params, 4, 24))
-				{
-					HelpMessage += Line + Environment.NewLine;
-				}
-			}
-
-			LogInformation(HelpMessage);
-		}
-
-		/// <summary>
-		/// Formats the given parameters as so:
-		///     -Param1     Param1 Description
-		///
-		///     -Param2      Param2 Description, this description is
-		///                  longer and splits onto a separate line. 
-		///
-		///     -Param3      Param3 Description continues as before. 
-		/// </summary>
-		/// <param name="Params">List of parameters arranged as "-ParamName Param Description"</param>
-		/// <param name="Indent">Indent from the left hand side</param>
-		/// <param name="DefaultRightPadding">The minimum padding from the start of the param name to the start of the description (resizes with larger param names)</param>
-		/// <returns></returns>
-		public static IEnumerable<string> FormatParams(IEnumerable<string> Params, int Indent, int DefaultRightPadding)
 		{
 			Dictionary<string, string> ParamDict = new Dictionary<string, string>(StringComparer.InvariantCultureIgnoreCase);
 
@@ -196,115 +145,11 @@ namespace AutomationTool
 				}
 			}
 
-			// string used to intent the param
-			string IndentString = string.Empty.PadRight(Indent);
-
-			// default the padding value
-			int RightPadding = DefaultRightPadding;
-
-			// If (Padding < longest param name) padding = longest param name + 1
-			foreach (var ParamName in ParamDict.Keys)
-			{
-				if (ParamName.Length + 1 > RightPadding)
-				{
-					RightPadding = ParamName.Length + 1;
-				}
-			}
-
-			// Get the window width, using a default value if there's no console attached to this process.
-			int WindowWidth;
-			try
-			{
-				WindowWidth = Console.WindowWidth;
-			}
-			catch
-			{
-				WindowWidth = 240;
-			}
-
-			if(WindowWidth <= 0)
-			{
-				WindowWidth = 240;
-			}
-
-			// Build the formatted params
-			foreach (var ParamName in ParamDict.Keys)
-			{
-				// build the param first, including intend and padding on the rights size
-				string ParamString = IndentString + ParamName.PadRight(RightPadding);
-
-				// Build the description line by line, adding the same amount of intending each time. 
-				foreach (var DescriptionLine in WordWrap(ParamDict[ParamName], WindowWidth - ParamString.Length))
-				{
-					// Formatting as following:
-					// <Indent>-param<Right Padding>Description<New line>
-					yield return ParamString + DescriptionLine;
-
-					// we replace the param string on subsequent lines with white space of the same length
-					ParamString = string.Empty.PadRight(IndentString.Length + RightPadding);
-				}
-			}
-		}
-
-		/// <summary>
-		/// Takes a given sentence and wraps it on a word by word basis so that no line exceeds the
-		/// set maximum line length. Words longer than a line are broken up. Returns the sentence as
-		/// a list of individual lines.
-		/// </summary>
-		/// <param name="sentence">The sentence to be wrapped</param>
-		/// <param name="MaximumLineLength">The maximum (non negative) length of the returned sentences</param>
-		/// <returns>The provided sentence in a list of individual lines no longer than MaximumLineLength</returns>
-		static List<string> WordWrap(string Sentence, int MaximumLineLength)
-		{
-			// Early out
-			if (Sentence.Length == 0)
-			{
-				return new List<string>();
-			}
-
-			string[] words = Sentence.Split(' ');
-			List<string> WrappedWords = new List<string>();
-
-			string CurrentSentence = string.Empty;
-			foreach (var word in words)
-			{
-				// if this is a very large word, split it
-				if (word.Length > MaximumLineLength)
-				{
-					// Top up the current line
-					WrappedWords.Add(CurrentSentence + word.Substring(0, MaximumLineLength - CurrentSentence.Length));
-
-					int length = MaximumLineLength - CurrentSentence.Length;
-					while (length + MaximumLineLength < word.Length)
-					{
-						// Place the starting lengths into their own lines
-						WrappedWords.Add(word.Substring(length, Math.Min(MaximumLineLength, word.Length - length)));
-						length += MaximumLineLength;
-					}
-
-					// then the trailing end into the next line
-					CurrentSentence += word.Substring(length, Math.Min(MaximumLineLength, word.Length - length)) + " ";
-				}
-				else
-				{
-					if (CurrentSentence.Length + word.Length > MaximumLineLength)
-					{
-						// next line and reset sentence
-						WrappedWords.Add(CurrentSentence);
-						CurrentSentence = string.Empty;
-					}
-
-					// Add the word to the current sentence.
-					CurrentSentence += word + " ";
-				}
-			}
-
-			if (CurrentSentence.Length > 0)
-			{
-				WrappedWords.Add(CurrentSentence);
-			}
+			List<string> Lines = new List<string>();
+			Lines.Add("");
+			Lines.AddRange(HelpUtils.GetHelpText(String.Format("{0} Help:", Command.Name), Description, ParamDict.ToList(), HelpUtils.WindowWidth - 1));
 
-			return WrappedWords;
+			Log.TraceInformation("{0}", String.Join(Environment.NewLine, Lines));
 		}
 
 		/// <summary>
diff --git a/Engine/Source/Programs/AutomationTool/BuildGraph/Graph.cs b/Engine/Source/Programs/AutomationTool/BuildGraph/Graph.cs
--- a/Engine/Source/Programs/AutomationTool/BuildGraph/Graph.cs
+++ b/Engine/Source/Programs/AutomationTool/BuildGraph/Graph.cs
@@ -589,25 +589,27 @@ namespace AutomationTool
 			if((PrintOptions & GraphPrintOptions.ShowCommandLineOptions) != 0)
 			{
 				// Get the list of messages
-				List<string> Messages = new List<string>();
+				List<KeyValuePair<string, string>> Parameters = new List<KeyValuePair<string, string>>();
 				foreach(GraphOption Option in Options)
 				{
-					StringBuilder Message = new StringBuilder();
-					Message.AppendFormat("-set:{0}=... {1}", Option.Name, Option.Description);
+					string Name = String.Format("-set:{0}=...", Option.Name);
+
+					StringBuilder Description = new StringBuilder(Option.Description);
 					if(!String.IsNullOrEmpty(Option.DefaultValue))
 					{
-						Message.AppendFormat(" (Default: {0})", Option.DefaultValue);
+						Description.AppendFormat(" (Default: {0})", Option.DefaultValue);
 					}
-					Messages.Add(Message.ToString());
+
+					Parameters.Add(new KeyValuePair<string, string>(Name, Description.ToString()));
 				}
 
 				// Format them to the log
-				if(Messages.Count > 0)
+				if(Parameters.Count > 0)
 				{
 					CommandUtils.LogInformation("");
 					CommandUtils.LogInformation("Options:");
 					CommandUtils.LogInformation("");
-					foreach(string Line in CommandUtils.FormatParams(Messages, 4, 24))
+					foreach(string Line in HelpUtils.Tabulate(Parameters, 4, 24, HelpUtils.WindowWidth - 1))
 					{
 						CommandUtils.LogInformation(Line);
 					}
diff --git a/Engine/Source/Programs/DotNETCommon/DotNETUtilities/CommandLineArguments.cs b/Engine/Source/Programs/DotNETCommon/DotNETUtilities/CommandLineArguments.cs
--- a/Engine/Source/Programs/DotNETCommon/DotNETUtilities/CommandLineArguments.cs
+++ b/Engine/Source/Programs/DotNETCommon/DotNETUtilities/CommandLineArguments.cs
@@ -3,6 +3,7 @@
 using System;
 using System.Collections;
 using System.Collections.Generic;
+using System.ComponentModel;
 using System.Diagnostics;
 using System.Linq;
 using System.Reflection;
@@ -571,6 +572,43 @@ namespace Tools.DotNETCommon
 			}
 		}
 
+		/// <summary>
+		/// Gets the prefix for a particular argument
+		/// </summary>
+		/// <param name="FieldInfo">The field hosting the attribute</param>
+		/// <param name="Attribute">The attribute instance</param>
+		/// <returns>Prefix for this argument</returns>
+		private static string GetArgumentPrefix(FieldInfo FieldInfo, CommandLineAttribute Attribute)
+		{
+			// Get the inner field type, unwrapping nullable types
+			Type ValueType = FieldInfo.FieldType;
+			if (ValueType.IsGenericType && ValueType.GetGenericTypeDefinition() == typeof(Nullable<>))
+			{
+				ValueType = ValueType.GetGenericArguments()[0];
+			}
+
+			string Prefix = Attribute.Prefix;
+			if (Prefix == null)
+			{
+				if (ValueType == typeof(bool))
+				{
+					Prefix = String.Format("-{0}", FieldInfo.Name);
+				}
+				else
+				{
+					Prefix = String.Format("-{0}=", FieldInfo.Name);
+				}
+			}
+			else
+			{
+				if (ValueType != typeof(bool) && Attribute.Value == null && !Prefix.EndsWith("=") && !Prefix.EndsWith(":"))
+				{
+					Prefix = Prefix + "=";
+				}
+			}
+			return Prefix;
+		}
+
 		/// <summary>
 		/// Applies these arguments to fields with the [CommandLine] attribute in the given object.
 		/// </summary>
@@ -593,33 +631,8 @@ namespace Tools.DotNETCommon
 					IEnumerable<CommandLineAttribute> Attributes = FieldInfo.GetCustomAttributes<CommandLineAttribute>();
 					foreach(CommandLineAttribute Attribute in Attributes)
 					{
-						// Get the inner field type, unwrapping nullable types
-						Type ValueType = FieldInfo.FieldType;
-						if(ValueType.IsGenericType && ValueType.GetGenericTypeDefinition() == typeof(Nullable<>))
-						{
-							ValueType = ValueType.GetGenericArguments()[0];
-						}
-
 						// Get the appropriate prefix for this attribute
-						string Prefix = Attribute.Prefix;
-						if(Prefix == null)
-						{
-							if(ValueType == typeof(bool))
-							{
-								Prefix = String.Format("-{0}", FieldInfo.Name);
-							}
-							else
-							{
-								Prefix = String.Format("-{0}=", FieldInfo.Name);
-							}
-						}
-						else
-						{
-							if(ValueType != typeof(bool) && Attribute.Value == null && !Prefix.EndsWith("=") && !Prefix.EndsWith(":"))
-							{
-								Prefix = Prefix + "=";
-							}
-						}
+						string Prefix = GetArgumentPrefix(FieldInfo, Attribute);
 
 						// Get the value with the correct prefix
 						int FirstIndex;
@@ -698,6 +711,52 @@ namespace Tools.DotNETCommon
 			}
 		}
 
+		/// <summary>
+		/// Gets help text for the arguments of a given type
+		/// </summary>
+		/// <param name="Type">The type to find parameters for</param>
+		/// <returns>List of parameters</returns>
+		public static List<KeyValuePair<string, string>> GetParameters(Type Type)
+		{
+			List<KeyValuePair<string, string>> Parameters = new List<KeyValuePair<string, string>>();
+			for (Type TargetType = Type; TargetType != typeof(object); TargetType = TargetType.BaseType)
+			{
+				foreach (FieldInfo FieldInfo in TargetType.GetFields(BindingFlags.Instance | BindingFlags.GetField | BindingFlags.GetProperty | BindingFlags.Public | BindingFlags.NonPublic | BindingFlags.DeclaredOnly))
+				{
+					List<CommandLineAttribute> Attributes = FieldInfo.GetCustomAttributes<CommandLineAttribute>().ToList();
+					if(Attributes.Count > 0)
+					{
+						StringBuilder DescriptionBuilder = new StringBuilder();
+						foreach (DescriptionAttribute Attribute in FieldInfo.GetCustomAttributes<DescriptionAttribute>())
+						{
+							if(DescriptionBuilder.Length > 0)
+							{
+								DescriptionBuilder.Append("\n");
+							}
+							DescriptionBuilder.Append(Attribute.Description);
+						}
+
+						string Description = DescriptionBuilder.ToString();
+						if (Description.Length == 0)
+						{
+							Description = "No description available.";
+						}
+
+						foreach (CommandLineAttribute Attribute in Attributes)
+						{
+							string Prefix = GetArgumentPrefix(FieldInfo, Attribute);
+							if(Prefix.EndsWith("=", StringComparison.Ordinal))
+							{
+								Prefix += "...";
+							}
+							Parameters.Add(new KeyValuePair<string, string>(Prefix, Description));
+						}
+					}
+				}
+			}
+			return Parameters;
+		}
+
 		/// <summary>
 		/// Splits a command line into individual arguments
 		/// </summary>
diff --git a/Engine/Source/Programs/DotNETCommon/DotNETUtilities/DotNETUtilities.csproj b/Engine/Source/Programs/DotNETCommon/DotNETUtilities/DotNETUtilities.csproj
--- a/Engine/Source/Programs/DotNETCommon/DotNETUtilities/DotNETUtilities.csproj
+++ b/Engine/Source/Programs/DotNETCommon/DotNETUtilities/DotNETUtilities.csproj
@@ -82,6 +82,7 @@
     <Compile Include="FileReference.cs" />
     <Compile Include="FileSystemReference.cs" />
     <Compile Include="FileUtils.cs" />
+    <Compile Include="HelpUtils.cs" />
     <Compile Include="Json.cs" />
     <Compile Include="JsonObject.cs" />
     <Compile Include="JsonWriter.cs" />
diff --git a/Engine/Source/Programs/DotNETCommon/DotNETUtilities/HelpUtils.cs b/Engine/Source/Programs/DotNETCommon/DotNETUtilities/HelpUtils.cs
new file mode 100644
--- /dev/null
+++ b/Engine/Source/Programs/DotNETCommon/DotNETUtilities/HelpUtils.cs
@@ -0,0 +1,172 @@
+using System;
+using System.Collections.Generic;
+using System.ComponentModel;
+using System.Linq;
+using System.Text;
+using System.Threading.Tasks;
+
+namespace Tools.DotNETCommon
+{
+	/// <summary>
+	/// Utility functions for showing help for objects
+	/// </summary>
+	public static class HelpUtils
+	{
+		/// <summary>
+		/// Gets the width of the window for formatting purposes
+		/// </summary>
+		public static int WindowWidth
+		{
+			get
+			{
+				// Get the window width, using a default value if there's no console attached to this process.
+				int NewWindowWidth;
+				try
+				{
+					NewWindowWidth = Console.WindowWidth;
+				}
+				catch
+				{
+					NewWindowWidth = 240;
+				}
+
+				if (NewWindowWidth <= 0)
+				{
+					NewWindowWidth = 240;
+				}
+
+				return NewWindowWidth;
+			}
+		}
+
+		/// <summary>
+		/// Prints help for the given object type
+		/// </summary>
+		/// <param name="Type">Type to print help for</param>
+		public static void PrintHelp(string Title, Type Type)
+		{
+			List<string> Lines = GetHelpText(Title, Type);
+			foreach(string Line in Lines)
+			{
+				Log.TraceInformation("{0}", Line);
+			}
+		}
+
+		/// <summary>
+		/// Get the help text for a given type, using any attribute-driven command line arguments, and the [Description] attribute for help text.
+		/// </summary>
+		/// <param name="Title">Title for the help text</param>
+		/// <param name="Type">Type to query</param>
+		/// <param name="MaxWidth">Maximum width for the returned lines</param>
+		/// <returns>List of help lines</returns>
+		public static List<string> GetHelpText(string Title, Type Type)
+		{
+			return GetHelpText(Title, GetDescription(Type), CommandLineArguments.GetParameters(Type), WindowWidth - 1);
+		}
+
+		/// <summary>
+		/// Get the help text for a given type, using any attribute-driven command line arguments, and the [Description] attribute for help text.
+		/// </summary>
+		/// <param name="Title">Title for the help text</param>
+		/// <param name="Description">Help description</param>
+		/// <param name="Parameters">Arguments for the command</param>
+		/// <param name="MaxWidth">Maximum width for the returned lines</param>
+		/// <returns>List of help lines</returns>
+		public static List<string> GetHelpText(string Title, string Description, List<KeyValuePair<string, string>> Parameters, int MaxWidth)
+		{
+			List<string> Lines = new List<string>();
+
+			if (!String.IsNullOrEmpty(Title))
+			{
+				Lines.AddRange(StringUtils.WordWrap(Title, MaxWidth));
+			}
+
+			if (!String.IsNullOrEmpty(Description))
+			{
+				if (Lines.Count > 0)
+				{
+					Lines.Add("");
+				}
+				Lines.AddRange(StringUtils.WordWrap(Description, MaxWidth));
+			}
+
+			if (Parameters.Count > 0)
+			{
+				if (Lines.Count > 0)
+				{
+					Lines.Add("");
+				}
+				Lines.Add("Parameters:");
+				Lines.AddRange(Tabulate(Parameters, 4, 24, MaxWidth));
+			}
+
+			return Lines;
+		}
+
+		/// <summary>
+		/// Gets the description from a type
+		/// </summary>
+		/// <param name="Type">The type to get a description for</param>
+		/// <returns>The description text</returns>
+		public static string GetDescription(Type Type)
+		{
+			StringBuilder DescriptionText = new StringBuilder();
+			foreach (DescriptionAttribute Attribute in Type.GetCustomAttributes(typeof(DescriptionAttribute), false))
+			{
+				if (DescriptionText.Length > 0)
+				{
+					DescriptionText.AppendLine();
+				}
+				DescriptionText.AppendLine(Attribute.Description);
+			}
+			return DescriptionText.ToString();
+		}
+
+		/// <summary>
+		/// Formats the given parameters as so:
+		///     -Param1     Param1 Description
+		///
+		///     -Param2      Param2 Description, this description is
+		///                  longer and splits onto a separate line. 
+		///
+		///     -Param3      Param3 Description continues as before. 
+		/// </summary>
+		/// <param name="Items">List of parameters arranged as "-ParamName Param Description"</param>
+		/// <param name="Indent">Indent from the left hand side</param>
+		/// <param name="DefaultRightPadding">The minimum padding from the start of the param name to the start of the description (resizes with larger param names)</param>
+		/// <returns></returns>
+		public static List<string> Tabulate(List<KeyValuePair<string, string>> Items, int Indent, int MinFirstColumnWidth, int MaxWidth)
+		{
+			List<string> Lines = new List<string>();
+			if(Items.Count > 0)
+			{
+				// string used to intent the param
+				string IndentString = new string(' ', Indent);
+
+				// default the padding value
+				int RightPadding = Math.Max(MinFirstColumnWidth, Items.Max(x => x.Key.Length + 1));
+
+				// Build the formatted params
+				foreach(KeyValuePair<string, string> Item in Items)
+				{
+					// build the param first, including intend and padding on the rights size
+					string ParamString = IndentString + Item.Key.PadRight(RightPadding);
+
+					// Build the description line by line, adding the same amount of intending each time. 
+					List<string> DescriptionLines = StringUtils.WordWrap(Item.Value, MaxWidth - ParamString.Length);
+
+					foreach(string DescriptionLine in DescriptionLines)
+					{
+						// Formatting as following:
+						// <Indent>-param<Right Padding>Description<New line>
+						Lines.Add(ParamString + DescriptionLine);
+
+						// we replace the param string on subsequent lines with white space of the same length
+						ParamString = string.Empty.PadRight(IndentString.Length + RightPadding);
+					}
+				}
+			}
+			return Lines;
+		}
+	}
+}
diff --git a/Engine/Source/Programs/DotNETCommon/DotNETUtilities/StringUtils.cs b/Engine/Source/Programs/DotNETCommon/DotNETUtilities/StringUtils.cs
--- a/Engine/Source/Programs/DotNETCommon/DotNETUtilities/StringUtils.cs
+++ b/Engine/Source/Programs/DotNETCommon/DotNETUtilities/StringUtils.cs
@@ -25,6 +25,67 @@ namespace Tools.DotNETCommon
 			return Result;
 		}
 
+		/// <summary>
+		/// Takes a given sentence and wraps it on a word by word basis so that no line exceeds the
+		/// set maximum line length. Words longer than a line are broken up. Returns the sentence as
+		/// a list of individual lines.
+		/// </summary>
+		/// <param name="Text">The text to be wrapped</param>
+		/// <param name="MaxWidth">The maximum (non negative) length of the returned sentences</param>
+		/// <param name="Lines">Receives a list of word-wrapped lines</param>
+		public static List<string> WordWrap(string Text, int MaxWidth)
+		{
+			// Early out
+			if (Text.Length == 0)
+			{
+				return new List<string>();
+			}
+
+			string[] Words = Text.Split(' ');
+			List<string> WrappedWords = new List<string>();
+
+			string CurrentSentence = string.Empty;
+			foreach (var Word in Words)
+			{
+				// if this is a very large word, split it
+				if (Word.Length > MaxWidth)
+				{
+					// Top up the current line
+					WrappedWords.Add(CurrentSentence + Word.Substring(0, MaxWidth - CurrentSentence.Length));
+
+					int length = MaxWidth - CurrentSentence.Length;
+					while (length + MaxWidth < Word.Length)
+					{
+						// Place the starting lengths into their own lines
+						WrappedWords.Add(Word.Substring(length, Math.Min(MaxWidth, Word.Length - length)));
+						length += MaxWidth;
+					}
+
+					// then the trailing end into the next line
+					CurrentSentence += Word.Substring(length, Math.Min(MaxWidth, Word.Length - length)) + " ";
+				}
+				else
+				{
+					if (CurrentSentence.Length + Word.Length > MaxWidth)
+					{
+						// next line and reset sentence
+						WrappedWords.Add(CurrentSentence);
+						CurrentSentence = string.Empty;
+					}
+
+					// Add the word to the current sentence.
+					CurrentSentence += Word + " ";
+				}
+			}
+
+			if (CurrentSentence.Length > 0)
+			{
+				WrappedWords.Add(CurrentSentence);
+			}
+
+			return WrappedWords;
+		}
+
 		/// <summary>
 		/// Extension method to allow formatting a string to a stringbuilder and appending a newline
 		/// </summary>
