Gauntlet: Add some additional android verbose logging to help diagnose issues
#rnx

diff --git a/Engine/Source/Programs/AutomationTool/Gauntlet/Platform/Android/Gauntlet.TargetDeviceAndroid.cs b/Engine/Source/Programs/AutomationTool/Gauntlet/Platform/Android/Gauntlet.TargetDeviceAndroid.cs
--- a/Engine/Source/Programs/AutomationTool/Gauntlet/Platform/Android/Gauntlet.TargetDeviceAndroid.cs
+++ b/Engine/Source/Programs/AutomationTool/Gauntlet/Platform/Android/Gauntlet.TargetDeviceAndroid.cs
@@ -213,9 +213,13 @@ namespace Gauntlet
 
 		protected void SaveArtifacts()
 		{
+
 			// copy remote artifacts to local
 			if (Directory.Exists(Install.AndroidDevice.LocalCachePath))
 			{
+
+				Log.Verbose("Deleting {0}", Install.AndroidDevice.LocalCachePath);
+
 				try
 				{
 					// don't consider this fatal, people often have the directory or a file open
@@ -235,14 +239,19 @@ namespace Gauntlet
 			catch (Exception Ex)
 			{
 				Log.Warning("Exception marking directory for cleanup {0}", Ex.Message);
-			}
+			}			
 
 			string LocalSaved = Path.Combine(Install.AndroidDevice.LocalCachePath, "Saved");
+			Log.Verbose("Creating {0}", LocalSaved);
 			Directory.CreateDirectory(LocalSaved);
 
+
+			
+
+
 			// pull all the artifacts
 			string ArtifactPullCommand = string.Format("pull {0} {1}", Install.AndroidDevice.DeviceArtifactPath, Install.AndroidDevice.LocalCachePath);
-			IProcessResult PullCmd = Install.AndroidDevice.RunAdbDeviceCommand(ArtifactPullCommand);
+			IProcessResult PullCmd = Install.AndroidDevice.RunAdbDeviceCommand(ArtifactPullCommand, bShouldLogCommand: Log.IsVerbose);
 
 			if (PullCmd.ExitCode != 0)
 			{
@@ -261,7 +270,7 @@ namespace Gauntlet
 			}
 
 			// pull the logcat over from device.
-			IProcessResult LogcatResult = Install.AndroidDevice.RunAdbDeviceCommand("logcat -d");
+			IProcessResult LogcatResult = Install.AndroidDevice.RunAdbDeviceCommand("logcat -d", bShouldLogCommand: Log.IsVerbose);
 
 			string LogcatFilename = "Logcat.log";
 			// Save logcat dump to local artifact path.
