Ignore any mismatched timestamps on DLLs for now; there are many generated as runtime dependencies which clash.

diff --git a/Engine/Source/Programs/AutomationTool/BuildGraph/TempStorage.cs b/Engine/Source/Programs/AutomationTool/BuildGraph/TempStorage.cs
--- a/Engine/Source/Programs/AutomationTool/BuildGraph/TempStorage.cs
+++ b/Engine/Source/Programs/AutomationTool/BuildGraph/TempStorage.cs
@@ -1133,6 +1133,10 @@ namespace AutomationTool
 			{
 				return true;
 			}
+			if (FileName.EndsWith(".dll", StringComparison.OrdinalIgnoreCase))
+			{
+				return true;
+			}
 			return false;
 		}
 	}
