UE-80607

Adding -TimeStamps to global args for AutomationTool since it should be.

#rnx

diff --git a/Engine/Source/Programs/AutomationTool/AutomationUtils/Automation.cs b/Engine/Source/Programs/AutomationTool/AutomationUtils/Automation.cs
--- a/Engine/Source/Programs/AutomationTool/AutomationUtils/Automation.cs
+++ b/Engine/Source/Programs/AutomationTool/AutomationUtils/Automation.cs
@@ -92,6 +92,7 @@ namespace AutomationTool
 
 		public static CommandLineArg CompileOnly = new CommandLineArg("-CompileOnly");
 		public static CommandLineArg Verbose = new CommandLineArg("-Verbose");
+		public static CommandLineArg TimeStamps = new CommandLineArg("-TimeStamps");
 		public static CommandLineArg Submit = new CommandLineArg("-Submit");
 		public static CommandLineArg NoSubmit = new CommandLineArg("-NoSubmit");
 		public static CommandLineArg NoP4 = new CommandLineArg("-NoP4");
