Added heartbeats to the gauntlet module which output log lines regularly to track test controller status. Test controllers can mark these heartbeats as active which denotes that the test is actively functioning as intended. In Csharp, Gauntlet test nodes listen for heartbeats and can optionally fail the test run if a heartbeat hasn't been detected for long enough or if heartbeats are not marked active for long enough.

#rnx

diff --git a/Engine/Plugins/Experimental/Gauntlet/Source/Gauntlet/Private/GauntletModule.cpp b/Engine/Plugins/Experimental/Gauntlet/Source/Gauntlet/Private/GauntletModule.cpp
--- a/Engine/Plugins/Experimental/Gauntlet/Source/Gauntlet/Private/GauntletModule.cpp
+++ b/Engine/Plugins/Experimental/Gauntlet/Source/Gauntlet/Private/GauntletModule.cpp
@@ -30,13 +30,11 @@ public:
 	FName	GetCurrentState() const override { return CurrentState; }
 	double	GetTimeInCurrentState() const override { return TimeInCurrentState; }
 
-	
 	FString			GetCurrentMap() const { return CurrentMap; }
 
-	void			SetScreenshotPeriod(float Period)
-	{
-		ScreenshotPeriod = Period;
-	}
+	void			SetScreenshotPeriod(float Period) { ScreenshotPeriod = Period; }
+
+	void			MarkHeartbeatActive(const FString& OptionalStatusMessage = FString());
 
 	UGauntletTestController*	GetTestController(UClass* ControllerClass) override;
 
@@ -65,6 +63,9 @@ protected:
 
 	void		TakeScreenshot();
 
+	/** Prints a heartbeat message to the log */
+	void		LogHeartbeat();
+
 private:
 
 	/** Handle to our tick callback */
@@ -85,6 +86,19 @@ private:
 
 	double							LastScreenshotTime;
 	float							ScreenshotPeriod;
+
+	/**
+	 * Heartbeat Internal Members
+	 */
+
+	/** True if the next heartbeat is marked Active rather than Idle */
+	bool bHeartbeatActive;
+	/** How often to regularly log heartbeat messages, in seconds */
+	float HeartbeatPeriod;
+	/** Platform time in seconds of the last heartbeat. Used with HeartbeatPeriod. */
+	float LastHeartbeatTime;
+	/** Optional status message to include with Active heartbeats */
+	FString HeartbeatStatus;
 };
 
 
@@ -102,6 +116,10 @@ void FGauntletModuleImpl::StartupModule()
 	TimeInCurrentState = 0;
 	LastScreenshotTime = 0.0f;
 	ScreenshotPeriod = 0.0f;
+	bHeartbeatActive = false;
+	HeartbeatPeriod = 0.0f;
+	LastHeartbeatTime = 0.f;
+	HeartbeatStatus = FString();
 
 	if (IsRunningGame() || IsRunningDedicatedServer())
 	{
@@ -116,6 +134,7 @@ void FGauntletModuleImpl::OnPostEngineInit()
 	FCoreUObjectDelegates::PreLoadMap.AddRaw(this, &FGauntletModuleImpl::InnerPreMapChange);
 
 	FParse::Value(FCommandLine::Get(), TEXT("gauntlet.screenshotperiod="), ScreenshotPeriod);
+	FParse::Value(FCommandLine::Get(), TEXT("gauntlet.heartbeatperiod="), HeartbeatPeriod);
 
 	LoadControllers();
 
@@ -324,6 +343,11 @@ void FGauntletModuleImpl::InnerTick(const float TimeDelta)
 			TakeScreenshot();
 			LastScreenshotTime = FPlatformTime::Seconds();
 		}
+
+		if (HeartbeatPeriod > 0.0f && (FPlatformTime::Seconds() - LastHeartbeatTime) > HeartbeatPeriod)
+		{
+			LogHeartbeat();
+		}
 	}
 
 	TimeInCurrentState += TimeDelta;
@@ -339,6 +363,25 @@ void FGauntletModuleImpl::TakeScreenshot()
 	FScreenshotRequest::RequestScreenshot(true);
 }
 
+void FGauntletModuleImpl::LogHeartbeat()
+{
+	FString HeartbeatMessage = bHeartbeatActive ? FString("GauntletHeartbeat: Active") : FString("GauntletHeartbeat: Idle");
+	UE_LOG(LogGauntlet, Display, TEXT("%s %s"), *HeartbeatMessage, *HeartbeatStatus);
+	bHeartbeatActive = false;
+	HeartbeatStatus.Empty();
+	LastHeartbeatTime = FPlatformTime::Seconds();
+}
+
+void FGauntletModuleImpl::MarkHeartbeatActive(const FString& OptionalStatusMessage /*= FString()*/)
+{
+	bHeartbeatActive = true;
+	// If a status message is given, heartbeat immediately so that the message is printed at the relevant time.
+	if (!OptionalStatusMessage.IsEmpty())
+	{
+		HeartbeatStatus = OptionalStatusMessage;
+		LogHeartbeat();
+	}
+}
 
 UGauntletTestController* FGauntletModuleImpl::GetTestController(UClass* ControllerClass)
 {
diff --git a/Engine/Plugins/Experimental/Gauntlet/Source/Gauntlet/Private/GauntletTestController.cpp b/Engine/Plugins/Experimental/Gauntlet/Source/Gauntlet/Private/GauntletTestController.cpp
--- a/Engine/Plugins/Experimental/Gauntlet/Source/Gauntlet/Private/GauntletTestController.cpp
+++ b/Engine/Plugins/Experimental/Gauntlet/Source/Gauntlet/Private/GauntletTestController.cpp
@@ -43,6 +43,17 @@ FGauntletModule* UGauntletTestController::GetGauntlet()
 	return ParentModule;
 }
 
+void UGauntletTestController::MarkHeartbeatActive(const FString& OptionalStatusMessage /*= FString()*/)
+{
+	FString StatusMessage = OptionalStatusMessage;
+	if (!StatusMessage.IsEmpty())
+	{
+		StatusMessage = FString::Printf(TEXT("[%s] %s"), *GetName(), *StatusMessage);
+	}
+
+	GetGauntlet()->MarkHeartbeatActive(StatusMessage);
+}
+
 UWorld* UGauntletTestController::GetWorld() const
 {
 	return GWorld;
diff --git a/Engine/Plugins/Experimental/Gauntlet/Source/Gauntlet/Public/GauntletModule.h b/Engine/Plugins/Experimental/Gauntlet/Source/Gauntlet/Public/GauntletModule.h
--- a/Engine/Plugins/Experimental/Gauntlet/Source/Gauntlet/Public/GauntletModule.h
+++ b/Engine/Plugins/Experimental/Gauntlet/Source/Gauntlet/Public/GauntletModule.h
@@ -82,6 +82,11 @@ public:
 	 */
 	virtual void			SetScreenshotPeriod(float Period) = 0;
 
+	/**
+	 * Mark the next heartbeat as active. If a status message is given, immediately log an active heartbeat with that message.
+	 */
+	virtual void			MarkHeartbeatActive(const FString& OptionalStatusMessage = FString()) = 0;
+
 	/**
 	 * Returns the first controller (if any) matching the provided name
 	 */
diff --git a/Engine/Plugins/Experimental/Gauntlet/Source/Gauntlet/Public/GauntletTestController.h b/Engine/Plugins/Experimental/Gauntlet/Source/Gauntlet/Public/GauntletTestController.h
--- a/Engine/Plugins/Experimental/Gauntlet/Source/Gauntlet/Public/GauntletTestController.h
+++ b/Engine/Plugins/Experimental/Gauntlet/Source/Gauntlet/Public/GauntletTestController.h
@@ -94,11 +94,14 @@ protected:
 	 */
 	FGauntletModule* GetGauntlet();
 
+	/**
+	 * Marks the next heartbeat as active and immediately forces a heartbeat with the given status message if one is specified.
+	 * Test controllers should call this regularly as purposeful actions happens or meaningful states are achieved within the test.
+	 */
+	void MarkHeartbeatActive(const FString& OptionalStatusMessage = FString());
+
 private:
 
 	FGauntletModule*  ParentModule;
-};
-
-
-
 
+};
diff --git a/Engine/Source/Programs/AutomationTool/Gauntlet/Unreal/Base/Gauntlet.UnrealTestConfiguration.cs b/Engine/Source/Programs/AutomationTool/Gauntlet/Unreal/Base/Gauntlet.UnrealTestConfiguration.cs
--- a/Engine/Source/Programs/AutomationTool/Gauntlet/Unreal/Base/Gauntlet.UnrealTestConfiguration.cs
+++ b/Engine/Source/Programs/AutomationTool/Gauntlet/Unreal/Base/Gauntlet.UnrealTestConfiguration.cs
@@ -123,6 +123,48 @@ namespace Gauntlet
 
 	}
 
+	/// <summary>
+	/// Collection of parameters that control how heartbeats coming from the native gauntlet controller for this role should be handled.
+	/// To make best use of this, your GauntletTestController should regularly call MarkHeartbeatActive().
+	/// Set bExpectHeartbeats to true to enable killing the App Instance when expected heartbeats are not detected.
+	/// </summary>
+	public class UnrealHeartbeatOptions
+	{
+		/// <summary>
+		/// The amount of time between regular heartbeats. This value is passed along through the command line.
+		/// </summary>
+		public float HeartbeatPeriod;
+
+		/// <summary>
+		/// Set to true to allow the App Instance to be killed when expected heartbeats are not detected. If left false, heartbeat timeouts will not result in any action or timeouts.
+		/// </summary>
+		public bool bExpectHeartbeats;
+		
+		/// <summary>
+		/// The max amount of time allowed before the first "active" heartbeat is detected
+		/// </summary>
+		public float TimeoutBeforeFirstActiveHeartbeat;
+
+		/// <summary>
+		/// The max amount of time allowed between "active" heartbeats
+		/// </summary>
+		public float TimeoutBetweenActiveHeartbeats;
+
+		/// <summary>
+		/// The max amount of time allowed between any heartbeats, active or not
+		/// </summary>
+		public float TimeoutBetweenAnyHeartbeats;
+
+		public UnrealHeartbeatOptions(float InHeartbeatPeriod = 30f, bool bShouldExpectHeartbeats = false, float InTimeoutBeforeFirstActiveHeartbeat = 0f, float InTimeoutBetweenActiveHeartbeats = 0f, float InTimeoutBetweenAnyHeartbeats = 90f)
+		{
+			HeartbeatPeriod = InHeartbeatPeriod;
+			bExpectHeartbeats = bShouldExpectHeartbeats;
+			TimeoutBeforeFirstActiveHeartbeat = InTimeoutBeforeFirstActiveHeartbeat;
+			TimeoutBetweenActiveHeartbeats = InTimeoutBetweenActiveHeartbeats;
+			TimeoutBetweenAnyHeartbeats = InTimeoutBetweenAnyHeartbeats;
+		}
+
+	}
 
 	/// <summary>
 	///	TestConfiguration describes the setup that is required for a specific test. 
@@ -231,6 +273,11 @@ namespace Gauntlet
 		/// </summary>
 		public bool AllRolesExit { get; set; }
 
+		/// <summary>
+		/// The collection of options which define heartbeat behavior
+		/// </summary>
+		public UnrealHeartbeatOptions HeartbeatOptions { get; set; }
+
 		// Member variables 
 
 		/// <summary>
@@ -247,6 +294,8 @@ namespace Gauntlet
 
 			// create the role structure
 			RequiredRoles = new Dictionary<UnrealTargetRole, List<UnrealTestRole>>();
+
+			HeartbeatOptions = new UnrealHeartbeatOptions();
 		}
 
 		/// <summary>
@@ -399,6 +448,12 @@ namespace Gauntlet
 			}
 
 			AppConfig.CommandLine += " -stdout -AllowStdOutLogVerbosity";
+
+			float HeartbeatPeriod = Globals.Params.ParseValue("HeartbeatPeriod", HeartbeatOptions.HeartbeatPeriod);
+			if (HeartbeatPeriod > 0)
+			{
+				AppConfig.CommandLine += string.Format(" -gauntlet.heartbeatperiod={0}", HeartbeatPeriod);
+			}
 		}
 	}
 
diff --git a/Engine/Source/Programs/AutomationTool/Gauntlet/Unreal/Base/Gauntlet.UnrealTestNode.cs b/Engine/Source/Programs/AutomationTool/Gauntlet/Unreal/Base/Gauntlet.UnrealTestNode.cs
--- a/Engine/Source/Programs/AutomationTool/Gauntlet/Unreal/Base/Gauntlet.UnrealTestNode.cs
+++ b/Engine/Source/Programs/AutomationTool/Gauntlet/Unreal/Base/Gauntlet.UnrealTestNode.cs
@@ -123,6 +123,9 @@ namespace Gauntlet
 
 		protected int TimeToWaitForProcesses { get; set; }
 
+		protected DateTime LastHeartbeatTime = DateTime.MinValue;
+		protected DateTime LastActiveHeartbeatTime = DateTime.MinValue;
+
 		// End  UnrealTestNode properties and members 
 
 		// artifact paths that have been used in this run
@@ -575,7 +578,7 @@ namespace Gauntlet
 			if (App != null)
 			{
 				UnrealLogParser Parser = new UnrealLogParser(App.StdOut);
-
+				
 				// TODO - hardcoded for Orion
 				List<string> TestLines = Parser.GetLogChannel("Gauntlet").ToList();
 
@@ -584,12 +587,25 @@ namespace Gauntlet
 				for (int i = LastLogCount; i < TestLines.Count; i++)
 				{
 					Log.Info(TestLines[i]);
+
+					if (Regex.IsMatch(TestLines[i], @".*GauntletHeartbeat\: Active.*"))
+					{
+						LastHeartbeatTime = DateTime.Now;
+						LastActiveHeartbeatTime = DateTime.Now;
+					}
+					else if (Regex.IsMatch(TestLines[i], @".*GauntletHeartbeat\: Idle.*"))
+					{
+						LastHeartbeatTime = DateTime.Now;
+					}
 				}
 
 				LastLogCount = TestLines.Count;
+
+				// Detect missed heartbeats and fail the test
+				CheckHeartbeat();
 			}
 
-			
+
 			// Check status and health after updating logs
 			if (GetTestStatus() == TestStatus.InProgress && IsTestRunning() == false)
 			{
@@ -831,6 +847,52 @@ namespace Gauntlet
 			return ExitCode;
 		}
 
+		private void CheckHeartbeat()
+		{
+			if (CachedConfig == null || CachedConfig.HeartbeatOptions.bExpectHeartbeats == false)
+			{
+				return;
+			}
+
+			UnrealHeartbeatOptions HeartbeatOptions = CachedConfig.HeartbeatOptions;
+
+			// First active heartbeat has not happened yet and timeout before first active heartbeat is enabled
+			if (LastActiveHeartbeatTime == DateTime.MinValue && HeartbeatOptions.TimeoutBeforeFirstActiveHeartbeat > 0)
+			{
+				double SecondsSinceSessionStart = DateTime.Now.Subtract(SessionStartTime).TotalSeconds;
+				if (SecondsSinceSessionStart > HeartbeatOptions.TimeoutBeforeFirstActiveHeartbeat)
+				{
+					Log.Error("{0} seconds have passed without detecting the first active Gauntlet heartbeat.", HeartbeatOptions.TimeoutBeforeFirstActiveHeartbeat);
+					MarkTestComplete();
+					SetUnrealTestResult(TestResult.TimedOut);
+				}
+			}
+
+			// First active heartbeat has happened and timeout between active heartbeats is enabled
+			if (LastActiveHeartbeatTime != DateTime.MinValue && HeartbeatOptions.TimeoutBetweenActiveHeartbeats > 0)
+			{
+				double SecondsSinceLastActiveHeartbeat = DateTime.Now.Subtract(LastActiveHeartbeatTime).TotalSeconds;
+				if (SecondsSinceLastActiveHeartbeat > HeartbeatOptions.TimeoutBetweenActiveHeartbeats)
+				{
+					Log.Error("{0} seconds have passed without detecting any active Gauntlet heartbeats.", HeartbeatOptions.TimeoutBetweenActiveHeartbeats);
+					MarkTestComplete();
+					SetUnrealTestResult(TestResult.TimedOut);
+				}
+			}
+
+			// First heartbeat has happened and timeout between heartbeats is enabled
+			if (LastHeartbeatTime != DateTime.MinValue && HeartbeatOptions.TimeoutBetweenAnyHeartbeats > 0)
+			{
+				double SecondsSinceLastHeartbeat = DateTime.Now.Subtract(LastHeartbeatTime).TotalSeconds;
+				if (SecondsSinceLastHeartbeat > HeartbeatOptions.TimeoutBetweenAnyHeartbeats)
+				{
+					Log.Error("{0} seconds have passed without detecting any Gauntlet heartbeats.", HeartbeatOptions.TimeoutBetweenAnyHeartbeats);
+					MarkTestComplete();
+					SetUnrealTestResult(TestResult.TimedOut);
+				}
+			}
+		}
+
 		/// <summary>
 		/// Returns a hash that represents the results of a role. Should be 0 if no fatal errors or ensures
 		/// </summary>
