Patch pak file won't be generated for EarlyStartup pak files

diff --git a/Engine/Source/Programs/AutomationTool/Scripts/CopyBuildToStagingDirectory.Automation.cs b/Engine/Source/Programs/AutomationTool/Scripts/CopyBuildToStagingDirectory.Automation.cs
--- a/Engine/Source/Programs/AutomationTool/Scripts/CopyBuildToStagingDirectory.Automation.cs
+++ b/Engine/Source/Programs/AutomationTool/Scripts/CopyBuildToStagingDirectory.Automation.cs
@@ -1753,6 +1753,38 @@ public partial class Project : CommandUtils
 		}
 	}
 
+	private static bool ShouldSkipGeneratingPatch(ConfigHierarchy PlatformGameConfig, string Filename)
+	{
+		bool bIsEarlyDownloaderPakFile = false;
+		string EarlyDownloaderPakFilePrefix = "";
+		if (PlatformGameConfig.GetString("/Script/UnrealEd.ProjectPackagingSettings", "EarlyDownloaderPakFilePrefix", out EarlyDownloaderPakFilePrefix))
+		{
+			bIsEarlyDownloaderPakFile = String.Equals(EarlyDownloaderPakFilePrefix, Filename, StringComparison.OrdinalIgnoreCase);
+		}
+
+		return bIsEarlyDownloaderPakFile;
+	}
+
+	private static string GetPakFilePostFix(bool bShouldGeneratePatch, bool bHasIterateSharedCookedBuild, ConfigHierarchy PlatformGameConfig, string Filename)
+	{
+		if (ShouldSkipGeneratingPatch(PlatformGameConfig, Filename))
+		{
+			return String.Empty;
+		}
+		else if(bShouldGeneratePatch)
+		{
+			return "_P";
+		}
+		else if(bHasIterateSharedCookedBuild)
+		{
+			// shared cooked builds will produce a patch
+			// then be combined with the shared cooked build
+			return "_S_P";
+		}
+
+		return String.Empty;
+	}
+
 	/// <summary>
 	/// Creates a pak file using response file.
 	/// </summary>
@@ -1769,18 +1801,6 @@ public partial class Project : CommandUtils
 			LogInformation("Generating patch required a based on release version flag");
 		}
 
-		string PostFix = "";
-		if (bShouldGeneratePatch)
-		{
-			PostFix += "_P";
-		}
-		else if (Params.HasIterateSharedCookedBuild)
-		{
-			// shared cooked builds will produce a patch
-			// then be combined with the shared cooked build
-			PostFix += "_S_P";
-		}
-
 		const string OutputFilenameExtension = ".pak";
 
 		// read some compression settings from the project (once, shared across all pak commands)
@@ -1791,6 +1811,8 @@ public partial class Project : CommandUtils
 			CompressionFormats = " -compressionformats=" + CompressionFormats;
 		}
 
+		Func<string, string> GetPostFix = PakFilename => GetPakFilePostFix(bShouldGeneratePatch, Params.HasIterateSharedCookedBuild, PlatformGameConfig, PakFilename);
+
 		// the game may want to control compression settings, but since it may be in a plugin that checks the commandline for the settings, we need to pass
 		// the settings directly on the UnrealPak commandline, and not put it into the batch file lines (plugins can't get the unrealpak command list, and
 		// there's not a great way to communicate random strings down into the plugins during plugin init time)
@@ -1807,10 +1829,16 @@ public partial class Project : CommandUtils
 		{
 			foreach (CreatePakParams PakParams in PakParamsList)
 			{
+				if(ShouldSkipGeneratingPatch(PlatformGameConfig, PakParams.PakName))
+				{
+					continue;
+				}
+
 				string OutputFilename = PakParams.PakName + "-" + SC.FinalCookPlatform;
 				string ExistingPatchSearchPath = SC.StageTargetPlatform.GetReleasePakFilePath(SC, Params, null);
 				if (Directory.Exists(ExistingPatchSearchPath))
 				{
+					string PostFix = GetPostFix(PakParams.PakName);
 					IEnumerable<string> PakFileSet = Directory.EnumerateFiles(ExistingPatchSearchPath, OutputFilename + "*" + PostFix + OutputFilenameExtension);
 					NumPakFiles += PakFileSet.Count();
 
@@ -1845,7 +1873,8 @@ public partial class Project : CommandUtils
 		foreach (CreatePakParams PakParams in PakParamsList)
 		{
 			string OutputFilename = PakParams.PakName + "-" + SC.FinalCookPlatform;
-			if (bShouldGeneratePatch)
+			string PostFix = GetPostFix(PakParams.PakName);
+			if (bShouldGeneratePatch && !ShouldSkipGeneratingPatch(PlatformGameConfig, PakParams.PakName))
 			{
 				OutputFilename = OutputFilename + "_" + TargetPatchIndex;
 			}
@@ -1952,7 +1981,7 @@ public partial class Project : CommandUtils
 			}
 
 			string PatchSourceContentPath = null;
-			if (bShouldGeneratePatch)
+			if (bShouldGeneratePatch && !ShouldSkipGeneratingPatch(PlatformGameConfig, PakParams.PakName))
 			{
 				// don't include the post fix in this filename because we are looking for the source pak path
 				string PakFilename = PakParams.PakName + "-" + SC.FinalCookPlatform + "*.pak";
@@ -1994,7 +2023,7 @@ public partial class Project : CommandUtils
 					}
 
 					string AdditionalArgs = String.Empty;
-					if (bShouldGeneratePatch)
+					if (bShouldGeneratePatch && !ShouldSkipGeneratingPatch(PlatformGameConfig, PakParams.PakName))
 					{
 						string PatchSeekOptMode = String.Empty;
 						string PatchSeekOptMaxGapSize = String.Empty;
@@ -2085,7 +2114,7 @@ public partial class Project : CommandUtils
 					string VersionString = Params.ChunkInstallVersionString;
 					string ChunkInstallBasePath = CombinePaths(Params.ChunkInstallDirectory, SC.FinalCookPlatform);
 					string RawDataPath = CombinePaths(ChunkInstallBasePath, VersionString, PakName);
-					string RawDataPakPath = CombinePaths(RawDataPath, PakName + "-" + SC.FinalCookPlatform + PostFix + ".pak");
+					string RawDataPakPath = CombinePaths(RawDataPath, PakName + "-" + SC.FinalCookPlatform + GetPostFix(PakName) + ".pak");
 					bool bPakFilesAreSigned = InternalUtils.SafeFileExists(Path.ChangeExtension(OutputLocation.FullName, ".sig"));
 
 					//copy the pak chunk to the raw data folder
@@ -2153,7 +2182,7 @@ public partial class Project : CommandUtils
 				SC.FilesToStage.UFSFiles.Add(OutputRelativeLocation, OutputLocation);
 
 				// add the base pak files to deployment as well
-				if (bShouldGeneratePatch)
+				if (bShouldGeneratePatch && !ShouldSkipGeneratingPatch(PlatformGameConfig, PakName))
 				{
 					string ExistingPatchSearchPath = SC.StageTargetPlatform.GetReleasePakFilePath(SC, Params, null);
 					if (Directory.Exists(ExistingPatchSearchPath))
@@ -2375,8 +2404,10 @@ public partial class Project : CommandUtils
 		ConfigHierarchy PlatformGameConfig = ConfigCache.ReadHierarchy(ConfigHierarchyType.Game, DirectoryReference.FromFile(Params.RawProjectPath), SC.StageTargetPlatform.IniPlatformType);
 
 		bool bShouldGenerateEarlyDownloaderPakFile = false, bForceOneChunkPerFile = false;
+		string EarlyDownloaderPakFilePrefix = string.Empty;
 		PlatformGameConfig.GetBool("/Script/UnrealEd.ProjectPackagingSettings", "GenerateEarlyDownloaderPakFile", out bShouldGenerateEarlyDownloaderPakFile);
 		PlatformGameConfig.GetBool("/Script/UnrealEd.ProjectPackagingSettings", "bForceOneChunkPerFile", out bForceOneChunkPerFile);
+		PlatformGameConfig.GetString("/Script/UnrealEd.ProjectPackagingSettings", "EarlyDownloaderPakFilePrefix", out EarlyDownloaderPakFilePrefix);
 
 		List<ChunkDefinition> ChunkDefinitions = new List<ChunkDefinition>();
 		List<PakFileRules> PakRulesList = GetPakFileRules(Params, SC);
@@ -2525,7 +2556,7 @@ public partial class Project : CommandUtils
 		// Deprecated in favor of DefaultPakFileRules.ini
 		if (bShouldGenerateEarlyDownloaderPakFile)
 		{
-			const string EarlyChunkName = "pakChunkEarly";
+			string EarlyChunkName = EarlyDownloaderPakFilePrefix;
 			ChunkDefinition EarlyChunk = new ChunkDefinition(EarlyChunkName);
 
 			EarlyChunk.bCompressed = true;
