Gauntlet: fix issue with parallel tests copying servers/clients multiple times
#rnx

diff --git a/Engine/Source/Programs/AutomationTool/Gauntlet/Platform/Mac/Gauntlet.TargetDeviceMac.cs b/Engine/Source/Programs/AutomationTool/Gauntlet/Platform/Mac/Gauntlet.TargetDeviceMac.cs
--- a/Engine/Source/Programs/AutomationTool/Gauntlet/Platform/Mac/Gauntlet.TargetDeviceMac.cs
+++ b/Engine/Source/Programs/AutomationTool/Gauntlet/Platform/Mac/Gauntlet.TargetDeviceMac.cs
@@ -177,9 +177,7 @@ namespace Gauntlet
 
 				if (!SkipDeploy)
 				{
-					Log.Info("Installing {0} to {1}", AppConfig.Name, ToString());
-					Log.Verbose("\tCopying {0} to {1}", BuildPath, DestPath);
-					Gauntlet.Utils.SystemHelpers.CopyDirectory(BuildPath, DestPath, Utils.SystemHelpers.CopyOptions.Mirror);
+					StagedBuild.InstallBuildParallel(AppConfig, InBuild, BuildPath, DestPath, ToString());
 				}
 				else
 				{
diff --git a/Engine/Source/Programs/AutomationTool/Gauntlet/Platform/Windows/Gauntlet.TargetDeviceWindows.cs b/Engine/Source/Programs/AutomationTool/Gauntlet/Platform/Windows/Gauntlet.TargetDeviceWindows.cs
--- a/Engine/Source/Programs/AutomationTool/Gauntlet/Platform/Windows/Gauntlet.TargetDeviceWindows.cs
+++ b/Engine/Source/Programs/AutomationTool/Gauntlet/Platform/Windows/Gauntlet.TargetDeviceWindows.cs
@@ -363,9 +363,7 @@ namespace Gauntlet
 
 				if (!SkipDeploy)
 				{
-					Log.Info("Installing {0} to {1}", AppConfig.Name, ToString());
-					Log.Verbose("\tCopying {0} to {1}", BuildPath, DestPath);
-					Gauntlet.Utils.SystemHelpers.CopyDirectory(BuildPath, DestPath, Utils.SystemHelpers.CopyOptions.Mirror);
+					StagedBuild.InstallBuildParallel(AppConfig, InBuild, BuildPath, DestPath, ToString());
 				}
 				else
 				{
diff --git a/Engine/Source/Programs/AutomationTool/Gauntlet/Unreal/BuildSource/Gauntlet.StagedBuildSource.cs b/Engine/Source/Programs/AutomationTool/Gauntlet/Unreal/BuildSource/Gauntlet.StagedBuildSource.cs
--- a/Engine/Source/Programs/AutomationTool/Gauntlet/Unreal/BuildSource/Gauntlet.StagedBuildSource.cs
+++ b/Engine/Source/Programs/AutomationTool/Gauntlet/Unreal/BuildSource/Gauntlet.StagedBuildSource.cs
@@ -57,6 +57,101 @@ namespace Gauntlet
 			Flags = BuildFlags.CanReplaceCommandLine | BuildFlags.CanReplaceExecutable | BuildFlags.Loose;
 		}
 
+		enum InstallStatus
+		{
+			Error,
+			Installing,
+			Installed
+		}
+
+		static Dictionary<StagedBuild, InstallStatus> LocalInstalls = new Dictionary<StagedBuild, InstallStatus>();
+
+		/// <summary>
+		/// When running parallel tests using staged builds, local desktop devices only need one copy of client/server 
+		/// this method is used to coordinate the copy across multiple local devices
+		/// </summary>
+		public static void InstallBuildParallel(UnrealAppConfig AppConfig, StagedBuild InBuild, string BuildPath, string DestPath, string Desc)
+		{
+			// In parallel tests, we only want to copy the client/server once
+			bool Install = false;
+			InstallStatus Status = InstallStatus.Error;
+
+			lock (Globals.MainLock)
+			{
+				if (!LocalInstalls.ContainsKey(InBuild))
+				{
+					Install = true;
+					LocalInstalls[InBuild] = Status = InstallStatus.Installing;
+				}
+				else
+				{
+					Status = LocalInstalls[InBuild];
+				}
+			}
+
+			// check if we've already errored in another thread
+			if (Status == InstallStatus.Error)
+			{
+				throw new AutomationException("Parallel build error installing {0} to {1}", BuildPath, DestPath);
+			}
+
+			if (Install)
+			{
+				try
+				{
+					Log.Info("Installing {0} to {1}", AppConfig.Name, Desc);
+					Log.Verbose("\tCopying {0} to {1}", BuildPath, DestPath);
+					Utils.SystemHelpers.CopyDirectory(BuildPath, DestPath, Utils.SystemHelpers.CopyOptions.Mirror);
+				}
+				catch (Exception Ex)
+				{
+					lock (Globals.MainLock)
+					{
+						LocalInstalls[InBuild] = InstallStatus.Error;
+					}
+
+					throw Ex;
+				}
+
+				lock (Globals.MainLock)
+				{
+					LocalInstalls[InBuild] = InstallStatus.Installed;
+				}
+
+			}
+			else
+			{
+				DateTime StartTime = DateTime.Now;
+				while (true)
+				{
+					if ((DateTime.Now - StartTime).TotalMinutes > 60)
+					{
+						throw new AutomationException("Parallel build error installing {0} to {1}, timed out after an hour", BuildPath, DestPath);
+					}
+
+					Thread.Sleep(1000);
+
+					lock (Globals.MainLock)
+					{
+						Status = LocalInstalls[InBuild];
+					}
+
+					// install process failed in other thread, disk full, etc
+					if (Status == InstallStatus.Error)
+					{
+						throw new AutomationException("Error installing parallel build from {0} to {1}", BuildPath, DestPath);
+					}
+
+					if (Status == InstallStatus.Installed)
+					{
+						Log.Verbose("Parallel build successfully installed from {0} to {1}", BuildPath, DestPath);
+						break;
+					}
+				}
+			}
+
+		}
+
 		public static IEnumerable<T> CreateFromPath<T>(UnrealTargetPlatform InPlatform, string InProjectName, string InPath, string InExecutableExtension)
 			where T : StagedBuild
 		{
