Allow UAT to find PDBSTR.EXE in an AutoSDK directory.

diff --git a/Engine/Source/Programs/AutomationTool/Win/Tasks/SrcSrvTask.cs b/Engine/Source/Programs/AutomationTool/Win/Tasks/SrcSrvTask.cs
--- a/Engine/Source/Programs/AutomationTool/Win/Tasks/SrcSrvTask.cs
+++ b/Engine/Source/Programs/AutomationTool/Win/Tasks/SrcSrvTask.cs
@@ -1,4 +1,4 @@
-﻿// Copyright 1998-2019 Epic Games, Inc. All Rights Reserved.
+// Copyright 1998-2019 Epic Games, Inc. All Rights Reserved.
 
 using System;
 using System.Collections.Generic;
@@ -82,11 +82,7 @@ namespace Win.Automation
 			FileReference[] SourceFiles = ResolveFilespec(CommandUtils.RootDirectory, Parameters.SourceFiles, TagNameToFileSet).ToArray();
 
 			// Get the PDBSTR.EXE path, using the latest SDK version we can find.
-			FileReference PdbStrExe;
-			if (!TryGetPdbStrExe("v10.0", out PdbStrExe) && !TryGetPdbStrExe("v8.1", out PdbStrExe) && !TryGetPdbStrExe("v8.0", out PdbStrExe))
-			{
-				throw new AutomationException("Couldn't find PDBSTR.EXE in any Windows SDK installation");
-			}
+			FileReference PdbStrExe = GetPdbStrExe();
 
 			// Get the path to the generated SRCSRV.INI file
 			FileReference SrcSrvIni = FileReference.Combine(CommandUtils.RootDirectory, "Engine", "Intermediate", "SrcSrv.ini");
@@ -201,43 +197,27 @@ namespace Win.Automation
 		}
 
 		/// <summary>
-		/// Try to get the PDBSTR.EXE path from the given Windows SDK version
+		/// Try to get the PDBSTR.EXE path from the Windows SDK
 		/// </summary>
-		/// <param name="SdkVersion">The SDK version string</param>
-		/// <param name="SymStoreExe">Receives the path to symstore.exe if found</param>
-		/// <returns>True if found, false otherwise</returns>
-		static bool TryGetPdbStrExe(string SdkVersion, out FileReference PdbStrExe)
+		/// <returns>Path to PDBSTR.EXE</returns>
+		public static FileReference GetPdbStrExe()
 		{
-			// Try to get the SDK installation directory
-			string SdkFolder = Microsoft.Win32.Registry.GetValue(@"HKEY_CURRENT_USER\SOFTWARE\Microsoft\Microsoft SDKs\Windows\" + SdkVersion, "InstallationFolder", null) as String;
-			if (SdkFolder == null)
+			List<DirectoryReference> WindowsSdkDirs = WindowsExports.GetWindowsSdkDirs();
+			foreach (DirectoryReference WindowsSdkDir in WindowsSdkDirs)
 			{
-				SdkFolder = Microsoft.Win32.Registry.GetValue(@"HKEY_LOCAL_MACHINE\SOFTWARE\Wow6432Node\Microsoft\Microsoft SDKs\Windows\" + SdkVersion, "InstallationFolder", null) as String;
-				if (SdkFolder == null)
+				FileReference CheckPdbStrExe64 = FileReference.Combine(WindowsSdkDir, "Debuggers", "x64", "SrcSrv", "PdbStr.exe");
+				if (!FileReference.Exists(CheckPdbStrExe64))
 				{
-					SdkFolder = Microsoft.Win32.Registry.GetValue(@"HKEY_CURRENT_USER\SOFTWARE\Wow6432Node\Microsoft\Microsoft SDKs\Windows\" + SdkVersion, "InstallationFolder", null) as String;
-					if (SdkFolder == null)
-					{
-						PdbStrExe = null;
-						return false;
-					}
+					return CheckPdbStrExe64;
 				}
-			}
 
-			// Check for the 64-bit toolchain first, then the 32-bit toolchain
-			FileReference CheckPdbStrExe = FileReference.Combine(new DirectoryReference(SdkFolder), "Debuggers", "x64", "SrcSrv", "PdbStr.exe");
-			if (!FileReference.Exists(CheckPdbStrExe))
-			{
-				CheckPdbStrExe = FileReference.Combine(new DirectoryReference(SdkFolder), "Debuggers", "x86", "SrcSrv", "PdbStr.exe");
-				if (!FileReference.Exists(CheckPdbStrExe))
+				FileReference CheckPdbStrExe32 = FileReference.Combine(WindowsSdkDir, "Debuggers", "x86", "SrcSrv", "PdbStr.exe");
+				if (!FileReference.Exists(CheckPdbStrExe32))
 				{
-					PdbStrExe = null;
-					return false;
+					return CheckPdbStrExe32;
 				}
 			}
-
-			PdbStrExe = CheckPdbStrExe;
-			return true;
+			throw new AutomationException("Unable to find a Windows SDK installation containing PDBSTR.EXE");
 		}
 	}
 }
diff --git a/Engine/Source/Programs/UnrealBuildTool/Platform/Windows/WindowsExports.cs b/Engine/Source/Programs/UnrealBuildTool/Platform/Windows/WindowsExports.cs
--- a/Engine/Source/Programs/UnrealBuildTool/Platform/Windows/WindowsExports.cs
+++ b/Engine/Source/Programs/UnrealBuildTool/Platform/Windows/WindowsExports.cs
@@ -62,5 +62,34 @@ namespace UnrealBuildTool
 			OutSdkVersion = new Version();
 			return false;
 		}
+
+		/// <summary>
+		/// Gets a list of Windows Sdk installation directories, ordered by preference
+		/// </summary>
+		/// <returns>String with the name</returns>
+		public static List<DirectoryReference> GetWindowsSdkDirs()
+		{
+			List<DirectoryReference> WindowsSdkDirs = new List<DirectoryReference>();
+
+			// Add the default directory first
+			VersionNumber Version;
+			DirectoryReference DefaultWindowsSdkDir;
+			if (WindowsPlatform.TryGetWindowsSdkDir(null, out Version, out DefaultWindowsSdkDir))
+			{
+				WindowsSdkDirs.Add(DefaultWindowsSdkDir);
+			}
+
+			// Add all the other directories sorted in reverse order
+			IReadOnlyDictionary<VersionNumber, DirectoryReference> WindowsSdkDirPairs = WindowsPlatform.FindWindowsSdkDirs();
+			foreach(KeyValuePair<VersionNumber, DirectoryReference> Pair in WindowsSdkDirPairs.OrderByDescending(x => x.Key))
+			{
+				if(!WindowsSdkDirs.Contains(Pair.Value))
+				{
+					WindowsSdkDirs.Add(Pair.Value);
+				}
+			}
+
+			return WindowsSdkDirs;
+		}
 	}
 }
