Also add new "Percentage" unit type to the telemetry enum

diff --git a/Engine/Source/Programs/AutomationTool/AutomationUtils/CommandUtils.cs b/Engine/Source/Programs/AutomationTool/AutomationUtils/CommandUtils.cs
--- a/Engine/Source/Programs/AutomationTool/AutomationUtils/CommandUtils.cs
+++ b/Engine/Source/Programs/AutomationTool/AutomationUtils/CommandUtils.cs
@@ -2365,6 +2365,7 @@ namespace AutomationTool
 		Minutes,
 		Bytes,
 		Megabytes,
+		Percentage,
 	}
 
 	/// <summary>
