Fix incorrect merge (Thanks dyanikoglu! PR #6256)

diff --git a/Engine/Source/Programs/AutomationTool/Scripts/BuildPhysX.Automation.cs b/Engine/Source/Programs/AutomationTool/Scripts/BuildPhysX.Automation.cs
--- a/Engine/Source/Programs/AutomationTool/Scripts/BuildPhysX.Automation.cs
+++ b/Engine/Source/Programs/AutomationTool/Scripts/BuildPhysX.Automation.cs
@@ -1344,7 +1344,7 @@ class BuildPhysX : BuildCommand
 		{
 			return "dSYM";
 		}
-		if (TargetData.Platform == UnrealTargetPlatform.Mac)
+		if (TargetData.Platform == UnrealTargetPlatform.Linux)
 		{
 			return "sym";
 		}
