UAT: Add logging to trace where telemetry events are going.

diff --git a/Engine/Source/Programs/AutomationTool/AutomationUtils/CommandUtils.cs b/Engine/Source/Programs/AutomationTool/AutomationUtils/CommandUtils.cs
--- a/Engine/Source/Programs/AutomationTool/AutomationUtils/CommandUtils.cs
+++ b/Engine/Source/Programs/AutomationTool/AutomationUtils/CommandUtils.cs
@@ -2414,6 +2414,7 @@ namespace AutomationTool
 		{
 			Samples.RemoveAll(x => x.Name == Name);
 			Samples.Add(new TelemetrySample() { Name = Name, Value = Value, Units = Units });
+			Log.TraceLog("Added telemetry value: {0} = {1} ({2})", Name, Value, Units);
 		}
 
 		/// <summary>
@@ -2476,6 +2477,7 @@ namespace AutomationTool
 		/// <param name="FileName"></param>
 		public void Write(string FileName)
 		{
+			Log.TraceLog("Writing telemetry to {0}...", FileName);
 			using (JsonWriter Writer = new JsonWriter(FileName))
 			{
 				Writer.WriteObjectStart();
diff --git a/Engine/Source/Programs/AutomationTool/BuildGraph/Tasks/CommandTask.cs b/Engine/Source/Programs/AutomationTool/BuildGraph/Tasks/CommandTask.cs
--- a/Engine/Source/Programs/AutomationTool/BuildGraph/Tasks/CommandTask.cs
+++ b/Engine/Source/Programs/AutomationTool/BuildGraph/Tasks/CommandTask.cs
@@ -1,4 +1,4 @@
-﻿// Copyright 1998-2019 Epic Games, Inc. All Rights Reserved.
+// Copyright 1998-2019 Epic Games, Inc. All Rights Reserved.
 
 using System;
 using System.Collections.Generic;
@@ -110,13 +110,19 @@ namespace BuildGraph.Tasks
 			CommandUtils.RunUAT(CommandUtils.CmdEnv, CommandLine.ToString(), Identifier: Parameters.Name);
 
 			// Merge in any new telemetry data that was produced
-			if (Parameters.MergeTelemetryWithPrefix != null)
+			if (TelemetryFile != null && FileReference.Exists(TelemetryFile))
 			{
+				Log.TraceLog("Merging telemetry from {0}", TelemetryFile);
+
 				TelemetryData NewTelemetry;
 				if (TelemetryData.TryRead(TelemetryFile, out NewTelemetry))
 				{
 					CommandUtils.Telemetry.Merge(Parameters.MergeTelemetryWithPrefix, NewTelemetry);
 				}
+				else
+				{
+					Log.TraceWarning("Unable to read UAT telemetry file from {0}", TelemetryFile);
+				}
 			}
 		}
 
