Gauntlet: don't surface iOS crashes if they happen during request exit
#rnx

diff --git a/Engine/Source/Programs/AutomationTool/Gauntlet/Platform/IOS/Gauntlet.TargetDeviceIOS.cs b/Engine/Source/Programs/AutomationTool/Gauntlet/Platform/IOS/Gauntlet.TargetDeviceIOS.cs
--- a/Engine/Source/Programs/AutomationTool/Gauntlet/Platform/IOS/Gauntlet.TargetDeviceIOS.cs
+++ b/Engine/Source/Programs/AutomationTool/Gauntlet/Platform/IOS/Gauntlet.TargetDeviceIOS.cs
@@ -965,7 +965,7 @@ namespace Gauntlet
 		class FrameInfo
 		{
 			public string Module;
-			public string Symbol;
+			public string Symbol = String.Empty;
 			public string Address;
 			public string Offset;
 			public string Source;
@@ -1033,11 +1033,11 @@ namespace Gauntlet
 			Regex UnsymbolicatedFrameRegex = new Regex(@"\*?frame\s#(?<framenum>\d+):\s0x(?<address>[\da-f]+)\s(?<module>.+)\`(?<symbol>.+)(\s\+\s(?<offset>\d+))?");
 
 			LinkedList<string> CrashLog = new LinkedList<string>(Regex.Split(LogOutput, "\r\n|\r|\n"));
-						
+
 			List<ThreadInfo> Threads = new List<ThreadInfo>();
-			ThreadInfo Thread = null;			
+			ThreadInfo Thread = null;
 
-			var LineNode = CrashLog.First;			
+			var LineNode = CrashLog.First;
 			while (LineNode != null)
 			{
 				string Line = LineNode.Value.Trim();
@@ -1113,7 +1113,7 @@ namespace Gauntlet
 						Thread.Frames.Add(Frame);
 					}
 					else
-					{						
+					{
 						Thread = null;
 					}
 
@@ -1144,7 +1144,7 @@ namespace Gauntlet
 
 						Threads.Add(Thread);
 
-					}					
+					}
 				}
 
 				LineNode = LineNode.Next;
@@ -1154,7 +1154,14 @@ namespace Gauntlet
 
 			if (Threads.Count > 0 && Thread == null)
 			{
-				Log.Warning("Unable to parse full crash callstack");				
+				Log.Warning("Unable to parse full crash callstack");
+			}
+
+
+			// Do not want to surface crashes which happen as a result of requesting exit
+			if (Thread != null && Thread.Frames.SingleOrDefault(F => F.Symbol.Contains("::RequestExit")) != null)
+			{
+				Thread = null;
 			}
 
 			return Thread;
