Resurrect [REDACTED_HANDLE]'s CopyBuildToStaging() optimizations.
Re-submit my change to fix pakchunk1 missing on mobile.

diff --git a/Engine/Source/Editor/UnrealEd/Private/Commandlets/AssetRegistryGenerator.cpp b/Engine/Source/Editor/UnrealEd/Private/Commandlets/AssetRegistryGenerator.cpp
--- a/Engine/Source/Editor/UnrealEd/Private/Commandlets/AssetRegistryGenerator.cpp
+++ b/Engine/Source/Editor/UnrealEd/Private/Commandlets/AssetRegistryGenerator.cpp
@@ -339,7 +339,7 @@ bool FAssetRegistryGenerator::GenerateStreamingInstallManifest(int64 InExtraFlav
 				int32 PakchunkIndex = GetPakchunkIndex(ChunkID);
 				if (PakchunkIndex >= FinalChunkManifests.Num())
 				{
-					FinalChunkManifests.AddZeroed(PakchunkIndex - FinalChunkManifests.Num() + 1);
+					FinalChunkManifests.AddDefaulted(PakchunkIndex - FinalChunkManifests.Num() + 1);
 				}
 				checkf(PakchunkIndex < FinalChunkManifests.Num(), TEXT("Chunk %i out of range. %i manifests available"), PakchunkIndex, FinalChunkManifests.Num() - 1);
 				checkf(FinalChunkManifests[PakchunkIndex] == nullptr, TEXT("Manifest already exists for chunk %i"), PakchunkIndex);
@@ -368,7 +368,7 @@ bool FAssetRegistryGenerator::GenerateStreamingInstallManifest(int64 InExtraFlav
 		ChunkLayerFile->Serialize(TCHAR_TO_ANSI(*LayerString), LayerString.Len());
 
 		// Is this chunk empty?
-		if (!FinalChunkManifests[PakchunkIndex] || FinalChunkManifests[PakchunkIndex]->Num() == 0)
+		if (!FinalChunkManifests[PakchunkIndex])
 		{
 			continue;
 		}
@@ -626,7 +626,8 @@ bool FAssetRegistryGenerator::SaveManifests(FSandboxPlatformFile* InSandboxFile,
 		// Collect all unique chunk indices and map all files to their chunks
 		for (int32 PakchunkIndex = 0; PakchunkIndex < FinalChunkManifests.Num(); ++PakchunkIndex)
 		{
-			if (FinalChunkManifests[PakchunkIndex] && FinalChunkManifests[PakchunkIndex]->Num())
+			check(FinalChunkManifests[PakchunkIndex]);
+			if (FinalChunkManifests[PakchunkIndex]->Num())
 			{
 				PakchunkIndicesInUse.Add(PakchunkIndex);
 				for (auto& Filename : *FinalChunkManifests[PakchunkIndex])
@@ -1440,7 +1441,7 @@ bool FAssetRegistryGenerator::GenerateAssetChunkInformationCSV(const FString& Ou
 	TArray<TUniquePtr<FArchive>> ChunkFiles;
 	if (bWriteIndividualFiles)
 	{
-		for (int32 PakchunkIndex = 0, ChunkNum = FinalChunkManifests.Num(); PakchunkIndex < ChunkNum; ++PakchunkIndex)
+		for (int32 PakchunkIndex = 0; PakchunkIndex < FinalChunkManifests.Num(); ++PakchunkIndex)
 		{
 			FArchive* ChunkFile = IFileManager::Get().CreateFileWriter(*FPaths::Combine(*OutputPath, *FString::Printf(TEXT("Chunks%dInfo.csv"), PakchunkIndex)));
 			if (ChunkFile == nullptr)
@@ -1689,12 +1690,12 @@ void FAssetRegistryGenerator::FixupPackageDependenciesForChunks(FSandboxPlatform
 
 	for (int32 PakchunkIndex = 0, MaxPakchunk = ChunkManifests.Num(); PakchunkIndex < MaxPakchunk; ++PakchunkIndex)
 	{
-		FinalChunkManifests.Add(nullptr);
+		FinalChunkManifests.Add(new FChunkPackageSet());
 		if (!ChunkManifests[PakchunkIndex])
 		{
 			continue;
 		}
-		FinalChunkManifests[PakchunkIndex] = new FChunkPackageSet();
+		
 		for (auto It = ChunkManifests[PakchunkIndex]->CreateConstIterator(); It; ++It)
 		{
 			AddPackageAndDependenciesToChunk(FinalChunkManifests[PakchunkIndex], It.Key(), It.Value(), PakchunkIndex, InSandboxFile);
@@ -1711,12 +1712,10 @@ void FAssetRegistryGenerator::FixupPackageDependenciesForChunks(FSandboxPlatform
 	//Once complete, Add any remaining assets (that are not assigned to a chunk) to the first chunk.
 	if (FinalChunkManifests.Num() == 0)
 	{
-		FinalChunkManifests.Add(nullptr);
-	}
-	if (!FinalChunkManifests[0])
-	{
-		FinalChunkManifests[0] = new FChunkPackageSet();
+		FinalChunkManifests.Add(new FChunkPackageSet());
 	}
+	check(FinalChunkManifests[0]);
+	
 	// Copy the remaining assets
 	auto RemainingAssets = UnassignedPackageSet;
 	for (auto It = RemainingAssets.CreateConstIterator(); It; ++It)
@@ -1750,17 +1749,18 @@ void FAssetRegistryGenerator::FixupPackageDependenciesForChunks(FSandboxPlatform
 	for (int32 PakchunkIndex = 0, MaxPakchunk = ChunkManifests.Num(); PakchunkIndex < MaxPakchunk; ++PakchunkIndex)
 	{
 		const int32 ChunkManifestNum = ChunkManifests[PakchunkIndex] ? ChunkManifests[PakchunkIndex]->Num() : 0;
-		const int32 FinalChunkManifestNum = FinalChunkManifests[PakchunkIndex] ? FinalChunkManifests[PakchunkIndex]->Num() : 0;
+		const int32 FinalChunkManifestNum = FinalChunkManifests[PakchunkIndex]->Num();
 		UE_LOG(LogAssetRegistryGenerator, Log, TEXT("Chunk: %i, Started with %i packages, Final after dependency resolve: %i"), PakchunkIndex, ChunkManifestNum, FinalChunkManifestNum);
 	}
 	
 	// Fix up the asset registry to reflect this chunk layout
 	for (int32 PakchunkIndex = 0 ; PakchunkIndex < FinalChunkManifests.Num(); ++PakchunkIndex)
 	{
-		if (PakchunkIndex >= FinalChunkManifests.Num() || !FinalChunkManifests[PakchunkIndex])
+		if (PakchunkIndex >= FinalChunkManifests.Num())
 		{
 			continue;
 		}
+		check(FinalChunkManifests[PakchunkIndex]);
 		for (const TPair<FName, FString>& Asset : *FinalChunkManifests[PakchunkIndex])
 		{
 			const TArray<const FAssetData*> AssetIndexArray = State.GetAssetsByPackageName(Asset.Key);
diff --git a/Engine/Source/Programs/AutomationTool/Scripts/CopyBuildToStagingDirectory.Automation.cs b/Engine/Source/Programs/AutomationTool/Scripts/CopyBuildToStagingDirectory.Automation.cs
--- a/Engine/Source/Programs/AutomationTool/Scripts/CopyBuildToStagingDirectory.Automation.cs
+++ b/Engine/Source/Programs/AutomationTool/Scripts/CopyBuildToStagingDirectory.Automation.cs
@@ -1518,7 +1518,7 @@ public partial class Project : CommandUtils
 	/// </summary>
 	/// <param name="Params"></param>
 	/// <param name="SC"></param>
-	private static void ApplyPakFileRules(List<PakFileRules> RulesList, KeyValuePair<string, string> StagingFile, List<string> ModifyPakList, out bool bExcludeFromPaks)
+	private static void ApplyPakFileRules(List<PakFileRules> RulesList, KeyValuePair<string, string> StagingFile, HashSet<ChunkDefinition> ModifyPakList, Dictionary<string, ChunkDefinition> ChunkNameToDefinition, out bool bExcludeFromPaks)
 	{
 		bExcludeFromPaks = false;
 
@@ -1551,10 +1551,10 @@ public partial class Project : CommandUtils
 				}
 
 				bExcludeFromPaks = PakRules.bExcludeFromPaks;
-				if (PakRules.OverridePaks != null)
+				if (PakRules.OverridePaks != null && ModifyPakList != null)
 				{
 					ModifyPakList.Clear();
-					PakRules.OverridePaks.ForEach(val => { ModifyPakList.Add(val); });
+					ModifyPakList.UnionWith(PakRules.OverridePaks.Select(x => ChunkNameToDefinition[x]));
 					//LogInformation("Setting pak assignment for file {0} to {1}", StagingFile.Key, string.Join(", ", PakRules.OverridePaks));
 				}
 				else if (bExcludeFromPaks)
@@ -1588,10 +1588,8 @@ public partial class Project : CommandUtils
 		// Apply the pak file rules, this can remove things but will not override the pak file name
 		foreach (var StagingFile in UnrealPakResponseFile)
 		{
-			List<string> PakList = new List<string>();
-
 			bool bExcludeFromPaks = false;
-			ApplyPakFileRules(PakRulesList, StagingFile, PakList, out bExcludeFromPaks);
+			ApplyPakFileRules(PakRulesList, StagingFile, null, null, out bExcludeFromPaks);
 
 			if (bExcludeFromPaks)
 			{
@@ -2468,11 +2466,32 @@ public partial class Project : CommandUtils
 			Dictionary<string, ChunkDefinition> OptionalChunks = new Dictionary<string, ChunkDefinition>();
 			ChunkDefinition DefaultChunk = ChunkDefinitions[DefaultChunkIndex];
 
+			Dictionary<string, List<ChunkDefinition>> FileNameToChunks = new Dictionary<string, List<ChunkDefinition>>();
+			foreach (ChunkDefinition Chunk in ChunkDefinitions)
+			{
+				foreach (string FileName in Chunk.Manifest)
+				{
+					List<ChunkDefinition> Chunks;
+					if (!FileNameToChunks.TryGetValue(FileName, out Chunks))
+					{
+						Chunks = new List<ChunkDefinition>();
+						FileNameToChunks.Add(FileName, Chunks);
+					}
+					Chunks.Add(Chunk);
+				}
+			}
+
+			Dictionary<string, ChunkDefinition> ChunkNameToDefinition = new Dictionary<string, ChunkDefinition>();
+			foreach (ChunkDefinition Chunk in ChunkDefinitions)
+			{
+				ChunkNameToDefinition.Add(Chunk.ChunkName, Chunk);
+			}
+
 			foreach (var StagingFile in StagingManifestResponseFile)
 			{
 				bool bAddedToChunk = false;
 				bool bExcludeFromPaks = false;
-				List<string> PakList = new List<string>();
+				HashSet<ChunkDefinition> PakList = new HashSet<ChunkDefinition>();
 
 				string OriginalFilename = StagingFile.Key;
 				string NoExtension = CombinePaths(Path.GetDirectoryName(OriginalFilename), Path.GetFileNameWithoutExtension(OriginalFilename));
@@ -2485,21 +2504,26 @@ public partial class Project : CommandUtils
 				string NoExtensionReplaceSlashes = NoExtension.Replace('/', '\\');
 
 				// First read manifest
-				for (int ChunkIndex = 0; ChunkIndex < ChunkDefinitions.Count; ++ChunkIndex)
+				List<ChunkDefinition> Chunks;
+				if (FileNameToChunks.TryGetValue(OriginalFilename, out Chunks))
 				{
-					ChunkDefinition Chunk = ChunkDefinitions[ChunkIndex];
-
-					if (Chunk.Manifest.Contains(OriginalFilename) ||
-						Chunk.Manifest.Contains(OriginalReplaceSlashes) ||
-						Chunk.Manifest.Contains(NoExtension) ||
-						Chunk.Manifest.Contains(NoExtensionReplaceSlashes))
-					{
-						PakList.Add(Chunk.ChunkName);
-					}
+					PakList.UnionWith(Chunks);
+				}
+				if (FileNameToChunks.TryGetValue(OriginalReplaceSlashes, out Chunks))
+				{
+					PakList.UnionWith(Chunks);
+				}
+				if (FileNameToChunks.TryGetValue(NoExtension, out Chunks))
+				{
+					PakList.UnionWith(Chunks);
+				}
+				if (FileNameToChunks.TryGetValue(NoExtensionReplaceSlashes, out Chunks))
+				{
+					PakList.UnionWith(Chunks);
 				}
 
 				// Now run through the pak rules which may override things
-				ApplyPakFileRules(PakRulesList, StagingFile, PakList, out bExcludeFromPaks);
+				ApplyPakFileRules(PakRulesList, StagingFile, PakList, ChunkNameToDefinition, out bExcludeFromPaks);
 
 				if (bExcludeFromPaks)
 				{
@@ -2507,36 +2531,31 @@ public partial class Project : CommandUtils
 				}
 
 				// Actually add to chunk
-				for (int ChunkIndex = 0; ChunkIndex < ChunkDefinitions.Count; ++ChunkIndex)
+				foreach (ChunkDefinition Chunk in PakList)
 				{
-					ChunkDefinition Chunk = ChunkDefinitions[ChunkIndex];
+					ChunkDefinition TargetChunk = Chunk;
 
-					if (PakList.Contains(Chunk.ChunkName))
+					string OrigExt = Path.GetExtension(OriginalFilename);
+					if (OrigExt.Equals(OptionalBulkDataFileExtension))
 					{
-						ChunkDefinition TargetChunk = Chunk;
-
-						string OrigExt = Path.GetExtension(OriginalFilename);
-						if (OrigExt.Equals(OptionalBulkDataFileExtension))
+						// any optional files encountered we want to put in a separate pak file
+						string OptionalChunkName = Chunk.ChunkName + "optional";
+						if (!OptionalChunks.TryGetValue(OptionalChunkName, out TargetChunk))
 						{
-							// any optional files encountered we want to put in a separate pak file
-							string OptionalChunkName = Chunk.ChunkName + "optional";
-							if (!OptionalChunks.TryGetValue(OptionalChunkName, out TargetChunk))
-							{
-								TargetChunk = new ChunkDefinition(OptionalChunkName);
-								TargetChunk.RequestedEncryptionKeyGuid = Chunk.RequestedEncryptionKeyGuid;
-								TargetChunk.EncryptionKeyGuid = Chunk.EncryptionKeyGuid;
-								OptionalChunks.Add(OptionalChunkName, TargetChunk);
-							}
+							TargetChunk = new ChunkDefinition(OptionalChunkName);
+							TargetChunk.RequestedEncryptionKeyGuid = Chunk.RequestedEncryptionKeyGuid;
+							TargetChunk.EncryptionKeyGuid = Chunk.EncryptionKeyGuid;
+							OptionalChunks.Add(OptionalChunkName, TargetChunk);
 						}
+					}
 
-						TargetChunk.ResponseFile.Add(StagingFile.Key, StagingFile.Value);
-						bAddedToChunk = true;
+					TargetChunk.ResponseFile.Add(StagingFile.Key, StagingFile.Value);
+					bAddedToChunk = true;
 
-						if (bForceOneChunkPerFile)
-						{
-							// Files are only allowed to be in a single chunk
-							break;
-						}
+					if (bForceOneChunkPerFile)
+					{
+						// Files are only allowed to be in a single chunk
+						break;
 					}
 				}
 				if (!bAddedToChunk)
