#UE4 Staging fix (build health)

diff --git a/Engine/Source/Programs/AutomationTool/Scripts/CopyBuildToStagingDirectory.Automation.cs b/Engine/Source/Programs/AutomationTool/Scripts/CopyBuildToStagingDirectory.Automation.cs
--- a/Engine/Source/Programs/AutomationTool/Scripts/CopyBuildToStagingDirectory.Automation.cs
+++ b/Engine/Source/Programs/AutomationTool/Scripts/CopyBuildToStagingDirectory.Automation.cs
@@ -1532,7 +1532,7 @@ public partial class Project : CommandUtils
 		{
 			if (PakRules.Filter.Matches(StagingFile.Key))
 			{
-				if (ModifyPakList.Count > 0)
+				if (ModifyPakList != null && ModifyPakList.Count > 0)
 				{
 					// Only override the existing list if bOverrideChunkManifest is set
 					if (!PakRules.bOverrideChunkManifest)
@@ -1551,7 +1551,7 @@ public partial class Project : CommandUtils
 				}
 
 				bExcludeFromPaks = PakRules.bExcludeFromPaks;
-				if (PakRules.OverridePaks != null && ModifyPakList != null)
+				if (PakRules.OverridePaks != null && ModifyPakList != null && ChunkNameToDefinition != null)
 				{
 					ModifyPakList.Clear();
 					ModifyPakList.UnionWith(PakRules.OverridePaks.Select(x => ChunkNameToDefinition[x]));
