Refactored ini key stripping in UAT to  be driven by config settings, exposed in the project packaging settings
Added the ability to set a list of config sections to strip while staging

diff --git a/Engine/Config/BaseGame.ini b/Engine/Config/BaseGame.ini
--- a/Engine/Config/BaseGame.ini
+++ b/Engine/Config/BaseGame.ini
@@ -98,6 +98,19 @@ bSkipMovies=False
 +EarlyDownloaderPakFileFiles=...\*.uproject
 +EarlyDownloaderPakFileFiles=...\global_sf*.metalmap
 
++IniKeyBlacklist=KeyStorePassword
++IniKeyBlacklist=KeyPassword
++IniKeyBlacklist=rsa.privateexp
++IniKeyBlacklist=rsa.modulus
++IniKeyBlacklist=rsa.publicexp
++IniKeyBlacklist=aes.key
++IniKeyBlacklist=SigningPublicExponent
++IniKeyBlacklist=SigningModulus
++IniKeyBlacklist=SigningPrivateExponent
++IniKeyBlacklist=EncryptionKey
++IniKeyBlacklist=IniKeyBlacklist
++IniKeyBlacklist=IniSectionBlacklist
+
 [/Script/Engine.HUD]
 DebugDisplay=AI
 
diff --git a/Engine/Source/Editor/UnrealEd/Classes/Settings/ProjectPackagingSettings.h b/Engine/Source/Editor/UnrealEd/Classes/Settings/ProjectPackagingSettings.h
--- a/Engine/Source/Editor/UnrealEd/Classes/Settings/ProjectPackagingSettings.h
+++ b/Engine/Source/Editor/UnrealEd/Classes/Settings/ProjectPackagingSettings.h
@@ -355,6 +355,14 @@ public:
 	UPROPERTY(config, EditAnywhere, Category = Packaging, AdvancedDisplay)
 	TArray<FString> CompressedChunkWildcard;
 
+	/** List of ini file keys to strip when packaging */
+	UPROPERTY(config, EditAnywhere, Category = Packaging)
+	TArray<FString> IniKeyBlacklist;
+
+	/** List of ini file sections to strip when packaging */
+	UPROPERTY(config, EditAnywhere, Category = Packaging)
+	TArray<FString> IniSectionBlacklist;
+
 	/**
 	 * List of specific files to include with GenerateEarlyDownloaderPakFile
 	 */
diff --git a/Engine/Source/Programs/AutomationTool/AutomationUtils/CommandUtils.cs b/Engine/Source/Programs/AutomationTool/AutomationUtils/CommandUtils.cs
--- a/Engine/Source/Programs/AutomationTool/AutomationUtils/CommandUtils.cs
+++ b/Engine/Source/Programs/AutomationTool/AutomationUtils/CommandUtils.cs
@@ -1346,7 +1346,7 @@ namespace AutomationTool
 		/// <param name="Dest">The full path to the destination file</param>
 		/// <param name="bAllowDifferingTimestamps">If true, will always skip a file if the destination exists, even if timestamp differs; defaults to false</param>
 		/// <returns>True if the operation was successful, false otherwise.</returns>
-		public static void CopyFileIncremental(FileReference Source, FileReference Dest, bool bAllowDifferingTimestamps = false, bool bFilterSpecialLinesFromIniFiles = false)
+		public static void CopyFileIncremental(FileReference Source, FileReference Dest, bool bAllowDifferingTimestamps = false, List<string> IniKeyBlacklist = null, List<string> IniSectionBlacklist = null)
 		{
 			if (InternalUtils.SafeFileExists(Dest.FullName, true))
 			{
@@ -1374,7 +1374,7 @@ namespace AutomationTool
 			{
 				throw new AutomationException("Failed to delete {0} for copy", Dest);
 			}
-			if (!InternalUtils.SafeCopyFile(Source.FullName, Dest.FullName, bFilterSpecialLinesFromIniFiles:bFilterSpecialLinesFromIniFiles))
+			if (!InternalUtils.SafeCopyFile(Source.FullName, Dest.FullName, IniKeyBlacklist:IniKeyBlacklist, IniSectionBlacklist:IniSectionBlacklist))
 			{
 				throw new AutomationException("Failed to copy {0} to {1}", Source, Dest);
 			}
diff --git a/Engine/Source/Programs/AutomationTool/AutomationUtils/DeploymentContext.cs b/Engine/Source/Programs/AutomationTool/AutomationUtils/DeploymentContext.cs
--- a/Engine/Source/Programs/AutomationTool/AutomationUtils/DeploymentContext.cs
+++ b/Engine/Source/Programs/AutomationTool/AutomationUtils/DeploymentContext.cs
@@ -258,6 +258,17 @@ public class DeploymentContext //: ProjectParams
 	/// </summary>
 	public HashSet<StagedFileReference> BlacklistConfigFiles = new HashSet<StagedFileReference>();
 
+
+	/// <summary>
+	/// List of ini keys to strip when staging
+	/// </summary>
+	public List<string> IniKeyBlacklist = null;
+
+	/// <summary>
+	/// List of ini sections to strip when staging
+	/// </summary>
+	public List<string> IniSectionBlacklist = null;
+
 	/// <summary>
 	/// List of localization targets that are not included in staged build. By default, all project Content/Localization targets are automatically staged.
 	/// This list is read from the +BlacklistLocalizationTargets=... array in the [Staging] section of *Game.ini files.
@@ -477,23 +488,29 @@ public class DeploymentContext //: ProjectParams
 		ReadConfigFileList(GameConfig, "Staging", "WhitelistConfigFiles", WhitelistConfigFiles);
 		ReadConfigFileList(GameConfig, "Staging", "BlacklistConfigFiles", BlacklistConfigFiles);
 
-        // If we were configured to use manifests across the whole project, then this platform should use manifests.
-        // Otherwise, read whether we are generating chunks from the ProjectPackagingSettings ini.
-        if (InForceChunkManifests)
-        {
-            PlatformUsesChunkManifests = true;
-        }
-        else
-        {
-            ConfigHierarchy GameIni = ConfigCache.ReadHierarchy(ConfigHierarchyType.Game, ProjectRoot, InTargetPlatform.PlatformType);
-            String IniPath = "/Script/UnrealEd.ProjectPackagingSettings";
-            bool bSetting = false;
-            if (GameIni.GetBool(IniPath, "bGenerateChunks", out bSetting))
-            {
-                PlatformUsesChunkManifests = bSetting;
-            }
-        }
-    }
+		// Grab the game ini data
+		ConfigHierarchy GameIni = ConfigCache.ReadHierarchy(ConfigHierarchyType.Game, ProjectRoot, InTargetPlatform.PlatformType);
+		String IniPath = "/Script/UnrealEd.ProjectPackagingSettings";
+
+		// Read the config blacklists
+		GameIni.GetArray(IniPath, "IniKeyBlacklist", out IniKeyBlacklist);
+		GameIni.GetArray(IniPath, "IniSectionBlacklist", out IniSectionBlacklist);
+
+		// If we were configured to use manifests across the whole project, then this platform should use manifests.
+		// Otherwise, read whether we are generating chunks from the ProjectPackagingSettings ini.
+		if (InForceChunkManifests)
+		{
+			PlatformUsesChunkManifests = true;
+		}
+		else
+		{
+			bool bSetting = false;
+			if (GameIni.GetBool(IniPath, "bGenerateChunks", out bSetting))
+			{
+				PlatformUsesChunkManifests = bSetting;
+			}
+		}
+	}
 
 	/// <summary>
 	/// Read a list of whitelisted or blacklisted config files names from a config file
diff --git a/Engine/Source/Programs/AutomationTool/AutomationUtils/Utils.cs b/Engine/Source/Programs/AutomationTool/AutomationUtils/Utils.cs
--- a/Engine/Source/Programs/AutomationTool/AutomationUtils/Utils.cs
+++ b/Engine/Source/Programs/AutomationTool/AutomationUtils/Utils.cs
@@ -276,37 +276,45 @@ namespace AutomationTool
 			return Result;
 		}
 
-		// @todo: This could be passed in from elsewhere, and this should be somehow done per ini section
-		// but this will get it so that games won't ship passwords
-		private static string[] LinesToFilter = new string[]
-		{
-			"KeyStorePassword",
-			"KeyPassword",
-			"rsa.privateexp",
-			"rsa.modulus",
-			"rsa.publicexp",
-			"aes.key",
-			"SigningPublicExponent",
-			"SigningModulus",
-			"SigningPrivateExponent",
-			"EncryptionKey"
-		};
-
-		private static void FilterIniFile(string SourceName, string TargetName)
+		// Characters that can appear at the start of
+		private static char[] IgnoredIniValuePrefixes = { '+', '-', ' ', '\t', ';' };
+
+		private static void FilterIniFile(string SourceName, string TargetName, List<string> IniKeyBlacklist, List<string> InSectionBlacklist)
 		{
 			string[] Lines = File.ReadAllLines(SourceName);
 			StringBuilder NewLines = new StringBuilder("");
+			bool bFilteringSection = false;
 
-			foreach (string Line in Lines)
+			foreach (string OriginalLine in Lines)
 			{
+				string Line = OriginalLine.Trim();
+				bool bFiltered = bFilteringSection;
+
 				// look for each filter on each line
-				bool bFiltered = false;
-				foreach (string Filter in LinesToFilter)
+				if (!bFiltered)
 				{
-					if (Line.StartsWith(Filter + "="))
+					string TrimmedLine = Line.TrimStart(IgnoredIniValuePrefixes);
+					foreach (string Filter in IniKeyBlacklist)
 					{
-						bFiltered = true;
-						break;
+						if (TrimmedLine.StartsWith(Filter + "="))
+						{
+							bFiltered = true;
+							break;
+						}
+					}
+				}
+
+				if (InSectionBlacklist != null)
+				{
+					if (Line.StartsWith("[") && Line.EndsWith("]"))
+					{
+						string SectionName = Line.Substring(1, Line.Length - 2);
+						bFilteringSection = bFiltered = InSectionBlacklist.Contains(SectionName);
+
+						if (bFilteringSection)
+						{
+							Log.TraceLog("Filtering config section '{0}'", SectionName);
+						}
 					}
 				}
 
@@ -334,7 +342,7 @@ namespace AutomationTool
 		/// <param name="SourceName">Source name</param>
 		/// <param name="TargetName">Target name</param>
 		/// <returns>True if the operation was successful, false otherwise.</returns>
-		public static bool SafeCopyFile(string SourceName, string TargetName, bool bQuiet = false, bool bFilterSpecialLinesFromIniFiles = false)
+		public static bool SafeCopyFile(string SourceName, string TargetName, bool bQuiet = false, List<string> IniKeyBlacklist = null, List<string> IniSectionBlacklist = null)
 		{
 			if (!bQuiet)
 			{
@@ -351,9 +359,9 @@ namespace AutomationTool
 				try
 				{
 					bool bSkipSizeCheck = false;
-					if (bFilterSpecialLinesFromIniFiles && Path.GetExtension(SourceName) == ".ini")
+					if (IniKeyBlacklist != null && Path.GetExtension(SourceName) == ".ini")
 					{
-						FilterIniFile(SourceName, TargetName);
+						FilterIniFile(SourceName, TargetName, IniKeyBlacklist, IniSectionBlacklist);
 						// ini files may change size, don't check
 						bSkipSizeCheck = true;
 					}
diff --git a/Engine/Source/Programs/AutomationTool/Scripts/CopyBuildToStagingDirectory.Automation.cs b/Engine/Source/Programs/AutomationTool/Scripts/CopyBuildToStagingDirectory.Automation.cs
--- a/Engine/Source/Programs/AutomationTool/Scripts/CopyBuildToStagingDirectory.Automation.cs
+++ b/Engine/Source/Programs/AutomationTool/Scripts/CopyBuildToStagingDirectory.Automation.cs
@@ -1245,7 +1245,7 @@ public partial class Project : CommandUtils
 		}
 	}
 
-	public static void CopyManifestFilesToStageDir(Dictionary<StagedFileReference, FileReference> Mapping, DirectoryReference StageDir, DirectoryReference ManifestDir, string ManifestName, HashSet<StagedFileReference> CRCFiles, string PlatformName)
+	public static void CopyManifestFilesToStageDir(Dictionary<StagedFileReference, FileReference> Mapping, DirectoryReference StageDir, DirectoryReference ManifestDir, string ManifestName, HashSet<StagedFileReference> CRCFiles, string PlatformName, List<string> IniKeyBlacklist, List<string> IniSectionBlacklist)
 	{
 		LogInformation("Copying {0} to staging directory: {1}", ManifestName, StageDir);
 		FileReference ManifestPath = null;
@@ -1262,7 +1262,7 @@ public partial class Project : CommandUtils
 			FileReference Dest = FileReference.Combine(StageDir, Pair.Key.Name);
 			if (Src != Dest)  // special case for things created in the staging directory, like the pak file
 			{
-				CopyFileIncremental(Src, Dest, bFilterSpecialLinesFromIniFiles: true);
+				CopyFileIncremental(Src, Dest, IniKeyBlacklist:IniKeyBlacklist, IniSectionBlacklist:IniSectionBlacklist);
 			}
 		}
 		if (ManifestPath != null && Mapping.Count > 0)
@@ -1301,7 +1301,7 @@ public partial class Project : CommandUtils
 
 	public static void CopyUsingStagingManifest(ProjectParams Params, DeploymentContext SC)
 	{
-		CopyManifestFilesToStageDir(SC.FilesToStage.NonUFSFiles, SC.StageDirectory, SC.DebugStageDirectory, "NonUFSFiles", SC.StageTargetPlatform.GetFilesForCRCCheck(), SC.StageTargetPlatform.PlatformType.ToString());
+		CopyManifestFilesToStageDir(SC.FilesToStage.NonUFSFiles, SC.StageDirectory, SC.DebugStageDirectory, "NonUFSFiles", SC.StageTargetPlatform.GetFilesForCRCCheck(), SC.StageTargetPlatform.PlatformType.ToString(), SC.IniKeyBlacklist, SC.IniSectionBlacklist);
 
 		bool bStageUnrealFileSystemFiles = !Params.CookOnTheFly && !Params.UsePak(SC.StageTargetPlatform) && !Params.FileServer;
 		if (bStageUnrealFileSystemFiles)
@@ -1319,14 +1319,14 @@ public partial class Project : CommandUtils
 					throw new AutomationException("File '{0}' is set to be staged from '{1}' for project and '{2}' for crash reporter", Pair.Key, ExistingLocation, Pair.Value);
 				}
 			}
-			CopyManifestFilesToStageDir(UFSFiles, SC.StageDirectory, SC.DebugStageDirectory, "UFSFiles", SC.StageTargetPlatform.GetFilesForCRCCheck(), SC.StageTargetPlatform.PlatformType.ToString());
+			CopyManifestFilesToStageDir(UFSFiles, SC.StageDirectory, SC.DebugStageDirectory, "UFSFiles", SC.StageTargetPlatform.GetFilesForCRCCheck(), SC.StageTargetPlatform.PlatformType.ToString(), SC.IniKeyBlacklist, SC.IniSectionBlacklist);
 		}
 
 		// Copy debug files last
 		// they do not respect the DeployLowerCaseFilenames() setting, but if copied to a case-insensitive staging directory first they determine the casing for outer directories (like Engine/Content) 
 		if (!Params.NoDebugInfo)
 		{
-			CopyManifestFilesToStageDir(SC.FilesToStage.NonUFSDebugFiles, SC.DebugStageDirectory, SC.DebugStageDirectory, "DebugFiles", SC.StageTargetPlatform.GetFilesForCRCCheck(), SC.StageTargetPlatform.PlatformType.ToString());
+			CopyManifestFilesToStageDir(SC.FilesToStage.NonUFSDebugFiles, SC.DebugStageDirectory, SC.DebugStageDirectory, "DebugFiles", SC.StageTargetPlatform.GetFilesForCRCCheck(), SC.StageTargetPlatform.PlatformType.ToString(), SC.IniKeyBlacklist, SC.IniSectionBlacklist);
 		}
 	}
 
@@ -1689,7 +1689,7 @@ public partial class Project : CommandUtils
 				string SubFolder = Pair.Key.Name.Replace('/', Path.DirectorySeparatorChar);
 				FileReference NewIniFilename = FileReference.Combine(SC.ProjectRoot, "Saved", "Temp", SC.PlatformDir, SubFolder);
 				InternalUtils.SafeCreateDirectory(NewIniFilename.Directory.FullName, true);
-				InternalUtils.SafeCopyFile(Src.FullName, NewIniFilename.FullName, bFilterSpecialLinesFromIniFiles: true);
+				InternalUtils.SafeCopyFile(Src.FullName, NewIniFilename.FullName, IniKeyBlacklist:SC.IniKeyBlacklist, IniSectionBlacklist:SC.IniSectionBlacklist);
 				Src = NewIniFilename;
 			}
 
