Gauntlet -

Fixed issue where EngineTest arguments were being appended without a space.

Added temp code to use a preferred domain when getting the host IP address

Specify multihome to editor instances to force selection of the preferred adaptor

diff --git a/Engine/Source/Programs/AutomationTool/Gauntlet/Unreal/Utils/Gauntlet.UnrealHelpers.cs b/Engine/Source/Programs/AutomationTool/Gauntlet/Unreal/Utils/Gauntlet.UnrealHelpers.cs
--- a/Engine/Source/Programs/AutomationTool/Gauntlet/Unreal/Utils/Gauntlet.UnrealHelpers.cs
+++ b/Engine/Source/Programs/AutomationTool/Gauntlet/Unreal/Utils/Gauntlet.UnrealHelpers.cs
@@ -9,6 +9,7 @@ using UnrealBuildTool;
 using System.Linq;
 using System.Net.NetworkInformation;
 using System.Text.RegularExpressions;
+using System.Net;
 
 namespace Gauntlet
 {
@@ -164,27 +165,40 @@ namespace Gauntlet
 		/// <summary>
 		/// Gets the filehost IP to provide to devkits by examining our local adapters and
 		/// returning the one that's active and on the local LAN (based on DNS assignment)
+		/// AG-TODO: PreferredDomain should be in the master config
 		/// </summary>
 		/// <returns></returns>
-		public static string GetHostIpAddress()
+		public static string GetHostIpAddress(string PreferredDomain="epicgames.net")
 		{
-			NetworkInterface[] Interfaces = NetworkInterface.GetAllNetworkInterfaces();
-			foreach (NetworkInterface adapter in Interfaces)
+			// Default to the first address with a valid prefix
+			var LocalAddress = Dns.GetHostEntry(Dns.GetHostName()).AddressList
+				.Where(o => o.AddressFamily == System.Net.Sockets.AddressFamily.InterNetwork
+					&& o.GetAddressBytes()[0] != 169)
+				.FirstOrDefault();
+
+			var ActiveInterfaces = NetworkInterface.GetAllNetworkInterfaces()
+				.Where(I => I.OperationalStatus == OperationalStatus.Up);
+
+			bool MultipleInterfaces = ActiveInterfaces.Count() > 1;
+
+			if (MultipleInterfaces)
 			{
-				if (adapter.OperationalStatus == OperationalStatus.Up)
+				// Now, lots of Epic PCs have virtual adapters etc, so see if there's one that's on our network and if so use that IP
+				var PreferredInterface = ActiveInterfaces
+					.Where(I => I.GetIPProperties().DnsSuffix.IndexOf(PreferredDomain, StringComparison.OrdinalIgnoreCase) >= 0)
+					.SelectMany(I => I.GetIPProperties().UnicastAddresses)
+					.Where(A => A.Address.AddressFamily == System.Net.Sockets.AddressFamily.InterNetwork)
+					.FirstOrDefault();
+
+				if (PreferredInterface != null)
 				{
-					IPInterfaceProperties IP = adapter.GetIPProperties();
-					for (int Index = 0; Index < IP.UnicastAddresses.Count; ++Index)
-					{
-						if (IP.UnicastAddresses[Index].IsDnsEligible)
-						{
-							// use default port
-							return IP.UnicastAddresses[Index].Address.ToString();
-						}
-					}
+					LocalAddress = PreferredInterface.Address;
 				}
 			}
-			return "";
+	
+			string HostIP = Globals.Params.ParseValue("hostip", "");
+			HostIP = string.IsNullOrEmpty(HostIP) ? LocalAddress.ToString() : HostIP;
+			return HostIP;
 		}
 
 		static public string GetExecutableName(string ProjectName, UnrealTargetPlatform Platform, UnrealTargetConfiguration Config, UnrealTargetRole Role, string Extension)
