Making sure that CheckHeartbeat in UnrealTestNode's TickTest doesn't try to process heartbeats once the test has finished

diff --git a/Engine/Source/Programs/AutomationTool/Gauntlet/Unreal/Base/Gauntlet.UnrealTestNode.cs b/Engine/Source/Programs/AutomationTool/Gauntlet/Unreal/Base/Gauntlet.UnrealTestNode.cs
--- a/Engine/Source/Programs/AutomationTool/Gauntlet/Unreal/Base/Gauntlet.UnrealTestNode.cs
+++ b/Engine/Source/Programs/AutomationTool/Gauntlet/Unreal/Base/Gauntlet.UnrealTestNode.cs
@@ -849,7 +849,9 @@ namespace Gauntlet
 
 		private void CheckHeartbeat()
 		{
-			if (CachedConfig == null || CachedConfig.HeartbeatOptions.bExpectHeartbeats == false)
+			if (CachedConfig == null 
+				|| CachedConfig.HeartbeatOptions.bExpectHeartbeats == false
+				|| GetTestStatus() != TestStatus.InProgress)
 			{
 				return;
 			}
