Removed explicit reference to Null device which is now pulled in by the wildcard

diff --git a/Engine/Source/Programs/AutomationTool/Gauntlet/Gauntlet.Automation.csproj b/Engine/Source/Programs/AutomationTool/Gauntlet/Gauntlet.Automation.csproj
--- a/Engine/Source/Programs/AutomationTool/Gauntlet/Gauntlet.Automation.csproj
+++ b/Engine/Source/Programs/AutomationTool/Gauntlet/Gauntlet.Automation.csproj
@@ -84,7 +84,6 @@
     <Compile Include="Framework\Utils\Gauntlet.IniConfigUtil.cs" />
     <Compile Include="Framework\Utils\Gauntlet.MarkdownBuilder.cs" />
     <Compile Include="Framework\Utils\Gauntlet.HtmlBuilder.cs" />
-    <Compile Include="Platform\Null\Gauntlet.TargetDeviceNull.cs" />
     <Compile Include="SelfTest\Unreal\Gauntlet.SelfTest.TestUnrealBuildParsing.cs" />
     <Compile Include="Properties\Resources.Designer.cs">
       <AutoGen>True</AutoGen>
