Fix inverted logic for discovering PDBSTR path.

#rnx

diff --git a/Engine/Source/Programs/AutomationTool/Win/Tasks/SrcSrvTask.cs b/Engine/Source/Programs/AutomationTool/Win/Tasks/SrcSrvTask.cs
--- a/Engine/Source/Programs/AutomationTool/Win/Tasks/SrcSrvTask.cs
+++ b/Engine/Source/Programs/AutomationTool/Win/Tasks/SrcSrvTask.cs
@@ -206,13 +206,13 @@ namespace Win.Automation
 			foreach (DirectoryReference WindowsSdkDir in WindowsSdkDirs)
 			{
 				FileReference CheckPdbStrExe64 = FileReference.Combine(WindowsSdkDir, "Debuggers", "x64", "SrcSrv", "PdbStr.exe");
-				if (!FileReference.Exists(CheckPdbStrExe64))
+				if (FileReference.Exists(CheckPdbStrExe64))
 				{
 					return CheckPdbStrExe64;
 				}
 
 				FileReference CheckPdbStrExe32 = FileReference.Combine(WindowsSdkDir, "Debuggers", "x86", "SrcSrv", "PdbStr.exe");
-				if (!FileReference.Exists(CheckPdbStrExe32))
+				if (FileReference.Exists(CheckPdbStrExe32))
 				{
 					return CheckPdbStrExe32;
 				}
