Add support for using AutoSDK in the HoloLens automation project.

#rnx

diff --git a/Engine/Source/Programs/AutomationTool/HoloLens/HoloLens.Automation.csproj b/Engine/Source/Programs/AutomationTool/HoloLens/HoloLens.Automation.csproj
--- a/Engine/Source/Programs/AutomationTool/HoloLens/HoloLens.Automation.csproj
+++ b/Engine/Source/Programs/AutomationTool/HoloLens/HoloLens.Automation.csproj
@@ -72,7 +72,10 @@
 		<Reference Include="Windows">
 		  <!-- This is a dummy entry to throw off the parser in DotNETUtilities's CsProjectInfo which doesn't understand MSBuild functions -->
 		  <HintPath />
-		  <HintPath>$(Windows10SdkPath)\UnionMetadata\$(Windows10SdkLatestVersion)\Windows.winmd</HintPath>
+		  <!-- Try using the version from AutoSDK -->
+		  <HintPath Condition="'$(UE_SDKS_ROOT)' != ''">$(UE_SDKS_ROOT)\HostWin64\Win64\Windows Kits\10\UnionMetadata\10.0.18362.0\Windows.winmd</HintPath>
+		  <!-- Use the version from the current Windows SDK -->
+		  <HintPath Condition="'$([System.IO.File]::Exists(%(HintPath)))'=='false'">$(Windows10SdkPath)\UnionMetadata\$(Windows10SdkLatestVersion)\Windows.winmd</HintPath>
 		  <!-- 16299 SDK (possibly others?) has ProductVersion == 10.0.16299 in the registry, but pathing needs the final .0 -->
 		  <HintPath Condition="'$([System.IO.File]::Exists(%(HintPath)))'=='false'">$(Windows10SdkPath)\UnionMetadata\$(Windows10SdkLatestVersion).0\Windows.winmd</HintPath>
 		  <!-- Fall back to old, unversioned location -->
