Taking into account target name instead of just the project name

diff --git a/Engine/Source/Programs/AutomationTool/Lumin/LuminPlatform.Automation.cs b/Engine/Source/Programs/AutomationTool/Lumin/LuminPlatform.Automation.cs
--- a/Engine/Source/Programs/AutomationTool/Lumin/LuminPlatform.Automation.cs
+++ b/Engine/Source/Programs/AutomationTool/Lumin/LuminPlatform.Automation.cs
@@ -553,7 +553,9 @@ public class LuminPlatform : Platform
 			return;
 		}
 
-		Deploy.InitUPL(Params.ShortProjectName, SC.ProjectRoot, SC.StageTargets[0].Receipt.Configuration);
+		//requires target name instead of just the project name
+		string TargetName = Params.ClientCookedTargets[0];
+		Deploy.InitUPL(TargetName, SC.ProjectRoot, SC.StageTargets[0].Receipt.Configuration);
 
 		string MpkName = GetFinalMpkName(Params, SC);
 		string PackageName = GetPackageName(Params);
