MagicLeap install command file for Mac: Set directory to current before trying to install

diff --git a/Engine/Source/Programs/AutomationTool/Lumin/LuminPlatform.Automation.cs b/Engine/Source/Programs/AutomationTool/Lumin/LuminPlatform.Automation.cs
--- a/Engine/Source/Programs/AutomationTool/Lumin/LuminPlatform.Automation.cs
+++ b/Engine/Source/Programs/AutomationTool/Lumin/LuminPlatform.Automation.cs
@@ -374,6 +374,7 @@ public class LuminPlatform : Platform
 			LogInformation("Writing shell for install");
 			BatchLines = new string[] {
 				"#!/bin/sh",
+				"cd \"`dirname \"$0\"`\"",
 				"MLSDK_ROOT=$MLSDK",
 				"if [ -z \"$MLSDK_ROOT\" ]; then",
 					"\tMLSDK_ROOT=\"" + Environment.GetEnvironmentVariable("MLSDK") + " \"",
