Parameters to LuminDeploy function changed with the integration, changing to match the new method signature.

diff --git a/Engine/Source/Programs/AutomationTool/Lumin/LuminPlatform.Automation.cs b/Engine/Source/Programs/AutomationTool/Lumin/LuminPlatform.Automation.cs
--- a/Engine/Source/Programs/AutomationTool/Lumin/LuminPlatform.Automation.cs
+++ b/Engine/Source/Programs/AutomationTool/Lumin/LuminPlatform.Automation.cs
@@ -563,7 +563,7 @@ public class LuminPlatform : Platform
 
 		//requires target name instead of just the project name
 		string TargetName = Params.ClientCookedTargets[0];
-		Deploy.InitUPL(TargetName, SC.ProjectRoot, SC.StageTargets[0].Receipt.Configuration);
+		Deploy.InitUPL(SC.StageTargets[0].Receipt);
 
 		string MpkName = GetFinalMpkName(Params, SC);
 		string PackageName = GetPackageName(Params);
