Changed the HoloLens automation tool project to use an OS check for referencing WinRT assemblies

diff --git a/Engine/Source/Programs/AutomationTool/HoloLens/HoloLens.Automation.csproj b/Engine/Source/Programs/AutomationTool/HoloLens/HoloLens.Automation.csproj
--- a/Engine/Source/Programs/AutomationTool/HoloLens/HoloLens.Automation.csproj
+++ b/Engine/Source/Programs/AutomationTool/HoloLens/HoloLens.Automation.csproj
@@ -50,27 +50,31 @@
       <HintPath>..\..\..\..\..\Engine\Binaries\ThirdParty\WindowsDevicePortalWrapper\WindowsDevicePortalWrapper.dll</HintPath>
     </Reference>
   </ItemGroup>
-  <PropertyGroup>
-    <Windows10SdkPath>$(Registry:HKEY_LOCAL_MACHINE\SOFTWARE\Wow6432Node\Microsoft\Microsoft SDKs\Windows\v10.0@InstallationFolder)</Windows10SdkPath>
-    <Windows10SdkLatestVersion>$(Registry:HKEY_LOCAL_MACHINE\SOFTWARE\Wow6432Node\Microsoft\Microsoft SDKs\Windows\v10.0@ProductVersion)</Windows10SdkLatestVersion>
-  </PropertyGroup>
-  <ItemGroup>
-    <Reference Include="System.Runtime.WindowsRuntime">
-      <HintPath>$(SystemDrive)\Program Files (x86)\Reference Assemblies\Microsoft\Framework\.NETCore\$(TargetFrameworkVersion)\System.Runtime.WindowsRuntime.dll</HintPath>
-      <HintPath Condition="'$([System.IO.File]::Exists(%(HintPath)))'=='false'">C:\Program Files (x86)\Reference Assemblies\Microsoft\Framework\.NETCore\v4.5\System.Runtime.WindowsRuntime.dll</HintPath>
-      <Private>True</Private>
-    </Reference>
-    <Reference Include="Windows">
-      <!-- This is a dummy entry to throw off the parser in DotNETUtilities's CsProjectInfo which doesn't understand MSBuild functions -->
-      <HintPath />
-      <HintPath>$(Windows10SdkPath)\UnionMetadata\$(Windows10SdkLatestVersion)\Windows.winmd</HintPath>
-      <!-- 16299 SDK (possibly others?) has ProductVersion == 10.0.16299 in the registry, but pathing needs the final .0 -->
-      <HintPath Condition="'$([System.IO.File]::Exists(%(HintPath)))'=='false'">$(Windows10SdkPath)\UnionMetadata\$(Windows10SdkLatestVersion).0\Windows.winmd</HintPath>
-      <!-- Fall back to old, unversioned location -->
-      <HintPath Condition="'$([System.IO.File]::Exists(%(HintPath)))'=='false'">$(Windows10SdkPath)\UnionMetadata\Windows.winmd</HintPath>
-      <Private>false</Private>
-    </Reference>
-  </ItemGroup>
+    <Choose>
+    <When Condition="'$(OS)' == 'Windows_NT'">
+	  <PropertyGroup>
+		<Windows10SdkPath>$(Registry:HKEY_LOCAL_MACHINE\SOFTWARE\Wow6432Node\Microsoft\Microsoft SDKs\Windows\v10.0@InstallationFolder)</Windows10SdkPath>
+		<Windows10SdkLatestVersion>$(Registry:HKEY_LOCAL_MACHINE\SOFTWARE\Wow6432Node\Microsoft\Microsoft SDKs\Windows\v10.0@ProductVersion)</Windows10SdkLatestVersion>
+	  </PropertyGroup>
+	  <ItemGroup>
+		<Reference Include="System.Runtime.WindowsRuntime">
+		  <HintPath>$(SystemDrive)\Program Files (x86)\Reference Assemblies\Microsoft\Framework\.NETCore\$(TargetFrameworkVersion)\System.Runtime.WindowsRuntime.dll</HintPath>
+		  <HintPath Condition="'$([System.IO.File]::Exists(%(HintPath)))'=='false'">C:\Program Files (x86)\Reference Assemblies\Microsoft\Framework\.NETCore\v4.5\System.Runtime.WindowsRuntime.dll</HintPath>
+		  <Private>True</Private>
+		</Reference>
+		<Reference Include="Windows">
+		  <!-- This is a dummy entry to throw off the parser in DotNETUtilities's CsProjectInfo which doesn't understand MSBuild functions -->
+		  <HintPath />
+		  <HintPath>$(Windows10SdkPath)\UnionMetadata\$(Windows10SdkLatestVersion)\Windows.winmd</HintPath>
+		  <!-- 16299 SDK (possibly others?) has ProductVersion == 10.0.16299 in the registry, but pathing needs the final .0 -->
+		  <HintPath Condition="'$([System.IO.File]::Exists(%(HintPath)))'=='false'">$(Windows10SdkPath)\UnionMetadata\$(Windows10SdkLatestVersion).0\Windows.winmd</HintPath>
+		  <!-- Fall back to old, unversioned location -->
+		  <HintPath Condition="'$([System.IO.File]::Exists(%(HintPath)))'=='false'">$(Windows10SdkPath)\UnionMetadata\Windows.winmd</HintPath>
+		  <Private>false</Private>
+		</Reference>
+	  </ItemGroup>
+	</When>
+  </Choose>
   <ItemGroup>
     <Compile Include="..\..\DotNETCommon\MetaData.cs">
       <Link>Properties\MetaData.cs</Link>
