Fixed more places that __MonoCS__ wasn't defined
Wrapped WinRT specific references in #if !__MonoCS__

diff --git a/Engine/Build/BatchFiles/RunUAT.sh b/Engine/Build/BatchFiles/RunUAT.sh
--- a/Engine/Build/BatchFiles/RunUAT.sh
+++ b/Engine/Build/BatchFiles/RunUAT.sh
@@ -73,7 +73,7 @@ if [ "$UATCompileArg" = "-compile" ]; then
 	else
 		echo Compiling AutomationTool with xbuild
 
-		ARGS="/p:Configuration=Development /p:Platform=AnyCPU /verbosity:quiet /nologo"
+		ARGS="/p:DefineConstants=__MonoCS__ /p:Configuration=Development /p:Platform=AnyCPU /verbosity:quiet /nologo"
 		ARGS="${ARGS} /p:TargetFrameworkProfile="
 
 		echo "xbuild Source/Programs/AutomationTool/AutomationTool.csproj $ARGS"
diff --git a/Engine/Source/Programs/AutomationTool/AutomationUtils/MacHostPlatform.cs b/Engine/Source/Programs/AutomationTool/AutomationUtils/MacHostPlatform.cs
--- a/Engine/Source/Programs/AutomationTool/AutomationUtils/MacHostPlatform.cs
+++ b/Engine/Source/Programs/AutomationTool/AutomationUtils/MacHostPlatform.cs
@@ -99,7 +99,7 @@ namespace AutomationTool
 			if (AppName == "xbuild")
 			{
 				AppName = "sh";
-				CommandLine = "-c 'xbuild " + (String.IsNullOrEmpty(CommandLine) ? "" : CommandLine) + " /p:DefineConstants=MONO /verbosity:quiet /nologo |grep -i error; if [ $? -ne 1 ]; then exit 1; else exit 0; fi'";
+				CommandLine = "-c 'xbuild " + (String.IsNullOrEmpty(CommandLine) ? "" : CommandLine) + " /p:DefineConstants=MONO /p:DefineConstants=__MonoCS__ /verbosity:quiet /nologo |grep -i error; if [ $? -ne 1 ]; then exit 1; else exit 0; fi'";
 			}
 			if (AppName.EndsWith(".exe") || ((AppName.Contains("/Binaries/Win64/") || AppName.Contains("/Binaries/Mac/")) && string.IsNullOrEmpty(Path.GetExtension(AppName))))
 			{
diff --git a/Engine/Source/Programs/AutomationTool/HoloLens/HoloLensPlatform.Automation.cs b/Engine/Source/Programs/AutomationTool/HoloLens/HoloLensPlatform.Automation.cs
--- a/Engine/Source/Programs/AutomationTool/HoloLens/HoloLensPlatform.Automation.cs
+++ b/Engine/Source/Programs/AutomationTool/HoloLens/HoloLensPlatform.Automation.cs
@@ -182,6 +182,7 @@ namespace HoloLens.Automation
 		}
 	}
 
+#if !__MonoCS__
 	class HoloLensDevicePortalCreatedProcess : IProcessResult
 	{
 		object StateLock;
@@ -334,6 +335,7 @@ namespace HoloLens.Automation
 			}
 		}
 	}
+#endif
 
 	public class HoloLensPlatform : Platform
     {
@@ -951,6 +953,7 @@ LogWarning("PackagePakFiles intermediate dir {0}", IntermediateDirectory);
 
 		private void DeployToLocalDevice(ProjectParams Params, DeploymentContext SC)
 		{
+#if !__MonoCS__
 			bool bRequiresPackage = Params.Package || SC.StageTargetPlatform.RequiresPackageToDeploy;
 			string AppxManifestPath = GetAppxManifestPath(SC);
 			string Name;
@@ -1013,10 +1016,12 @@ LogWarning("PackagePakFiles intermediate dir {0}", IntermediateDirectory);
 			{
 				LogWarning("Failed to apply a loopback exemption to the deployed app.  Connection to a local cook server will fail.");
 			}
+#endif
 		}
 
 		private void DeployToRemoteDevice(string DeviceAddress, ProjectParams Params, DeploymentContext SC)
 		{
+#if !__MonoCS__
 			if (Params.Package || SC.StageTargetPlatform.RequiresPackageToDeploy)
 			{
 				Microsoft.Tools.WindowsDevicePortal.DefaultDevicePortalConnection conn = new Microsoft.Tools.WindowsDevicePortal.DefaultDevicePortalConnection(DeviceAddress, Params.DeviceUsername, Params.DevicePassword);
@@ -1078,10 +1083,12 @@ LogWarning("PackagePakFiles intermediate dir {0}", IntermediateDirectory);
 			{
 				throw new AutomationException(ExitCode.Error_AppInstallFailed, "Remote deployment of unpackaged apps is not supported.");
 			}
+#endif
 		}
 
 		private IProcessResult RunUsingLauncherTool(string DeviceAddress, ERunOptions ClientRunFlags, string ClientApp, string ClientCmdLine, ProjectParams Params)
 		{
+#if !__MonoCS__
 			string Name;
 			string Publisher;
 			string PrimaryAppId = "App";
@@ -1128,10 +1135,14 @@ LogWarning("PackagePakFiles intermediate dir {0}", IntermediateDirectory);
 				HoloLensProcessResult.DisposeProcess();
 			}
 			return HoloLensProcessResult;
+#else
+			return null;
+#endif
 		}
 
 		private IProcessResult RunUsingDevicePortal(string DeviceAddress, ERunOptions ClientRunFlags, string ClientApp, string ClientCmdLine, ProjectParams Params)
 		{
+#if !__MonoCS__
 			string Name;
 			string Publisher;
 			GetPackageInfo(ClientApp, out Name, out Publisher);
@@ -1189,6 +1200,9 @@ LogWarning("PackagePakFiles intermediate dir {0}", IntermediateDirectory);
 			{
 				throw new AutomationException(ExitCode.Error_LauncherFailed, e, e.Message);
 			}
+#else
+			return null;
+#endif
 		}
 
 		private string GetAppxManifestPath(DeploymentContext SC)
@@ -1219,6 +1233,7 @@ LogWarning("PackagePakFiles intermediate dir {0}", IntermediateDirectory);
 
 		private bool ShouldAcceptCertificate(System.Security.Cryptography.X509Certificates.X509Certificate2 Certificate, bool Unattended)
 		{
+#if !__MonoCS__
 			if (AcceptThumbprints.Contains(Certificate.Thumbprint))
 			{
 				return true;
@@ -1242,8 +1257,12 @@ LogWarning("PackagePakFiles intermediate dir {0}", IntermediateDirectory);
 			}
 
 			throw new AutomationException(ExitCode.Error_CertificateNotFound, "Cannot connect to remote device: certificate is untrusted and user declined to accept.");
+#else
+			return false;
+#endif
 		}
 
+#if !__MonoCS__
 		private void Portal_AppInstallStatus(Microsoft.Tools.WindowsDevicePortal.DevicePortal sender, Microsoft.Tools.WindowsDevicePortal.ApplicationInstallStatusEventArgs args)
 		{
 			if (args.Status == Microsoft.Tools.WindowsDevicePortal.ApplicationInstallStatus.Failed)
@@ -1267,6 +1286,7 @@ LogWarning("PackagePakFiles intermediate dir {0}", IntermediateDirectory);
 				LogLog(args.Message);
 			}
 		}
+#endif
 
         private string[] GetPathToVCLibsPackages(bool UseDebugCrt, WindowsCompiler Compiler)
 		{
