Added runtime checks for running inside mono to the HoloLens AutomationTool and UnrealBuildTool code

diff --git a/Engine/Source/Programs/AutomationTool/HoloLens/HoloLensPlatform.Automation.cs b/Engine/Source/Programs/AutomationTool/HoloLens/HoloLensPlatform.Automation.cs
--- a/Engine/Source/Programs/AutomationTool/HoloLens/HoloLensPlatform.Automation.cs
+++ b/Engine/Source/Programs/AutomationTool/HoloLens/HoloLensPlatform.Automation.cs
@@ -954,7 +954,12 @@ LogWarning("PackagePakFiles intermediate dir {0}", IntermediateDirectory);
 		private void DeployToLocalDevice(ProjectParams Params, DeploymentContext SC)
 		{
 #if !__MonoCS__
-			bool bRequiresPackage = Params.Package || SC.StageTargetPlatform.RequiresPackageToDeploy;
+            if (Utils.IsRunningOnMono)
+            {
+                return;
+            }
+
+            bool bRequiresPackage = Params.Package || SC.StageTargetPlatform.RequiresPackageToDeploy;
 			string AppxManifestPath = GetAppxManifestPath(SC);
 			string Name;
 			string Publisher;
@@ -1022,6 +1027,11 @@ LogWarning("PackagePakFiles intermediate dir {0}", IntermediateDirectory);
 		private void DeployToRemoteDevice(string DeviceAddress, ProjectParams Params, DeploymentContext SC)
 		{
 #if !__MonoCS__
+			if (Utils.IsRunningOnMono)
+            {
+                return;
+            }
+
 			if (Params.Package || SC.StageTargetPlatform.RequiresPackageToDeploy)
 			{
 				Microsoft.Tools.WindowsDevicePortal.DefaultDevicePortalConnection conn = new Microsoft.Tools.WindowsDevicePortal.DefaultDevicePortalConnection(DeviceAddress, Params.DeviceUsername, Params.DevicePassword);
@@ -1089,7 +1099,12 @@ LogWarning("PackagePakFiles intermediate dir {0}", IntermediateDirectory);
 		private IProcessResult RunUsingLauncherTool(string DeviceAddress, ERunOptions ClientRunFlags, string ClientApp, string ClientCmdLine, ProjectParams Params)
 		{
 #if !__MonoCS__
-			string Name;
+            if (Utils.IsRunningOnMono)
+            {
+                return;
+            }
+
+            string Name;
 			string Publisher;
 			string PrimaryAppId = "App";
 			GetPackageInfo(ClientApp, out Name, out Publisher);
@@ -1143,7 +1158,12 @@ LogWarning("PackagePakFiles intermediate dir {0}", IntermediateDirectory);
 		private IProcessResult RunUsingDevicePortal(string DeviceAddress, ERunOptions ClientRunFlags, string ClientApp, string ClientCmdLine, ProjectParams Params)
 		{
 #if !__MonoCS__
-			string Name;
+            if (Utils.IsRunningOnMono)
+            {
+                return;
+            }
+
+            string Name;
 			string Publisher;
 			GetPackageInfo(ClientApp, out Name, out Publisher);
 
@@ -1234,7 +1254,12 @@ LogWarning("PackagePakFiles intermediate dir {0}", IntermediateDirectory);
 		private bool ShouldAcceptCertificate(System.Security.Cryptography.X509Certificates.X509Certificate2 Certificate, bool Unattended)
 		{
 #if !__MonoCS__
-			if (AcceptThumbprints.Contains(Certificate.Thumbprint))
+            if (Utils.IsRunningOnMono)
+            {
+                return false;
+            }
+
+            if (AcceptThumbprints.Contains(Certificate.Thumbprint))
 			{
 				return true;
 			}
diff --git a/Engine/Source/Programs/UnrealBuildTool/Platform/HoloLens/UEBuildHoloLens.cs b/Engine/Source/Programs/UnrealBuildTool/Platform/HoloLens/UEBuildHoloLens.cs
--- a/Engine/Source/Programs/UnrealBuildTool/Platform/HoloLens/UEBuildHoloLens.cs
+++ b/Engine/Source/Programs/UnrealBuildTool/Platform/HoloLens/UEBuildHoloLens.cs
@@ -763,13 +763,18 @@ namespace UnrealBuildTool
 
 		static HoloLensPlatformSDK()
 		{
+#if !__MonoCS__
+			if (Utils.IsRunningOnMono)
+			{
+				return;
+			}
+
 			string Version = "v10.0";
 			string[] possibleRegLocations =
 			{
 				@"HKEY_LOCAL_MACHINE\SOFTWARE\Wow6432Node\Microsoft\Microsoft SDKs\Windows\",
 				@"HKEY_CURRENT_USER\SOFTWARE\Wow6432Node\Microsoft\Microsoft SDKs\Windows\"
 			};
-#if !__MonoCS__
 			foreach (string regLocation in possibleRegLocations)
 			{
 				object Result = Microsoft.Win32.Registry.GetValue(@"HKEY_LOCAL_MACHINE\SOFTWARE\Wow6432Node\Microsoft\Microsoft SDKs\Windows\" + Version, "InstallationFolder", null);
@@ -787,7 +792,7 @@ namespace UnrealBuildTool
 
 		protected override SDKStatus HasRequiredManualSDKInternal()
 		{
-			return bIsInstalled ? SDKStatus.Valid : SDKStatus.Invalid;
+			return (!Utils.IsRunningOnMono && bIsInstalled) ? SDKStatus.Valid : SDKStatus.Invalid;
 		}
 	}
 
