Fixed return error

diff --git a/Engine/Source/Programs/AutomationTool/HoloLens/HoloLensPlatform.Automation.cs b/Engine/Source/Programs/AutomationTool/HoloLens/HoloLensPlatform.Automation.cs
--- a/Engine/Source/Programs/AutomationTool/HoloLens/HoloLensPlatform.Automation.cs
+++ b/Engine/Source/Programs/AutomationTool/HoloLens/HoloLensPlatform.Automation.cs
@@ -1101,7 +1101,7 @@ LogWarning("PackagePakFiles intermediate dir {0}", IntermediateDirectory);
 #if !__MonoCS__
             if (Utils.IsRunningOnMono)
             {
-                return;
+                return null;
             }
 
             string Name;
@@ -1160,7 +1160,7 @@ LogWarning("PackagePakFiles intermediate dir {0}", IntermediateDirectory);
 #if !__MonoCS__
             if (Utils.IsRunningOnMono)
             {
-                return;
+                return null;
             }
 
             string Name;
