Changed the AutomationTool and UBT projects to use choose/when for defining __MonoCS__ when not on WinNT

diff --git a/Engine/Build/BatchFiles/Mac/GenerateProjectFiles.sh b/Engine/Build/BatchFiles/Mac/GenerateProjectFiles.sh
--- a/Engine/Build/BatchFiles/Mac/GenerateProjectFiles.sh
+++ b/Engine/Build/BatchFiles/Mac/GenerateProjectFiles.sh
@@ -25,7 +25,7 @@ fi
 "${BASE_PATH}/../FindPlatformExtensionSources.sh"
 
 if [ -f "$BASE_PATH/../../../Source/Programs/UnrealBuildTool/UnrealBuildTool.csproj" ]; then
-	xbuild "$BASE_PATH/../../../Source/Programs/UnrealBuildTool/UnrealBuildTool.csproj" /p:DefineConstants=__MonoCS__ /property:Configuration="Development" /verbosity:quiet /nologo |grep -i error
+	xbuild "$BASE_PATH/../../../Source/Programs/UnrealBuildTool/UnrealBuildTool.csproj" /property:Configuration="Development" /verbosity:quiet /nologo |grep -i error
 fi
 
 WANT_AOT="`defaults read com.epicgames.ue4 MonoAOT`"
diff --git a/Engine/Build/BatchFiles/Mac/RunXBuild.sh b/Engine/Build/BatchFiles/Mac/RunXBuild.sh
--- a/Engine/Build/BatchFiles/Mac/RunXBuild.sh
+++ b/Engine/Build/BatchFiles/Mac/RunXBuild.sh
@@ -20,7 +20,7 @@ if [ ! -f Build/BatchFiles/Mac/RunXBuild.sh ]; then
 	exit 1
 fi
 
-xbuild /verbosity:quiet /nologo /p:DefineConstants=__MonoCS__ "$@" |grep -i error
+xbuild /verbosity:quiet /nologo "$@" |grep -i error
 if [ $? -ne 1 ]; then
 	exit 1
 else
diff --git a/Engine/Build/BatchFiles/RunUAT.sh b/Engine/Build/BatchFiles/RunUAT.sh
--- a/Engine/Build/BatchFiles/RunUAT.sh
+++ b/Engine/Build/BatchFiles/RunUAT.sh
@@ -76,7 +76,7 @@ if [ "$UATCompileArg" = "-compile" ]; then
 		# make sure the UBT project has references to auto-discovered platform extension source files
 		"${SCRIPT_DIR}/FindPlatformExtensionSources.sh"
 
-		ARGS="/p:DefineConstants=__MonoCS__ /p:Configuration=Development /p:Platform=AnyCPU /verbosity:quiet /nologo"
+		ARGS="/p:Configuration=Development /p:Platform=AnyCPU /verbosity:quiet /nologo"
 		ARGS="${ARGS} /p:TargetFrameworkProfile="
 
 		echo "xbuild Source/Programs/AutomationTool/AutomationTool.csproj $ARGS"
diff --git a/Engine/Source/Programs/AutomationTool/HoloLens/HoloLens.Automation.csproj b/Engine/Source/Programs/AutomationTool/HoloLens/HoloLens.Automation.csproj
--- a/Engine/Source/Programs/AutomationTool/HoloLens/HoloLens.Automation.csproj
+++ b/Engine/Source/Programs/AutomationTool/HoloLens/HoloLens.Automation.csproj
@@ -50,7 +50,14 @@
       <HintPath>..\..\..\..\..\Engine\Binaries\ThirdParty\WindowsDevicePortalWrapper\WindowsDevicePortalWrapper.dll</HintPath>
     </Reference>
   </ItemGroup>
-    <Choose>
+  <Choose>
+    <When Condition="'$(OS)' != 'Windows_NT'">
+	  <PropertyGroup>
+		<DefineConstants>__MonoCS__</DefineConstants>
+	  </PropertyGroup>
+	</When>
+  </Choose>
+  <Choose>
     <When Condition="'$(OS)' == 'Windows_NT'">
 	  <PropertyGroup>
 		<Windows10SdkPath>$(Registry:HKEY_LOCAL_MACHINE\SOFTWARE\Wow6432Node\Microsoft\Microsoft SDKs\Windows\v10.0@InstallationFolder)</Windows10SdkPath>
diff --git a/Engine/Source/Programs/UnrealBuildTool/UnrealBuildTool.csproj b/Engine/Source/Programs/UnrealBuildTool/UnrealBuildTool.csproj
--- a/Engine/Source/Programs/UnrealBuildTool/UnrealBuildTool.csproj
+++ b/Engine/Source/Programs/UnrealBuildTool/UnrealBuildTool.csproj
@@ -294,6 +294,13 @@
 	  </ItemGroup>
 	</When>
   </Choose>
+  <Choose>
+    <When Condition="'$(OS)' != 'Windows_NT'">
+	  <PropertyGroup>
+		<DefineConstants>__MonoCS__</DefineConstants>
+	  </PropertyGroup>
+	</When>
+  </Choose>
   <Choose>
     <When Condition="'$(OS)' == 'Windows_NT'">
 	  <PropertyGroup>
