Suppress XML documentation warning

diff --git a/Engine/Build/BatchFiles/Linux/GenerateProjectFiles.sh b/Engine/Build/BatchFiles/Linux/GenerateProjectFiles.sh
--- a/Engine/Build/BatchFiles/Linux/GenerateProjectFiles.sh
+++ b/Engine/Build/BatchFiles/Linux/GenerateProjectFiles.sh
@@ -26,7 +26,7 @@ source "$BASE_PATH/SetupMono.sh" "$BASE_PATH"
 "${BASE_PATH}/../FindPlatformExtensionSources.sh"
 
 if [ -f "$BASE_PATH/../../../Source/Programs/UnrealBuildTool/UnrealBuildTool.csproj" ]; then
-	xbuild "$BASE_PATH/../../../Source/Programs/UnrealBuildTool/UnrealBuildTool.csproj" /property:Configuration="Development" /verbosity:quiet /nologo
+	xbuild "$BASE_PATH/../../../Source/Programs/UnrealBuildTool/UnrealBuildTool.csproj" /property:Configuration="Development" /verbosity:quiet /nologo /p:NoWarn=1591
 fi
 
 # pass all parameters to UBT
diff --git a/Engine/Build/BatchFiles/Mac/GenerateProjectFiles.sh b/Engine/Build/BatchFiles/Mac/GenerateProjectFiles.sh
--- a/Engine/Build/BatchFiles/Mac/GenerateProjectFiles.sh
+++ b/Engine/Build/BatchFiles/Mac/GenerateProjectFiles.sh
@@ -25,7 +25,7 @@ fi
 "${BASE_PATH}/../FindPlatformExtensionSources.sh"
 
 if [ -f "$BASE_PATH/../../../Source/Programs/UnrealBuildTool/UnrealBuildTool.csproj" ]; then
-	xbuild "$BASE_PATH/../../../Source/Programs/UnrealBuildTool/UnrealBuildTool.csproj" /property:Configuration="Development" /verbosity:quiet /nologo |grep -i error
+	xbuild "$BASE_PATH/../../../Source/Programs/UnrealBuildTool/UnrealBuildTool.csproj" /property:Configuration="Development" /verbosity:quiet /nologo /p:NoWarn=1591 |grep -i error
 fi
 
 WANT_AOT="`defaults read com.epicgames.ue4 MonoAOT`"
diff --git a/Engine/Build/BatchFiles/RunUAT.sh b/Engine/Build/BatchFiles/RunUAT.sh
--- a/Engine/Build/BatchFiles/RunUAT.sh
+++ b/Engine/Build/BatchFiles/RunUAT.sh
@@ -76,7 +76,7 @@ if [ "$UATCompileArg" = "-compile" ]; then
 		# make sure the UBT project has references to auto-discovered platform extension source files
 		"${SCRIPT_DIR}/FindPlatformExtensionSources.sh"
 
-		ARGS="/p:Configuration=Development /p:Platform=AnyCPU /verbosity:quiet /nologo"
+		ARGS="/p:Configuration=Development /p:Platform=AnyCPU /verbosity:quiet /nologo /p:NoWarn=1591"
 		ARGS="${ARGS} /p:TargetFrameworkProfile="
 
 		echo "xbuild Source/Programs/AutomationTool/AutomationTool.csproj $ARGS"
diff --git a/Engine/Source/Programs/AutomationTool/AutomationUtils/ScriptCompiler.cs b/Engine/Source/Programs/AutomationTool/AutomationUtils/ScriptCompiler.cs
--- a/Engine/Source/Programs/AutomationTool/AutomationUtils/ScriptCompiler.cs
+++ b/Engine/Source/Programs/AutomationTool/AutomationUtils/ScriptCompiler.cs
@@ -258,7 +258,7 @@ namespace AutomationTool
 				throw new AutomationException(String.Format("Unable to build Project {0}. Project file not found.", ProjectFile));
 			}
 
-			string CmdLine = String.Format("\"{0}\" /verbosity:quiet /nologo /target:Build /property:Configuration={1} /property:Platform=AnyCPU /p:TreatWarningsAsErrors=false /p:NoWarn=\"612,618,672\" /p:BuildProjectReferences=true",
+			string CmdLine = String.Format("\"{0}\" /verbosity:quiet /nologo /target:Build /property:Configuration={1} /property:Platform=AnyCPU /p:TreatWarningsAsErrors=false /p:NoWarn=\"612,618,672,1591\" /p:BuildProjectReferences=true",
 				ProjectFile, BuildConfig);
 
 			// Compile the project
