Support for running EngineTest automated tests on Lumin via Gauntlet

Disabled by default until we get hardware set up

diff --git a/Engine/Source/Programs/AutomationTool/Gauntlet/Gauntlet.Automation.csproj b/Engine/Source/Programs/AutomationTool/Gauntlet/Gauntlet.Automation.csproj
--- a/Engine/Source/Programs/AutomationTool/Gauntlet/Gauntlet.Automation.csproj
+++ b/Engine/Source/Programs/AutomationTool/Gauntlet/Gauntlet.Automation.csproj
@@ -158,6 +158,10 @@
       <Name>Localization.Automation</Name>
       <Private>False</Private>
     </ProjectReference>
+    <ProjectReference Include="..\Lumin\Lumin.Automation.csproj">
+      <Project>{17779110-5cab-43a0-9dfb-4df67df29b42}</Project>
+      <Name>Lumin.Automation</Name>
+    </ProjectReference>
     <ProjectReference Include="..\XLocLocalization\XLocLocalization.Automation.csproj">
       <Project>{7e2d6784-c9b7-40c6-a89b-0a8afddce60d}</Project>
       <Name>XLocLocalization.Automation</Name>
@@ -222,6 +226,14 @@
       </ItemGroup>
     </When>
   </Choose>
+  <Choose>
+    <When Condition="Exists('Platform\Lumin\Gauntlet.TargetDeviceLumin.cs')">
+      <ItemGroup>
+        <Compile Include="Platform\Lumin\Gauntlet.TargetDeviceLumin.cs" />
+        <Compile Include="Platform\Lumin\Gauntlet.LuminBuildSource.cs" />
+      </ItemGroup>
+    </When>
+  </Choose>
   <Choose>
     <When Condition="Exists('Platform\Mac\Gauntlet.TargetDeviceMac.cs')">
       <ItemGroup>
diff --git a/Engine/Source/Programs/AutomationTool/Gauntlet/Platform/Lumin/Gauntlet.LuminBuildSource.cs b/Engine/Source/Programs/AutomationTool/Gauntlet/Platform/Lumin/Gauntlet.LuminBuildSource.cs
new file mode 100644
--- /dev/null
+++ b/Engine/Source/Programs/AutomationTool/Gauntlet/Platform/Lumin/Gauntlet.LuminBuildSource.cs
@@ -0,0 +1,156 @@
+﻿// Copyright 1998-2019 Epic Games, Inc. All Rights Reserved.
+
+using System;
+using System.Collections.Generic;
+using System.IO;
+using AutomationTool;
+using UnrealBuildTool;
+using System.Threading;
+using System.Text.RegularExpressions;
+using System.Linq;
+
+namespace Gauntlet
+{
+	public class LuminBuild : IBuild
+	{
+		public UnrealTargetConfiguration Configuration { get; protected set; }
+
+		public string SourceMpkPath;
+
+		public string LuminPackageName;
+		public BuildFlags Flags { get; protected set; }
+
+		public UnrealTargetPlatform Platform { get { return UnrealTargetPlatform.Lumin; } }
+
+		public LuminBuild(UnrealTargetConfiguration InConfig, string InLuminPackageName, string InMpkPath, BuildFlags InFlags)
+		{
+			Configuration = InConfig;
+			LuminPackageName = InLuminPackageName;
+			SourceMpkPath = InMpkPath;
+			Flags = InFlags;
+		}
+
+		public bool CanSupportRole(UnrealTargetRole RoleType)
+		{
+			if (RoleType.IsClient())
+			{
+				return true;
+			}
+
+			return false;
+		}
+
+		public static IEnumerable<LuminBuild> CreateFromPath(string InProjectName, string InPath)
+		{
+			string BuildPath = InPath;
+
+			List<LuminBuild> DiscoveredBuilds = new List<LuminBuild>();
+
+			DirectoryInfo Di = new DirectoryInfo(BuildPath);
+
+			// find all install batchfiles
+			FileInfo[] InstallFiles = Di.GetFiles("Install_*");
+
+			foreach (FileInfo Fi in InstallFiles)
+			{
+				Log.Verbose("Pulling install data from {0}", Fi.FullName);
+
+				string AbsPath = Fi.Directory.FullName;
+
+				// read contents and replace linefeeds (regex doesn't stop on them :((
+				string BatContents = File.ReadAllText(Fi.FullName).Replace(Environment.NewLine, "\n");
+
+				// Replace .bat with .mpk and strip up to and including the first _, that is then our MPK name
+				string SourceMpkPath = Regex.Replace(Fi.Name, ".bat", ".mpk", RegexOptions.IgnoreCase);
+				SourceMpkPath = SourceMpkPath.Substring(SourceMpkPath.IndexOf("_") + 1);
+				SourceMpkPath = Path.Combine(AbsPath, SourceMpkPath);
+
+				Match Info = Regex.Match(BatContents, @"install\s+(-u+)\s""%~dp0\\(.+)""");
+				string LuminPackageName = Info.Groups[2].ToString();
+	
+				if (string.IsNullOrEmpty(SourceMpkPath))
+				{
+					Log.Warning("No MPK found for build at {0}", Fi.FullName);
+					continue;
+				}
+
+				if (string.IsNullOrEmpty(LuminPackageName))
+				{
+					Log.Warning("No product name found for build at {0}", Fi.FullName);
+					continue;
+				}
+
+				UnrealTargetConfiguration UnrealConfig = UnrealHelpers.GetConfigurationFromExecutableName(InProjectName, Fi.Name);
+
+				// Lumin builds are always packaged, and we can always replace the command line
+				BuildFlags Flags = BuildFlags.Packaged | BuildFlags.CanReplaceCommandLine;
+				
+				if (AbsPath.Contains("Bulk"))
+				{
+					Flags |= BuildFlags.Bulk;
+				}
+				else
+				{
+					Flags |= BuildFlags.NotBulk;
+				}
+
+				LuminBuild NewBuild = new LuminBuild(UnrealConfig, LuminPackageName, SourceMpkPath, Flags);
+
+				DiscoveredBuilds.Add(NewBuild);
+
+				Log.Verbose("Found {0} {1} build at {2}", UnrealConfig, ((Flags & BuildFlags.Bulk) == BuildFlags.Bulk) ? "(bulk)" : "(not bulk)", AbsPath);
+
+			}
+
+			return DiscoveredBuilds;
+		}
+	}
+	
+	public class LuminBuildSource : IFolderBuildSource
+	{	
+		public string BuildName { get { return "LuminBuildSource";  } }
+
+		public bool CanSupportPlatform(UnrealTargetPlatform InPlatform)
+		{
+			return InPlatform == UnrealTargetPlatform.Lumin;
+		}
+
+		public string ProjectName { get; protected set; }
+
+		public LuminBuildSource()
+		{
+		}
+
+		public List<IBuild> GetBuildsAtPath(string InProjectName, string InPath, int MaxRecursion = 3)
+		{		
+			List<IBuild> Builds = new List<IBuild>();
+
+			// First try creating a build just from this path
+			IEnumerable<LuminBuild> FoundBuilds = LuminBuild.CreateFromPath(InProjectName, InPath);
+
+			// If we found one we were pointed at a specific case so we're done
+			if (FoundBuilds.Any())
+			{
+				Builds.AddRange(FoundBuilds);
+			}
+			else
+			{
+				// Ok, find all directories that match our platform and add those
+				DirectoryInfo TopDir = new DirectoryInfo(InPath);
+
+				// find all directories that begin with Lumin
+				IEnumerable<string> LuminDirs = Directory.EnumerateDirectories(InPath, "Lumin*", SearchOption.AllDirectories);
+
+				foreach (string DirPath in LuminDirs)
+				{
+					
+
+					FoundBuilds = LuminBuild.CreateFromPath(InProjectName, DirPath);
+					Builds.AddRange(FoundBuilds);
+				}
+			}
+
+			return Builds;
+		}
+	}
+}
\ No newline at end of file
diff --git a/Engine/Source/Programs/AutomationTool/Gauntlet/Platform/Lumin/Gauntlet.TargetDeviceLumin.cs b/Engine/Source/Programs/AutomationTool/Gauntlet/Platform/Lumin/Gauntlet.TargetDeviceLumin.cs
new file mode 100644
--- /dev/null
+++ b/Engine/Source/Programs/AutomationTool/Gauntlet/Platform/Lumin/Gauntlet.TargetDeviceLumin.cs
@@ -0,0 +1,930 @@
+// Copyright 1998-2019 Epic Games, Inc. All Rights Reserved.
+
+using System;
+using System.Collections.Generic;
+using System.IO;
+using AutomationTool;
+using UnrealBuildTool;
+using System.Text.RegularExpressions;
+using System.Linq;
+using System.Text;
+using System.Threading;
+using Tools.DotNETCommon;
+
+namespace Gauntlet
+{
+	// device data from json
+	public sealed class LuminDeviceData
+	{
+		// remote device settings (wifi)
+
+		// host of PC which is tethered 
+		public string hostIP { get; set; }
+
+		// public key 
+		public string publicKey { get; set; }
+
+		// private key
+		public string privateKey { get; set; }
+	}
+
+	// become IAppInstance when implemented enough
+	class LuminAppInstance : IAppInstance
+	{
+		protected TargetDeviceLumin LuminDevice;
+
+		protected LuminAppInstall Install;
+
+		internal IProcessResult LaunchProcess;
+
+		internal bool bHaveSavedArtifacts;
+
+		public string CommandLine { get { return Install.CommandLine; } }
+
+		public LuminAppInstance(TargetDeviceLumin InDevice, LuminAppInstall InInstall, IProcessResult InProcess)
+		{
+			LuminDevice = InDevice;
+			Install = InInstall;
+			LaunchProcess = InProcess;
+		}
+
+		public string ArtifactPath
+		{
+			get
+			{
+				if (bHaveSavedArtifacts == false)
+				{
+					if (HasExited)
+					{
+						SaveArtifacts();
+						bHaveSavedArtifacts = true;
+					}
+				}
+				
+				return Path.Combine(LuminDevice.LocalCachePath, "Saved");
+			}
+		}
+
+		public ITargetDevice Device
+		{
+			get
+			{
+				return LuminDevice;
+			}
+		}
+
+		public bool HasExited
+		{
+			get
+			{
+				return IsActivityRunning();
+			}
+		}
+
+		/// <summary>
+		/// Checks on device whether the activity is running, this is an expensive shell with output operation
+		/// the result is cached, with checks at ActivityCheckDelta seconds
+		/// </summary>
+		private bool IsActivityRunning()
+		{
+			if (ActivityExited)
+			{
+				return false;
+			}
+
+			if ((DateTime.UtcNow - ActivityCheckTime) < ActivityCheckDelta)
+			{
+				return false;
+			}
+
+			ActivityCheckTime = DateTime.UtcNow;
+
+			// get activities filtered by our package name
+			IProcessResult ActivityQuery = LuminDevice.RunMldbDeviceCommand("ps", true);
+			while (ActivityQuery.Output == "")
+			{
+				ActivityQuery = LuminDevice.RunMldbDeviceCommand("ps",true);
+			}
+
+			// We have exited if our activity doesn't appear in the activity query or is not the focused activity.
+			bool bActivityPresent = ActivityQuery.Output.Contains(LuminDevice.GetQualifiedProjectName(Install.Name));
+			bool bActivityInForeground = ActivityQuery.Output.Contains(".fullscreen"); // Technically only checking that ANY app is fullscreen, but should work for now
+			bool bHasExited = !bActivityPresent || !bActivityInForeground;
+			if (bHasExited)
+			{
+				ActivityExited = true;
+				// The activity has exited, make sure entire activity log has been captured, sleep to allow time for the log to flush
+				Thread.Sleep(5000);
+				UpdateCachedLog(true);
+				Log.VeryVerbose("{0}: process exited, Activity running={1}, Activity in foreground={2} ", ToString(), bActivityPresent.ToString(), bActivityInForeground.ToString());
+			}
+
+			return bHasExited;
+
+		}
+
+		private static readonly TimeSpan ActivityCheckDelta = TimeSpan.FromSeconds(10);
+		private DateTime ActivityCheckTime = DateTime.UtcNow;
+		private bool ActivityExited = false;
+
+		public bool WasKilled { get; protected set; }
+
+		/// <summary>
+		/// The output of the test activity
+		/// </summary>
+		public string StdOut
+		{
+			get
+			{
+				UpdateCachedLog();
+				return String.IsNullOrEmpty(ActivityLogCached) ? String.Empty : ActivityLogCached;
+			}
+		}
+		/// <summary>
+		/// Updates cached activity log by running a shell command returning the full log from device (possibly over wifi)
+		/// The result is cached and updated at ActivityLogDelta frequency
+		/// </summary>
+		private void UpdateCachedLog(bool ForceUpdate = false)
+		{
+			if (!ForceUpdate && (ActivityLogTime == DateTime.MinValue || ((DateTime.UtcNow - ActivityLogTime) < ActivityLogDelta)))
+			{
+				return;
+			}
+
+			if (Install.LuminDevice != null && Install.LuminDevice.Disposed)
+			{
+				Log.Warning("Attempting to cache log using disposed Lumin device");
+				return;
+			}
+
+			string GetLogCommand = string.Format("log -d");
+			IProcessResult LogQuery = Install.LuminDevice.RunMldbDeviceCommand(GetLogCommand, true, false);
+
+			if (LogQuery.ExitCode != 0)
+			{
+				Log.VeryVerbose("Unable to query activity stdout on device {0}", Install.LuminDevice.Name);
+			}
+			else
+			{
+				ActivityLogCached = LogQuery.Output;
+			}
+
+			ActivityLogTime = DateTime.UtcNow;
+
+			// the activity has exited, mark final log sentinel 
+			if (ActivityExited)
+			{
+				ActivityLogTime = DateTime.MinValue;
+			}
+
+		}
+
+		private static readonly TimeSpan ActivityLogDelta = TimeSpan.FromSeconds(15);
+		private DateTime ActivityLogTime = DateTime.UtcNow - ActivityLogDelta;
+		private string ActivityLogCached = string.Empty;
+		
+
+		public int WaitForExit()
+		{
+			if (!HasExited)
+			{
+				if (!LaunchProcess.HasExited)
+				{
+					LaunchProcess.WaitForExit();
+				}
+
+			}
+
+			return ExitCode;
+		}
+
+		public void Kill()
+		{
+			if (!HasExited)
+			{
+				WasKilled = true;
+				Install.LuminDevice.KillRunningProcess(Install.LuminDevice.GetQualifiedProjectName(Install.Name));
+			}
+		}
+		public int ExitCode { get { return LaunchProcess.ExitCode; } }
+
+		protected void SaveArtifacts()
+		{
+			// copy remote artifacts to local
+			if (Directory.Exists(Install.LuminDevice.LocalCachePath))
+			{
+				try
+				{
+					// don't consider this fatal, people often have the directory or a file open
+					Directory.Delete(Install.LuminDevice.LocalCachePath, true);
+				}
+				catch
+				{
+					Log.Warning("Failed to remove old cache folder {0}", Install.LuminDevice.LocalCachePath);
+				}
+			}
+
+			// mark it as a temp dir (will also create it)
+			Utils.SystemHelpers.MarkDirectoryForCleanup(Install.LuminDevice.LocalCachePath);
+
+			string LocalSaved = Path.Combine(Install.LuminDevice.LocalCachePath, "Saved");
+			Directory.CreateDirectory(LocalSaved);
+
+			string QualifiedPackageName = this.LuminDevice.GetQualifiedProjectName(this.LuminDevice.InstalledAppName);
+
+			// pull all the artifacts
+			string ArtifactPullCommand = string.Format("pull -p {0} {1} {2}", QualifiedPackageName, Install.LuminDevice.DeviceArtifactPath, Install.LuminDevice.LocalCachePath);
+			IProcessResult PullCmd = Install.LuminDevice.RunMldbDeviceCommand(ArtifactPullCommand);
+
+			if (PullCmd.ExitCode != 0)
+			{
+				Log.Warning("Failed to retrieve artifacts. {0}", PullCmd.Output);
+			}
+			else
+			{
+				// update final cached stdout property
+				string LogFilename = string.Format("{0}/Logs/{1}.log", LocalSaved, Install.Name);
+				if (File.Exists(LogFilename))
+				{
+					ActivityLogCached = File.ReadAllText(LogFilename);
+					ActivityLogTime = DateTime.MinValue;
+				}
+
+			}
+
+			// pull the logcat over from device.
+			IProcessResult LogcatResult = Install.LuminDevice.RunMldbDeviceCommand("log -d");
+
+			string LogcatFilename = "Logcat.log";
+			// Save logcat dump to local artifact path.
+			File.WriteAllText(Path.Combine(LocalSaved, LogcatFilename), LogcatResult.Output);
+
+			// check for crash dumps
+			string CrashDumpListCommand = "crashdump list";
+			IProcessResult CrashDumpListCmd = Install.LuminDevice.RunMldbDeviceCommand(CrashDumpListCommand);
+
+			if (PullCmd.Output.Contains(QualifiedPackageName))
+			{
+				try
+				{
+					// Attempt to pull relevant crash dump (unsure what kind of file this returns, this might be the wrong file extension)
+					string GetCrashDumpCommand = string.Format("mldb crashdump get {0} {1}", QualifiedPackageName, Path.Combine(Install.LuminDevice.LocalCachePath, string.Format("{0}.dmp", QualifiedPackageName)));
+					IProcessResult GetCrashDumpCmmd = Install.LuminDevice.RunMldbDeviceCommand(GetCrashDumpCommand);
+				}
+				catch
+				{
+					Log.Warning("Failed to retrieve crash dump for {0}", QualifiedPackageName);
+				}
+			}
+
+			Install.LuminDevice.PostRunCleanup();
+		}
+	}
+
+	class LuminAppInstall : IAppInstall
+	{
+		public string Name { get; protected set; }
+
+		public string LuminPackageName { get; protected set; }
+
+		public TargetDeviceLumin LuminDevice { get; protected set; }
+
+		public ITargetDevice Device { get { return LuminDevice; } }
+
+		public string CommandLine { get; protected set; }
+
+		public IAppInstance Run()
+		{
+			return LuminDevice.Run(this);
+		}
+
+		public LuminAppInstall(TargetDeviceLumin InDevice, string InName, string InLuminPackageName, string InCommandLine)
+		{
+			LuminDevice = InDevice;
+			Name = InName;
+			LuminPackageName = InLuminPackageName;
+			CommandLine = InCommandLine;
+		}
+	}
+
+	public class DefaultLuminDevices : IDefaultDeviceSource
+	{
+		public bool CanSupportPlatform(UnrealTargetPlatform? Platform)
+		{
+			return Platform == UnrealTargetPlatform.Lumin;
+		}
+
+		public ITargetDevice[] GetDefaultDevices()
+		{
+			return TargetDeviceLumin.GetDefaultDevices();
+		}
+	}
+
+	public class LuminDeviceFactory : IDeviceFactory
+	{
+		public bool CanSupportPlatform(UnrealTargetPlatform? Platform)
+		{
+			return Platform == UnrealTargetPlatform.Lumin;
+		}
+
+		public ITargetDevice CreateDevice(string InRef, string InParam)
+		{
+			LuminDeviceData DeviceData = null;
+
+			if (!String.IsNullOrEmpty(InParam))
+			{
+				DeviceData = fastJSON.JSON.Instance.ToObject<LuminDeviceData>(InParam);
+			}
+
+			return new TargetDeviceLumin(InRef, DeviceData);
+		}
+	}
+
+	/// <summary>
+	/// Lumin implementation of a device that can run applications
+	/// </summary>
+	public class TargetDeviceLumin : ITargetDevice
+	{ 
+		/// <summary>
+		/// Friendly name for this target
+		/// </summary>
+		public string Name { get; protected set; }
+
+		public string InstalledAppName { get; protected set; }
+
+		/// <summary>
+		/// Low-level device name
+		/// </summary>
+		public string DeviceName { get; protected set; }
+
+		/// <summary>
+		/// Platform type.
+		/// </summary>
+		public UnrealTargetPlatform? Platform { get { return UnrealTargetPlatform.Lumin; } }
+
+		/// <summary>
+		/// Options for executing commands
+		/// </summary>
+		public CommandUtils.ERunOptions RunOptions { get; set; }
+
+		/// <summary>
+		/// Temp path we use to push/pull things from the device
+		/// </summary>
+		public string LocalCachePath { get; protected set; }
+
+
+		/// <summary>
+		/// Artifact (e.g. Saved) path on the device
+		/// </summary>
+		public string DeviceArtifactPath { get; protected set;  }
+
+		/// <summary>
+		/// Path to a command line if installed
+		/// </summary>
+		protected string CommandLineFilePath { get; set; }
+
+		public bool IsAvailable
+		{
+			get
+			{
+                // ensure our device is present in 'mldb devices' output.
+				var AllDevices = GetAllConnectedDevices();
+
+				if (AllDevices.Keys.Contains(DeviceName) == false)
+				{
+					return false;
+				}
+
+				if (AllDevices[DeviceName] == false)
+				{
+					Log.Warning("Device {0} is connected but we are not authorized", DeviceName);
+					return false;
+				}
+
+				// any device will do, but only one at a time.
+				return true;
+            }
+		}
+
+		protected Dictionary<EIntendedBaseCopyDirectory, string> LocalDirectoryMappings { get; set; }
+        void SetUpDirectoryMappings()
+        {
+            LocalDirectoryMappings = new Dictionary<EIntendedBaseCopyDirectory, string>();
+        }
+
+        public void PopulateDirectoryMappings(string ProjectDir)
+		{
+			LocalDirectoryMappings.Add(EIntendedBaseCopyDirectory.Build, Path.Combine(ProjectDir, "Build"));
+			LocalDirectoryMappings.Add(EIntendedBaseCopyDirectory.Binaries, Path.Combine(ProjectDir, "Binaries"));
+			LocalDirectoryMappings.Add(EIntendedBaseCopyDirectory.Config, Path.Combine(ProjectDir, "Config"));
+            LocalDirectoryMappings.Add(EIntendedBaseCopyDirectory.Content, Path.Combine(ProjectDir, "Content"));
+            LocalDirectoryMappings.Add(EIntendedBaseCopyDirectory.Demos, Path.Combine(ProjectDir, "Demos"));
+            LocalDirectoryMappings.Add(EIntendedBaseCopyDirectory.Profiling, Path.Combine(ProjectDir, "Profiling"));
+            LocalDirectoryMappings.Add(EIntendedBaseCopyDirectory.Saved, ProjectDir);
+        }
+        public bool IsConnected { get	{ return IsAvailable; }	}
+
+		protected bool IsExistingDevice = false;		
+
+		/// <summary>
+		/// Constructor
+		/// </summary>
+		/// <param name="InReferenceName"></param>
+		/// <param name="InRemoveOnDestruction"></param>
+		public TargetDeviceLumin(string InDeviceName = "", LuminDeviceData DeviceData = null)
+		{
+			DeviceName = InDeviceName;
+			
+			// If no device name or its 'default' then use the first default device
+			if (string.IsNullOrEmpty(DeviceName) || DeviceName.Equals("default", StringComparison.OrdinalIgnoreCase))
+			{
+				var DefaultDevices = GetAllAvailableDevices();	
+
+				if (DefaultDevices.Count() == 0)
+				{
+					if (GetAllConnectedDevices().Count > 0)
+					{
+						throw new AutomationException("No default device available. One or more devices are connected but unauthorized. See 'mldb devices'");
+					}
+					else
+					{
+						throw new AutomationException("No default device available. See 'mldb devices'");
+					}
+				}
+
+				DeviceName = DefaultDevices.First();
+			}
+
+			if (Log.IsVerbose)
+			{
+				RunOptions = CommandUtils.ERunOptions.NoWaitForExit;
+			}
+			else
+			{
+				RunOptions = CommandUtils.ERunOptions.NoWaitForExit | CommandUtils.ERunOptions.NoLoggingOfRunCommand;
+			}
+
+			// if this is not a connected device then remove when done
+			var ConnectedDevices = GetAllConnectedDevices();
+
+			IsExistingDevice = ConnectedDevices.Keys.Contains(DeviceName);
+
+			if (!IsExistingDevice)
+			{
+				// TODO: adb uses 5555 by default, not sure about mldb?
+				if (DeviceName.Contains(":") == false)
+				{
+					DeviceName = DeviceName + ":5555";
+				}
+
+				lock (Globals.MainLock)
+				{
+					using (var PauseEC = new ScopedSuspendECErrorParsing())
+					{
+						IProcessResult MldbResult = RunMldbGlobalCommand(string.Format("connect {0}", DeviceName));
+
+						if (MldbResult.ExitCode != 0)
+						{
+							throw new AutomationException("mldb failed to connect to {0}. {1}", DeviceName, MldbResult.Output);
+						}
+					}
+
+					Log.Info("Connected to {0}", DeviceName);
+
+					// Need to sleep for mldb service process to register, otherwise get an unauthorized (especially on parallel device use)
+					Thread.Sleep(5000);
+				}
+			}
+
+			LocalDirectoryMappings = new Dictionary<EIntendedBaseCopyDirectory, string>();
+
+			// for IP devices need to sanitize this
+			Name = DeviceName.Replace(":", "_");
+
+			// Path we use for artifacts, we'll create it later when we need it
+			LocalCachePath = Path.Combine(Globals.TempDir, "LuminDevice_" + Name);
+
+			ConnectedDevices = GetAllConnectedDevices();
+
+            SetUpDirectoryMappings();
+
+            // sanity check that it was now dound
+            if (ConnectedDevices.Keys.Contains(DeviceName) == false)
+			{
+				throw new AutomationException("Failed to find new device {0} in connection list", DeviceName);
+			}
+
+			if (ConnectedDevices[DeviceName] == false)
+			{
+				Dispose();
+				throw new AutomationException("Device {0} is connected but this PC is not authorized.", DeviceName);
+			}
+		}
+
+		~TargetDeviceLumin()
+		{
+			Dispose(false);
+		}
+
+		#region IDisposable Support
+		private bool disposedValue = false; // To detect redundant calls
+
+		protected virtual void Dispose(bool disposing)
+		{
+			if (!disposedValue)
+			{
+				try
+				{
+					if (!IsExistingDevice)
+					{
+						// disconnect
+						RunMldbGlobalCommand(string.Format("disconnect {0}", DeviceName), true, false, true);
+
+						Log.Info("Disconnected {0}", DeviceName);
+					}					
+				}
+				catch (Exception Ex)
+				{
+					Log.Warning("TargetDeviceLumin.Dispose() threw: {0}", Ex.Message);
+				}
+				finally
+				{
+					disposedValue = true;			
+				}
+
+			}
+		}
+
+		// This code added to correctly implement the disposable pattern.
+		public void Dispose()
+		{
+			// Do not change this code. Put cleanup code in Dispose(bool disposing) above.
+			Dispose(true);
+		}
+
+		public bool Disposed
+		{
+			get
+			{
+				return disposedValue;
+			}
+			
+		}
+
+		#endregion
+
+		/// <summary>
+		/// Returns a list of locally connected devices (e.g. 'mldb devices'). 
+		/// </summary>
+		/// <returns></returns>
+		static private Dictionary<string, bool> GetAllConnectedDevices()
+		{
+           var Result = RunMldbGlobalCommand("devices", true, false);
+
+            MatchCollection DeviceMatches = Regex.Matches(Result.Output, @"^([\d\w\.\:]{6,32})\s+(\w+)", RegexOptions.Multiline);
+
+            var DeviceList = DeviceMatches.Cast<Match>().ToDictionary(
+                M => M.Groups[1].ToString(),
+                M => !M.Groups[2].ToString().ToLower().Contains("unauthorized")
+            );
+
+            return DeviceList;
+		}
+
+		static private IEnumerable<string> GetAllAvailableDevices()
+		{
+			var AllDevices = GetAllConnectedDevices();
+			return AllDevices.Keys.Where(D => AllDevices[D] == true);
+		}
+
+		static public ITargetDevice[] GetDefaultDevices()
+		{
+			var Result = RunMldbGlobalCommand("devices");
+
+			MatchCollection DeviceMatches = Regex.Matches(Result.Output, @"([\d\w\.\:]{8,32})\s+device");
+
+			List<ITargetDevice> Devices = new List<ITargetDevice>();
+
+			foreach (string Device in GetAllAvailableDevices())
+			{
+				ITargetDevice NewDevice = new TargetDeviceLumin(Device);
+				Devices.Add(NewDevice);
+			}
+
+			return Devices.ToArray();
+		}
+
+		internal void PostRunCleanup()
+		{
+			// Delete the commandline file, if someone installs an MPK on top of ours
+			// they will get very confusing behavior...
+			if (string.IsNullOrEmpty(CommandLineFilePath) == false)
+			{
+				Log.Verbose("Removing {0}", CommandLineFilePath);
+				DeleteFileFromDevice(CommandLineFilePath);
+				CommandLineFilePath = null;
+			}
+		}
+
+		public bool IsOn
+		{
+			get
+			{
+				string CommandLine = "get-state";
+				IProcessResult OnAndUnlockedQuery = RunMldbDeviceCommand(CommandLine);
+
+				return OnAndUnlockedQuery.Output.Contains("device");
+			}
+		}
+
+		public bool PowerOn()
+		{
+			Log.Verbose("{0}: Powering on", ToString());
+			string CommandLine = "wait-for-device";
+			RunMldbDeviceCommand(CommandLine);
+			return true;
+		}
+		public bool PowerOff()
+		{
+			Log.Verbose("{0}: Powering off", ToString());
+
+			string CommandLine = "shutdown";
+			RunMldbDeviceCommand(CommandLine);
+			return true;
+		}
+
+		public bool Reboot()
+		{
+			Log.Verbose("{0}: Rebooting", ToString());
+
+			string CommandLine = "reboot";
+			RunMldbDeviceCommand(CommandLine);
+			return true;
+		}
+
+		public bool Connect()
+		{
+			return true;
+		}
+
+		public bool Disconnect()
+		{
+			return true;
+		}
+
+		public override string ToString()
+		{
+			if (Name == DeviceName)
+			{
+				return Name;
+			}
+			return string.Format("{0} ({1})", Name, DeviceName);
+		}
+
+		protected bool DeleteFileFromDevice(string DestPath)
+		{
+			string QualifiedPackageName = GetQualifiedProjectName(InstalledAppName);
+			var MldbResult = RunMldbDeviceCommand(string.Format("rm -f -p {0}", QualifiedPackageName, DestPath));
+			return MldbResult.ExitCode == 0;
+		}
+
+		//public bool CopyFileToDevice(string PackageName, string SourcePath, string DestPath, bool IgnoreDependencies = false)
+		//{
+			// This function existed for Android, but I don't think we necessarily need it for Lumin.  No separate OBB-style data currently.
+			// If needed in the future can base it off Gauntlet.TargetDeviceAndroid.cs' implementation.
+		//}
+
+		// This is the com.yourcompany.projectname that is shown for installs on device
+		public string GetQualifiedProjectName(string ProjectName)
+		{
+			// ask the .ini system for what version to use
+			ConfigHierarchy Ini = ConfigCache.ReadHierarchy(ConfigHierarchyType.Engine, DirectoryReference.FromString(Path.Combine(Environment.CurrentDirectory, ProjectName).ToString()), UnrealTargetPlatform.Lumin);
+			// check for project override of NDK API level
+			string QualifiedPackageName;
+			Ini.GetString("/Script/LuminRuntimeSettings.LuminRuntimeSettings", "PackageName", out QualifiedPackageName);
+			QualifiedPackageName = QualifiedPackageName.Trim();
+			QualifiedPackageName = QualifiedPackageName.Replace("[PROJECT]", ProjectName);
+			return QualifiedPackageName.ToLower();
+		}
+
+		public IAppInstall InstallApplication(UnrealAppConfig AppConfig)
+		{
+			LuminBuild Build = AppConfig.Build as LuminBuild;
+
+			string QualifiedPackageName = GetQualifiedProjectName(AppConfig.ProjectName);
+
+			InstalledAppName = AppConfig.ProjectName;
+
+			// Ensure MPK exists
+			if (Build == null)
+			{
+				throw new AutomationException("Invalid build for Lumin!");
+			}
+
+			// kill any currently running instance:
+			KillRunningProcess(QualifiedPackageName);
+
+			bool SkipDeploy = Globals.Params.ParseParam("SkipDeploy");
+
+			if (SkipDeploy == false)
+			{
+				// remote dir used to save things
+				string RemoteDir = "/documents/c2/" + AppConfig.ProjectName.ToLower();				
+				//string DependencyDir = RemoteDir + "/deps";
+
+				// device artifact path, always clear between runs
+				// This clear is from andorid, but currently not needed on mldb.  Everything is removed when you remove the pak
+				// If in the future mldb supports a "sidecar" of data like the obb on Android, might need to delete that here
+				DeviceArtifactPath = string.Format("{0}/{1}/saved", RemoteDir.ToLower(), AppConfig.ProjectName.ToLower());
+
+				if (Globals.Params.ParseParam("cleandevice"))
+				{
+					Log.Info("Cleaning previous builds due to presence of -cleandevice");
+
+					// we need to ununstall then install the mpk - don't care if it fails, may have been deleted
+					Log.Info("Uninstalling {0}", QualifiedPackageName);
+					RunMldbDeviceCommand(string.Format("uninstall {0}", QualifiedPackageName));
+				}
+
+				// path to the MPK to install.
+				string MpkPath = Build.SourceMpkPath;
+
+				// check for a local newer executable
+				if (Globals.Params.ParseParam("dev"))
+				{
+					string MpkFileName = UnrealHelpers.GetExecutableName(AppConfig.ProjectName, UnrealTargetPlatform.Lumin, AppConfig.Configuration, AppConfig.ProcessType, "mpk");
+
+					string LocalMPK = Path.Combine(Environment.CurrentDirectory, AppConfig.ProjectName, "Binaries/Lumin", MpkFileName);
+
+					bool LocalFileExists = File.Exists(LocalMPK);
+					bool LocalFileNewer = LocalFileExists && File.GetLastWriteTime(LocalMPK) > File.GetLastWriteTime(MpkPath);
+
+					Log.Verbose("Checking for newer binary at {0}", LocalMPK);
+					Log.Verbose("LocalFile exists: {0}. Newer: {1}", LocalFileExists, LocalFileNewer);
+
+					if (LocalFileExists && LocalFileNewer)
+					{
+						MpkPath = LocalMPK;
+					}
+				}
+
+				// first install the MPK
+				RunMldbDeviceCommand(string.Format("install -u {0}", MpkPath));
+
+				// If in the future mldb supports a "sidecar" of data like the obb on Android, might need to install that here
+				// The Gauntlet.TargetDeviceAndroid.cs implementation should be a good starting point
+			}
+
+			LuminAppInstall AppInstall = new LuminAppInstall(this, AppConfig.ProjectName, Build.LuminPackageName, AppConfig.CommandLine);
+
+			return AppInstall;
+		}
+
+		public IAppInstance Run(IAppInstall App)
+		{
+			LuminAppInstall LuminInstall = App as LuminAppInstall;
+
+			if (LuminInstall == null)
+			{
+				throw new Exception("AppInstance is of incorrect type!");
+			}
+
+			string QualifiedPackageName = GetQualifiedProjectName(LuminInstall.Name);
+
+			// wake the device - we can install while its asleep but not run
+			PowerOn();
+
+			// kill any currently running instance:
+			KillRunningProcess(QualifiedPackageName);
+
+			string LaunchActivity = "";// LuminPlatform.GetLaunchableActivityName();
+
+			Log.Info("Launching {0} on '{1}' ", QualifiedPackageName, ToString());
+			Log.Verbose("\t{0}", LuminInstall.CommandLine);
+
+			// Clear the device's log in preparation for the test..
+			RunMldbDeviceCommand("log -c");
+
+			// start the app on device!
+			string CommandLine = "launch --auto-net-privs " + QualifiedPackageName + " -i \"" + LuminInstall.CommandLine + "\"";
+			IProcessResult Process = RunMldbDeviceCommand(CommandLine, false, true);
+
+			return new LuminAppInstance(this, LuminInstall, Process);
+		}
+
+		/// <summary>
+		/// Runs an MLDB command, automatically adding the name of the current device to
+		/// the arguments sent to mldb
+		/// </summary>
+		/// <param name="Args"></param>
+		/// <param name="Wait"></param>
+		/// <param name="Input"></param>
+		/// <returns></returns>
+		public IProcessResult RunMldbDeviceCommand(string Args, bool Wait=true, bool bShouldLogCommand = true, bool bPauseErrorParsing = false, bool bNoStdOutRedirect = false)
+		{
+			if (string.IsNullOrEmpty(DeviceName) == false)
+			{
+				Args = string.Format("-s {0} {1}", DeviceName, Args);
+			}
+
+			return RunMldbGlobalCommand(Args, Wait, bShouldLogCommand, bPauseErrorParsing, bNoStdOutRedirect);
+		}
+
+		/// <summary>
+		/// Runs an MLDB command, automatically adding the name of the current device to
+		/// the arguments sent to mldb
+		/// </summary>
+		/// <param name="Args"></param>
+		/// <param name="Wait"></param>
+		/// <param name="Input"></param>
+		/// <returns></returns>
+		public string RunMldbDeviceCommandAndGetOutput(string Args)
+		{
+			if (string.IsNullOrEmpty(DeviceName) == false)
+			{
+				Args = string.Format("-s {0} {1}", DeviceName, Args);
+			}
+
+			IProcessResult Result = RunMldbGlobalCommand(Args);
+
+			if (Result.ExitCode != 0)
+			{
+				throw new AutomationException("mldb command {0} failed. {1}", Args, Result.Output);
+			}
+
+			return Result.Output;
+		}
+
+		/// <summary>
+		/// Runs an MLDB command at the global scope
+		/// </summary>
+		/// <param name="Args"></param>
+		/// <param name="Wait"></param>
+		/// <returns></returns>
+		public static IProcessResult RunMldbGlobalCommand(string Args, bool Wait = true, bool bShouldLogCommand = true, bool bPauseErrorParsing = false, bool bNoStdOutRedirect = false)
+		{
+			CommandUtils.ERunOptions RunOptions = CommandUtils.ERunOptions.AppMustExist;
+			if (!Wait)
+			{
+				RunOptions |= CommandUtils.ERunOptions.NoWaitForExit;
+			}
+
+			if (Log.IsVeryVerbose)
+			{
+				RunOptions |= CommandUtils.ERunOptions.AllowSpew;
+			}
+			else
+			{
+				RunOptions |= CommandUtils.ERunOptions.NoLoggingOfRunCommand;
+			}
+
+			if (bNoStdOutRedirect)
+			{
+				RunOptions |= CommandUtils.ERunOptions.NoStdOutRedirect;
+			}
+	
+			if (bShouldLogCommand)
+			{
+				Log.Info("Running MLDB Command: mldb {0}", Args);
+			}
+
+			IProcessResult Process;
+
+			using (bPauseErrorParsing ? new ScopedSuspendECErrorParsing() : null)
+			{
+				Process = LuminPlatform.RunDeviceCommand(null, null, Args, null, RunOptions);
+
+				if (Wait)
+				{
+					Process.WaitForExit();
+				}
+			}
+			
+			return Process;
+		}
+
+		/// <summary>
+		/// Enable Lumin permissions which would otherwise block automation with permission requests
+		/// </summary>
+		public void EnablePermissions(string LuminPackageName)
+		{
+			throw new AutomationException("Not implemented for Lumin");
+		}
+
+		public void KillRunningProcess(string LuminPackageName)
+		{
+			Log.Verbose("{0}: Killing process '{1}' ", ToString(), LuminPackageName);
+			string KillProcessCommand = string.Format("terminate {0}", LuminPackageName);
+			RunMldbDeviceCommand(KillProcessCommand, false, true, false, true);
+		}
+
+		public Dictionary<EIntendedBaseCopyDirectory, string> GetPlatformDirectoryMappings()
+		{
+			if (LocalDirectoryMappings.Count == 0)
+			{
+				Log.Warning("Platform directory mappings have not been populated for this platform! This should be done within InstallApplication()");
+			}
+			return LocalDirectoryMappings;
+		}
+	}
+}
\ No newline at end of file
diff --git a/Engine/Source/Programs/AutomationTool/Lumin/LuminPlatform.Automation.cs b/Engine/Source/Programs/AutomationTool/Lumin/LuminPlatform.Automation.cs
--- a/Engine/Source/Programs/AutomationTool/Lumin/LuminPlatform.Automation.cs
+++ b/Engine/Source/Programs/AutomationTool/Lumin/LuminPlatform.Automation.cs
@@ -652,6 +652,11 @@ public class LuminPlatform : Platform
 		string UninstallBatchName = GetFinalBatchName(Params, SC, true);
 		BatchLines = GenerateUninstallBatchFile(PackageName, Params);
 		File.WriteAllLines(UninstallBatchName, BatchLines);
+		
+		// Copy mpk and .bat file to staging directory, because this is where Gauntlet looks for them to discover the build
+		File.Copy(MpkName, SC.CookSourceRuntimeRootDir + "\\" + Path.GetFileName(MpkName).ToString(), true);
+		File.Copy(BatchName, SC.CookSourceRuntimeRootDir + "\\" + Path.GetFileName(BatchName).ToString(), true);
+		File.Copy(UninstallBatchName, SC.CookSourceRuntimeRootDir + "\\" + Path.GetFileName(UninstallBatchName).ToString(), true);
 
 		// If needed, make the batch files able to execute
 		if (Utils.IsRunningOnMono)
@@ -1110,19 +1115,19 @@ public class LuminPlatform : Platform
 
 	#region Implementation Details
 
-	private string GetDeviceCommandLine(ProjectParams Params, string SerialNumber, string Args)
+	private static string GetDeviceCommandLine(ProjectParams Params, string SerialNumber, string Args)
 	{
-		if (SerialNumber != "")
+		if (SerialNumber != null && SerialNumber != "")
 		{
 			SerialNumber = "-s " + SerialNumber;
 		}
 
-		return string.Format("{0} {1}", SerialNumber, Args);
+		return string.Format("{0} {1}", SerialNumber != null ? SerialNumber : "", Args);
 	}
 
 	static string LastSpewFilename = "";
 
-	public string ADBSpewFilter(string Message)
+	public static string MLDBSpewFilter(string Message)
 	{
 		if (Message.StartsWith("[") && Message.Contains("%]"))
 		{
@@ -1143,7 +1148,7 @@ public class LuminPlatform : Platform
 		return Message;
 	}
 
-	private string DeviceCommand
+	private static string DeviceCommand
 	{
 		get
 		{
@@ -1174,12 +1179,12 @@ public class LuminPlatform : Platform
 		}
 	}
 
-	private IProcessResult RunDeviceCommand(ProjectParams Params, string SerialNumber, string Args, string Input = null, ERunOptions Options = ERunOptions.Default)
+	public static IProcessResult RunDeviceCommand(ProjectParams Params, string SerialNumber, string Args, string Input = null, ERunOptions Options = ERunOptions.Default)
 	{
 		if (Options.HasFlag(ERunOptions.AllowSpew) || Options.HasFlag(ERunOptions.SpewIsVerbose))
 		{
 			LastSpewFilename = "";
-			return Run(DeviceCommand, GetDeviceCommandLine(Params, SerialNumber, Args), Input, Options, SpewFilterCallback: new ProcessResult.SpewFilterCallbackType(ADBSpewFilter));
+			return Run(DeviceCommand, GetDeviceCommandLine(Params, SerialNumber, Args), Input, Options, SpewFilterCallback: new ProcessResult.SpewFilterCallbackType(MLDBSpewFilter));
 		}
 		return Run(DeviceCommand, GetDeviceCommandLine(Params, SerialNumber, Args), Input, Options);
 	}
@@ -1187,7 +1192,7 @@ public class LuminPlatform : Platform
 	private string RunAndLogDeviceCommand(ProjectParams Params, string SerialNumber, string Args, out int SuccessCode)
 	{
 		LastSpewFilename = "";
-		return RunAndLog(CmdEnv, DeviceCommand, GetDeviceCommandLine(Params, SerialNumber, Args), out SuccessCode, SpewFilterCallback: new ProcessResult.SpewFilterCallbackType(ADBSpewFilter));
+		return RunAndLog(CmdEnv, DeviceCommand, GetDeviceCommandLine(Params, SerialNumber, Args), out SuccessCode, SpewFilterCallback: new ProcessResult.SpewFilterCallbackType(MLDBSpewFilter));
 	}
 
 	private const int DeployMaxParallelCommands = 6;
