Deleting unused variable.

diff --git a/Engine/Source/Programs/AutomationTool/Gauntlet/Platform/Lumin/Gauntlet.TargetDeviceLumin.cs b/Engine/Source/Programs/AutomationTool/Gauntlet/Platform/Lumin/Gauntlet.TargetDeviceLumin.cs
--- a/Engine/Source/Programs/AutomationTool/Gauntlet/Platform/Lumin/Gauntlet.TargetDeviceLumin.cs
+++ b/Engine/Source/Programs/AutomationTool/Gauntlet/Platform/Lumin/Gauntlet.TargetDeviceLumin.cs
@@ -796,9 +796,7 @@ namespace Gauntlet
 
 			// kill any currently running instance:
 			KillRunningProcess(QualifiedPackageName);
-
-			string LaunchActivity = "";// LuminPlatform.GetLaunchableActivityName();
-
+			
 			Log.Info("Launching {0} on '{1}' ", QualifiedPackageName, ToString());
 			Log.Verbose("\t{0}", LuminInstall.CommandLine);
 
