Fixed issue where the provisioning profile is null when using automatic signing.

Something else is missing here but this fixes the issue for now.
The filename is returned by GetProvisioningData. It's explicitly set to null if bAutomaticSigning is true. Later on we attempt to build the path using that value which blows up if it's null. Interestingly we check for null later on but we don't do the check before trying to construct the FileReference. Added that check.

diff --git a/Engine/Source/Programs/AutomationTool/IOS/IOSPlatform.Automation.cs b/Engine/Source/Programs/AutomationTool/IOS/IOSPlatform.Automation.cs
--- a/Engine/Source/Programs/AutomationTool/IOS/IOSPlatform.Automation.cs
+++ b/Engine/Source/Programs/AutomationTool/IOS/IOSPlatform.Automation.cs
@@ -1693,7 +1693,12 @@ public class IOSPlatform : Platform
 			MobileProvisionDir = DirectoryReference.Combine(DirectoryReference.GetSpecialFolder(Environment.SpecialFolder.LocalApplicationData), "Apple Computer", "MobileDevice", "Provisioning Profiles");
 		}
 
-		FileReference MobileProvisionFile = FileReference.Combine(MobileProvisionDir, Params.Provision);
+		FileReference MobileProvisionFile = null;
+		
+		if(MobileProvisionDir != null && Params.Provision != null)
+		{
+			MobileProvisionFile = FileReference.Combine(MobileProvisionDir, Params.Provision);
+		}
 
 		// distribution build
 		bool bForDistribution = Params.Distribution;
