Temp code to debug builder problem.

#rnx

diff --git a/Engine/Source/Programs/AutomationTool/BuildGraph/TempStorage.cs b/Engine/Source/Programs/AutomationTool/BuildGraph/TempStorage.cs
--- a/Engine/Source/Programs/AutomationTool/BuildGraph/TempStorage.cs
+++ b/Engine/Source/Programs/AutomationTool/BuildGraph/TempStorage.cs
@@ -695,6 +695,19 @@ namespace AutomationTool
 				FileReference SharedFileListLocation = GetTaggedFileListLocation(SharedDir, NodeName, TagName);
 				if(!FileReference.Exists(SharedFileListLocation))
 				{
+					for(DirectoryReference ParentDir = SharedFileListLocation.Directory; ParentDir != null; ParentDir = ParentDir.ParentDirectory)
+					{
+						bool bExists = DirectoryReference.Exists(ParentDir);
+						Log.TraceInformation("{0}: {1}", ParentDir, bExists);
+						if(bExists)
+						{
+							foreach(FileReference Item in DirectoryReference.EnumerateFiles(ParentDir))
+							{
+								Log.TraceInformation(" FILE: {0}", Item);
+							}
+							break;
+						}
+					}
 					throw new AutomationException("Missing local or shared file list - {0}", SharedFileListLocation.FullName);
 				}
 
