Updating Expand node in BuildGraph to expand the values set in arguments to allow for use of propertys.
Ex: <Expand Name="MyMacro" Parameter1="$(ValueProperty1)" /> will set the Parameter1 argument to the expanded value of $(ValueProperty1)

#rnx

diff --git a/Engine/Source/Programs/AutomationTool/BuildGraph/Script.cs b/Engine/Source/Programs/AutomationTool/BuildGraph/Script.cs
--- a/Engine/Source/Programs/AutomationTool/BuildGraph/Script.cs
+++ b/Engine/Source/Programs/AutomationTool/BuildGraph/Script.cs
@@ -1256,7 +1256,7 @@ namespace AutomationTool
 							int Index;
 							if(Macro.ArgumentNameToIndex.TryGetValue(Attribute.Name, out Index))
 							{
-								Arguments[Index] = Attribute.Value;
+								Arguments[Index] = ExpandProperties(Element, Attribute.Value);
 							}
 							else
 							{
