Temp hack to fix building UAT without HMTL5.

#rnx

diff --git a/Engine/Source/Programs/AutomationTool/Scripts/CopyBuildToStagingDirectory.Automation.cs b/Engine/Source/Programs/AutomationTool/Scripts/CopyBuildToStagingDirectory.Automation.cs
--- a/Engine/Source/Programs/AutomationTool/Scripts/CopyBuildToStagingDirectory.Automation.cs
+++ b/Engine/Source/Programs/AutomationTool/Scripts/CopyBuildToStagingDirectory.Automation.cs
@@ -553,7 +553,7 @@ public partial class Project : CommandUtils
 			{
 				// Work out which ICU data version we use for this platform
 				var ICUDataVersion = "icudt64l";
-				if (SC.StageTargetPlatform.PlatformType == UnrealTargetPlatform.HoloLens || SC.StageTargetPlatform.PlatformType == UnrealTargetPlatform.TVOS || SC.StageTargetPlatform.PlatformType == UnrealTargetPlatform.HTML5)
+				if (SC.StageTargetPlatform.PlatformType == UnrealTargetPlatform.HoloLens || SC.StageTargetPlatform.PlatformType == UnrealTargetPlatform.TVOS || SC.StageTargetPlatform.PlatformType.ToString() == "HTML5")
 				{
 					ICUDataVersion = "icudt53l";
 				}
