Gauntlet: Fix issue with non-bulk test configuration ue4commandline.txt and log location on Android

diff --git a/Engine/Source/Programs/AutomationTool/Gauntlet/Platform/Android/Gauntlet.TargetDeviceAndroid.cs b/Engine/Source/Programs/AutomationTool/Gauntlet/Platform/Android/Gauntlet.TargetDeviceAndroid.cs
--- a/Engine/Source/Programs/AutomationTool/Gauntlet/Platform/Android/Gauntlet.TargetDeviceAndroid.cs
+++ b/Engine/Source/Programs/AutomationTool/Gauntlet/Platform/Android/Gauntlet.TargetDeviceAndroid.cs
@@ -847,7 +847,15 @@ namespace Gauntlet
 				string StorageLocation = StorageQueryResult.Output.Trim(); // "/mnt/sdcard";
 
 				// remote dir used to save things
-				string RemoteDir = StorageLocation + "/UE4Game/" + AppConfig.ProjectName;				
+				string RemoteDir = StorageLocation + "/UE4Game/" + AppConfig.ProjectName;
+
+				// if using a non-bulk test configuration, implies com.epicgames.ue4.GameActivity.bUseExternalFilesDir = true
+				// @todo: query this from the apk? 
+				if (AppConfig.Configuration == UnrealTargetConfiguration.Test && ((Build.Flags & BuildFlags.NotBulk) == BuildFlags.NotBulk))
+				{
+					RemoteDir = StorageLocation + "/Android/data/" + Build.AndroidPackageName + "/files/UE4Game/" + AppConfig.ProjectName;
+				}
+
 				string DependencyDir = RemoteDir + "/deps";
 
 				// device artifact path, always clear between runs
