HTML5 fixes

#rnx

diff --git a/Engine/Platforms/HTML5/Plugins/Media/WebMMedia/WebMMedia_HTML5.uplugin b/Engine/Platforms/HTML5/Plugins/Media/WebMMedia/WebMMedia_HTML5.uplugin
deleted file mode 100644
diff --git a/Engine/Platforms/HTML5/Source/Runtime/OpenGLDrv/OpenGLDrv_HTML5.Build.cs b/Engine/Platforms/HTML5/Source/Runtime/OpenGLDrv/OpenGLDrv_HTML5.Build.cs
--- a/Engine/Platforms/HTML5/Source/Runtime/OpenGLDrv/OpenGLDrv_HTML5.Build.cs
+++ b/Engine/Platforms/HTML5/Source/Runtime/OpenGLDrv/OpenGLDrv_HTML5.Build.cs
@@ -2,6 +2,7 @@
 
 using UnrealBuildTool;
 
+[SupportedPlatformsAttribute(new string[] {"Win64", "Mac", "Linux", "HTML5", "LinuxAArch64"})]
 public class OpenGLDrv_HTML5 : OpenGLDrv
 {
 	public OpenGLDrv_HTML5(ReadOnlyTargetRules Target) : base(Target)
diff --git a/Engine/Source/Programs/AutomationTool/Scripts/CopyBuildToStagingDirectory.Automation.cs b/Engine/Source/Programs/AutomationTool/Scripts/CopyBuildToStagingDirectory.Automation.cs
--- a/Engine/Source/Programs/AutomationTool/Scripts/CopyBuildToStagingDirectory.Automation.cs
+++ b/Engine/Source/Programs/AutomationTool/Scripts/CopyBuildToStagingDirectory.Automation.cs
@@ -553,7 +553,7 @@ public partial class Project : CommandUtils
 			{
 				// Work out which ICU data version we use for this platform
 				var ICUDataVersion = "icudt64l";
-				if (SC.StageTargetPlatform.PlatformType == UnrealTargetPlatform.HoloLens || SC.StageTargetPlatform.PlatformType == UnrealTargetPlatform.TVOS || SC.StageTargetPlatform.PlatformType.ToString() == "HTML5")
+				if (SC.StageTargetPlatform.PlatformType == UnrealTargetPlatform.HoloLens || SC.StageTargetPlatform.PlatformType == UnrealTargetPlatform.TVOS)
 				{
 					ICUDataVersion = "icudt53l";
 				}
diff --git a/Engine/Source/Runtime/OpenGLDrv/OpenGLDrv.Build.cs b/Engine/Source/Runtime/OpenGLDrv/OpenGLDrv.Build.cs
--- a/Engine/Source/Runtime/OpenGLDrv/OpenGLDrv.Build.cs
+++ b/Engine/Source/Runtime/OpenGLDrv/OpenGLDrv.Build.cs
@@ -2,7 +2,7 @@
 
 using UnrealBuildTool;
 
-[SupportedPlatformsAttribute(new string[] {"Win32", "Win64", "Mac", "Linux", "HTML5", "Android", "LinuxAArch64"})]
+[SupportedPlatformsAttribute(new string[] {"Win32", "Win64", "Mac", "Linux", "Android", "LinuxAArch64"})]
 public class OpenGLDrv : ModuleRules
 {
 	public OpenGLDrv(ReadOnlyTargetRules Target) : base(Target)
