Removing Regex line from Gauntlet which is preventing cooking for iOS

diff --git a/Engine/Source/Programs/AutomationTool/Gauntlet/Gauntlet.Automation.csproj b/Engine/Source/Programs/AutomationTool/Gauntlet/Gauntlet.Automation.csproj
--- a/Engine/Source/Programs/AutomationTool/Gauntlet/Gauntlet.Automation.csproj
+++ b/Engine/Source/Programs/AutomationTool/Gauntlet/Gauntlet.Automation.csproj
@@ -215,7 +215,6 @@
   files and include them in the build. They will also appear in VisualStudio -->
   <ItemGroup>
     <Compile Include="../../../../Platforms/*/Source/Programs/AutomationTool/Gauntlet/**/*.cs">
-      <Link>Platform/$([System.Text.RegularExpressions.Regex]::Split(`%(Identity)`, `[\\/]`)[5])/_</Link>
     </Compile>
   </ItemGroup>
   
