Fixing bugs in build graphs.
- Empty aggregates shouldn't be added to a preprocessed build graph. They do no work and cause a schema failure.
- Schema should support True and False for booleans.
- Generating a preprocessed XML should not early out by default. PreprocessOnly flag added to support exiting after the preprocessed XML is generated.

#rnx

diff --git a/Engine/Source/Programs/AutomationTool/BuildGraph/BuildGraph.cs b/Engine/Source/Programs/AutomationTool/BuildGraph/BuildGraph.cs
--- a/Engine/Source/Programs/AutomationTool/BuildGraph/BuildGraph.cs
+++ b/Engine/Source/Programs/AutomationTool/BuildGraph/BuildGraph.cs
@@ -80,6 +80,7 @@ namespace AutomationTool
 			string SchemaFileName = ParseParamValue("Schema", null);
 			string ExportFileName = ParseParamValue("Export", null);
 			string PreprocessedFileName = ParseParamValue("Preprocess", null);
+			bool bPreprocessOnly = ParseParam("PreprocessOnly");
 			string SharedStorageDir = ParseParamValue("SharedStorageDir", null);
 			string SingleNodeName = ParseParamValue("SingleNode", null);
 			string TriggerName = ParseParamValue("Trigger", null);
@@ -364,7 +365,10 @@ namespace AutomationTool
 				FileReference PreprocessedFileLocation = new FileReference(PreprocessedFileName);
 				LogInformation("Writing {0}...", PreprocessedFileLocation);
 				Graph.Write(PreprocessedFileLocation, (SchemaFileName != null)? new FileReference(SchemaFileName) : null);
-				return ExitCode.Success;
+				if (bPreprocessOnly)
+				{
+					return ExitCode.Success;
+				}
 			}
 
 			// Find the triggers which we are explicitly running.
diff --git a/Engine/Source/Programs/AutomationTool/BuildGraph/Graph.cs b/Engine/Source/Programs/AutomationTool/BuildGraph/Graph.cs
--- a/Engine/Source/Programs/AutomationTool/BuildGraph/Graph.cs
+++ b/Engine/Source/Programs/AutomationTool/BuildGraph/Graph.cs
@@ -395,6 +395,12 @@ namespace AutomationTool
 
 				foreach (KeyValuePair<string, Node[]> Aggregate in AggregateNameToNodes)
 				{
+					// If the aggregate has no required elements, skip it.
+					if (!Aggregate.Value.Any())
+					{
+						continue;
+					}
+
 					Writer.WriteStartElement("Aggregate");
 					Writer.WriteAttributeString("Name", Aggregate.Key);
 					Writer.WriteAttributeString("Requires", String.Join(";", Aggregate.Value.Select(x => x.Name)));
diff --git a/Engine/Source/Programs/AutomationTool/BuildGraph/Schema.cs b/Engine/Source/Programs/AutomationTool/BuildGraph/Schema.cs
--- a/Engine/Source/Programs/AutomationTool/BuildGraph/Schema.cs
+++ b/Engine/Source/Programs/AutomationTool/BuildGraph/Schema.cs
@@ -353,7 +353,7 @@ namespace AutomationTool
 			NewSchema.Items.Add(CreateSimpleTypeFromRegex(GetTypeName(ScriptSchemaStandardType.NameOrTagList), "(" + NameOrTagListPattern + "|" + StringWithPropertiesPattern + ")"));
 			NewSchema.Items.Add(CreateSimpleTypeFromRegex(GetTypeName(ScriptSchemaStandardType.QualifiedName), "(" + QualifiedNamePattern + "|" + StringWithPropertiesPattern + ")"));
 			NewSchema.Items.Add(CreateSimpleTypeFromRegex(GetTypeName(ScriptSchemaStandardType.BalancedString), BalancedStringPattern));
-			NewSchema.Items.Add(CreateSimpleTypeFromRegex(GetTypeName(ScriptSchemaStandardType.Boolean), "(" + "true" + "|" + "false" + "|" + StringWithPropertiesPattern + ")"));
+			NewSchema.Items.Add(CreateSimpleTypeFromRegex(GetTypeName(ScriptSchemaStandardType.Boolean), "(true|True|false|False|" + StringWithPropertiesPattern + ")"));
 			NewSchema.Items.Add(CreateSimpleTypeFromRegex(GetTypeName(ScriptSchemaStandardType.Integer), "(" + "(-?[1-9][0-9]*|0)" + "|" + StringWithPropertiesPattern + ")"));
 			foreach(XmlSchemaComplexType Type in TaskNameToType.Values)
 			{
