UE-81878 Unable to package projects for Android error MSB4186: Invalid static method invocation syntax

diff --git a/Engine/Source/Programs/AutomationTool/HoloLens/HoloLens.Automation.csproj b/Engine/Source/Programs/AutomationTool/HoloLens/HoloLens.Automation.csproj
--- a/Engine/Source/Programs/AutomationTool/HoloLens/HoloLens.Automation.csproj
+++ b/Engine/Source/Programs/AutomationTool/HoloLens/HoloLens.Automation.csproj
@@ -74,13 +74,14 @@
 		  <HintPath />
 		  <!-- Use the version from the current Windows SDK -->
 		  <HintPath>$(Windows10SdkPath)\UnionMetadata\$(Windows10SdkLatestVersion)\Windows.winmd</HintPath>
-          <!-- Try using the version from AutoSDK -->
-          <HintPath Condition="'$(UE_SDKS_ROOT)' != '' And Exists('$(UE_SDKS_ROOT)\HostWin64\Win64\Windows Kits\10\UnionMetadata\10.0.18362.0\Windows.winmd')">$(UE_SDKS_ROOT)\HostWin64\Win64\Windows Kits\10\UnionMetadata\10.0.18362.0\Windows.winmd</HintPath>
-          <!-- 16299 SDK (possibly others?) has ProductVersion == 10.0.16299 in the registry, but pathing needs the final .0 -->
+		  <!-- Try using the version from AutoSDK -->
+		  <SDKsWinmdPath Condition="'$(UE_SDKS_ROOT)' != ''">$(UE_SDKS_ROOT)\HostWin64\Win64\Windows Kits\10\UnionMetadata\10.0.18362.0\Windows.winmd</SDKsWinmdPath>
+		  <HintPath Condition="'$(UE_SDKS_ROOT)' != '' And '$([System.IO.File]::Exists(%(SDKsWinmdPath)))'!='false'">$(UE_SDKS_ROOT)\HostWin64\Win64\Windows Kits\10\UnionMetadata\10.0.18362.0\Windows.winmd</HintPath>
+		  <!-- 16299 SDK (possibly others?) has ProductVersion == 10.0.16299 in the registry, but pathing needs the final .0 -->
 		  <HintPath Condition="'$([System.IO.File]::Exists(%(HintPath)))'=='false'">$(Windows10SdkPath)\UnionMetadata\$(Windows10SdkLatestVersion).0\Windows.winmd</HintPath>
 		  <!-- Fall back to old, unversioned location -->
 		  <HintPath Condition="'$([System.IO.File]::Exists(%(HintPath)))'=='false'">$(Windows10SdkPath)\UnionMetadata\Windows.winmd</HintPath>
-          <Private>false</Private>
+		  <Private>false</Private>
 		</Reference>
 	  </ItemGroup>
 	</When>
