Change: 9625820
UE-81878 Unable to package projects for Android error MSB4186: Invalid static method invocation syntax
-Merging //UE4/Dev-Main to Release-4.24 (//UE4/Release-4.24)

Files:
//UE4/Main/Engine/Source/Programs/AutomationTool/HoloLens/HoloLens.Automation.csproj#4

diff --git a/Engine/Source/Programs/AutomationTool/HoloLens/HoloLens.Automation.csproj b/Engine/Source/Programs/AutomationTool/HoloLens/HoloLens.Automation.csproj
--- a/Engine/Source/Programs/AutomationTool/HoloLens/HoloLens.Automation.csproj
+++ b/Engine/Source/Programs/AutomationTool/HoloLens/HoloLens.Automation.csproj
@@ -72,10 +72,11 @@
 		<Reference Include="Windows">
 		  <!-- This is a dummy entry to throw off the parser in DotNETUtilities's CsProjectInfo which doesn't understand MSBuild functions -->
 		  <HintPath />
-		  <!-- Try using the version from AutoSDK -->
-		  <HintPath Condition="'$(UE_SDKS_ROOT)' != ''">$(UE_SDKS_ROOT)\HostWin64\Win64\Windows Kits\10\UnionMetadata\10.0.18362.0\Windows.winmd</HintPath>
 		  <!-- Use the version from the current Windows SDK -->
-		  <HintPath Condition="'$([System.IO.File]::Exists(%(HintPath)))'=='false'">$(Windows10SdkPath)\UnionMetadata\$(Windows10SdkLatestVersion)\Windows.winmd</HintPath>
+		  <HintPath>$(Windows10SdkPath)\UnionMetadata\$(Windows10SdkLatestVersion)\Windows.winmd</HintPath>
+		  <!-- Try using the version from AutoSDK -->
+		  <SDKsWinmdPath Condition="'$(UE_SDKS_ROOT)' != ''">$(UE_SDKS_ROOT)\HostWin64\Win64\Windows Kits\10\UnionMetadata\10.0.18362.0\Windows.winmd</SDKsWinmdPath>
+		  <HintPath Condition="'$(UE_SDKS_ROOT)' != '' And '$([System.IO.File]::Exists(%(SDKsWinmdPath)))'!='false'">$(UE_SDKS_ROOT)\HostWin64\Win64\Windows Kits\10\UnionMetadata\10.0.18362.0\Windows.winmd</HintPath>
 		  <!-- 16299 SDK (possibly others?) has ProductVersion == 10.0.16299 in the registry, but pathing needs the final .0 -->
 		  <HintPath Condition="'$([System.IO.File]::Exists(%(HintPath)))'=='false'">$(Windows10SdkPath)\UnionMetadata\$(Windows10SdkLatestVersion).0\Windows.winmd</HintPath>
 		  <!-- Fall back to old, unversioned location -->
