Exclude the project metadata folder from DLC builds.

diff --git a/Engine/Source/Programs/AutomationTool/Scripts/CopyBuildToStagingDirectory.Automation.cs b/Engine/Source/Programs/AutomationTool/Scripts/CopyBuildToStagingDirectory.Automation.cs
--- a/Engine/Source/Programs/AutomationTool/Scripts/CopyBuildToStagingDirectory.Automation.cs
+++ b/Engine/Source/Programs/AutomationTool/Scripts/CopyBuildToStagingDirectory.Automation.cs
@@ -440,6 +440,7 @@ public partial class Project : CommandUtils
 			// Put all of the cooked dir into the staged dir
 			SC.PlatformCookDir = String.IsNullOrEmpty(Params.CookOutputDir) ? DirectoryReference.Combine(DLCRoot, "Saved", "Cooked", SC.CookPlatform) : DirectoryReference.Combine(new DirectoryReference(Params.CookOutputDir), SC.CookPlatform);
 			DirectoryReference PlatformEngineDir = DirectoryReference.Combine(SC.PlatformCookDir, "Engine");
+			DirectoryReference ProjectMetadataDir = DirectoryReference.Combine(SC.PlatformCookDir, SC.ShortProjectName, "Metadata");
 			string RelativeDLCRootPath = DLCRoot.MakeRelativeTo(SC.LocalRoot);
 			SC.MetadataDir = DirectoryReference.Combine(SC.PlatformCookDir, RelativeDLCRootPath, "Metadata");
 
@@ -465,7 +466,7 @@ public partial class Project : CommandUtils
 			foreach (FileReference CookedFile in CookedFiles)
 			{
 				// Skip metadata directory
-				if (CookedFile.Directory.IsUnderDirectory(SC.MetadataDir))
+				if (CookedFile.Directory.IsUnderDirectory(SC.MetadataDir) || CookedFile.Directory.IsUnderDirectory(ProjectMetadataDir))
 				{
 					continue;
 				}
