Do not expand $(RootDir) variables when preprocessing a build graph script. The root directory may vary on each agent running it.

For now, the $(HostPlatform) variable is still expanded, due to its prevalence in determining defaults when running locally. Care must be taken to ensure that this does not affect logic for other platforms when building on the farm, eg. by overriding output paths.

#rnx

diff --git a/Engine/Source/Programs/AutomationTool/BuildGraph/BuildGraph.cs b/Engine/Source/Programs/AutomationTool/BuildGraph/BuildGraph.cs
--- a/Engine/Source/Programs/AutomationTool/BuildGraph/BuildGraph.cs
+++ b/Engine/Source/Programs/AutomationTool/BuildGraph/BuildGraph.cs
@@ -80,7 +80,6 @@ namespace AutomationTool
 			string SchemaFileName = ParseParamValue("Schema", null);
 			string ExportFileName = ParseParamValue("Export", null);
 			string PreprocessedFileName = ParseParamValue("Preprocess", null);
-			bool bPreprocessOnly = ParseParam("PreprocessOnly");
 			string SharedStorageDir = ParseParamValue("SharedStorageDir", null);
 			string SingleNodeName = ParseParamValue("SingleNode", null);
 			string TriggerName = ParseParamValue("Trigger", null);
@@ -127,12 +126,21 @@ namespace AutomationTool
 			DefaultProperties["EscapedBranch"] = P4Enabled ? CommandUtils.EscapePath(P4Env.Branch) : "Unknown";
 			DefaultProperties["Change"] = P4Enabled ? P4Env.Changelist.ToString() : "0";
 			DefaultProperties["CodeChange"] = P4Enabled ? P4Env.CodeChangelist.ToString() : "0";
-			DefaultProperties["RootDir"] = CommandUtils.RootDirectory.FullName;
 			DefaultProperties["IsBuildMachine"] = IsBuildMachine ? "true" : "false";
 			DefaultProperties["HostPlatform"] = HostPlatform.Current.HostEditorPlatform.ToString();
 			DefaultProperties["RestrictedFolderNames"] = String.Join(";", RestrictedFolder.GetNames());
 			DefaultProperties["RestrictedFolderFilter"] = String.Join(";", RestrictedFolder.GetNames().Select(x => String.Format(".../{0}/...", x)));
 
+			// Prevent expansion of the root directory if we're just preprocessing the output. They may vary by machine.
+			if (PreprocessedFileName == null)
+			{
+				DefaultProperties["RootDir"] = CommandUtils.RootDirectory.FullName;
+			}
+			else
+			{
+				DefaultProperties["RootDir"] = null;
+			}
+
 			// Attempt to read existing Build Version information
 			BuildVersion Version;
 			if (BuildVersion.TryRead(BuildVersion.GetDefaultFileName(), out Version))
@@ -198,7 +206,7 @@ namespace AutomationTool
 
 			// Read the script from disk
 			Graph Graph;
-			if(!ScriptReader.TryRead(new FileReference(ScriptFileName), Arguments, DefaultProperties, Schema, out Graph))
+			if(!ScriptReader.TryRead(new FileReference(ScriptFileName), Arguments, DefaultProperties, PreprocessedFileName != null, Schema, out Graph))
 			{
 				return ExitCode.Error_Unknown;
 			}
@@ -365,10 +373,7 @@ namespace AutomationTool
 				FileReference PreprocessedFileLocation = new FileReference(PreprocessedFileName);
 				LogInformation("Writing {0}...", PreprocessedFileLocation);
 				Graph.Write(PreprocessedFileLocation, (SchemaFileName != null)? new FileReference(SchemaFileName) : null);
-				if (bPreprocessOnly)
-				{
-					return ExitCode.Success;
-				}
+				bListOnly = true;
 			}
 
 			// Find the triggers which we are explicitly running.
diff --git a/Engine/Source/Programs/AutomationTool/BuildGraph/Script.cs b/Engine/Source/Programs/AutomationTool/BuildGraph/Script.cs
--- a/Engine/Source/Programs/AutomationTool/BuildGraph/Script.cs
+++ b/Engine/Source/Programs/AutomationTool/BuildGraph/Script.cs
@@ -231,6 +231,11 @@ namespace AutomationTool
 		/// </summary>
 		Dictionary<string, ScriptMacro> MacroNameToDefinition = new Dictionary<string, ScriptMacro>();
 
+		/// <summary>
+		/// Preprocess the file, but do not expand any values that are not portable (eg. paths on the local machine)
+		/// </summary>
+		bool bPreprocessOnly;
+
 		/// <summary>
 		/// Schema for the script
 		/// </summary>
@@ -245,9 +250,11 @@ namespace AutomationTool
 		/// Private constructor. Use ScriptReader.TryRead() to read a script file.
 		/// </summary>
 		/// <param name="DefaultProperties">Default properties available to the script</param>
+		/// <param name="bPreprocessOnly">Preprocess the file, but do not expand any values that are not portable (eg. paths on the local machine)</param>
 		/// <param name="Schema">Schema for the script</param>
-		private ScriptReader(IDictionary<string, string> DefaultProperties, ScriptSchema Schema)
+		private ScriptReader(IDictionary<string, string> DefaultProperties, bool bPreprocessOnly, ScriptSchema Schema)
 		{
+			this.bPreprocessOnly = bPreprocessOnly;
 			this.Schema = Schema;
 
 			EnterScope();
@@ -264,10 +271,11 @@ namespace AutomationTool
 		/// <param name="File">File to read from</param>
 		/// <param name="Arguments">Arguments passed in to the graph on the command line</param>
 		/// <param name="DefaultProperties">Default properties available to the script</param>
+		/// <param name="bPreprocessOnly">Preprocess the file, but do not expand any values that are not portable (eg. paths on the local machine)</param>
 		/// <param name="Schema">Schema for the script</param>
 		/// <param name="Graph">If successful, the graph constructed from the given script</param>
 		/// <returns>True if the graph was read, false if there were errors</returns>
-		public static bool TryRead(FileReference File, Dictionary<string, string> Arguments, Dictionary<string, string> DefaultProperties, ScriptSchema Schema, out Graph Graph)
+		public static bool TryRead(FileReference File, Dictionary<string, string> Arguments, Dictionary<string, string> DefaultProperties, bool bPreprocessOnly, ScriptSchema Schema, out Graph Graph)
 		{
 			// Check the file exists before doing anything.
 			if (!FileReference.Exists(File))
@@ -278,7 +286,7 @@ namespace AutomationTool
 			}
 
 			// Read the file and build the graph
-			ScriptReader Reader = new ScriptReader(DefaultProperties, Schema);
+			ScriptReader Reader = new ScriptReader(DefaultProperties, bPreprocessOnly, Schema);
 			if (!Reader.TryRead(File, Arguments) || Reader.NumErrors > 0)
 			{
 				Graph = null;
@@ -1424,11 +1432,25 @@ namespace AutomationTool
 			}
 			else if (ValueType == typeof(FileReference))
 			{
-				return CustomTask.ResolveFile(ValueText);
+				if (bPreprocessOnly)
+				{
+					return new FileReference(ValueText, FileReference.Sanitize.None);
+				}
+				else
+				{
+					return CustomTask.ResolveFile(ValueText);
+				}
 			}
 			else if (ValueType == typeof(DirectoryReference))
 			{
-				return CustomTask.ResolveDirectory(ValueText);
+				if (bPreprocessOnly)
+				{
+					return new DirectoryReference(ValueText, DirectoryReference.Sanitize.None);
+				}
+				else
+				{
+					return CustomTask.ResolveDirectory(ValueText);
+				}
 			}
 			else if (ValueType == typeof(UnrealTargetPlatform))
 			{
@@ -1952,11 +1974,20 @@ namespace AutomationTool
 					Value = "";
 				}
 
-				// Replace the variable, or skip past it
-				Result = Result.Substring(0, Idx) + Value + Result.Substring(EndIdx + 1);
+				// Check if we've got a value for this variable
+				if (Value == null)
+				{
+					// Do not expand it; must be preprocessing the script.
+					Idx = EndIdx;
+				}
+				else
+				{
+					// Replace the variable, or skip past it
+					Result = Result.Substring(0, Idx) + Value + Result.Substring(EndIdx + 1);
 
-				// Make sure we skip over the expanded variable; we don't want to recurse on it.
-				Idx += Value.Length;
+					// Make sure we skip over the expanded variable; we don't want to recurse on it.
+					Idx += Value.Length;
+				}
 			}
 			return Result;
 		}
