Adding WriteFileTask so we can write out a "build complete" file onto the share so the stuff Rene is working on can upload the build.

#rnx

diff --git a/Engine/Source/Programs/AutomationTool/BuildGraph/BuildGraph.Automation.csproj b/Engine/Source/Programs/AutomationTool/BuildGraph/BuildGraph.Automation.csproj
--- a/Engine/Source/Programs/AutomationTool/BuildGraph/BuildGraph.Automation.csproj
+++ b/Engine/Source/Programs/AutomationTool/BuildGraph/BuildGraph.Automation.csproj
@@ -78,6 +78,7 @@
     <Compile Include="Script.cs" />
     <Compile Include="Tasks\AgeStoreTask.cs" />
     <Compile Include="Tasks\CommandletTask.cs" />
+    <Compile Include="Tasks\WriteFileTask.cs" />
     <Compile Include="Tasks\CsCompileTask.cs" />
     <Compile Include="Tasks\MoveTask.cs" />
     <Compile Include="Tasks\MsBuildTask.cs" />
diff --git a/Engine/Source/Programs/AutomationTool/BuildGraph/Tasks/WriteTextFileTask.cs b/Engine/Source/Programs/AutomationTool/BuildGraph/Tasks/WriteTextFileTask.cs
new file mode 100644
--- /dev/null
+++ b/Engine/Source/Programs/AutomationTool/BuildGraph/Tasks/WriteTextFileTask.cs
@@ -0,0 +1,140 @@
+// Copyright 1998-2019 Epic Games, Inc. All Rights Reserved.
+
+using AutomationTool;
+using System;
+using System.Collections.Generic;
+using System.Linq;
+using System.Xml;
+using Tools.DotNETCommon;
+
+namespace BuildGraph.Tasks
+{
+	/// <summary>
+	/// Parameters for a <see cref="WriteTextFileTask"/>.
+	/// </summary>
+	public class WriteTextFileTaskParameters
+	{
+		/// <summary>
+		/// Path to the file to write.
+		/// </summary>
+		[TaskParameter]
+		public FileReference File;
+
+		/// <summary>
+		/// Optional, whether or not to append to the file rather than overwrite.
+		/// </summary>
+		[TaskParameter(Optional = true)]
+		public bool Append;
+
+		/// <summary>
+		/// The text to write to the file.
+		/// </summary>
+		[TaskParameter(Optional = true)]
+		public string Text;
+
+		/// <summary>
+		/// If specified, causes the given list of files to be printed after the given message.
+		/// </summary>
+		[TaskParameter(Optional = true, ValidationType = TaskParameterValidationType.FileSpec)]
+		public string Files;
+
+		/// <summary>
+		/// Tag to be applied to build products of this task.
+		/// </summary>
+		[TaskParameter(Optional = true, ValidationType = TaskParameterValidationType.TagList)]
+		public string Tag;
+	}
+
+	/// <summary>
+	/// Writes text to a file.
+	/// </summary>
+	[TaskElement("WriteTextFile", typeof(WriteTextFileTaskParameters))]
+	public class WriteTextFileTask : CustomTask
+	{
+		/// <summary>
+		/// Parameters for this task.
+		/// </summary>
+		WriteTextFileTaskParameters Parameters;
+
+		/// <summary>
+		/// Constructor.
+		/// </summary>
+		/// <param name="InParameters">Parameters for this task.</param>
+		public WriteTextFileTask(WriteTextFileTaskParameters InParameters)
+		{
+			Parameters = InParameters;
+		}
+
+		/// <summary>
+		/// Execute the task.
+		/// </summary>
+		/// <param name="Job">Information about the current job.</param>
+		/// <param name="BuildProducts">Set of build products produced by this node.</param>
+		/// <param name="TagNameToFileSet">Mapping from tag names to the set of files they include.</param>
+		public override void Execute(JobContext Job, HashSet<FileReference> BuildProducts, Dictionary<string, HashSet<FileReference>> TagNameToFileSet)
+		{
+			string FileText = Parameters.Text;
+
+			// If any files or tagsets are provided, add them to the text output.
+			if (!String.IsNullOrEmpty(Parameters.Files))
+			{
+				if (!string.IsNullOrWhiteSpace(FileText))
+				{
+					FileText += Environment.NewLine;
+				}
+
+				HashSet<FileReference> Files = ResolveFilespec(CommandUtils.RootDirectory, Parameters.Files, TagNameToFileSet);
+				if (Files.Any())
+				{
+					FileText += string.Join(Environment.NewLine, Files.Select(f => f.FullName));
+				}
+			}
+
+			// Make sure output folder exists.
+			if (!DirectoryReference.Exists(Parameters.File.Directory))
+			{
+				DirectoryReference.CreateDirectory(Parameters.File.Directory);
+			}
+
+			if (Parameters.Append)
+			{
+				CommandUtils.LogInformation(string.Format("Appending text to file '{0}': {1}", Parameters.File, FileText));
+				FileReference.AppendAllText(Parameters.File, Environment.NewLine + FileText);
+			}
+			else
+			{
+				CommandUtils.LogInformation(string.Format("Writing text to file '{0}': {1}", Parameters.File, FileText));
+				FileReference.WriteAllText(Parameters.File, FileText);
+			}
+		}
+
+		/// <summary>
+		/// Output this task out to an XML writer.
+		/// </summary>
+		public override void Write(XmlWriter Writer)
+		{
+			Write(Writer, Parameters);
+		}
+
+		/// <summary>
+		/// Find all the tags which are used as inputs to this task
+		/// </summary>
+		/// <returns>The tag names which are read by this task</returns>
+		public override IEnumerable<string> FindConsumedTagNames()
+		{
+			foreach(string TagName in FindTagNamesFromFilespec(Parameters.Files))
+			{
+				yield return TagName;
+			}
+		}
+
+		/// <summary>
+		/// Find all the tags which are modified by this task
+		/// </summary>
+		/// <returns>The tag names which are modified by this task</returns>
+		public override IEnumerable<string> FindProducedTagNames()
+		{
+			return FindTagNamesFromList(Parameters.Tag);
+		}
+	}
+}
diff --git a/Engine/Source/Programs/DotNETCommon/DotNETUtilities/FileReference.cs b/Engine/Source/Programs/DotNETCommon/DotNETUtilities/FileReference.cs
--- a/Engine/Source/Programs/DotNETCommon/DotNETUtilities/FileReference.cs
+++ b/Engine/Source/Programs/DotNETCommon/DotNETUtilities/FileReference.cs
@@ -574,6 +574,69 @@ namespace Tools.DotNETCommon
 			File.WriteAllText(Location.FullName, Contents, Encoding);
 		}
 
+		/// <summary>
+		/// Appends the contents to a file
+		/// </summary>
+		/// <param name="Location">Location of the file</param>
+		/// <param name="Contents">Contents to append to the file</param>
+		public static void AppendAllLines(FileReference Location, IEnumerable<string> Contents)
+		{
+			File.AppendAllLines(Location.FullName, Contents);
+		}
+
+		/// <summary>
+		/// Appends the contents to a file
+		/// </summary>
+		/// <param name="Location">Location of the file</param>
+		/// <param name="Contents">Contents to append to the file</param>
+		/// <param name="Encoding">The encoding to use when parsing the file</param>
+		public static void AppendAllLines(FileReference Location, IEnumerable<string> Contents, Encoding Encoding)
+		{
+			File.AppendAllLines(Location.FullName, Contents, Encoding);
+		}
+
+		/// <summary>
+		/// Appends the contents to a file
+		/// </summary>
+		/// <param name="Location">Location of the file</param>
+		/// <param name="Contents">Contents to append to the file</param>
+		public static void AppendAllLines(FileReference Location, string[] Contents)
+		{
+			File.AppendAllLines(Location.FullName, Contents);
+		}
+
+		/// <summary>
+		/// Appends the contents to a file
+		/// </summary>
+		/// <param name="Location">Location of the file</param>
+		/// <param name="Contents">Contents to append to the file</param>
+		/// <param name="Encoding">The encoding to use when parsing the file</param>
+		public static void AppendAllLines(FileReference Location, string[] Contents, Encoding Encoding)
+		{
+			File.AppendAllLines(Location.FullName, Contents, Encoding);
+		}
+
+		/// <summary>
+		/// Appends the contents to a file
+		/// </summary>
+		/// <param name="Location">Location of the file</param>
+		/// <param name="Contents">Contents to append to the file</param>
+		public static void AppendAllText(FileReference Location, string Contents)
+		{
+			File.AppendAllText(Location.FullName, Contents);
+		}
+
+		/// <summary>
+		/// Appends the contents to a file
+		/// </summary>
+		/// <param name="Location">Location of the file</param>
+		/// <param name="Contents">Contents to append to the file</param>
+		/// <param name="Encoding">The encoding to use when parsing the file</param>
+		public static void AppendAllText(FileReference Location, string Contents, Encoding Encoding)
+		{
+			File.AppendAllText(Location.FullName, Contents, Encoding);
+		}
+
 		#endregion
 	}
 
