Correcting file name in csproj.

#rnx

diff --git a/Engine/Source/Programs/AutomationTool/BuildGraph/BuildGraph.Automation.csproj b/Engine/Source/Programs/AutomationTool/BuildGraph/BuildGraph.Automation.csproj
--- a/Engine/Source/Programs/AutomationTool/BuildGraph/BuildGraph.Automation.csproj
+++ b/Engine/Source/Programs/AutomationTool/BuildGraph/BuildGraph.Automation.csproj
@@ -78,7 +78,6 @@
     <Compile Include="Script.cs" />
     <Compile Include="Tasks\AgeStoreTask.cs" />
     <Compile Include="Tasks\CommandletTask.cs" />
-    <Compile Include="Tasks\WriteFileTask.cs" />
     <Compile Include="Tasks\CsCompileTask.cs" />
     <Compile Include="Tasks\MoveTask.cs" />
     <Compile Include="Tasks\MsBuildTask.cs" />
@@ -101,6 +100,7 @@
     <Compile Include="Tasks\StripTask.cs" />
     <Compile Include="Tasks\SetVersionTask.cs" />
     <Compile Include="Tasks\UnzipTask.cs" />
+    <Compile Include="Tasks\WriteTextFileTask.cs" />
     <Compile Include="Tasks\ZipTask.cs" />
     <Compile Include="TempStorage.cs" />
     <Compile Include="Tasks\ModifyConfigTask.cs" />
