Adjust behavior in SetVersionTask to better mirror UpdateLocalVersion by setting CompatibleChange to Change if no value was set.

#rnx

diff --git a/Engine/Source/Programs/AutomationTool/BuildGraph/Tasks/SetVersionTask.cs b/Engine/Source/Programs/AutomationTool/BuildGraph/Tasks/SetVersionTask.cs
--- a/Engine/Source/Programs/AutomationTool/BuildGraph/Tasks/SetVersionTask.cs
+++ b/Engine/Source/Programs/AutomationTool/BuildGraph/Tasks/SetVersionTask.cs
@@ -94,7 +94,8 @@ namespace AutomationTool.Tasks
 		public override void Execute(JobContext Job, HashSet<FileReference> BuildProducts, Dictionary<string, HashSet<FileReference>> TagNameToFileSet)
 		{
 			// Update the version files
-			List<FileReference> VersionFiles = UE4Build.StaticUpdateVersionFiles(Parameters.Change, Parameters.CompatibleChange, Parameters.Branch, Parameters.Build, Parameters.Licensee, Parameters.Promoted, !Parameters.SkipWrite);
+			int CompatibleChange = (Parameters.CompatibleChange == 0) ? Parameters.Change : Parameters.CompatibleChange;
+			List<FileReference> VersionFiles = UE4Build.StaticUpdateVersionFiles(Parameters.Change, CompatibleChange, Parameters.Branch, Parameters.Build, Parameters.Licensee, Parameters.Promoted, !Parameters.SkipWrite);
 
 			// Apply the optional tag to them
 			foreach(string TagName in FindTagNamesFromList(Parameters.Tag))
