Determine the path to SYMSTORE.EXE and PDBCOPY.EXE from UBT's list of known Windows SDK directories, so we can use the AutoSDK version of these tools if available.

diff --git a/Engine/Source/Programs/AutomationTool/Win/WinPlatform.Automation.cs b/Engine/Source/Programs/AutomationTool/Win/WinPlatform.Automation.cs
--- a/Engine/Source/Programs/AutomationTool/Win/WinPlatform.Automation.cs
+++ b/Engine/Source/Programs/AutomationTool/Win/WinPlatform.Automation.cs
@@ -297,58 +297,37 @@ public abstract class BaseWinPlatform : Platform
     /// <summary>
     /// Try to get the SYMSTORE.EXE path from the given Windows SDK version
     /// </summary>
-    /// <param name="SdkVersion">The SDK version string</param>
-    /// <param name="SymStoreExe">Receives the path to symstore.exe if found</param>
-    /// <returns>True if found, false otherwise</returns>
-    private static bool TryGetSymStoreExe(string SdkVersion, out FileReference SymStoreExe)
+    /// <returns>Path to SYMSTORE.EXE</returns>
+    private static FileReference GetSymStoreExe()
     {
-        // Try to get the SDK installation directory
-        string SdkFolder = Registry.GetValue(@"HKEY_CURRENT_USER\SOFTWARE\Microsoft\Microsoft SDKs\Windows\" + SdkVersion, "InstallationFolder", null) as String;
-        if (SdkFolder == null)
-        {
-            SdkFolder = Registry.GetValue(@"HKEY_LOCAL_MACHINE\SOFTWARE\Wow6432Node\Microsoft\Microsoft SDKs\Windows\" + SdkVersion, "InstallationFolder", null) as String;
-            if (SdkFolder == null)
-            {
-                SdkFolder = Registry.GetValue(@"HKEY_CURRENT_USER\SOFTWARE\Wow6432Node\Microsoft\Microsoft SDKs\Windows\" + SdkVersion, "InstallationFolder", null) as String;
-                if (SdkFolder == null)
-                {
-                    SymStoreExe = null;
-                    return false;
-                }
-            }
-        }
-
-        // Check for the 64-bit toolchain first, then the 32-bit toolchain
-        FileReference CheckSymStoreExe = FileReference.Combine(new DirectoryReference(SdkFolder), "Debuggers", "x64", "SymStore.exe");
-        if (!FileReference.Exists(CheckSymStoreExe))
-        {
-            CheckSymStoreExe = FileReference.Combine(new DirectoryReference(SdkFolder), "Debuggers", "x86", "SymStore.exe");
-            if (!FileReference.Exists(CheckSymStoreExe))
-            {
-                SymStoreExe = null;
-                return false;
-            }
-        }
+		List<KeyValuePair<string, DirectoryReference>> WindowsSdkDirs = WindowsExports.GetWindowsSdkDirs();
+		foreach (DirectoryReference WindowsSdkDir in WindowsSdkDirs.Select(x => x.Value))
+		{
+			FileReference SymStoreExe64 = FileReference.Combine(WindowsSdkDir, "Debuggers", "x64", "SymStore.exe");
+			if (FileReference.Exists(SymStoreExe64))
+			{
+				return SymStoreExe64;
+			}
 
-        SymStoreExe = CheckSymStoreExe;
-        return true;
+			FileReference SymStoreExe32 = FileReference.Combine(WindowsSdkDir, "Debuggers", "x86", "SymStore.exe");
+			if (FileReference.Exists(SymStoreExe32))
+			{
+				return SymStoreExe32;
+			}
+		}
+		throw new AutomationException("Unable to find a Windows SDK installation containing PDBSTR.EXE");
     }
 
 	public static bool TryGetPdbCopyLocation(out FileReference OutLocation)
 	{
 		// Try to find an installation of the Windows 10 SDK
-		string SdkInstallFolder =
-			(Microsoft.Win32.Registry.GetValue("HKEY_CURRENT_USER\\Software\\Microsoft\\Windows Kits\\Installed Roots", "KitsRoot10", null) as string) ??
-			(Microsoft.Win32.Registry.GetValue("HKEY_LOCAL_MACHINE\\Software\\Microsoft\\Windows Kits\\Installed Roots", "KitsRoot10", null) as string) ??
-			(Microsoft.Win32.Registry.GetValue("HKEY_CURRENT_USER\\Software\\Wow6432Node\\Microsoft\\Windows Kits\\Installed Roots", "KitsRoot10", null) as string) ??
-			(Microsoft.Win32.Registry.GetValue("HKEY_LOCAL_MACHINE\\Software\\Wow6432Node\\Microsoft\\Windows Kits\\Installed Roots", "KitsRoot10", null) as string);
-
-		if(!String.IsNullOrEmpty(SdkInstallFolder))
+		List<KeyValuePair<string, DirectoryReference>> WindowsSdkDirs = WindowsExports.GetWindowsSdkDirs();
+		foreach (DirectoryReference WindowsSdkDir in WindowsSdkDirs.Select(x => x.Value))
 		{
-			FileReference Location = FileReference.Combine(new DirectoryReference(SdkInstallFolder), "Debuggers", "x64", "PDBCopy.exe");
-			if(FileReference.Exists(Location))
+			FileReference PdbCopyExe = FileReference.Combine(WindowsSdkDir, "Debuggers", "x64", "PdbCopy.exe");
+			if (FileReference.Exists(PdbCopyExe))
 			{
-				OutLocation = Location;
+				OutLocation = PdbCopyExe;
 				return true;
 			}
 		}
@@ -409,12 +388,7 @@ public abstract class BaseWinPlatform : Platform
 	public override bool PublishSymbols(DirectoryReference SymbolStoreDirectory, List<FileReference> Files, string Product, string BuildVersion = null)
     {
         // Get the SYMSTORE.EXE path, using the latest SDK version we can find.
-        FileReference SymStoreExe;
-        if (!TryGetSymStoreExe("v10.0", out SymStoreExe) && !TryGetSymStoreExe("v8.1", out SymStoreExe) && !TryGetSymStoreExe("v8.0", out SymStoreExe))
-        {
-            CommandUtils.LogError("Couldn't find SYMSTORE.EXE in any Windows SDK installation");
-            return false;
-        }
+        FileReference SymStoreExe = GetSymStoreExe();
 
 		List<FileReference> FilesToAdd = Files.Where(x => x.HasExtension(".pdb") || x.HasExtension(".exe") || x.HasExtension(".dll")).ToList();
 		if(FilesToAdd.Count > 0)
