Added Gauntlet AutoParam "DisableHeartbeatTimeout" to disable tests failing as intended due to missed heartbeats.

#rnx

diff --git a/Engine/Source/Programs/AutomationTool/Gauntlet/Unreal/Base/Gauntlet.UnrealTestConfiguration.cs b/Engine/Source/Programs/AutomationTool/Gauntlet/Unreal/Base/Gauntlet.UnrealTestConfiguration.cs
--- a/Engine/Source/Programs/AutomationTool/Gauntlet/Unreal/Base/Gauntlet.UnrealTestConfiguration.cs
+++ b/Engine/Source/Programs/AutomationTool/Gauntlet/Unreal/Base/Gauntlet.UnrealTestConfiguration.cs
@@ -278,6 +278,12 @@ namespace Gauntlet
 		/// </summary>
 		public UnrealHeartbeatOptions HeartbeatOptions { get; set; }
 
+		/// <summary>
+		/// Prevents heartbeats timeouts from being checked so that tests will not fail from missed heartbeats
+		/// </summary>
+		[AutoParam(false)]
+		public bool DisableHeartbeatTimeout { get; set; }
+
 		// Member variables 
 
 		/// <summary>
diff --git a/Engine/Source/Programs/AutomationTool/Gauntlet/Unreal/Base/Gauntlet.UnrealTestNode.cs b/Engine/Source/Programs/AutomationTool/Gauntlet/Unreal/Base/Gauntlet.UnrealTestNode.cs
--- a/Engine/Source/Programs/AutomationTool/Gauntlet/Unreal/Base/Gauntlet.UnrealTestNode.cs
+++ b/Engine/Source/Programs/AutomationTool/Gauntlet/Unreal/Base/Gauntlet.UnrealTestNode.cs
@@ -122,7 +122,7 @@ namespace Gauntlet
 		protected DateTime TimeOfFirstMissingProcess;
 
 		protected int TimeToWaitForProcesses { get; set; }
-
+		
 		protected DateTime LastHeartbeatTime = DateTime.MinValue;
 		protected DateTime LastActiveHeartbeatTime = DateTime.MinValue;
 
@@ -850,6 +850,7 @@ namespace Gauntlet
 		private void CheckHeartbeat()
 		{
 			if (CachedConfig == null 
+				|| CachedConfig.DisableHeartbeatTimeout
 				|| CachedConfig.HeartbeatOptions.bExpectHeartbeats == false
 				|| GetTestStatus() != TestStatus.InProgress)
 			{
