Generate a --build-id for physx

diff --git a/Engine/Source/Programs/AutomationTool/Scripts/BuildPhysX.Automation.cs b/Engine/Source/Programs/AutomationTool/Scripts/BuildPhysX.Automation.cs
--- a/Engine/Source/Programs/AutomationTool/Scripts/BuildPhysX.Automation.cs
+++ b/Engine/Source/Programs/AutomationTool/Scripts/BuildPhysX.Automation.cs
@@ -949,7 +949,7 @@ class BuildPhysX_Linux : BuildPhysX.MakefileTargetPlatform
 	private string GetBundledLinuxLibCxxFlags()
     {
 		string CxxFlags = "\"-I " + ThirdPartySourceDirectory + "/Linux/LibCxx/include -I " + ThirdPartySourceDirectory + "/Linux/LibCxx/include/c++/v1\"";
-		string CxxLinkerFlags = "\"-stdlib=libc++ -nodefaultlibs -L " 
+		string CxxLinkerFlags = "\"-stdlib=libc++ -nodefaultlibs -Wl,--build-id -L " 
 			+ ThirdPartySourceDirectory + "/Linux/LibCxx/lib/Linux/x86_64-unknown-linux-gnu/ " 
 			+ ThirdPartySourceDirectory + "/Linux/LibCxx/lib/Linux/x86_64-unknown-linux-gnu/libc++.a " 
 			+ ThirdPartySourceDirectory + "/Linux/LibCxx/lib/Linux/x86_64-unknown-linux-gnu/libc++abi.a -lm -lc -lgcc_s\"";
@@ -991,7 +991,7 @@ class BuildPhysX_Linux : BuildPhysX.MakefileTargetPlatform
 		string OriginalToolchainPath = Environment.GetEnvironmentVariable("LINUX_MULTIARCH_ROOT");
 		if (!string.IsNullOrEmpty(OriginalToolchainPath))
 		{
-			string ToolchainPathToUse = OriginalToolchainPath.Replace("v13_clang-7.0.1-centos7", "v12_clang-6.0.1-centos7");
+			string ToolchainPathToUse = OriginalToolchainPath.Replace("v15_clang-8.0.1-centos7", "v12_clang-6.0.1-centos7");
 			LogInformation("Working around problems with newer clangs: {0} -> {1}", OriginalToolchainPath, ToolchainPathToUse);
 			Environment.SetEnvironmentVariable("LINUX_MULTIARCH_ROOT", ToolchainPathToUse);
 		}
