Adding support for importing a pregenerated schema to get around licensee workspace not having the same tasks as non-licensee, causing a schema failure unrelated to the work the agent will be doing.

#rnx

diff --git a/Engine/Source/Programs/AutomationTool/BuildGraph/BuildGraph.cs b/Engine/Source/Programs/AutomationTool/BuildGraph/BuildGraph.cs
--- a/Engine/Source/Programs/AutomationTool/BuildGraph/BuildGraph.cs
+++ b/Engine/Source/Programs/AutomationTool/BuildGraph/BuildGraph.cs
@@ -47,6 +47,7 @@ namespace AutomationTool
 	[Help("Target=<Name>", "Name of the node or output tag to be built")]
 	[Help("Schema", "Generates a schema to the default location")]
 	[Help("Schema=<FileName>", "Generate a schema describing valid script documents, including all the known tasks")]
+	[Help("ImportSchema=<FileName>", "Imports a schema from an existing schema file")]
 	[Help("Set:<Property>=<Value>", "Sets a named property to the given value")]
 	[Help("Clean", "Cleans all cached state of completed build nodes before running")]
 	[Help("CleanNode=<Name>[+<Name>...]", "Cleans just the given nodes before running")]
@@ -78,6 +79,7 @@ namespace AutomationTool
 			string TargetNames = ParseParamValue("Target", null);
 			string DocumentationFileName = ParseParamValue("Documentation", null);
 			string SchemaFileName = ParseParamValue("Schema", null);
+			string ImportSchemaFileName = ParseParamValue("ImportSchema", null);
 			string ExportFileName = ParseParamValue("Export", null);
 			string PreprocessedFileName = ParseParamValue("Preprocess", null);
 			string SharedStorageDir = ParseParamValue("SharedStorageDir", null);
@@ -184,16 +186,25 @@ namespace AutomationTool
 				return ExitCode.Success;
 			}
 
-			// Create a schema for the given tasks
-			ScriptSchema Schema = new ScriptSchema(NameToTask);
-			if(SchemaFileName != null)
+			// Import schema if one is passed in
+			ScriptSchema Schema;
+			if (ImportSchemaFileName != null)
 			{
-				FileReference FullSchemaFileName = new FileReference(SchemaFileName);
-				LogInformation("Writing schema to {0}...", FullSchemaFileName.FullName);
-				Schema.Export(FullSchemaFileName);
-				if(ScriptFileName == null)
+				Schema = ScriptSchema.Import(FileReference.FromString(ImportSchemaFileName), NameToTask);
+			}
+			else
+			{
+				// Create a schema for the given tasks
+				Schema = new ScriptSchema(NameToTask);
+				if (SchemaFileName != null)
 				{
-					return ExitCode.Success;
+					FileReference FullSchemaFileName = new FileReference(SchemaFileName);
+					LogInformation("Writing schema to {0}...", FullSchemaFileName.FullName);
+					Schema.Export(FullSchemaFileName);
+					if (ScriptFileName == null)
+					{
+						return ExitCode.Success;
+					}
 				}
 			}
 
diff --git a/Engine/Source/Programs/AutomationTool/BuildGraph/Schema.cs b/Engine/Source/Programs/AutomationTool/BuildGraph/Schema.cs
--- a/Engine/Source/Programs/AutomationTool/BuildGraph/Schema.cs
+++ b/Engine/Source/Programs/AutomationTool/BuildGraph/Schema.cs
@@ -3,12 +3,14 @@
 using System;
 using System.Collections.Generic;
 using System.Diagnostics;
+using System.IO;
 using System.Linq;
 using System.Reflection;
 using System.Text;
 using System.Threading.Tasks;
 using System.Xml;
 using System.Xml.Schema;
+using System.Xml.Serialization;
 using Tools.DotNETCommon;
 using UnrealBuildTool;
 
@@ -262,6 +264,12 @@ namespace AutomationTool
 		/// </summary>
 		const string BalancedStringPattern = "[^\\$]*" + "(" + "(" + PropertyPattern + "|" + "\\$[^\\(]" + ")" + "[^\\$]*" + ")*" + "\\$?";
 
+		private ScriptSchema(XmlSchema Schema, Dictionary<string, ScriptTask> InNameToTask)
+		{
+			CompiledSchema = Schema;
+			NameToTask = InNameToTask;
+		}
+
 		/// <summary>
 		/// Constructor
 		/// </summary>
@@ -402,6 +410,38 @@ namespace AutomationTool
 			}
 		}
 
+		/// <summary>
+		/// Imports a schema from a file
+		/// </summary>
+		/// <param name="File">The XML file to import</param>
+		/// <param name="NameToTask">Mapping of task name to information about how to construct it</param>
+		/// <returns>A <see cref="ScriptSchema"/> deserialized from the XML file, or null if file doesn't exist</returns>
+		public static ScriptSchema Import(FileReference File, Dictionary<string, ScriptTask> NameToTask)
+		{
+			if (!FileReference.Exists(File))
+			{
+				return null;
+			}
+			
+			using (XmlReader SchemaFile = XmlReader.Create(File.FullName))
+			{
+				ScriptSchema ImportedSchema = new ScriptSchema(XmlSchema.Read(SchemaFile, ValidationCallback), NameToTask);
+				return ImportedSchema;
+			}
+		}
+
+		static void ValidationCallback(object Sender, ValidationEventArgs EventArgs)
+		{
+			if (EventArgs.Severity == XmlSeverityType.Warning)
+			{
+				Log.WriteLine(LogEventType.Warning, "WARNING: {0}", EventArgs.Message);
+			}
+			else if (EventArgs.Severity == XmlSeverityType.Error)
+			{
+				Log.WriteLine(LogEventType.Error, "ERROR: {0}", EventArgs.Message);
+			}
+		}
+
 		/// <summary>
 		/// Gets the bare name for the given script type
 		/// </summary>
