Re-add platform localization

#rnx

diff --git a/Engine/Platforms/HTML5/Config/HTML5Engine.ini b/Engine/Platforms/HTML5/Config/HTML5Engine.ini
--- a/Engine/Platforms/HTML5/Config/HTML5Engine.ini
+++ b/Engine/Platforms/HTML5/Config/HTML5Engine.ini
@@ -1,6 +1,9 @@
 [HTMl5Section]
 bIsHTMl5=true
 
+[Internationalization]
++LocalizationPaths=../../../Engine/Platforms/HTML5/Content/Localization/HTML5
+
 [Audio]
 ; This audio backend can be used by running with the -noaudiomixer commandline argument:
 AudioDeviceModuleName=ALAudio
diff --git a/Engine/Platforms/HTML5/Config/Localization/HTML5.ini b/Engine/Platforms/HTML5/Config/Localization/HTML5.ini
new file mode 100644
--- /dev/null
+++ b/Engine/Platforms/HTML5/Config/Localization/HTML5.ini
@@ -0,0 +1,54 @@
+[CommonSettings]
+SourcePath=Platforms/HTML5/Content/Localization/HTML5
+DestinationPath=Platforms/HTML5/Content/Localization/HTML5
+ManifestName=HTML5.manifest
+ArchiveName=HTML5.archive
+PortableObjectName=HTML5.po
+ResourceName=HTML5.locres
+NativeCulture=en
+CulturesToGenerate=en
+CulturesToGenerate=ja
+CulturesToGenerate=ko
+CulturesToGenerate=zh-Hans
+
+;Gather text from source code
+[GatherTextStep0]
+CommandletClass=GatherTextFromSource
+SearchDirectoryPaths=Platforms/HTML5/Source/Runtime/
+SearchDirectoryPaths=Platforms/HTML5/Source/Developer/
+SearchDirectoryPaths=Platforms/HTML5/Plugins/
+SearchDirectoryPaths=Platforms/HTML5/Config/
+ExcludePathFilters=Platforms/HTML5/Config/Localization/*
+ExcludePathFilters=Platforms/HTML5/Plugins/Online/*
+ExcludePathFilters=*/NoRedist/*
+ExcludePathFilters=*/NotForLicensees/*
+FileNameFilters=*.cpp
+FileNameFilters=*.h
+FileNameFilters=*.c
+FileNameFilters=*.inl
+FileNameFilters=*.mm
+FileNameFilters=*.ini
+ShouldGatherFromEditorOnlyData=true
+
+;Write Manifest
+[GatherTextStep1]
+CommandletClass=GenerateGatherManifest
+
+;Write Archives
+[GatherTextStep2]
+CommandletClass=GenerateGatherArchive
+bPurgeOldEmptyEntries=true
+
+;Import localized PO files
+[GatherTextStep3]
+CommandletClass=InternationalizationExport
+bImportLoc=true
+
+;Write Localized Text Resource
+[GatherTextStep4]
+CommandletClass=GenerateTextLocalizationResource
+
+;Export PO files
+[GatherTextStep5]
+CommandletClass=InternationalizationExport
+bExportLoc=true
diff --git a/Engine/Source/Programs/AutomationTool/Scripts/CopyBuildToStagingDirectory.Automation.cs b/Engine/Source/Programs/AutomationTool/Scripts/CopyBuildToStagingDirectory.Automation.cs
--- a/Engine/Source/Programs/AutomationTool/Scripts/CopyBuildToStagingDirectory.Automation.cs
+++ b/Engine/Source/Programs/AutomationTool/Scripts/CopyBuildToStagingDirectory.Automation.cs
@@ -303,6 +303,20 @@ public partial class Project : CommandUtils
 		}
 	}
 
+	private static void StageLocalizationDataForTargetsInDirectory(DeploymentContext SC, List<string> CulturesToStage, DirectoryReference RootDirectory)
+	{
+		if (DirectoryReference.Exists(RootDirectory))
+		{
+			foreach (DirectoryReference LocTargetDirectory in DirectoryReference.EnumerateDirectories(RootDirectory))
+			{
+				if (!SC.BlacklistLocalizationTargets.Contains(LocTargetDirectory.MakeRelativeTo(RootDirectory)))
+				{
+					StageLocalizationDataForTarget(SC, CulturesToStage, LocTargetDirectory);
+				}
+			}
+		}
+	}
+
 	private static void StageLocalizationDataForCulture(DeploymentContext SC, string CultureName, DirectoryReference SourceDirectory)
 	{
 		CultureName = CultureName.Replace('-', '_');
@@ -613,19 +627,7 @@ public partial class Project : CommandUtils
 				}
 
 				// Stage all project localization targets
-				{
-					DirectoryReference ProjectLocRootDirectory = DirectoryReference.Combine(SC.ProjectRoot, "Content", "Localization");
-					if (DirectoryReference.Exists(ProjectLocRootDirectory))
-					{
-						foreach (DirectoryReference ProjectLocTargetDirectory in DirectoryReference.EnumerateDirectories(ProjectLocRootDirectory))
-						{
-							if (!SC.BlacklistLocalizationTargets.Contains(ProjectLocTargetDirectory.MakeRelativeTo(ProjectLocRootDirectory)))
-							{
-								StageLocalizationDataForTarget(SC, CulturesToStage, ProjectLocTargetDirectory);
-							}
-						}
-					}
-				}
+				StageLocalizationDataForTargetsInDirectory(SC, CulturesToStage, DirectoryReference.Combine(SC.ProjectRoot, "Content", "Localization"));
 
 				// Stage all plugin localization targets
 				foreach (KeyValuePair<StagedFileReference, FileReference> StagedPlugin in StagedPlugins)
@@ -633,6 +635,13 @@ public partial class Project : CommandUtils
 					StageLocalizationDataForPlugin(SC, CulturesToStage, StagedPlugin.Value);
 				}
 
+				// Stage any platform extension localization targets
+				{
+					string PlatformExtensionName = SC.StageTargetPlatform.PlatformType.ToString();
+					StageLocalizationDataForTargetsInDirectory(SC, CulturesToStage, DirectoryReference.Combine(SC.LocalRoot, "Engine", "Platforms", PlatformExtensionName, "Content", "Localization"));
+					StageLocalizationDataForTargetsInDirectory(SC, CulturesToStage, DirectoryReference.Combine(SC.ProjectRoot, "Platforms", PlatformExtensionName, "Content", "Localization"));
+				}
+
 				// Stage any additional UFS and NonUFS paths specified in the project ini files; these dirs are relative to the game content directory
 				if (PlatformGameConfig != null)
 				{
diff --git a/Engine/Source/Programs/AutomationTool/Scripts/Localisation.Automation.cs b/Engine/Source/Programs/AutomationTool/Scripts/Localisation.Automation.cs
--- a/Engine/Source/Programs/AutomationTool/Scripts/Localisation.Automation.cs
+++ b/Engine/Source/Programs/AutomationTool/Scripts/Localisation.Automation.cs
@@ -21,6 +21,7 @@ using Tools.DotNETCommon;
 [Help("LocalizationSteps", "Optional comma separated list of localization steps to perform [Download, Gather, Import, Export, Compile, GenerateReports, Upload] (default is all). Only valid for projects using a modular config.")]
 [Help("IncludePlugins", "Optional flag to include plugins from within the given UEProjectDirectory as part of the gather. This may optionally specify a comma separated list of the specific plugins to gather (otherwise all plugins will be gathered).")]
 [Help("ExcludePlugins", "Optional comma separated list of plugins to exclude from the gather.")]
+[Help("IncludePlatforms", "Optional flag to include platforms from within the given UEProjectDirectory as part of the gather.")]
 [Help("AdditionalCommandletArguments", "Optional arguments to pass to the gather process.")]
 class Localize : BuildCommand
 {
@@ -119,6 +120,8 @@ class Localize : BuildCommand
 			}
 		}
 
+		var ShouldGatherPlatforms = ParseParam("IncludePlatforms");
+
 		var AdditionalCommandletArguments = ParseParamValue("AdditionalCommandletArguments");
 		if (AdditionalCommandletArguments == null)
 		{
@@ -133,11 +136,32 @@ class Localize : BuildCommand
 			LocalizationBatches.Add(new LocalizationBatch(UEProjectDirectory, UEProjectDirectory, "", LocalizationProjectNames));
 		}
 
+		// Build up any additional batches needed for platforms
+		if (ShouldGatherPlatforms)
+		{
+			var PlatformsRootDirectory = new DirectoryReference(CombinePaths(UEProjectRoot, UEProjectDirectory, "Platforms"));
+			if (DirectoryReference.Exists(PlatformsRootDirectory))
+			{
+				foreach (DirectoryReference PlatformDirectory in DirectoryReference.EnumerateDirectories(PlatformsRootDirectory))
+				{
+					// Find the localization targets defined for this platform
+					var PlatformTargetNames = GetLocalizationTargetsFromDirectory(new DirectoryReference(CombinePaths(PlatformDirectory.FullName, "Config", "Localization")));
+					if (PlatformTargetNames.Count > 0)
+					{
+						var RootRelativePluginPath = PlatformDirectory.MakeRelativeTo(new DirectoryReference(UEProjectRoot));
+						RootRelativePluginPath = RootRelativePluginPath.Replace('\\', '/'); // Make sure we use / as these paths are used with P4
+
+						LocalizationBatches.Add(new LocalizationBatch(UEProjectDirectory, RootRelativePluginPath, "", PlatformTargetNames));
+					}
+				}
+			}
+		}
+
 		// Build up any additional batches needed for plugins
 		if (ShouldGatherPlugins)
 		{
-			var PluginsRootDirectory = CombinePaths(UEProjectRoot, UEProjectDirectory);
-			IReadOnlyList<PluginInfo> AllPlugins = Plugins.ReadPluginsFromDirectory(new DirectoryReference(PluginsRootDirectory), "Plugins", UEProjectName.Length == 0 ? PluginType.Engine : PluginType.Project);
+			var PluginsRootDirectory = new DirectoryReference(CombinePaths(UEProjectRoot, UEProjectDirectory));
+			IReadOnlyList<PluginInfo> AllPlugins = Plugins.ReadPluginsFromDirectory(PluginsRootDirectory, "Plugins", UEProjectName.Length == 0 ? PluginType.Engine : PluginType.Project);
 
 			// Add a batch for each plugin that meets our criteria
 			var AvailablePluginNames = new HashSet<string>();
@@ -461,6 +485,40 @@ class Localize : BuildCommand
 		return ProjectImportExportInfo;
 	}
 
+	private List<string> GetLocalizationTargetsFromDirectory(DirectoryReference ConfigDirectory)
+	{
+		var LocalizationTargets = new List<string>();
+
+		if (DirectoryReference.Exists(ConfigDirectory))
+		{
+			var FileSuffixes = new[] {
+				"_Gather",
+				"_Import",
+				"_Export",
+				"_Compile",
+				"_GenerateReports",
+			};
+
+			foreach (FileReference ConfigFile in DirectoryReference.EnumerateFiles(ConfigDirectory))
+			{
+				string LocalizationTarget = ConfigFile.GetFileNameWithoutExtension();
+				foreach (var FileSuffix in FileSuffixes)
+				{
+					if (LocalizationTarget.EndsWith(FileSuffix))
+					{
+						LocalizationTarget = LocalizationTarget.Remove(LocalizationTarget.Length - FileSuffix.Length);
+					}
+				}
+				if (!LocalizationTargets.Contains(LocalizationTarget))
+				{
+					LocalizationTargets.Add(LocalizationTarget);
+				}
+			}
+		}
+
+		return LocalizationTargets;
+	}
+
 	private Dictionary<string, byte[]> GetPOFileHashes(IReadOnlyList<LocalizationBatch> LocalizationBatches, string UEProjectRoot)
 	{
 		var AllFiles = new Dictionary<string, byte[]>();
