Gauntlet: Guard iOS LLDB crash parse and warn if parse more than one current thread
#rnx

diff --git a/Engine/Source/Programs/AutomationTool/Gauntlet/Platform/IOS/Gauntlet.TargetDeviceIOS.cs b/Engine/Source/Programs/AutomationTool/Gauntlet/Platform/IOS/Gauntlet.TargetDeviceIOS.cs
--- a/Engine/Source/Programs/AutomationTool/Gauntlet/Platform/IOS/Gauntlet.TargetDeviceIOS.cs
+++ b/Engine/Source/Programs/AutomationTool/Gauntlet/Platform/IOS/Gauntlet.TargetDeviceIOS.cs
@@ -1004,20 +1004,29 @@ namespace Gauntlet
 		/// </summary>		
 		public static string GenerateCrashLog(string LogOutput)
 		{
-			DateTime TimeStamp;
-			int Frame;
-			ThreadInfo Thread = ParseCallstack(LogOutput, out TimeStamp, out Frame);
-			if (Thread == null)
+			try
 			{
-				return null;
-			}
+				DateTime TimeStamp;
+				int Frame;
+				ThreadInfo Thread = ParseCallstack(LogOutput, out TimeStamp, out Frame);
+				if (Thread == null)
+				{
+					return null;
+				}
+
+				StringBuilder CrashLog = new StringBuilder();
+				CrashLog.Append(string.Format("[{0}:000][{1}]LogCore: === Fatal Error: ===\n", TimeStamp.ToString("yyyy.mm.dd - H.mm.ss"), Frame));
+				CrashLog.Append(string.Format("Error: Thread #{0} {1}\n", Thread.Num, Thread.Status));
+				CrashLog.Append(string.Join("\n", Thread.Frames));
 
-			StringBuilder CrashLog = new StringBuilder();
-			CrashLog.Append(string.Format("[{0}:000][{1}]LogCore: === Fatal Error: ===\n", TimeStamp.ToString("yyyy.mm.dd - H.mm.ss"), Frame));
-			CrashLog.Append(string.Format("Error: Thread #{0} {1}\n", Thread.Num, Thread.Status));
-			CrashLog.Append(string.Join("\n", Thread.Frames));
+				return CrashLog.ToString();
+			} 
+			catch (Exception Ex)
+			{
+				Log.Warning("Exception parsing LLDB callstack {0}", Ex.Message);
+			}
 
-			return CrashLog.ToString();
+			return null;
 
 		}
 
@@ -1150,7 +1159,12 @@ namespace Gauntlet
 				LineNode = LineNode.Next;
 			}
 
-			Thread = Threads.SingleOrDefault(T => T.Current == true);
+			if (Threads.Count(T => T.Current == true) > 1)
+			{
+				Log.Warning("LLDB debug parsed more than one current thread");
+			}
+
+			Thread = Threads.FirstOrDefault(T => T.Current == true);
 
 			if (Threads.Count > 0 && Thread == null)
 			{
@@ -1159,7 +1173,7 @@ namespace Gauntlet
 
 
 			// Do not want to surface crashes which happen as a result of requesting exit
-			if (Thread != null && Thread.Frames.SingleOrDefault(F => F.Symbol.Contains("::RequestExit")) != null)
+			if (Thread != null && Thread.Frames.FirstOrDefault(F => F.Symbol.Contains("::RequestExit")) != null)
 			{
 				Thread = null;
 			}
