Added support for switch and ps4 being able to build content patches.
Required much rework of the chunk upload process.  Converted most of it to a new C# task.
Justin.Marcus; Ben.Marsh

diff --git a/Engine/Source/Programs/AutomationTool/AutomationUtils/MCPPublic.cs b/Engine/Source/Programs/AutomationTool/AutomationUtils/MCPPublic.cs
--- a/Engine/Source/Programs/AutomationTool/AutomationUtils/MCPPublic.cs
+++ b/Engine/Source/Programs/AutomationTool/AutomationUtils/MCPPublic.cs
@@ -135,12 +135,19 @@ namespace EpicGames.MCP.Automation
 		/// PS4 platform
 		/// </summary>
 		PS4,
+
+		/// <summary>
+		/// Switch platform
+		/// </summary>
+		Switch,
+
+
 	}
 
-    /// <summary>
-    /// Enum that defines CDN types
-    /// </summary>
-    public enum CDNType
+	/// <summary>
+	/// Enum that defines CDN types
+	/// </summary>
+	public enum CDNType
     {
         /// <summary>
         /// Internal HTTP CDN server
diff --git a/Engine/Source/Programs/AutomationTool/BuildGraph/Script.cs b/Engine/Source/Programs/AutomationTool/BuildGraph/Script.cs
--- a/Engine/Source/Programs/AutomationTool/BuildGraph/Script.cs
+++ b/Engine/Source/Programs/AutomationTool/BuildGraph/Script.cs
@@ -1279,7 +1279,7 @@ namespace AutomationTool
 					{
 						if(Arguments[Idx] == null)
 						{
-							LogWarning(Element, "Macro '{0}' is missing argument '{1}'", Macro.ArgumentNameToIndex.First(x => x.Value == Idx).Key);
+							LogWarning(Element, "Macro '{0}' is missing argument '{1}'", Macro.Name, Macro.ArgumentNameToIndex.First(x => x.Value == Idx).Key);
 							bHasMissingArguments = true;
 						}
 					}
