Changed snapshot values to nullable types so that we only report stats that we actually captured. Previously, if a stat failed to capture, we would mistakenly report zero values.

#rnx
#test RunFort -test=ProfileGo

diff --git a/Engine/Source/Programs/AutomationTool/Gauntlet/Unreal/Utils/Gauntlet.UnrealSnapshotParser.cs b/Engine/Source/Programs/AutomationTool/Gauntlet/Unreal/Utils/Gauntlet.UnrealSnapshotParser.cs
--- a/Engine/Source/Programs/AutomationTool/Gauntlet/Unreal/Utils/Gauntlet.UnrealSnapshotParser.cs
+++ b/Engine/Source/Programs/AutomationTool/Gauntlet/Unreal/Utils/Gauntlet.UnrealSnapshotParser.cs
@@ -13,27 +13,28 @@ namespace Gauntlet
 		public string Name;
 		public float ProfileLength;
 
-		public float CpuUsedMemory;
-		public float CpuPeakMemory;
-		public float PhysicalUsedMemory;
-		public float PhysicalPeakMemory;
-		public float MVP;
-		public float AvgFps;
-		public float Hitches;
-		public float AvgHitches;
-		public float DynamicRes;
-		public float GTTime;
-		public float RTTime;
-		public float GPUTime;
-		public float FTTime;
-		public float RHIT;
-		public int DrawCalls;
-		public int DrawnPrims;
-		public int UnbuiltHLODs;
+		public float? CpuUsedMemory;
+		public float? CpuPeakMemory;
+		public float? PhysicalUsedMemory;
+		public float? PhysicalPeakMemory;
+		public float? MVP;
+		public float? AvgFps;
+		public float? Hitches;
+		public float? AvgHitches;
+		public float? DynamicRes;
+		public float? GTTime;
+		public float? RTTime;
+		public float? GPUTime;
+		public float? FTTime;
+		public float? RHIT;
+		public int? DrawCalls;
+		public int? DrawnPrims;
+		public int? UnbuiltHLODs;
 
 		public UnrealHealthSnapshot()
 		{
 			Name = "Unknown";
+			ProfileLength = 0;
 		}
 
 		public override string ToString()
@@ -45,20 +46,62 @@ namespace Gauntlet
 				SB.AppendFormat("Snapshot {0}\n", Name);
 				SB.AppendFormat("Duration\t{0}\n", ProfileLength);
 
-				SB.AppendFormat("Memory:\t\t{0} MB\n", CpuUsedMemory);
-				SB.AppendFormat("MVP:\t\t{0:0.00}\n", MVP);
-				SB.AppendFormat("AvgFps:\t\t{0:0.00}\n", AvgFps);
-				SB.AppendFormat("HPM:\t\t{0:0.00}\n", Hitches);
-				SB.AppendFormat("AvgH:\t\t{0:0.00}ms\n", AvgHitches);
-				SB.AppendFormat("DynRes:\t\t{0:0.00}\n", DynamicRes);
-				SB.AppendFormat("GT:\t\t{0:0.00}ms\n", GTTime);
-				SB.AppendFormat("RT:\t\t{0:0.00}ms\n", RTTime);
-				SB.AppendFormat("GPU:\t\t{0:0.00}ms\n", GPUTime);
-				SB.AppendFormat("FT:\t\t{0:0.00}ms\n", FTTime);
-				SB.AppendFormat("RHIT:\t\t{0:0.00}ms\n", RHIT);
-				SB.AppendFormat("DrawCalls:\t{0}\n", DrawCalls);
-				SB.AppendFormat("DrawnPrims:\t{0}\n", DrawnPrims);
-				SB.AppendFormat("UnbuiltHLODs:\t{0}\n", UnbuiltHLODs);
+				if (CpuUsedMemory.HasValue)
+				{
+					SB.AppendFormat("Memory:\t\t{0} MB\n", CpuUsedMemory.Value);
+				}
+				if (MVP.HasValue)
+				{
+					SB.AppendFormat("MVP:\t\t{0:0.00}\n", MVP.Value);
+				}
+				if (AvgFps.HasValue)
+				{
+					SB.AppendFormat("AvgFps:\t\t{0:0.00}\n", AvgFps.Value);
+				}
+				if (Hitches.HasValue)
+				{
+					SB.AppendFormat("HPM:\t\t{0:0.00}\n", Hitches.Value);
+				}
+				if (AvgHitches.HasValue)
+				{
+					SB.AppendFormat("AvgH:\t\t{0:0.00}ms\n", AvgHitches.Value);
+				}
+				if (DynamicRes.HasValue)
+				{
+					SB.AppendFormat("DynRes:\t\t{0:0.00}\n", DynamicRes.Value);
+				}
+				if (GTTime.HasValue)
+				{
+					SB.AppendFormat("GT:\t\t{0:0.00}ms\n", GTTime.Value);
+				}
+				if (RTTime.HasValue)
+				{
+					SB.AppendFormat("RT:\t\t{0:0.00}ms\n", RTTime.Value);
+				}
+				if (GPUTime.HasValue)
+				{
+					SB.AppendFormat("GPU:\t\t{0:0.00}ms\n", GPUTime.Value);
+				}
+				if (FTTime.HasValue)
+				{
+					SB.AppendFormat("FT:\t\t{0:0.00}ms\n", FTTime.Value);
+				}
+				if (RHIT.HasValue)
+				{
+					SB.AppendFormat("RHIT:\t\t{0:0.00}ms\n", RHIT.Value);
+				}
+				if (DrawCalls.HasValue)
+				{
+					SB.AppendFormat("DrawCalls:\t{0}\n", DrawCalls.Value);
+				}
+				if (DrawnPrims.HasValue)
+				{
+					SB.AppendFormat("DrawnPrims:\t{0}\n", DrawnPrims.Value);
+				}
+				if (UnbuiltHLODs.HasValue)
+				{
+					SB.AppendFormat("UnbuiltHLODs:\t{0}\n", UnbuiltHLODs.Value);
+				}
 			}
 			else
 			{
@@ -74,21 +117,21 @@ namespace Gauntlet
 				Name = Groups[1];
 			});
 
-			RegexUtil.MatchAndApplyGroups(InContent, @"CPU Memory:[\s\w]+?([\d\.]+)MB,[\s\w]+?([\d\.]+)MB", (Groups) =>
+			RegexUtil.MatchAndApplyGroups(InContent, @"MeasuredPerfTime:\s(\d.+) S", (Groups) =>
 			{
-				CpuUsedMemory += Convert.ToSingle(Groups[1]);
-				CpuPeakMemory += Convert.ToSingle(Groups[2]);
+				ProfileLength = Convert.ToSingle(Groups[1]);
 			});
 
-			RegexUtil.MatchAndApplyGroups(InContent, @"Physical Memory:[\s\w]+?([\d\.]+)MB,[\s\w:]+?([\d\.]+)MB", (Groups) =>
+			RegexUtil.MatchAndApplyGroups(InContent, @"CPU Memory:[\s\w]+?([\d\.]+)MB,[\s\w]+?([\d\.]+)MB", (Groups) =>
 			{
-				PhysicalUsedMemory += Convert.ToSingle(Groups[1]);
-				PhysicalPeakMemory += Convert.ToSingle(Groups[2]);
+				CpuUsedMemory = Convert.ToSingle(Groups[1]);
+				CpuPeakMemory = Convert.ToSingle(Groups[2]);
 			});
 
-			RegexUtil.MatchAndApplyGroups(InContent, @"MeasuredPerfTime:\s(\d.+) S", (Groups) =>
+			RegexUtil.MatchAndApplyGroups(InContent, @"Physical Memory:[\s\w]+?([\d\.]+)MB,[\s\w:]+?([\d\.]+)MB", (Groups) =>
 			{
-				ProfileLength += Convert.ToSingle(Groups[1]);
+				PhysicalUsedMemory = Convert.ToSingle(Groups[1]);
+				PhysicalPeakMemory = Convert.ToSingle(Groups[2]);
 			});
 
 			RegexUtil.MatchAndApplyGroups(InContent, @"MVP:\s(\d.+)%,\s.*AvgFPS:(\d.+),\s.*HitchesPerMinute:\s(\d.+),\sAvg\sHitch\s(\d.+)ms", (Groups) =>
@@ -189,6 +232,7 @@ namespace Gauntlet
 
 			CreateFromLog(LogContents, SnapshotTitles);
 		}
+
 		protected virtual void CreateFromLog(string LogContents, string InTitle)
 		{
 			UnrealLogParser Parser = new UnrealLogParser(LogContents);
@@ -220,54 +264,69 @@ namespace Gauntlet
 				{
 					SessionTime += Snap.ProfileLength;
 
-					MVP.Add(Snap.MVP);
+					if (Snap.MVP.HasValue)
+					{
+						MVP.Add(Snap.MVP.Value);
+					}
 
-					if (Snap.AvgFps > 0)
+					if (Snap.AvgFps.HasValue)
 					{
-						AvgFps.Add(Snap.AvgFps);
+						AvgFps.Add(Snap.AvgFps.Value);
 					}
 
-					Hitches.Add(Snap.Hitches);
-					AvgHitches.Add(Snap.AvgHitches);
+					if (Snap.Hitches.HasValue)
+					{
+						Hitches.Add(Snap.Hitches.Value);
+					}
+
+					if (Snap.AvgHitches.HasValue)
+					{
+						AvgHitches.Add(Snap.AvgHitches.Value);
+					}
+
+					if (Snap.DynamicRes.HasValue)
+					{
+						DynamicRes.Add(Snap.DynamicRes.Value);
+					}
 
-					if (Snap.DynamicRes > 0)
+					if (Snap.GTTime.HasValue)
 					{
-						DynamicRes.Add(Snap.DynamicRes);
+						GTTime.Add(Snap.GTTime.Value);
 					}
 
-					if (Snap.GTTime > 0)
+					if (Snap.RTTime.HasValue)
 					{
-						GTTime.Add(Snap.GTTime);
+						RTTime.Add(Snap.RTTime.Value);
 					}
 
-					if (Snap.RTTime > 0)
+					if (Snap.GPUTime.HasValue)
 					{
-						RTTime.Add(Snap.RTTime);
+						GPUTime.Add(Snap.GPUTime.Value);
 					}
 
-					if (Snap.GPUTime > 0)
+					if (Snap.FTTime.HasValue)
 					{
-						GPUTime.Add(Snap.GPUTime);
+						FTTime.Add(Snap.FTTime.Value);
 					}
 
-					if (Snap.FTTime > 0)
+					if (Snap.RHIT.HasValue)
 					{
-						FTTime.Add(Snap.FTTime);
+						RHIT.Add(Snap.RHIT.Value);
 					}
 
-					if (Snap.RHIT > 0)
+					if (Snap.DrawCalls.HasValue)
 					{
-						RHIT.Add(Snap.RHIT);
+						DrawCalls.Add(Snap.DrawCalls.Value);
 					}
 
-					if (Snap.DrawCalls > 0)
+					if (Snap.DrawnPrims.HasValue)
 					{
-						DrawCalls.Add(Snap.DrawCalls);
-						DrawnPrims.Add(Snap.DrawnPrims);
+						DrawnPrims.Add(Snap.DrawnPrims.Value);
 					}
-					if (Snap.UnbuiltHLODs > 0)
+
+					if (Snap.UnbuiltHLODs.HasValue)
 					{
-						UnbuiltHLODs.Add(Snap.UnbuiltHLODs);
+						UnbuiltHLODs.Add(Snap.UnbuiltHLODs.Value);
 					}
 				}
 
@@ -276,16 +335,16 @@ namespace Gauntlet
 				{
 					var LastSnapshot = Snapshots.Last();
 
-					float SnapshotPeakMemory = LastSnapshot.PhysicalPeakMemory;
+					float? SnapshotPeakMemory = LastSnapshot.PhysicalPeakMemory;
 
 					//if PeakMemory is reporting as 0, use Memory if it's higher than our last
-					if (SnapshotPeakMemory == 0)
+					if (SnapshotPeakMemory.HasValue)
 					{
-						Log.Info("PeakMemory reported as 0mb");
+						PeakMemory = SnapshotPeakMemory.Value;
 					}
 					else
 					{
-						PeakMemory = SnapshotPeakMemory;
+						Log.Info("PeakMemory reported as 0mb");
 					}
 				}
 			}
@@ -302,8 +361,8 @@ namespace Gauntlet
 			if (SampleCount > 0)
 			{
 				SB.AppendFormat("Performance report from {0} samples and {1} seconds\n", SampleCount, SessionTime);
-
 				SB.AppendFormat("Peak Memory: {0} MB\n", PeakMemory);
+
 				if (MVP.Count > 0)
 				{
 					SB.AppendFormat("MVP:\t{0:0.00} (Min: {1:0.00}, Max: {2:0.00})\n", MVP.Average(), MVP.Min(), MVP.Max());
