Gauntlet: Making sure we don't overwrite artifacts OR output logs if we have multiple of one or more roles (in particular, servers, which were presumed to be no-more-than-one before this CL)
#rnx

diff --git a/Engine/Source/Programs/AutomationTool/Gauntlet/Unreal/Base/Gauntlet.UnrealSession.cs b/Engine/Source/Programs/AutomationTool/Gauntlet/Unreal/Base/Gauntlet.UnrealSession.cs
--- a/Engine/Source/Programs/AutomationTool/Gauntlet/Unreal/Base/Gauntlet.UnrealSession.cs
+++ b/Engine/Source/Programs/AutomationTool/Gauntlet/Unreal/Base/Gauntlet.UnrealSession.cs
@@ -1192,27 +1192,41 @@ namespace Gauntlet
 		/// <returns></returns>
 		public IEnumerable<UnrealRoleArtifacts> SaveRoleArtifacts(UnrealTestContext Context, UnrealSessionInstance TestInstance, string OutputPath)
 		{
-			int ClientCount = 1;
-			int DummyClientCount = 1;
+			Dictionary<UnrealTargetRole, int> RoleCounts = new Dictionary<UnrealTargetRole, int>();
+			int DummyClientCount = 0;
 
 			List<UnrealRoleArtifacts> AllArtifacts = new List<UnrealRoleArtifacts>();
 
 			foreach (UnrealSessionInstance.RoleInstance App in TestInstance.RunningRoles)
 			{
-				bool IsServer = App.Role.RoleType.IsServer();
 				string RoleName = (App.Role.IsDummy() ? "Dummy" : "") + App.Role.RoleType.ToString();
 				string FolderName = RoleName;
 
-				if (IsServer == false)
+				int RoleCount = 1;
+
+				if (App.Role.IsDummy())
 				{
-					if ((!App.Role.IsDummy() && ClientCount++ > 1))
+					DummyClientCount++;
+					RoleCount = DummyClientCount;
+				}
+				else
+				{
+					if (!RoleCounts.ContainsKey(App.Role.RoleType))
 					{
-						FolderName += string.Format("_{0:00}", ClientCount);
+						RoleCounts.Add(App.Role.RoleType, 1);
 					}
-					else if (App.Role.IsDummy() && DummyClientCount++ > 1)
+					else
 					{
-						FolderName += string.Format("_{0:00}", DummyClientCount);
+						RoleCounts[App.Role.RoleType]++;
 					}
+
+					RoleCount = RoleCounts[App.Role.RoleType];
+				}
+
+				
+				if (RoleCount > 1)
+				{
+					FolderName += string.Format("_{0:00}", RoleCount);
 				}
 
 				string DestPath = Path.Combine(OutputPath, FolderName);
