Add very basic versioning stats to base UnrealTestNode

diff --git a/Engine/Source/Programs/AutomationTool/Gauntlet/Unreal/Base/Gauntlet.UnrealTestNode.cs b/Engine/Source/Programs/AutomationTool/Gauntlet/Unreal/Base/Gauntlet.UnrealTestNode.cs
--- a/Engine/Source/Programs/AutomationTool/Gauntlet/Unreal/Base/Gauntlet.UnrealTestNode.cs
+++ b/Engine/Source/Programs/AutomationTool/Gauntlet/Unreal/Base/Gauntlet.UnrealTestNode.cs
@@ -94,7 +94,7 @@ namespace Gauntlet
 		/// <summary>
 		/// Used to track how much of our log has been written out
 		/// </summary>
-		private int LastLogCount ;
+		private int LastLogCount;
 
 		private int CurrentPass;
 
@@ -102,6 +102,22 @@ namespace Gauntlet
 
 		static protected DateTime SessionStartTime = DateTime.MinValue;
 
+		/// <summary>
+		/// Standard semantic versioning for tests. Should be overwritten within individual tests, and individual test maintainers
+		/// are responsible for updating their own versions. See https://semver.org/ for more info on maintaining semantic versions.
+		/// </summary>
+		protected int TestVersionMajor = 1;
+
+		protected int TestVersionMinor = 0;
+
+		protected int TestVersionFix = 0;
+
+		/// <summary>
+		/// Get fully assembled version string.
+		/// </summary>
+		/// <returns>The version of this test in Major.Minor.Fix format.</returns>
+		public string TestVersion { get { return string.Format("{0}.{1}.{2}", TestVersionMajor, TestVersionMinor, TestVersionFix); } }
+
 		/// <summary>
 		/// Path to the directory that logs and other artifacts are copied to after the test run.
 		/// </summary>
