Change test version number to use a version object instead of hand-rolling it.

diff --git a/Engine/Source/Programs/AutomationTool/Gauntlet/Unreal/Base/Gauntlet.UnrealTestNode.cs b/Engine/Source/Programs/AutomationTool/Gauntlet/Unreal/Base/Gauntlet.UnrealTestNode.cs
--- a/Engine/Source/Programs/AutomationTool/Gauntlet/Unreal/Base/Gauntlet.UnrealTestNode.cs
+++ b/Engine/Source/Programs/AutomationTool/Gauntlet/Unreal/Base/Gauntlet.UnrealTestNode.cs
@@ -106,17 +106,9 @@ namespace Gauntlet
 		/// Standard semantic versioning for tests. Should be overwritten within individual tests, and individual test maintainers
 		/// are responsible for updating their own versions. See https://semver.org/ for more info on maintaining semantic versions.
 		/// </summary>
-		protected int TestVersionMajor = 1;
-
-		protected int TestVersionMinor = 0;
-
-		protected int TestVersionFix = 0;
-
-		/// <summary>
-		/// Get fully assembled version string.
-		/// </summary>
-		/// <returns>The version of this test in Major.Minor.Fix format.</returns>
-		public string TestVersion { get { return string.Format("{0}.{1}.{2}", TestVersionMajor, TestVersionMinor, TestVersionFix); } }
+		/// 
+		protected Version TestVersion;
+		
 
 		/// <summary>
 		/// Path to the directory that logs and other artifacts are copied to after the test run.
@@ -157,6 +149,7 @@ namespace Gauntlet
 			LastLogCount = 0;
 			CurrentPass = 0;
 			NumPasses = 0;
+			TestVersion = new Version("1.0.0");
 			ArtifactPath = string.Empty;
 		}
 
