Fix enum issue with newer P4 binaries having "undone by" integration actions
#rnx

diff --git a/Engine/Source/Programs/AutomationTool/AutomationUtils/P4Utils.cs b/Engine/Source/Programs/AutomationTool/AutomationUtils/P4Utils.cs
--- a/Engine/Source/Programs/AutomationTool/AutomationUtils/P4Utils.cs
+++ b/Engine/Source/Programs/AutomationTool/AutomationUtils/P4Utils.cs
@@ -432,6 +432,13 @@ namespace AutomationTool
 		/// </summary>
 		[Description("undid")]
 		Undid,
+
+		/// <summary>
+		/// file was reverted to a previous revision
+		/// </summary>
+		[Description("undone by")]
+		UndoneBy
+
 	}
 
 	/// <summary>
