Added missing packaged check for PS4 builds

diff --git a/Engine/Source/Programs/AutomationTool/Gauntlet/Unreal/Base/Gauntlet.UnrealSession.cs b/Engine/Source/Programs/AutomationTool/Gauntlet/Unreal/Base/Gauntlet.UnrealSession.cs
--- a/Engine/Source/Programs/AutomationTool/Gauntlet/Unreal/Base/Gauntlet.UnrealSession.cs
+++ b/Engine/Source/Programs/AutomationTool/Gauntlet/Unreal/Base/Gauntlet.UnrealSession.cs
@@ -157,7 +157,7 @@ namespace Gauntlet
 				}
 			}
 
-			if (Globals.Params.ParseParam("packaged") && (InPlatform == UnrealTargetPlatform.Switch || InPlatform == UnrealTargetPlatform.XboxOne))
+			if (Globals.Params.ParseParam("packaged") && (InPlatform == UnrealTargetPlatform.Switch || InPlatform == UnrealTargetPlatform.XboxOne || InPlatform == UnrealTargetPlatform.PS4))
 			{
 				RequiredBuildFlags |= BuildFlags.Packaged;
 			}
