Gauntlet: warn and attempt to recover when a copy fails due to lack of long path support on windows
#rnx

diff --git a/Engine/Source/Programs/AutomationTool/Gauntlet/Framework/Gauntlet.Utils.cs b/Engine/Source/Programs/AutomationTool/Gauntlet/Framework/Gauntlet.Utils.cs
--- a/Engine/Source/Programs/AutomationTool/Gauntlet/Framework/Gauntlet.Utils.cs
+++ b/Engine/Source/Programs/AutomationTool/Gauntlet/Framework/Gauntlet.Utils.cs
@@ -1128,7 +1128,23 @@ namespace Gauntlet
 								{
 									Log.Error("File Copy failed with {0}.", ex.Message);
 								}
-								throw new Exception(string.Format("File Copy failed with {0}.", ex.Message));
+
+								// Warn with message if we're exceeding long path, otherwise throw an exception
+								const int MAX_PATH = 260;
+								bool LongPath = BuildHostPlatform.Current.Platform == UnrealTargetPlatform.Win64 && (SourcePath.Length >= MAX_PATH || DestFile.Length >= MAX_PATH);
+
+								if (!LongPath)
+								{
+									throw new Exception(string.Format("File Copy failed with {0}.", ex.Message));
+								}
+								else
+								{
+									string LongPathMessage = (Environment.OSVersion.Version.Major > 6) ?
+										"Long path detected, check that long paths are enabled." :
+										"Long path detected, OS version doesn't support long paths.";
+
+									Log.Warning("Long path file copy failed with {0}.  Please verify that this file is not required.", ex.Message, LongPathMessage);
+								}
 							}
 						}
 					} while (Copied == false);
