Gauntlet: Making sure we only flag a build NotBulk if the user is specifically requesting that delineation
This way queries that don't specifically request -notbulk can still get bulk builds, if those are otherwise applicable

diff --git a/Engine/Source/Programs/AutomationTool/Gauntlet/Unreal/BuildSource/Gauntlet.UnrealBuildSource.cs b/Engine/Source/Programs/AutomationTool/Gauntlet/Unreal/BuildSource/Gauntlet.UnrealBuildSource.cs
--- a/Engine/Source/Programs/AutomationTool/Gauntlet/Unreal/BuildSource/Gauntlet.UnrealBuildSource.cs
+++ b/Engine/Source/Programs/AutomationTool/Gauntlet/Unreal/BuildSource/Gauntlet.UnrealBuildSource.cs
@@ -585,7 +585,7 @@ namespace Gauntlet
                     {
                         Flags |= BuildFlags.Bulk;
                     }
-					else
+					else if(Globals.Params.ParseParam("notbulk"))
 					{
 						Flags |= BuildFlags.NotBulk;
 					}
