Gauntlet: Fixing issue with long path warning
#rnx

diff --git a/Engine/Source/Programs/AutomationTool/Gauntlet/Framework/Gauntlet.Utils.cs b/Engine/Source/Programs/AutomationTool/Gauntlet/Framework/Gauntlet.Utils.cs
--- a/Engine/Source/Programs/AutomationTool/Gauntlet/Framework/Gauntlet.Utils.cs
+++ b/Engine/Source/Programs/AutomationTool/Gauntlet/Framework/Gauntlet.Utils.cs
@@ -1143,6 +1143,8 @@ namespace Gauntlet
 										"Long path detected, check that long paths are enabled." :
 										"Long path detected, OS version doesn't support long paths.";
 
+									// Break out of loop with warning
+									Copied = true;
 									Log.Warning("Long path file copy failed with {0}.  Please verify that this file is not required.", ex.Message, LongPathMessage);
 								}
 							}
