Gauntlet: Fix exception when deserializing devices from json due to UnrealTargetPlatform being nullable
#rnx

diff --git a/Engine/Source/Programs/AutomationTool/Gauntlet/Framework/Devices/Gauntlet.DevicePool.cs b/Engine/Source/Programs/AutomationTool/Gauntlet/Framework/Devices/Gauntlet.DevicePool.cs
--- a/Engine/Source/Programs/AutomationTool/Gauntlet/Framework/Devices/Gauntlet.DevicePool.cs
+++ b/Engine/Source/Programs/AutomationTool/Gauntlet/Framework/Devices/Gauntlet.DevicePool.cs
@@ -37,8 +37,10 @@ namespace Gauntlet
 		public string DeviceData;
 
 		// legacy - remove!
+		[JsonConverter(typeof(UnrealTargetPlatformConvertor))]
 		public UnrealTargetPlatform Type;
 
+		[JsonConverter(typeof(UnrealTargetPlatformConvertor))]
 		public UnrealTargetPlatform? Platform;
 
 		public EPerfSpec PerfSpec;
diff --git a/Engine/Source/Programs/AutomationTool/Gauntlet/Unreal/Utils/Gauntlet.UnrealHelpers.cs b/Engine/Source/Programs/AutomationTool/Gauntlet/Unreal/Utils/Gauntlet.UnrealHelpers.cs
--- a/Engine/Source/Programs/AutomationTool/Gauntlet/Unreal/Utils/Gauntlet.UnrealHelpers.cs
+++ b/Engine/Source/Programs/AutomationTool/Gauntlet/Unreal/Utils/Gauntlet.UnrealHelpers.cs
@@ -10,6 +10,7 @@ using System.Linq;
 using System.Net.NetworkInformation;
 using System.Text.RegularExpressions;
 using System.Net;
+using Newtonsoft.Json;
 
 namespace Gauntlet
 {
@@ -371,4 +372,33 @@ namespace Gauntlet
 
 	}
 
+	/// <summary>
+	///  Converts between json and UnrealTargetPlatform
+	/// </summary>
+	public class UnrealTargetPlatformConvertor : JsonConverter
+	{
+		public override bool CanConvert(Type ObjectType)
+		{
+			return ObjectType == typeof(string) || ObjectType == typeof(UnrealTargetPlatform);
+		}
+
+		public override object ReadJson(JsonReader Reader, Type ObjectType, object ExistingValue, JsonSerializer Serializer)
+		{
+			UnrealTargetPlatform Platform;
+			if (!UnrealTargetPlatform.TryParse((string)Reader.Value, out Platform))
+			{
+				return null;
+			}
+			return Platform;
+		}
+
+		public override void WriteJson(JsonWriter Writer, object Value, JsonSerializer Serializer)
+		{
+			UnrealTargetPlatform? Platform = (UnrealTargetPlatform)Value;
+			Writer.WriteValue(Platform);
+		}
+
+	}
+
+
 }
\ No newline at end of file
