Gauntlet: Fix issue with post mortem callstack not being used when multiple callstacks.  Also, handle [Inline Function] parsing
#rnx

diff --git a/Engine/Source/Programs/AutomationTool/Gauntlet/Unreal/Utils/Gauntlet.UnrealLogParser.cs b/Engine/Source/Programs/AutomationTool/Gauntlet/Unreal/Utils/Gauntlet.UnrealLogParser.cs
--- a/Engine/Source/Programs/AutomationTool/Gauntlet/Unreal/Utils/Gauntlet.UnrealLogParser.cs
+++ b/Engine/Source/Programs/AutomationTool/Gauntlet/Unreal/Utils/Gauntlet.UnrealLogParser.cs
@@ -26,6 +26,7 @@ namespace Gauntlet
 			public string			Message;
 			public string[]			Callstack;
 			public bool				IsEnsure;
+			public bool				IsPostMortem;
 			
 			/// <summary>
 			/// Generate a string that represents a CallstackMessage formatted to be inserted into a log file.
@@ -473,6 +474,9 @@ namespace Gauntlet
 
 						string Line = SearchContent.Substring(0, EOL);
 
+						// collapse inline function
+						Line = Line.Replace("[Inline Function] ", "[InlineFunction]");
+
 						// Must have [Callstack] 0x00123456
 						// The module name is optional, must start with whitespace, and continues until next whites[ace
 						// filename is optional, must be in [file]
@@ -553,26 +557,18 @@ namespace Gauntlet
 				// check this trace to see if it's postmortem
 				if (Trace.Message.IndexOf("Postmortem Cause:", StringComparison.OrdinalIgnoreCase) != -1)
 				{
-					// we have post-mortem info, which should be much better than the game-generated stuff so 
-					// update any previous trace
-					var PrevTrace = FilteredTraces.LastOrDefault();
-
-					if (PrevTrace != null)
-					{
-						PrevTrace.Callstack = Trace.Callstack;
-					}
-					else
-					{
-						// no in-log callstack? Strange... but just go with this
-						FilteredTraces.Add(Trace);
-					}
-				}
-				else
-				{
-					FilteredTraces.Add(Trace);
+					// we have post-mortem info, which should be much better than the game-generated stuff and will be sorted to first position
+					Trace.IsPostMortem = true;
 				}
+
+				FilteredTraces.Add(Trace);
+			}
+
+			// If we have a post mortem crash, sort it to the front
+			if (FilteredTraces.FirstOrDefault((Trace) => { return Trace.IsPostMortem; }) != null)
+			{				
+				FilteredTraces.Sort((Trace1, Trace2) => { if (!Trace1.IsPostMortem && !Trace2.IsPostMortem) return 0; return Trace1.IsPostMortem ? -1 : 1; });
 			}
-			
 
 			return FilteredTraces;
 		}
