Changed SyncScript syntax to use //workspace/path instead of //depot/path. This addresses issues when the workspace is in an uncertain state because the user used stream switching with a failed update

#tests Synced both existing and missing projects on Mac and Windows

diff --git a/Engine/Source/Programs/AutomationTool/Scripts/SyncProject.Automation.cs b/Engine/Source/Programs/AutomationTool/Scripts/SyncProject.Automation.cs
--- a/Engine/Source/Programs/AutomationTool/Scripts/SyncProject.Automation.cs
+++ b/Engine/Source/Programs/AutomationTool/Scripts/SyncProject.Automation.cs
@@ -143,7 +143,8 @@ class SyncProject : BuildCommand
 					if (P4.FileExistsInDepot(P4Path))
 					{
 						P4WhereRecord Record = P4.Where(P4Path, AllowSpew: false).FirstOrDefault(x => x.DepotFile != null && !x.bUnmap);
-						P4ProjectPath = Record.DepotFile;
+						// make sure to sync with //workspace/path as it cleans up files if the user has stream switched
+						P4ProjectPath = Record.ClientFile;
 						Log.TraceInformation("Found project at {0}", P4ProjectPath);
 						break;
 					}
@@ -170,14 +171,14 @@ class SyncProject : BuildCommand
 		}
 		else
 		{
-				// See if the engine is in P4 too by checking the p4 location of a local file
-				string LocalEngineFile = CommandUtils.CombinePaths(CmdEnv.LocalRoot, "Engine", "Source", "UE4Editor.target.cs");
-				P4WhereRecord EngineRecord = P4.Where(LocalEngineFile, AllowSpew: false).FirstOrDefault(x => x.DepotFile != null && !x.bUnmap);
+			// See if the engine is in P4 too by checking the p4 location of a local file
+			string LocalEngineFile = CommandUtils.CombinePaths(CmdEnv.LocalRoot, "Engine", "Source", "UE4Editor.target.cs");
+			P4WhereRecord EngineRecord = P4.Where(LocalEngineFile, AllowSpew: false).FirstOrDefault(x => x.DepotFile != null && !x.bUnmap);
 
-
-				if (!ProjectOnly && P4.FileExistsInDepot(EngineRecord.DepotFile))
+			if (!ProjectOnly && P4.FileExistsInDepot(EngineRecord.DepotFile))
 			{
-				P4EnginePath = EngineRecord.DepotFile.Replace("Engine/Source/UE4Editor.target.cs", "");
+				// make sure to sync with //workspace/path as it cleans up files if the user has stream switched
+				P4EnginePath = EngineRecord.ClientFile.Replace("Engine/Source/UE4Editor.target.cs", "");
 				SyncPaths.Add(CommandUtils.CombinePaths(PathSeparator.Slash, P4EnginePath + "*"));
 				SyncPaths.Add(CommandUtils.CombinePaths(PathSeparator.Slash, P4EnginePath, "Engine", "..."));
 			}
