Fix looped directory numbering for test runs.

diff --git a/Engine/Source/Programs/AutomationTool/Gauntlet/Unreal/Base/Gauntlet.UnrealTestNode.cs b/Engine/Source/Programs/AutomationTool/Gauntlet/Unreal/Base/Gauntlet.UnrealTestNode.cs
--- a/Engine/Source/Programs/AutomationTool/Gauntlet/Unreal/Base/Gauntlet.UnrealTestNode.cs
+++ b/Engine/Source/Programs/AutomationTool/Gauntlet/Unreal/Base/Gauntlet.UnrealTestNode.cs
@@ -468,7 +468,7 @@ namespace Gauntlet
 			// if doing multiple passes, put each in a subdir
 			if (NumPasses > 1)
 			{
-				ArtifactPath = Path.Combine(ArtifactPath, string.Format("Pass_{0}_of_{1}", CurrentPass, NumPasses));
+				ArtifactPath = Path.Combine(ArtifactPath, string.Format("Pass_{0}_of_{1}", CurrentPass + 1, NumPasses));
 			}
 
 			// When running with -parallel we could have several identical tests (same test, configurations) in flight so
