Fix an issue with unreal log parser where a traced error, if the regex matched on the exact end of the string, would crash.

diff --git a/Engine/Source/Programs/AutomationTool/Gauntlet/Unreal/Utils/Gauntlet.UnrealLogParser.cs b/Engine/Source/Programs/AutomationTool/Gauntlet/Unreal/Utils/Gauntlet.UnrealLogParser.cs
--- a/Engine/Source/Programs/AutomationTool/Gauntlet/Unreal/Utils/Gauntlet.UnrealLogParser.cs
+++ b/Engine/Source/Programs/AutomationTool/Gauntlet/Unreal/Utils/Gauntlet.UnrealLogParser.cs
@@ -400,8 +400,9 @@ namespace Gauntlet
 					NewTrace.Position = TraceMatch.Index;
 					NewTrace.Message = TraceMatch.Groups[1].Value;
 
-					// Next there will be a message (in some cases), and a callstack..
-					string ErrorContent = Content.Substring(TraceMatch.Index + TraceMatch.Length + 1);
+					// If the regex matches the very end of the string, the substring will get an invalid range.
+					string ErrorContent = Content.Length <= TraceMatch.Index + TraceMatch.Length 
+						? string.Empty : Content.Substring(TraceMatch.Index + TraceMatch.Length + 1);
 
 					Match MsgMatch = Regex.Match(ErrorContent, @".+:\s*Error:\s*(.+)");
 
