Gauntlet - fixed exceptions that could occur when certain platforms don't exist on a system

#tests ran with & without Xbox & Switch folders present

diff --git a/Engine/Source/Programs/AutomationTool/Gauntlet/Platform/Android/Gauntlet.AndroidBuildSource.cs b/Engine/Source/Programs/AutomationTool/Gauntlet/Platform/Android/Gauntlet.AndroidBuildSource.cs
--- a/Engine/Source/Programs/AutomationTool/Gauntlet/Platform/Android/Gauntlet.AndroidBuildSource.cs
+++ b/Engine/Source/Programs/AutomationTool/Gauntlet/Platform/Android/Gauntlet.AndroidBuildSource.cs
@@ -11,8 +11,6 @@ using System.Linq;
 
 namespace Gauntlet
 {
-
-
 	public class AndroidBuild : IBuild
 	{
 		public UnrealTargetConfiguration Configuration { get; protected set; }
@@ -185,51 +183,54 @@ namespace Gauntlet
 
 			List<DirectoryInfo> AllDirs = new List<DirectoryInfo>();
 
+			//AndroidBuildSource BuildSource = null;
+			List<IBuild> Builds = new List<IBuild>();
+
 			// c:\path\to\build
 			DirectoryInfo PathDI = new DirectoryInfo(InPath);
 
-			if (PathDI.Name.IndexOf("Android", StringComparison.OrdinalIgnoreCase) >= 0)
+			if (PathDI.Exists)
 			{
-				AllDirs.Add(PathDI);
-			}
-
-			// find all directories that begin with Android
-			DirectoryInfo[] AndroidDirs = PathDI.GetDirectories("Android*", SearchOption.TopDirectoryOnly);
+				if (PathDI.Name.IndexOf("Android", StringComparison.OrdinalIgnoreCase) >= 0)
+				{
+					AllDirs.Add(PathDI);
+				}
 
-			AllDirs.AddRange(AndroidDirs);
+				// find all directories that begin with Android
+				DirectoryInfo[] AndroidDirs = PathDI.GetDirectories("Android*", SearchOption.TopDirectoryOnly);
 
-			List<DirectoryInfo> DirsToRecurse = AllDirs;
+				AllDirs.AddRange(AndroidDirs);
 
-			// now get subdirs
-			while (MaxRecursion-- > 0)
-			{
-				List<DirectoryInfo> DiscoveredDirs = new List<DirectoryInfo>();
+				List<DirectoryInfo> DirsToRecurse = AllDirs;
 
-				DirsToRecurse.ToList().ForEach((D) =>
+				// now get subdirs
+				while (MaxRecursion-- > 0)
 				{
-					DiscoveredDirs.AddRange(D.GetDirectories("*", SearchOption.TopDirectoryOnly));
-				});
+					List<DirectoryInfo> DiscoveredDirs = new List<DirectoryInfo>();
 
-				AllDirs.AddRange(DiscoveredDirs);
-				DirsToRecurse = DiscoveredDirs;
-			}
-
-			//AndroidBuildSource BuildSource = null;
-			List<IBuild> Builds = new List<IBuild>();
+					DirsToRecurse.ToList().ForEach((D) =>
+					{
+						DiscoveredDirs.AddRange(D.GetDirectories("*", SearchOption.TopDirectoryOnly));
+					});
 
-            string AndroidBuildFilter = Globals.Params.ParseValue("AndroidBuildFilter", "");
-            foreach (DirectoryInfo Di in AllDirs)
-			{
-				IEnumerable<AndroidBuild> FoundBuilds = AndroidBuild.CreateFromPath(InProjectName, Di.FullName);
+					AllDirs.AddRange(DiscoveredDirs);
+					DirsToRecurse = DiscoveredDirs;
+				}
 
-				if (FoundBuilds != null)
+				string AndroidBuildFilter = Globals.Params.ParseValue("AndroidBuildFilter", "");
+				foreach (DirectoryInfo Di in AllDirs)
 				{
-                    if (!string.IsNullOrEmpty(AndroidBuildFilter))
-                    {
-                        //IndexOf used because Contains must be case-sensitive
-                        FoundBuilds = FoundBuilds.Where(B => B.SourceApkPath.IndexOf(AndroidBuildFilter, StringComparison.OrdinalIgnoreCase) >= 0);
-                    }
-                    Builds.AddRange(FoundBuilds);
+					IEnumerable<AndroidBuild> FoundBuilds = AndroidBuild.CreateFromPath(InProjectName, Di.FullName);
+
+					if (FoundBuilds != null)
+					{
+						if (!string.IsNullOrEmpty(AndroidBuildFilter))
+						{
+							//IndexOf used because Contains must be case-sensitive
+							FoundBuilds = FoundBuilds.Where(B => B.SourceApkPath.IndexOf(AndroidBuildFilter, StringComparison.OrdinalIgnoreCase) >= 0);
+						}
+						Builds.AddRange(FoundBuilds);
+					}
 				}
 			}
 
diff --git a/Engine/Source/Programs/AutomationTool/Gauntlet/Platform/IOS/Gauntlet.IOSBuildSource.cs b/Engine/Source/Programs/AutomationTool/Gauntlet/Platform/IOS/Gauntlet.IOSBuildSource.cs
--- a/Engine/Source/Programs/AutomationTool/Gauntlet/Platform/IOS/Gauntlet.IOSBuildSource.cs
+++ b/Engine/Source/Programs/AutomationTool/Gauntlet/Platform/IOS/Gauntlet.IOSBuildSource.cs
@@ -224,51 +224,55 @@ namespace Gauntlet
 
 			List<DirectoryInfo> AllDirs = new List<DirectoryInfo>();
 
+			List<IBuild> Builds = new List<IBuild>();
+
 			// c:\path\to\build
 			DirectoryInfo PathDI = new DirectoryInfo(InPath);
 
-			if (PathDI.Name.IndexOf("IOS", StringComparison.OrdinalIgnoreCase) >= 0)
+			if (PathDI.Exists)
 			{
-				AllDirs.Add(PathDI);
-			}
-
-			// find all directories that begin with IOS
-			DirectoryInfo[] IOSDirs = PathDI.GetDirectories("IOS*", SearchOption.TopDirectoryOnly);
+				if (PathDI.Name.IndexOf("IOS", StringComparison.OrdinalIgnoreCase) >= 0)
+				{
+					AllDirs.Add(PathDI);
+				}
 
-			AllDirs.AddRange(IOSDirs);
+				// find all directories that begin with IOS
+				DirectoryInfo[] IOSDirs = PathDI.GetDirectories("IOS*", SearchOption.TopDirectoryOnly);
 
-			List<DirectoryInfo> DirsToRecurse = AllDirs;
+				AllDirs.AddRange(IOSDirs);
 
-			// now get subdirs
-			while (MaxRecursion-- > 0)
-			{
-				List<DirectoryInfo> DiscoveredDirs = new List<DirectoryInfo>();
+				List<DirectoryInfo> DirsToRecurse = AllDirs;
 
-				DirsToRecurse.ToList().ForEach((D) =>
+				// now get subdirs
+				while (MaxRecursion-- > 0)
 				{
-					DiscoveredDirs.AddRange(D.GetDirectories("*", SearchOption.TopDirectoryOnly));
-				});
+					List<DirectoryInfo> DiscoveredDirs = new List<DirectoryInfo>();
 
-				AllDirs.AddRange(DiscoveredDirs);
-				DirsToRecurse = DiscoveredDirs;
-			}
+					DirsToRecurse.ToList().ForEach((D) =>
+					{
+						DiscoveredDirs.AddRange(D.GetDirectories("*", SearchOption.TopDirectoryOnly));
+					});
 
-			//IOSBuildSource BuildSource = null;
-			List<IBuild> Builds = new List<IBuild>();
+					AllDirs.AddRange(DiscoveredDirs);
+					DirsToRecurse = DiscoveredDirs;
+				}
 
-			string IOSBuildFilter = Globals.Params.ParseValue("IOSBuildFilter", "");
-			foreach (DirectoryInfo Di in AllDirs)
-			{
-				IEnumerable<IOSBuild> FoundBuilds = IOSBuild.CreateFromPath(InProjectName, Di.FullName);
+				//IOSBuildSource BuildSource = null;
 
-				if (FoundBuilds != null)
+				string IOSBuildFilter = Globals.Params.ParseValue("IOSBuildFilter", "");
+				foreach (DirectoryInfo Di in AllDirs)
 				{
-					if (!string.IsNullOrEmpty(IOSBuildFilter))
+					IEnumerable<IOSBuild> FoundBuilds = IOSBuild.CreateFromPath(InProjectName, Di.FullName);
+
+					if (FoundBuilds != null)
 					{
-						//IndexOf used because Contains must be case-sensitive
-						FoundBuilds = FoundBuilds.Where(B => B.SourceIPAPath.IndexOf(IOSBuildFilter, StringComparison.OrdinalIgnoreCase) >= 0);
+						if (!string.IsNullOrEmpty(IOSBuildFilter))
+						{
+							//IndexOf used because Contains must be case-sensitive
+							FoundBuilds = FoundBuilds.Where(B => B.SourceIPAPath.IndexOf(IOSBuildFilter, StringComparison.OrdinalIgnoreCase) >= 0);
+						}
+						Builds.AddRange(FoundBuilds);
 					}
-					Builds.AddRange(FoundBuilds);
 				}
 			}
 
diff --git a/Engine/Source/Programs/AutomationTool/Gauntlet/Platform/Lumin/Gauntlet.LuminBuildSource.cs b/Engine/Source/Programs/AutomationTool/Gauntlet/Platform/Lumin/Gauntlet.LuminBuildSource.cs
--- a/Engine/Source/Programs/AutomationTool/Gauntlet/Platform/Lumin/Gauntlet.LuminBuildSource.cs
+++ b/Engine/Source/Programs/AutomationTool/Gauntlet/Platform/Lumin/Gauntlet.LuminBuildSource.cs
@@ -48,58 +48,61 @@ namespace Gauntlet
 
 			DirectoryInfo Di = new DirectoryInfo(BuildPath);
 
-			// find all install batchfiles
-			FileInfo[] InstallFiles = Di.GetFiles("Install_*");
-
-			foreach (FileInfo Fi in InstallFiles)
+			if (Di.Exists)
 			{
-				Log.Verbose("Pulling install data from {0}", Fi.FullName);
+				// find all install batchfiles
+				FileInfo[] InstallFiles = Di.GetFiles("Install_*");
 
-				string AbsPath = Fi.Directory.FullName;
+				foreach (FileInfo Fi in InstallFiles)
+				{
+					Log.Verbose("Pulling install data from {0}", Fi.FullName);
 
-				// read contents and replace linefeeds (regex doesn't stop on them :((
-				string BatContents = File.ReadAllText(Fi.FullName).Replace(Environment.NewLine, "\n");
+					string AbsPath = Fi.Directory.FullName;
 
-				// Replace .bat with .mpk and strip up to and including the first _, that is then our MPK name
-				string SourceMpkPath = Regex.Replace(Fi.Name, ".bat", ".mpk", RegexOptions.IgnoreCase);
-				SourceMpkPath = SourceMpkPath.Substring(SourceMpkPath.IndexOf("_") + 1);
-				SourceMpkPath = Path.Combine(AbsPath, SourceMpkPath);
+					// read contents and replace linefeeds (regex doesn't stop on them :((
+					string BatContents = File.ReadAllText(Fi.FullName).Replace(Environment.NewLine, "\n");
 
-				Match Info = Regex.Match(BatContents, @"install\s+(-u+)\s""%~dp0\\(.+)""");
-				string LuminPackageName = Info.Groups[2].ToString();
-	
-				if (string.IsNullOrEmpty(SourceMpkPath))
-				{
-					Log.Warning("No MPK found for build at {0}", Fi.FullName);
-					continue;
-				}
+					// Replace .bat with .mpk and strip up to and including the first _, that is then our MPK name
+					string SourceMpkPath = Regex.Replace(Fi.Name, ".bat", ".mpk", RegexOptions.IgnoreCase);
+					SourceMpkPath = SourceMpkPath.Substring(SourceMpkPath.IndexOf("_") + 1);
+					SourceMpkPath = Path.Combine(AbsPath, SourceMpkPath);
 
-				if (string.IsNullOrEmpty(LuminPackageName))
-				{
-					Log.Warning("No product name found for build at {0}", Fi.FullName);
-					continue;
-				}
+					Match Info = Regex.Match(BatContents, @"install\s+(-u+)\s""%~dp0\\(.+)""");
+					string LuminPackageName = Info.Groups[2].ToString();
 
-				UnrealTargetConfiguration UnrealConfig = UnrealHelpers.GetConfigurationFromExecutableName(InProjectName, Fi.Name);
+					if (string.IsNullOrEmpty(SourceMpkPath))
+					{
+						Log.Warning("No MPK found for build at {0}", Fi.FullName);
+						continue;
+					}
 
-				// Lumin builds are always packaged, and we can always replace the command line
-				BuildFlags Flags = BuildFlags.Packaged | BuildFlags.CanReplaceCommandLine;
-				
-				if (AbsPath.Contains("Bulk"))
-				{
-					Flags |= BuildFlags.Bulk;
-				}
-				else
-				{
-					Flags |= BuildFlags.NotBulk;
-				}
+					if (string.IsNullOrEmpty(LuminPackageName))
+					{
+						Log.Warning("No product name found for build at {0}", Fi.FullName);
+						continue;
+					}
+
+					UnrealTargetConfiguration UnrealConfig = UnrealHelpers.GetConfigurationFromExecutableName(InProjectName, Fi.Name);
 
-				LuminBuild NewBuild = new LuminBuild(UnrealConfig, LuminPackageName, SourceMpkPath, Flags);
+					// Lumin builds are always packaged, and we can always replace the command line
+					BuildFlags Flags = BuildFlags.Packaged | BuildFlags.CanReplaceCommandLine;
 
-				DiscoveredBuilds.Add(NewBuild);
+					if (AbsPath.Contains("Bulk"))
+					{
+						Flags |= BuildFlags.Bulk;
+					}
+					else
+					{
+						Flags |= BuildFlags.NotBulk;
+					}
 
-				Log.Verbose("Found {0} {1} build at {2}", UnrealConfig, ((Flags & BuildFlags.Bulk) == BuildFlags.Bulk) ? "(bulk)" : "(not bulk)", AbsPath);
+					LuminBuild NewBuild = new LuminBuild(UnrealConfig, LuminPackageName, SourceMpkPath, Flags);
 
+					DiscoveredBuilds.Add(NewBuild);
+
+					Log.Verbose("Found {0} {1} build at {2}", UnrealConfig, ((Flags & BuildFlags.Bulk) == BuildFlags.Bulk) ? "(bulk)" : "(not bulk)", AbsPath);
+
+				}
 			}
 
 			return DiscoveredBuilds;
@@ -125,29 +128,30 @@ namespace Gauntlet
 		{		
 			List<IBuild> Builds = new List<IBuild>();
 
-			// First try creating a build just from this path
-			IEnumerable<LuminBuild> FoundBuilds = LuminBuild.CreateFromPath(InProjectName, InPath);
-
-			// If we found one we were pointed at a specific case so we're done
-			if (FoundBuilds.Any())
+			if (Directory.Exists(InPath))
 			{
-				Builds.AddRange(FoundBuilds);
-			}
-			else
-			{
-				// Ok, find all directories that match our platform and add those
-				DirectoryInfo TopDir = new DirectoryInfo(InPath);
-
-				// find all directories that begin with Lumin
-				IEnumerable<string> LuminDirs = Directory.EnumerateDirectories(InPath, "Lumin*", SearchOption.AllDirectories);
+				// First try creating a build just from this path
+				IEnumerable<LuminBuild> FoundBuilds = LuminBuild.CreateFromPath(InProjectName, InPath);
 
-				foreach (string DirPath in LuminDirs)
+				// If we found one we were pointed at a specific case so we're done
+				if (FoundBuilds.Any())
 				{
-					
-
-					FoundBuilds = LuminBuild.CreateFromPath(InProjectName, DirPath);
 					Builds.AddRange(FoundBuilds);
 				}
+				else
+				{
+					// Ok, find all directories that match our platform and add those
+					DirectoryInfo TopDir = new DirectoryInfo(InPath);
+
+					// find all directories that begin with Lumin
+					IEnumerable<string> LuminDirs = Directory.EnumerateDirectories(InPath, "Lumin*", SearchOption.AllDirectories);
+
+					foreach (string DirPath in LuminDirs)
+					{
+						FoundBuilds = LuminBuild.CreateFromPath(InProjectName, DirPath);
+						Builds.AddRange(FoundBuilds);
+					}
+				}
 			}
 
 			return Builds;
diff --git a/Engine/Source/Programs/AutomationTool/Gauntlet/Unreal/BuildSource/Gauntlet.StagedBuildSource.cs b/Engine/Source/Programs/AutomationTool/Gauntlet/Unreal/BuildSource/Gauntlet.StagedBuildSource.cs
--- a/Engine/Source/Programs/AutomationTool/Gauntlet/Unreal/BuildSource/Gauntlet.StagedBuildSource.cs
+++ b/Engine/Source/Programs/AutomationTool/Gauntlet/Unreal/BuildSource/Gauntlet.StagedBuildSource.cs
@@ -253,55 +253,58 @@ namespace Gauntlet
 		{
 			List<DirectoryInfo> AllDirs = new List<DirectoryInfo>();
 
+			List<IBuild> Builds = new List<IBuild>();
+
 			// c:\path\to\build
 			DirectoryInfo PathDI = new DirectoryInfo(InPath);
 
-			// If the path is directly to a platform folder, add us
-			if (PathDI.Name.IndexOf(PlatformFolderPrefix, StringComparison.OrdinalIgnoreCase) >= 0)
+			if (PathDI.Exists)
 			{
-				AllDirs.Add(PathDI);
-			}
-
-			// Assume it's a folder of all build types so find all sub directories that begin with the foldername for this platform
-			IEnumerable<DirectoryInfo> MatchingDirs = PathDI.GetDirectories("*", SearchOption.TopDirectoryOnly);
+				// If the path is directly to a platform folder, add us
+				if (PathDI.Name.IndexOf(PlatformFolderPrefix, StringComparison.OrdinalIgnoreCase) >= 0)
+				{
+					AllDirs.Add(PathDI);
+				}
 
-			MatchingDirs = MatchingDirs.Where(D => D.Name.StartsWith(PlatformFolderPrefix, StringComparison.OrdinalIgnoreCase)).ToArray();
+				// Assume it's a folder of all build types so find all sub directories that begin with the foldername for this platform
+				IEnumerable<DirectoryInfo> MatchingDirs = PathDI.GetDirectories("*", SearchOption.TopDirectoryOnly);
 
-			AllDirs.AddRange(MatchingDirs);
+				MatchingDirs = MatchingDirs.Where(D => D.Name.StartsWith(PlatformFolderPrefix, StringComparison.OrdinalIgnoreCase)).ToArray();
 
-			List<DirectoryInfo> DirsToRecurse = AllDirs;
+				AllDirs.AddRange(MatchingDirs);
 
-			// now get subdirs
-			while (MaxRecursion-- > 0)
-			{
-				List<DirectoryInfo> DiscoveredDirs = new List<DirectoryInfo>();
+				List<DirectoryInfo> DirsToRecurse = AllDirs;
 
-				DirsToRecurse.ToList().ForEach((D) =>
+				// now get subdirs
+				while (MaxRecursion-- > 0)
 				{
-					DiscoveredDirs.AddRange(D.GetDirectories("*", SearchOption.TopDirectoryOnly));
-				});
-
-				AllDirs.AddRange(DiscoveredDirs);
-				DirsToRecurse = DiscoveredDirs;
-			}
+					List<DirectoryInfo> DiscoveredDirs = new List<DirectoryInfo>();
 
-			// every directory that contains a valid build should have at least a ProjectName folder
-			AllDirs = AllDirs.Where(D =>
-			{
-				var SubDirs = D.GetDirectories();
-				return SubDirs.Any(SD => SD.Name.Equals(InProjectName, StringComparison.OrdinalIgnoreCase));
+					DirsToRecurse.ToList().ForEach((D) =>
+					{
+						DiscoveredDirs.AddRange(D.GetDirectories("*", SearchOption.TopDirectoryOnly));
+					});
 
-			}).ToList();
+					AllDirs.AddRange(DiscoveredDirs);
+					DirsToRecurse = DiscoveredDirs;
+				}
 
-			List<IBuild> Builds = new List<IBuild>();
+				// every directory that contains a valid build should have at least a ProjectName folder
+				AllDirs = AllDirs.Where(D =>
+				{
+					var SubDirs = D.GetDirectories();
+					return SubDirs.Any(SD => SD.Name.Equals(InProjectName, StringComparison.OrdinalIgnoreCase));
 
-			foreach (DirectoryInfo Di in AllDirs)
-			{
-				IEnumerable<IBuild> FoundBuilds = StagedBuild.CreateFromPath<T>(Platform, InProjectName, Di.FullName, ExecutableExtension);
+				}).ToList();
 
-				if (FoundBuilds != null)
+				foreach (DirectoryInfo Di in AllDirs)
 				{
-					Builds.AddRange(FoundBuilds);
+					IEnumerable<IBuild> FoundBuilds = StagedBuild.CreateFromPath<T>(Platform, InProjectName, Di.FullName, ExecutableExtension);
+
+					if (FoundBuilds != null)
+					{
+						Builds.AddRange(FoundBuilds);
+					}
 				}
 			}
 
