Allow AAB and universal apks to be archived

#ue4
#android

diff --git a/Engine/Source/Programs/AutomationTool/Android/AndroidPlatform.Automation.cs b/Engine/Source/Programs/AutomationTool/Android/AndroidPlatform.Automation.cs
--- a/Engine/Source/Programs/AutomationTool/Android/AndroidPlatform.Automation.cs
+++ b/Engine/Source/Programs/AutomationTool/Android/AndroidPlatform.Automation.cs
@@ -966,15 +966,39 @@ public class AndroidPlatform : Platform
 					}
 				}
 
+				string APKDirectory = Path.GetDirectoryName(ApkName);
+				string APKNameWithoutExtension = Path.GetFileNameWithoutExtension(ApkName);
+
 				// copy optional unprotected APK if exists
-				string UnprotectedApkName = Path.Combine(Path.GetDirectoryName(ApkName), "unprotected_" + Path.GetFileName(ApkName));
+				string UnprotectedApkName = Path.Combine(APKDirectory, "unprotected_" + APKNameWithoutExtension + ".apk");
 				if (FileExists(UnprotectedApkName))
 				{
-					SC.ArchiveFiles(Path.GetDirectoryName(UnprotectedApkName), Path.GetFileName(UnprotectedApkName));
+					SC.ArchiveFiles(APKDirectory, Path.GetFileName(UnprotectedApkName));
+				}
+
+				// copy optional app bundle if exists
+				string AppBundleName = Path.Combine(APKDirectory, APKNameWithoutExtension + ".aab");
+				if (FileExists(AppBundleName))
+				{
+					SC.ArchiveFiles(APKDirectory, Path.GetFileName(AppBundleName));
+				}
+
+				// copy optional apks (zip of split apks) if exists
+				string APKSName = Path.Combine(APKDirectory, APKNameWithoutExtension + ".apks");
+				if (FileExists(APKSName))
+				{
+					SC.ArchiveFiles(APKDirectory, Path.GetFileName(APKSName));
+				}
+
+				// copy optional universal apk if exists
+				string UniversalApkName = Path.Combine(APKDirectory, APKNameWithoutExtension + "_universal.apk");
+				if (FileExists(UniversalApkName))
+				{
+					SC.ArchiveFiles(APKDirectory, Path.GetFileName(UniversalApkName));
 				}
 
 				// copy optional logs directory if exists
-				string LogsDirName = Path.Combine(Path.GetDirectoryName(ApkName), Path.GetFileName(ApkName) + ".logs");
+				string LogsDirName = Path.Combine(APKDirectory, APKNameWithoutExtension + ".logs");
 				if (DirectoryExists(LogsDirName))
 				{
 					SC.ArchiveFiles(LogsDirName);
