Add support for disabling secondary order in UnrealPak to reduce fragmentation in platforms that use delta patching

diff --git a/Engine/Config/BaseGame.ini b/Engine/Config/BaseGame.ini
--- a/Engine/Config/BaseGame.ini
+++ b/Engine/Config/BaseGame.ini
@@ -75,6 +75,7 @@ bSkipEditorContent=false
 bSharedMaterialNativeLibraries=False
 bShareMaterialShaderCode=False
 bSkipMovies=False
+bPakUsesSecondaryOrder=True
 
 +EarlyDownloaderPakFileFiles=...\Content\Internationalization\...\*.icu
 +EarlyDownloaderPakFileFiles=...\Content\Internationalization\...\*.brk
diff --git a/Engine/Source/Programs/AutomationTool/Scripts/CopyBuildToStagingDirectory.Automation.cs b/Engine/Source/Programs/AutomationTool/Scripts/CopyBuildToStagingDirectory.Automation.cs
--- a/Engine/Source/Programs/AutomationTool/Scripts/CopyBuildToStagingDirectory.Automation.cs
+++ b/Engine/Source/Programs/AutomationTool/Scripts/CopyBuildToStagingDirectory.Automation.cs
@@ -2067,9 +2067,17 @@ public partial class Project : CommandUtils
 				if (!bCopiedExistingPak)
 				{
 					FileReference PrimaryOrderFile   = (PakOrderFileLocations.Count >= 1) ? PakOrderFileLocations[0] : null;
-					FileReference SecondaryOrderFile = (PakOrderFileLocations.Count >= 2) ? PakOrderFileLocations[1] : null;
+					FileReference SecondaryOrderFile = null;
 
-                    string BulkOption = "";
+					// Add a secondary order if there is one specified
+					bool bUseSecondaryOrder = false;
+					PlatformGameConfig.GetBool("/Script/UnrealEd.ProjectPackagingSettings", "bPakUsesSecondaryOrder", out bUseSecondaryOrder);
+					if (bUseSecondaryOrder && PakOrderFileLocations.Count >= 2)
+					{
+						SecondaryOrderFile = PakOrderFileLocations[1];
+					}
+
+					string BulkOption = "";
                     {
                         ConfigHierarchy PlatformEngineConfig;
                         if (Params.EngineConfigs.TryGetValue(SC.StageTargetPlatform.PlatformType, out PlatformEngineConfig))
