Gauntlet android build searching now correctly parses install directories.
Ben.Salem,Josh.Engebretson

diff --git a/Engine/Source/Programs/AutomationTool/Gauntlet/Platform/Android/Gauntlet.AndroidBuildSource.cs b/Engine/Source/Programs/AutomationTool/Gauntlet/Platform/Android/Gauntlet.AndroidBuildSource.cs
--- a/Engine/Source/Programs/AutomationTool/Gauntlet/Platform/Android/Gauntlet.AndroidBuildSource.cs
+++ b/Engine/Source/Programs/AutomationTool/Gauntlet/Platform/Android/Gauntlet.AndroidBuildSource.cs
@@ -102,10 +102,14 @@ namespace Gauntlet
 				string BatContents = File.ReadAllText(Fi.FullName).Replace(Environment.NewLine, "\n");
 
 				// Replace .bat with .apk and strip up to and including the first _, that is then our APK name
-				string SourceApkPath = Regex.Replace(Fi.Name, ".bat", ".apk", RegexOptions.IgnoreCase);
-				SourceApkPath = SourceApkPath.Substring(SourceApkPath.IndexOf("_") + 1);
-				SourceApkPath = Path.Combine(AbsPath, SourceApkPath);
-
+				var SourceApkMatch = Regex.Match(BatContents, @" install\s+(.+\.apk)");
+				if ( SourceApkMatch.Groups.Count <= 0)
+				{
+					Log.Warning("Could not parse install command from {0}", Fi.FullName);
+					continue;
+				}
+				string SourceApkPath = Path.Combine(AbsPath,SourceApkMatch.Groups[1].ToString());
+				
 				// save com.companyname.product
 				string AndroidPackageName = Regex.Match(BatContents, @"uninstall\s+(com\..+)").Groups[1].ToString();
 
@@ -121,6 +125,12 @@ namespace Gauntlet
 					continue;
 				}
 
+				if (!File.Exists(SourceApkPath))
+				{
+					Log.Warning("Resolved APK name but it doesn't exist {0}", SourceApkPath);
+					continue;
+				}
+
 				if (string.IsNullOrEmpty(AndroidPackageName))
 				{
 					Log.Warning("No product name found for build at {0}", Fi.FullName);
