Filter some files out of long file path warning
#rnx

diff --git a/Engine/Source/Programs/AutomationTool/Gauntlet/Framework/Gauntlet.Utils.cs b/Engine/Source/Programs/AutomationTool/Gauntlet/Framework/Gauntlet.Utils.cs
--- a/Engine/Source/Programs/AutomationTool/Gauntlet/Framework/Gauntlet.Utils.cs
+++ b/Engine/Source/Programs/AutomationTool/Gauntlet/Framework/Gauntlet.Utils.cs
@@ -1145,7 +1145,18 @@ namespace Gauntlet
 
 									// Break out of loop with warning
 									Copied = true;
-									Log.Warning("Long path file copy failed with {0}.  Please verify that this file is not required.", ex.Message, LongPathMessage);
+
+									// Filter out some known unneeded files which can cause this warning, and log the message instead
+									string[] Blacklist = new string[]{ "UE4CC-XboxOne" };
+									string Message = string.Format("Long path file copy failed with {0}.  Please verify that this file is not required.", ex.Message);
+									if ( Blacklist.FirstOrDefault(B => { return SourcePath.IndexOf(B, StringComparison.OrdinalIgnoreCase) >= 0; }) == null)
+									{
+										Log.Warning(Message); 
+									}
+									else
+									{
+										Log.Info(Message);
+									}
 								}
 							}
 						}
