Enhacement JIRA UETOOL-1721: Removed auto-testing warnings.

diff --git a/Engine/Source/Programs/AutomationTool/Scripts/CopyBuildToStagingDirectory.Automation.cs b/Engine/Source/Programs/AutomationTool/Scripts/CopyBuildToStagingDirectory.Automation.cs
--- a/Engine/Source/Programs/AutomationTool/Scripts/CopyBuildToStagingDirectory.Automation.cs
+++ b/Engine/Source/Programs/AutomationTool/Scripts/CopyBuildToStagingDirectory.Automation.cs
@@ -1057,6 +1057,11 @@ public partial class Project : CommandUtils
 		}
 		else
 		{
+			if (NormalizedPath.StartsWith("layouts/"))
+			{
+				return true;
+			}
+
 			if (NormalizedPath.StartsWith("localization/"))
 			{
 				return false;
