Fix assets created with 4.24.1 builds not being compatible with 4.24.0 release.

CompatibleChangelist number was being overridden as part of the build process due to CL 9823639. The desired effect here was to ensure that Epic's compatible changelist number was cleared out for licensee builds; instead of just forcing it to be set, we clear it out if we're changing the licensee flag.

diff --git a/Engine/Source/Programs/AutomationTool/AutomationUtils/UE4Build.cs b/Engine/Source/Programs/AutomationTool/AutomationUtils/UE4Build.cs
--- a/Engine/Source/Programs/AutomationTool/AutomationUtils/UE4Build.cs
+++ b/Engine/Source/Programs/AutomationTool/AutomationUtils/UE4Build.cs
@@ -317,7 +317,7 @@ namespace AutomationTool
 			{
 				ChangelistNumber = ChangelistNumberOverride.HasValue? ChangelistNumberOverride.Value : CommandUtils.P4Env.Changelist;
 			}
-			int CompatibleChangelistNumber = ChangelistNumber;
+			int CompatibleChangelistNumber = 0;
 			if(bDoUpdateVersionFiles && CompatibleChangelistNumberOverride.HasValue)
 			{
 				CompatibleChangelistNumber = CompatibleChangelistNumberOverride.Value;
@@ -365,6 +365,10 @@ namespace AutomationTool
 					{
 						Version.CompatibleChangelist = CompatibleChangelistNumber;
 					}
+					else if(bIsLicenseeVersion != Version.IsLicenseeVersion)
+					{
+						Version.CompatibleChangelist = 0; // Clear out the compatible changelist number; it corresponds to a different P4 server.
+					}
 					Version.IsLicenseeVersion = bIsLicenseeVersion;
 					Version.IsPromotedBuild = bIsPromotedBuild;
 					Version.BranchName = Branch;
diff --git a/Engine/Source/Programs/AutomationTool/BuildGraph/Tasks/SetVersionTask.cs b/Engine/Source/Programs/AutomationTool/BuildGraph/Tasks/SetVersionTask.cs
--- a/Engine/Source/Programs/AutomationTool/BuildGraph/Tasks/SetVersionTask.cs
+++ b/Engine/Source/Programs/AutomationTool/BuildGraph/Tasks/SetVersionTask.cs
@@ -94,8 +94,7 @@ namespace AutomationTool.Tasks
 		public override void Execute(JobContext Job, HashSet<FileReference> BuildProducts, Dictionary<string, HashSet<FileReference>> TagNameToFileSet)
 		{
 			// Update the version files
-			int CompatibleChange = (Parameters.CompatibleChange == 0) ? Parameters.Change : Parameters.CompatibleChange;
-			List<FileReference> VersionFiles = UE4Build.StaticUpdateVersionFiles(Parameters.Change, CompatibleChange, Parameters.Branch, Parameters.Build, Parameters.Licensee, Parameters.Promoted, !Parameters.SkipWrite);
+			List<FileReference> VersionFiles = UE4Build.StaticUpdateVersionFiles(Parameters.Change, Parameters.CompatibleChange, Parameters.Branch, Parameters.Build, Parameters.Licensee, Parameters.Promoted, !Parameters.SkipWrite);
 
 			// Apply the optional tag to them
 			foreach(string TagName in FindTagNamesFromList(Parameters.Tag))
