Added support for splitting localization data into separate PAK chunks during cook

There is a new packaging setting, LocalizationTargetsToChunk, which lets you specify which of your localization targets should be chunked during cook.

If enabled, any localization entries corresponding to assets will be split into a separate LocRes file for the chunk containing that asset, and will be removed from any chunks (including the primary chunk) that do not contain that asset. This is useful to restrict localized text to the chunk that contains the corresponding asset for the text, and allows localization data to be selectively encrypted based on the chunking rules.

At runtime this chunked data will be loaded into the text localization manager automatically when the chunk is loaded.

diff --git a/Engine/Source/Editor/UnrealEd/Classes/Settings/ProjectPackagingSettings.h b/Engine/Source/Editor/UnrealEd/Classes/Settings/ProjectPackagingSettings.h
--- a/Engine/Source/Editor/UnrealEd/Classes/Settings/ProjectPackagingSettings.h
+++ b/Engine/Source/Editor/UnrealEd/Classes/Settings/ProjectPackagingSettings.h
@@ -281,6 +281,10 @@ public:
 	UPROPERTY(config, EditAnywhere, Category=Packaging, AdvancedDisplay, meta=(DisplayName="Localizations to Package"))
 	TArray<FString> CulturesToStage;
 
+	/** List of localization targets that should be chunked during cooking (if using chunks) */
+	UPROPERTY(config, EditAnywhere, Category=Packaging, AdvancedDisplay)
+	TArray<FString> LocalizationTargetsToChunk;
+
 	/**
 	 * Cook all things in the project content directory
 	 */
diff --git a/Engine/Source/Editor/UnrealEd/Private/Commandlets/AssetRegistryGenerator.cpp b/Engine/Source/Editor/UnrealEd/Private/Commandlets/AssetRegistryGenerator.cpp
--- a/Engine/Source/Editor/UnrealEd/Private/Commandlets/AssetRegistryGenerator.cpp
+++ b/Engine/Source/Editor/UnrealEd/Private/Commandlets/AssetRegistryGenerator.cpp
@@ -19,6 +19,7 @@
 #include "Interfaces/ITargetPlatform.h"
 #include "AssetRegistryModule.h"
 #include "GameDelegates.h"
+#include "Commandlets/IChunkDataGenerator.h"
 #include "Commandlets/ChunkDependencyInfo.h"
 #include "IPlatformFileSandboxWrapper.h"
 #include "Misc/ConfigCacheIni.h"
@@ -366,22 +367,24 @@ bool FAssetRegistryGenerator::GenerateStreamingInstallManifest(int64 InExtraFlav
 	// generate per-chunk pak list files
 	for (int32 PakchunkIndex = 0; PakchunkIndex < FinalChunkManifests.Num(); ++PakchunkIndex)
 	{
+		const FChunkPackageSet* Manifest = FinalChunkManifests[PakchunkIndex];
+
 		// Serialize chunk layers whether chunk is empty or not
 		int32 TargetLayer = 0;
-		FGameDelegates::Get().GetAssignLayerChunkDelegate().ExecuteIfBound(FinalChunkManifests[PakchunkIndex], Platform, PakchunkIndex, TargetLayer);
+		FGameDelegates::Get().GetAssignLayerChunkDelegate().ExecuteIfBound(Manifest, Platform, PakchunkIndex, TargetLayer);
 
 		FString LayerString = FString::Printf(TEXT("%d\r\n"), TargetLayer);
 		ChunkLayerFile->Serialize(TCHAR_TO_ANSI(*LayerString), LayerString.Len());
 
 		// Is this chunk empty?
-		if (!FinalChunkManifests[PakchunkIndex])
+		if (!Manifest)
 		{
 			continue;
 		}
 
 		int32 FilenameIndex = 0;
 		TArray<FString> ChunkFilenames;
-		FinalChunkManifests[PakchunkIndex]->GenerateValueArray(ChunkFilenames);
+		Manifest->GenerateValueArray(ChunkFilenames);
 		bool bFinishedAllFiles = false;
 		for (int32 SubChunkIndex = 0; !bFinishedAllFiles; ++SubChunkIndex)
 		{
@@ -430,6 +433,22 @@ bool FAssetRegistryGenerator::GenerateStreamingInstallManifest(int64 InExtraFlav
 				}
 			}
 
+			// Allow the extra data generation steps to run and add their output to the manifest
+			if (ChunkDataGenerators.Num() > 0)
+			{
+				TSet<FName> PackagesInChunk;
+				PackagesInChunk.Reserve(Manifest->Num());
+				for (const auto& ChunkManifestPair : *Manifest)
+				{
+					PackagesInChunk.Add(ChunkManifestPair.Key);
+				}
+
+				for (const TSharedRef<IChunkDataGenerator>& ChunkDataGenerator : ChunkDataGenerators)
+				{
+					ChunkDataGenerator->GenerateChunkDataFiles(PakchunkIndex, PackagesInChunk, Platform, InSandboxFile, ChunkFilenames);
+				}
+			}
+
 			if (bUseAssetManager && SubChunkIndex == 0)
 			{
 				// Sort so the order is consistent. If load order is important then it should be specified as a load order file to UnrealPak
@@ -1060,6 +1079,11 @@ void FAssetRegistryGenerator::BuildChunkManifest(const TSet<FName>& InCookedPack
 
 }
 
+void FAssetRegistryGenerator::RegisterChunkDataGenerator(TSharedRef<IChunkDataGenerator> InChunkDataGenerator)
+{
+	ChunkDataGenerators.Add(MoveTemp(InChunkDataGenerator));
+}
+
 void FAssetRegistryGenerator::PreSave(const TSet<FName>& InCookedPackages)
 {
 	if (bUseAssetManager)
diff --git a/Engine/Source/Editor/UnrealEd/Private/CookOnTheFlyServer.cpp b/Engine/Source/Editor/UnrealEd/Private/CookOnTheFlyServer.cpp
--- a/Engine/Source/Editor/UnrealEd/Private/CookOnTheFlyServer.cpp
+++ b/Engine/Source/Editor/UnrealEd/Private/CookOnTheFlyServer.cpp
@@ -48,6 +48,7 @@
 #include "Editor.h"
 #include "UnrealEdGlobals.h"
 #include "FileServerMessages.h"
+#include "LocalizationChunkDataGenerator.h"
 #include "Internationalization/Culture.h"
 #include "Serialization/ArrayReader.h"
 #include "Serialization/ArrayWriter.h"
@@ -6968,6 +6969,8 @@ void UCookOnTheFlyServer::StartCookByTheBook( const FCookByTheBookStartupOptions
 		FCoreUObjectDelegates::PackageCreatedForLoad.AddUObject(this, &UCookOnTheFlyServer::MaybeMarkPackageAsAlreadyLoaded);
 	}
 
+	const UProjectPackagingSettings* const PackagingSettings = GetDefault<UProjectPackagingSettings>();
+
 	// Find all the localized packages and map them back to their source package
 	{
 		UE_LOG(LogCook, Display, TEXT("Discovering localized assets"));
@@ -7010,9 +7013,29 @@ void UCookOnTheFlyServer::StartCookByTheBook( const FCookByTheBookStartupOptions
 			TArray<FName>& LocalizedPackageNames = CookByTheBookOptions->SourceToLocalizedPackageVariants.FindOrAdd(SourcePackageName);
 			LocalizedPackageNames.AddUnique(LocalizedPackageName);
 		}
-	}
 
-	const UProjectPackagingSettings* const PackagingSettings = GetDefault<UProjectPackagingSettings>();
+		// Get the list of localization targets to chunk, and remove any targets that we've been asked not to stage
+		TArray<FString> LocalizationTargetsToChunk = PackagingSettings->LocalizationTargetsToChunk;
+		{
+			TArray<FString> BlacklistLocalizationTargets;
+			GConfig->GetArray(TEXT("Staging"), TEXT("BlacklistLocalizationTargets"), BlacklistLocalizationTargets, GGameIni);
+			if (BlacklistLocalizationTargets.Num() > 0)
+			{
+				LocalizationTargetsToChunk.RemoveAll([&BlacklistLocalizationTargets](const FString& InLocalizationTarget)
+				{
+					return BlacklistLocalizationTargets.Contains(InLocalizationTarget);
+				});
+			}
+		}
+
+		if (LocalizationTargetsToChunk.Num() > 0 && AllCulturesToCook.Num() > 0)
+		{
+			for (TPair<FName, FAssetRegistryGenerator*>& Pair : RegistryGenerators)
+			{
+				Pair.Value->RegisterChunkDataGenerator(MakeShared<FLocalizationChunkDataGenerator>(MoveTemp(LocalizationTargetsToChunk), MoveTemp(AllCulturesToCook)));
+			}
+		}
+	}
 
 	PackageTracker->NeverCookPackageList.Empty();
 	{
diff --git a/Engine/Source/Editor/UnrealEd/Private/LocalizationChunkDataGenerator.cpp b/Engine/Source/Editor/UnrealEd/Private/LocalizationChunkDataGenerator.cpp
new file mode 100644
--- /dev/null
+++ b/Engine/Source/Editor/UnrealEd/Private/LocalizationChunkDataGenerator.cpp
@@ -0,0 +1,211 @@
+// Copyright 1998-2019 Epic Games, Inc. All Rights Reserved.
+
+#include "LocalizationChunkDataGenerator.h"
+#include "LocTextHelper.h"
+#include "TextLocalizationResourceGenerator.h"
+#include "Internationalization/TextLocalizationResource.h"
+
+#include "Misc/Paths.h"
+#include "Misc/PackageName.h"
+#include "IPlatformFileSandboxWrapper.h"
+
+DEFINE_LOG_CATEGORY_STATIC(LogLocalizationChunkDataGenerator, Log, All);
+
+FLocalizationChunkDataGenerator::FLocalizationChunkDataGenerator(TArray<FString> InLocalizationTargetsToChunk, TArray<FString> InAllCulturesToCook)
+	: LocalizationTargetsToChunk(MoveTemp(InLocalizationTargetsToChunk))
+	, AllCulturesToCook(MoveTemp(InAllCulturesToCook))
+{
+}
+
+void FLocalizationChunkDataGenerator::GenerateChunkDataFiles(const int32 InChunkId, const TSet<FName>& InPackagesInChunk, const FString& InPlatformName, FSandboxPlatformFile* InSandboxFile, TArray<FString>& OutChunkFilenames)
+{
+	// We can skip this chunk if it's empty
+	if (InPackagesInChunk.Num() == 0)
+	{
+		return;
+	}
+
+	ConditionalCacheLocalizationTargetData();
+
+	// We can skip this if we're not actually chunking any localization data
+	if (CachedLocalizationTargetHelpers.Num() == 0)
+	{
+		return;
+	}
+
+	for (const TSharedPtr<FLocTextHelper>& SourceLocTextHelper : CachedLocalizationTargetHelpers)
+	{
+		// If this target has no helper, then it was either an invalid target or failed to load when caching - skip it here
+		if (!SourceLocTextHelper)
+		{
+			continue;
+		}
+
+		// Chunk 0 is the only chunk that can contain non-asset localization data, as it acts as the "catch-all" since it's always available
+		// It is also the only chunk that doesn't gain a suffix to make it unique (as it is replacing the offline localization data that is usually staged verbatim)
+		const bool bIsPrimaryChunk = InChunkId == 0;
+
+		// Prepare to produce the localization target for this chunk
+		const TArray<FString> AvailableCulturesToCook = SourceLocTextHelper->GetAllCultures();
+		const FString ChunkTargetName = TextLocalizationResourceUtil::GetLocalizationTargetNameForChunkId(SourceLocTextHelper->GetTargetName(), InChunkId);
+		const FString ChunkTargetRoot = (InSandboxFile->GetSandboxDirectory() / InSandboxFile->GetGameSandboxDirectoryName() / TEXT("Content") / TEXT("Localization") / ChunkTargetName).Replace(TEXT("[Platform]"), *InPlatformName);
+
+		// Produce a filtered set of data that can be used to produce the LocRes for each chunk
+		bool bChunkHasText = false;
+		FLocTextHelper ChunkLocTextHelper(ChunkTargetRoot, FString::Printf(TEXT("%s.manifest"), *ChunkTargetName), FString::Printf(TEXT("%s.archive"), *ChunkTargetName), SourceLocTextHelper->GetNativeCulture(), SourceLocTextHelper->GetForeignCultures(), nullptr);
+		ChunkLocTextHelper.LoadAll(ELocTextHelperLoadFlags::Create); // Create the in-memory manifest and archives
+		SourceLocTextHelper->EnumerateSourceTexts([bIsPrimaryChunk, SourceLocTextHelper, &bChunkHasText, &ChunkLocTextHelper, &InPackagesInChunk, &AvailableCulturesToCook](TSharedRef<FManifestEntry> InManifestEntry)
+		{
+			for (const FManifestContext& ManifestContext : InManifestEntry->Contexts)
+			{
+				bool bIncludeInChunk = false;
+				if (FPackageName::IsValidObjectPath(ManifestContext.SourceLocation))
+				{
+					const FName SourcePackageName = *FPackageName::ObjectPathToPackageName(ManifestContext.SourceLocation);
+					bIncludeInChunk = InPackagesInChunk.Contains(SourcePackageName);
+				}
+				else
+				{
+					bIncludeInChunk = bIsPrimaryChunk;
+				}
+
+				if (bIncludeInChunk)
+				{
+					bChunkHasText = true;
+					ChunkLocTextHelper.AddSourceText(InManifestEntry->Namespace, InManifestEntry->Source, ManifestContext);
+					for (const FString& CultureToCook : AvailableCulturesToCook)
+					{
+						TSharedPtr<FArchiveEntry> SourceTranslationEntry = SourceLocTextHelper->FindTranslation(CultureToCook, InManifestEntry->Namespace, ManifestContext.Key, ManifestContext.KeyMetadataObj);
+						if (SourceTranslationEntry)
+						{
+							ChunkLocTextHelper.AddTranslation(CultureToCook, SourceTranslationEntry.ToSharedRef());
+						}
+					}
+				}
+			}
+
+			return true; // continue enumeration
+		}, false);
+
+		// If this chunk has no localization data then we can skip generating the LocRes (unless it's the primary chunk)
+		if (!bIsPrimaryChunk && !bChunkHasText)
+		{
+			continue;
+		}
+
+		// Save the manifest and archives for debug purposes, but don't add them to the build
+		// We don't care if this fails as it's only for debugging
+		ChunkLocTextHelper.SaveAll();
+
+		// Produce the LocMeta file for the chunk target
+		{
+			const FString ChunkLocMetaFilename = ChunkTargetRoot / FString::Printf(TEXT("%s.locmeta"), *ChunkTargetName);
+			
+			FTextLocalizationMetaDataResource ChunkLocMeta;
+			if (FTextLocalizationResourceGenerator::GenerateLocMeta(ChunkLocTextHelper, FString::Printf(TEXT("%s.locres"), *ChunkTargetName), ChunkLocMeta) && ChunkLocMeta.SaveToFile(ChunkLocMetaFilename))
+			{
+				OutChunkFilenames.Add(ChunkLocMetaFilename);
+			}
+			else
+			{
+				UE_LOG(LogLocalizationChunkDataGenerator, Error, TEXT("Failed to generate meta-data for localization target '%s' when chunking localization data."), *ChunkTargetName);
+			}
+		}
+
+		// Produce the LocRes files for each culture of the chunk target
+		for (const FString& CultureToCook : AvailableCulturesToCook)
+		{
+			// This is an extra sanity check as the native culture of the target may not be being cooked
+			if (!AllCulturesToCook.Contains(CultureToCook))
+			{
+				continue;
+			}
+
+			const FString ChunkLocResFilename = ChunkTargetRoot / CultureToCook / FString::Printf(TEXT("%s.locres"), *ChunkTargetName);
+
+			FTextLocalizationResource ChunkLocRes;
+			TMap<FName, TSharedRef<FTextLocalizationResource>> UnusedPerPlatformLocRes;
+			if (FTextLocalizationResourceGenerator::GenerateLocRes(ChunkLocTextHelper, CultureToCook, EGenerateLocResFlags::None, ChunkLocResFilename, ChunkLocRes, UnusedPerPlatformLocRes) && ChunkLocRes.SaveToFile(ChunkLocResFilename))
+			{
+				OutChunkFilenames.Add(ChunkLocResFilename);
+			}
+			else
+			{
+				UE_LOG(LogLocalizationChunkDataGenerator, Error, TEXT("Failed to generate resource data for localization target '%s' when chunking localization data."), *ChunkTargetName);
+			}
+		}
+	}
+}
+
+void FLocalizationChunkDataGenerator::ConditionalCacheLocalizationTargetData()
+{
+	// We can skip this if we're not actually chunking or staging any localization data
+	if (LocalizationTargetsToChunk.Num() == 0 || AllCulturesToCook.Num() == 0)
+	{
+		return;
+	}
+
+	// Already cached?
+	if (LocalizationTargetsToChunk.Num() == CachedLocalizationTargetHelpers.Num())
+	{
+		return;
+	}
+
+	CachedLocalizationTargetHelpers.Reset(LocalizationTargetsToChunk.Num());
+	for (const FString& LocalizationTarget : LocalizationTargetsToChunk)
+	{
+		TSharedPtr<FLocTextHelper>& SourceLocTextHelper = CachedLocalizationTargetHelpers.AddDefaulted_GetRef();
+
+		// Does this target exist?
+		// Note: We only allow game localization targets to be chunked, and the layout is assumed to follow our standard pattern (as used by the localization dashboard and FLocTextHelper)
+		const FString SourceRootPath = FPaths::ProjectContentDir() / TEXT("Localization") / LocalizationTarget;
+		if (!FPaths::DirectoryExists(SourceRootPath))
+		{
+			UE_LOG(LogLocalizationChunkDataGenerator, Warning, TEXT("Failed to find localization target for '%s' when chunking localization data. Is it a valid project localization target? - %s"), *LocalizationTarget, *SourceRootPath);
+			continue;
+		}
+
+		// Work out what the native culture is
+		FString SourceNativeCulture;
+		{
+			const FString SourceLocMetaFilename = SourceRootPath / FString::Printf(TEXT("%s.locmeta"), *LocalizationTarget);
+
+			FTextLocalizationMetaDataResource SourceLocMeta;
+			if (!SourceLocMeta.LoadFromFile(SourceLocMetaFilename))
+			{
+				UE_LOG(LogLocalizationChunkDataGenerator, Error, TEXT("Failed to load meta-data for localization target '%s' when chunking localization data. Re-compile the localization target to generate the LocMeta file."), *LocalizationTarget);
+				continue;
+			}
+
+			SourceNativeCulture = SourceLocMeta.NativeCulture;
+		}
+
+		// Work out which of the desired cultures this target actually supports
+		TArray<FString> SourceForeignCulturesToCook;
+		for (const FString& CultureToCook : AllCulturesToCook)
+		{
+			if (CultureToCook == SourceNativeCulture)
+			{
+				continue;
+			}
+
+			const FString LocalizationTargetCulturePath = SourceRootPath / CultureToCook;
+			if (FPaths::DirectoryExists(LocalizationTargetCulturePath))
+			{
+				SourceForeignCulturesToCook.Add(CultureToCook);
+			}
+		}
+
+		// Load the data for this target
+		SourceLocTextHelper = MakeShared<FLocTextHelper>(SourceRootPath, FString::Printf(TEXT("%s.manifest"), *LocalizationTarget), FString::Printf(TEXT("%s.archive"), *LocalizationTarget), SourceNativeCulture, SourceForeignCulturesToCook, nullptr);
+		{
+			FText LoadError;
+			if (!SourceLocTextHelper->LoadAll(ELocTextHelperLoadFlags::Load, &LoadError))
+			{
+				SourceLocTextHelper.Reset();
+				UE_LOG(LogLocalizationChunkDataGenerator, Error, TEXT("Failed to load data for localization target '%s' when chunking localization data: %s"), *LocalizationTarget, *LoadError.ToString());
+				continue;
+			}
+		}
+	}
+}
diff --git a/Engine/Source/Editor/UnrealEd/Private/LocalizationChunkDataGenerator.h b/Engine/Source/Editor/UnrealEd/Private/LocalizationChunkDataGenerator.h
new file mode 100644
--- /dev/null
+++ b/Engine/Source/Editor/UnrealEd/Private/LocalizationChunkDataGenerator.h
@@ -0,0 +1,34 @@
+// Copyright 1998-2019 Epic Games, Inc. All Rights Reserved.
+
+#pragma once
+
+#include "CoreMinimal.h"
+#include "Commandlets/IChunkDataGenerator.h"
+
+class FLocTextHelper;
+
+/**
+ * Implementation for splitting localization data into chunks when creating streaming install manifests.
+ */
+class FLocalizationChunkDataGenerator : public IChunkDataGenerator
+{
+public:
+	FLocalizationChunkDataGenerator(TArray<FString> InLocalizationTargetsToChunk, TArray<FString> InAllCulturesToCook);
+	virtual ~FLocalizationChunkDataGenerator() = default;
+
+	//~ IChunkDataGenerator
+	virtual void GenerateChunkDataFiles(const int32 InChunkId, const TSet<FName>& InPackagesInChunk, const FString& InPlatformName, FSandboxPlatformFile* InSandboxFile, TArray<FString>& OutChunkFilenames) override;
+
+private:
+	/** Update CachedLocalizationTargetData if needed */
+	void ConditionalCacheLocalizationTargetData();
+
+	/** List of localization targets that should be chunked */
+	TArray<FString> LocalizationTargetsToChunk;
+
+	/** Complete list of cultures to cook data for, including inferred parent cultures */
+	TArray<FString> AllCulturesToCook;
+
+	/** Cached localization target helpers, to avoid redundant work for each chunk */
+	TArray<TSharedPtr<FLocTextHelper>> CachedLocalizationTargetHelpers;
+};
diff --git a/Engine/Source/Editor/UnrealEd/Public/Commandlets/AssetRegistryGenerator.h b/Engine/Source/Editor/UnrealEd/Public/Commandlets/AssetRegistryGenerator.h
--- a/Engine/Source/Editor/UnrealEd/Public/Commandlets/AssetRegistryGenerator.h
+++ b/Engine/Source/Editor/UnrealEd/Public/Commandlets/AssetRegistryGenerator.h
@@ -12,6 +12,7 @@
 class FSandboxPlatformFile;
 class IAssetRegistry;
 class ITargetPlatform;
+class IChunkDataGenerator;
 class UChunkDependencyInfo;
 struct FChunkDependencyTreeNode;
 
@@ -73,6 +74,12 @@ public:
 	 */
 	void BuildChunkManifest(const TSet<FName>& CookedPackages, const TSet<FName>& DevelopmentOnlyPackages, FSandboxPlatformFile* InSandboxFile, bool bGenerateStreamingInstallManifest);
 
+	/**
+	 * Register a chunk data generator with this generator.
+	 * @note Should be called prior to SaveManifests.
+	 */
+	void RegisterChunkDataGenerator(TSharedRef<IChunkDataGenerator> InChunkDataGenerator);
+
 	/**
 	* PreSave
 	* Notify generator that we are about to save the registry and chunk manifests
@@ -194,6 +201,8 @@ private:
 	FChunkPackageSet				AllCookedPackageSet;
 	/** Array of Maps with chunks<->packages assignments. This version contains all dependent packages */
 	TArray<FChunkPackageSet*>		FinalChunkManifests;
+	/** Additional data generators used when creating chunks */
+	TArray<TSharedRef<IChunkDataGenerator>> ChunkDataGenerators;
 	/** Lookup table of used package names used when searching references. */
 	TSet<FName>						InspectedNames;
 	/** */
diff --git a/Engine/Source/Editor/UnrealEd/Public/Commandlets/IChunkDataGenerator.h b/Engine/Source/Editor/UnrealEd/Public/Commandlets/IChunkDataGenerator.h
new file mode 100644
--- /dev/null
+++ b/Engine/Source/Editor/UnrealEd/Public/Commandlets/IChunkDataGenerator.h
@@ -0,0 +1,28 @@
+// Copyright 1998-2019 Epic Games, Inc. All Rights Reserved.
+
+#pragma once
+
+#include "CoreMinimal.h"
+
+class FSandboxPlatformFile;
+
+/**
+ * Interface for generating extra data when creating streaming install manifests.
+ * @note See FAssetRegistryGenerator.
+ */
+class IChunkDataGenerator
+{
+public:
+	virtual ~IChunkDataGenerator() = default;
+
+	/**
+	 * Called to generate any additional files that should be added to the given chunk.
+	 *
+	 * @param InChunkId The ID of the chunk to generate the files for.
+	 * @param InPackagesInChunk The set of packages that are in the chunk.
+	 * @param InPlatformName The name of the platform this chunk is for.
+	 * @param InSandboxFile The sandbox location that staging data for this chunk should be written to.
+	 * @param OutChunkFilenames Array of filenames that belong to this chunk (to be appended to).
+	 */
+	virtual void GenerateChunkDataFiles(const int32 InChunkId, const TSet<FName>& InPackagesInChunk, const FString& InPlatformName, FSandboxPlatformFile* InSandboxFile, TArray<FString>& OutChunkFilenames) = 0;
+};
diff --git a/Engine/Source/Programs/AutomationTool/Scripts/CopyBuildToStagingDirectory.Automation.cs b/Engine/Source/Programs/AutomationTool/Scripts/CopyBuildToStagingDirectory.Automation.cs
--- a/Engine/Source/Programs/AutomationTool/Scripts/CopyBuildToStagingDirectory.Automation.cs
+++ b/Engine/Source/Programs/AutomationTool/Scripts/CopyBuildToStagingDirectory.Automation.cs
@@ -260,6 +260,29 @@ public partial class Project : CommandUtils
 		return CulturesToStage;
 	}
 
+	private static bool ShouldStageLocalizationTarget(DeploymentContext SC, ConfigHierarchy PlatformGameConfig, string LocalizationTarget)
+	{
+		if (SC.BlacklistLocalizationTargets.Contains(LocalizationTarget))
+		{
+			return false;
+		}
+
+		// Chunked targets can be added to the manifest during cook, so skip them during staging otherwise we'll overwrite the chunk 0 data with the non-chunked version
+		if (PlatformGameConfig != null)
+		{
+			List<string> LocalizationTargetsToChunk = null;
+			PlatformGameConfig.GetArray("/Script/UnrealEd.ProjectPackagingSettings", "LocalizationTargetsToChunk", out LocalizationTargetsToChunk);
+
+			if (LocalizationTargetsToChunk != null && LocalizationTargetsToChunk.Contains(LocalizationTarget))
+			{
+				var CookedLocalizationTargetDirectory = DirectoryReference.Combine(SC.PlatformCookDir, SC.ShortProjectName, "Content", "Localization", LocalizationTarget);
+				return !DirectoryReference.Exists(CookedLocalizationTargetDirectory);
+			}
+		}
+
+		return true;
+	}
+
 	private static void StageLocalizationDataForPlugin(DeploymentContext SC, List<string> CulturesToStage, FileReference Plugin)
 	{
 		PluginDescriptor Descriptor = PluginDescriptor.FromFile(Plugin);
@@ -270,7 +293,7 @@ public partial class Project : CommandUtils
 				if (LocalizationTarget.LoadingPolicy == LocalizationTargetDescriptorLoadingPolicy.Always || LocalizationTarget.LoadingPolicy == LocalizationTargetDescriptorLoadingPolicy.Game)
 				{
 					DirectoryReference PluginLocTargetDirectory = DirectoryReference.Combine(Plugin.Directory, "Content", "Localization", LocalizationTarget.Name);
-					if (!SC.BlacklistLocalizationTargets.Contains(LocalizationTarget.Name) && DirectoryReference.Exists(PluginLocTargetDirectory))
+					if (ShouldStageLocalizationTarget(SC, null, LocalizationTarget.Name) && DirectoryReference.Exists(PluginLocTargetDirectory))
 					{
 						StageLocalizationDataForTarget(SC, CulturesToStage, PluginLocTargetDirectory);
 					}
@@ -303,13 +326,13 @@ public partial class Project : CommandUtils
 		}
 	}
 
-	private static void StageLocalizationDataForTargetsInDirectory(DeploymentContext SC, List<string> CulturesToStage, DirectoryReference RootDirectory)
+	private static void StageLocalizationDataForTargetsInDirectory(DeploymentContext SC, ConfigHierarchy PlatformGameConfig, List<string> CulturesToStage, DirectoryReference RootDirectory)
 	{
 		if (DirectoryReference.Exists(RootDirectory))
 		{
 			foreach (DirectoryReference LocTargetDirectory in DirectoryReference.EnumerateDirectories(RootDirectory))
 			{
-				if (!SC.BlacklistLocalizationTargets.Contains(LocTargetDirectory.MakeRelativeTo(RootDirectory)))
+				if (ShouldStageLocalizationTarget(SC, PlatformGameConfig, LocTargetDirectory.MakeRelativeTo(RootDirectory)))
 				{
 					StageLocalizationDataForTarget(SC, CulturesToStage, LocTargetDirectory);
 				}
@@ -420,6 +443,13 @@ public partial class Project : CommandUtils
 			// Making a plugin
 			DirectoryReference DLCRoot = Params.DLCFile.Directory;
 
+			// Put all of the cooked dir into the staged dir
+			SC.PlatformCookDir = String.IsNullOrEmpty(Params.CookOutputDir) ? DirectoryReference.Combine(DLCRoot, "Saved", "Cooked", SC.CookPlatform) : DirectoryReference.Combine(new DirectoryReference(Params.CookOutputDir), SC.CookPlatform);
+			DirectoryReference PlatformEngineDir = DirectoryReference.Combine(SC.PlatformCookDir, "Engine");
+			DirectoryReference ProjectMetadataDir = DirectoryReference.Combine(SC.PlatformCookDir, SC.ShortProjectName, "Metadata");
+			string RelativeDLCRootPath = DLCRoot.MakeRelativeTo(SC.LocalRoot);
+			SC.MetadataDir = DirectoryReference.Combine(SC.PlatformCookDir, RelativeDLCRootPath, "Metadata");
+
 			// The .uplugin file is staged differently for different DLC
 			// The .uplugin file doesn't actually exist for mobile DLC
 			if (FileReference.Exists(Params.DLCFile))
@@ -451,13 +481,6 @@ public partial class Project : CommandUtils
 				SC.StageFiles(StagedFileType.NonUFS, BinariesDir, "UE4Server-*.dll", StageFilesSearch.AllDirectories);
 			}
 
-			// Put all of the cooked dir into the staged dir
-			SC.PlatformCookDir = String.IsNullOrEmpty(Params.CookOutputDir) ? DirectoryReference.Combine(DLCRoot, "Saved", "Cooked", SC.CookPlatform) : DirectoryReference.Combine(new DirectoryReference(Params.CookOutputDir), SC.CookPlatform);
-			DirectoryReference PlatformEngineDir = DirectoryReference.Combine(SC.PlatformCookDir, "Engine");
-			DirectoryReference ProjectMetadataDir = DirectoryReference.Combine(SC.PlatformCookDir, SC.ShortProjectName, "Metadata");
-			string RelativeDLCRootPath = DLCRoot.MakeRelativeTo(SC.LocalRoot);
-			SC.MetadataDir = DirectoryReference.Combine(SC.PlatformCookDir, RelativeDLCRootPath, "Metadata");
-
 			// Put the config files into the staged dir
 			DirectoryReference ConfigDir = DirectoryReference.Combine(DLCRoot, "Config");
 			if (DirectoryReference.Exists(ConfigDir))
@@ -465,7 +488,6 @@ public partial class Project : CommandUtils
 				SC.StageFiles(StagedFileType.UFS, ConfigDir, "*.ini", StageFilesSearch.AllDirectories);
 			}
 
-
 			if (Params.DLCActLikePatch)
 			{
 				DirectoryReference DLCContent = DirectoryReference.Combine(DLCRoot, "Content");
@@ -566,6 +588,24 @@ public partial class Project : CommandUtils
 
 			if (!Params.CookOnTheFly && !Params.SkipCookOnTheFly) // only stage the UFS files if we are not using cook on the fly
 			{
+				// Get the final output directory for cooked data
+				DirectoryReference CookOutputDir;
+				if (!String.IsNullOrEmpty(Params.CookOutputDir))
+				{
+					CookOutputDir = DirectoryReference.Combine(new DirectoryReference(Params.CookOutputDir), SC.CookPlatform);
+				}
+				else if (Params.CookInEditor)
+				{
+					CookOutputDir = DirectoryReference.Combine(SC.ProjectRoot, "Saved", "EditorCooked", SC.CookPlatform);
+				}
+				else
+				{
+					CookOutputDir = DirectoryReference.Combine(SC.ProjectRoot, "Saved", "Cooked", SC.CookPlatform);
+				}
+
+				SC.PlatformCookDir = CookOutputDir;
+				SC.MetadataDir = DirectoryReference.Combine(CookOutputDir, SC.ShortProjectName, "Metadata");
+
 				// Work out which ICU data version we use for this platform
 				var ICUDataVersion = "icudt64l";
 				if (SC.StageTargetPlatform.PlatformType == UnrealTargetPlatform.HoloLens || SC.StageTargetPlatform.PlatformType == UnrealTargetPlatform.TVOS)
@@ -586,7 +626,7 @@ public partial class Project : CommandUtils
 				StageConfigFiles(SC, DirectoryReference.Combine(SC.LocalRoot, "Engine", "Config"), null);
 
 				// Stage the engine localization target
-				if (!SC.BlacklistLocalizationTargets.Contains("Engine"))
+				if (ShouldStageLocalizationTarget(SC, null, "Engine"))
 				{
 					StageLocalizationDataForTarget(SC, CulturesToStage, DirectoryReference.Combine(SC.LocalRoot, "Engine", "Content", "Localization", "Engine"));
 				}
@@ -627,7 +667,7 @@ public partial class Project : CommandUtils
 				}
 
 				// Stage all project localization targets
-				StageLocalizationDataForTargetsInDirectory(SC, CulturesToStage, DirectoryReference.Combine(SC.ProjectRoot, "Content", "Localization"));
+				StageLocalizationDataForTargetsInDirectory(SC, PlatformGameConfig, CulturesToStage, DirectoryReference.Combine(SC.ProjectRoot, "Content", "Localization"));
 
 				// Stage all plugin localization targets
 				foreach (KeyValuePair<StagedFileReference, FileReference> StagedPlugin in StagedPlugins)
@@ -638,8 +678,8 @@ public partial class Project : CommandUtils
 				// Stage any platform extension localization targets
 				{
 					string PlatformExtensionName = SC.StageTargetPlatform.PlatformType.ToString();
-					StageLocalizationDataForTargetsInDirectory(SC, CulturesToStage, DirectoryReference.Combine(SC.LocalRoot, "Engine", "Platforms", PlatformExtensionName, "Content", "Localization"));
-					StageLocalizationDataForTargetsInDirectory(SC, CulturesToStage, DirectoryReference.Combine(SC.ProjectRoot, "Platforms", PlatformExtensionName, "Content", "Localization"));
+					StageLocalizationDataForTargetsInDirectory(SC, null, CulturesToStage, DirectoryReference.Combine(SC.LocalRoot, "Engine", "Platforms", PlatformExtensionName, "Content", "Localization"));
+					StageLocalizationDataForTargetsInDirectory(SC, PlatformGameConfig, CulturesToStage, DirectoryReference.Combine(SC.ProjectRoot, "Platforms", PlatformExtensionName, "Content", "Localization"));
 				}
 
 				// Stage any additional UFS and NonUFS paths specified in the project ini files; these dirs are relative to the game content directory
@@ -741,24 +781,6 @@ public partial class Project : CommandUtils
 					}
 				}
 
-				// Get the final output directory for cooked data
-				DirectoryReference CookOutputDir;
-				if (!String.IsNullOrEmpty(Params.CookOutputDir))
-				{
-					CookOutputDir = DirectoryReference.Combine(new DirectoryReference(Params.CookOutputDir), SC.CookPlatform);
-				}
-				else if (Params.CookInEditor)
-				{
-					CookOutputDir = DirectoryReference.Combine(SC.ProjectRoot, "Saved", "EditorCooked", SC.CookPlatform);
-				}
-				else
-				{
-					CookOutputDir = DirectoryReference.Combine(SC.ProjectRoot, "Saved", "Cooked", SC.CookPlatform);
-				}
-
-				SC.PlatformCookDir = CookOutputDir;
-				SC.MetadataDir = DirectoryReference.Combine(CookOutputDir, SC.ShortProjectName, "Metadata");
-
 				// Stage all the cooked data. Currently not filtering this by restricted folders, since we shouldn't mask invalid references by filtering them out.
 				if (DirectoryReference.Exists(CookOutputDir))
 				{
diff --git a/Engine/Source/Runtime/Core/Private/Internationalization/LocalizationResourceTextSource.cpp b/Engine/Source/Runtime/Core/Private/Internationalization/LocalizationResourceTextSource.cpp
--- a/Engine/Source/Runtime/Core/Private/Internationalization/LocalizationResourceTextSource.cpp
+++ b/Engine/Source/Runtime/Core/Private/Internationalization/LocalizationResourceTextSource.cpp
@@ -45,18 +45,37 @@ void FLocalizationResourceTextSource::GetLocalizedCultureNames(const ELocalizati
 
 void FLocalizationResourceTextSource::LoadLocalizedResources(const ELocalizationLoadFlags InLoadFlags, TArrayView<const FString> InPrioritizedCultures, FTextLocalizationResource& InOutNativeResource, FTextLocalizationResource& InOutLocalizedResource)
 {
-	const int32 BaseResourcePriority = GetPriority() * -1; // Flip the priority as larger text source priorities are more important, but smaller text resource priorities are more important
+	auto GetGameLocalizationPaths = [this]()
+	{
+		TArray<FString> LocalizationPaths = FPaths::GetGameLocalizationPaths();
+
+		if (ChunkIds.Num() > 0)
+		{
+			const TArray<FString> ChunkedLocalizationTargets = GetChunkedLocalizationTargets();
+			for (const FString& LocalizationTarget : ChunkedLocalizationTargets)
+			{
+				for (const int32 ChunkId : ChunkIds)
+				{
+					// Note: We only allow game localization targets to be chunked, and the layout is assumed to follow our standard pattern (as used by the localization dashboard and FLocTextHelper)
+					const FString LocalizationTargetForChunk = TextLocalizationResourceUtil::GetLocalizationTargetNameForChunkId(LocalizationTarget, ChunkId);
+					LocalizationPaths.Add(FPaths::ProjectContentDir() / TEXT("Localization") / LocalizationTargetForChunk);
+				}
+			}
+		}
+
+		return LocalizationPaths;
+	};
 
 	// Collect the localization paths to load from.
 	TArray<FString> GameNativePaths;
 	TArray<FString> GameLocalizationPaths;
 	if (ShouldLoadNativeGameData(InLoadFlags))
 	{
-		GameNativePaths += FPaths::GetGameLocalizationPaths();
+		GameNativePaths += GetGameLocalizationPaths();
 	}
 	else if (ShouldLoadGame(InLoadFlags))
 	{
-		GameLocalizationPaths += FPaths::GetGameLocalizationPaths();
+		GameLocalizationPaths += GetGameLocalizationPaths();
 	}
 
 	TArray<FString> EditorNativePaths;
@@ -115,6 +134,13 @@ void FLocalizationResourceTextSource::LoadLocalizedResources(const ELocalization
 		}
 	}
 
+	LoadLocalizedResourcesFromPaths(PrioritizedNativePaths, PrioritizedLocalizationPaths, GameNativePaths, InLoadFlags, InPrioritizedCultures, InOutNativeResource, InOutLocalizedResource);
+}
+
+void FLocalizationResourceTextSource::LoadLocalizedResourcesFromPaths(TArrayView<const FString> InPrioritizedNativePaths, TArrayView<const FString> InPrioritizedLocalizationPaths, TArrayView<const FString> InGameNativePaths, const ELocalizationLoadFlags InLoadFlags, TArrayView<const FString> InPrioritizedCultures, FTextLocalizationResource& InOutNativeResource, FTextLocalizationResource& InOutLocalizedResource) const
+{
+	const int32 BaseResourcePriority = GetPriority() * -1; // Flip the priority as larger text source priorities are more important, but smaller text resource priorities are more important
+
 	static const FString PlatformLocalizationFolderName = FPaths::GetPlatformLocalizationFolderName();
 	static const FString PlatformName = ANSI_TO_TCHAR(FPlatformProperties::IniPlatformName());
 	auto LoadLocalizationResourcesForCulture = [](FTextLocalizationResource& InOutLocRes, const FString& InLocalizationPath, const FString& InCulture, const FString& InLocResFilename, const int32 InLocResPriority)
@@ -140,9 +166,9 @@ void FLocalizationResourceTextSource::LoadLocalizedResources(const ELocalization
 	};
 
 	// Load the native texts first to ensure we always apply translations to a consistent base
-	if (PrioritizedNativePaths.Num() > 0)
+	if (InPrioritizedNativePaths.Num() > 0)
 	{
-		for (const FString& LocalizationPath : PrioritizedNativePaths)
+		for (const FString& LocalizationPath : InPrioritizedNativePaths)
 		{
 			if (!IFileManager::Get().DirectoryExists(*LocalizationPath))
 			{
@@ -164,12 +190,12 @@ void FLocalizationResourceTextSource::LoadLocalizedResources(const ELocalization
 	}
 
 	// The editor cheats and loads the games native localizations.
-	if (ShouldLoadNativeGameData(InLoadFlags) && GameNativePaths.Num() > 0)
+	if (ShouldLoadNativeGameData(InLoadFlags) && InGameNativePaths.Num() > 0)
 	{
 		const FString NativeGameCulture = TextLocalizationResourceUtil::GetNativeProjectCultureName();
 		if (!NativeGameCulture.IsEmpty())
 		{
-			for (const FString& LocalizationPath : GameNativePaths)
+			for (const FString& LocalizationPath : InGameNativePaths)
 			{
 				if (!IFileManager::Get().DirectoryExists(*LocalizationPath))
 				{
@@ -183,12 +209,12 @@ void FLocalizationResourceTextSource::LoadLocalizedResources(const ELocalization
 	}
 
 	// Read culture localization resources.
-	if (PrioritizedLocalizationPaths.Num() > 0)
+	if (InPrioritizedLocalizationPaths.Num() > 0)
 	{
 		for (int32 CultureIndex = 0; CultureIndex < InPrioritizedCultures.Num(); ++CultureIndex)
 		{
 			const FString& PrioritizedCultureName = InPrioritizedCultures[CultureIndex];
-			for (const FString& LocalizationPath : PrioritizedLocalizationPaths)
+			for (const FString& LocalizationPath : InPrioritizedLocalizationPaths)
 			{
 				if (!IFileManager::Get().DirectoryExists(*LocalizationPath))
 				{
@@ -201,3 +227,20 @@ void FLocalizationResourceTextSource::LoadLocalizedResources(const ELocalization
 		}
 	}
 }
+
+TArray<FString> FLocalizationResourceTextSource::GetChunkedLocalizationTargets()
+{
+	TArray<FString> ChunkedLocalizationTargets;
+	GConfig->GetArray(TEXT("/Script/UnrealEd.ProjectPackagingSettings"), TEXT("LocalizationTargetsToChunk"), ChunkedLocalizationTargets, GGameIni);
+
+	const TArray<FString>& DisabledLocalizationTargets = TextLocalizationResourceUtil::GetDisabledLocalizationTargets();
+	if (DisabledLocalizationTargets.Num() > 0)
+	{
+		ChunkedLocalizationTargets.RemoveAll([&DisabledLocalizationTargets](const FString& LocalizationTarget)
+		{
+			return DisabledLocalizationTargets.Contains(LocalizationTarget);
+		});
+	}
+	
+	return ChunkedLocalizationTargets;
+}
diff --git a/Engine/Source/Runtime/Core/Private/Internationalization/LocalizationResourceTextSource.h b/Engine/Source/Runtime/Core/Private/Internationalization/LocalizationResourceTextSource.h
--- a/Engine/Source/Runtime/Core/Private/Internationalization/LocalizationResourceTextSource.h
+++ b/Engine/Source/Runtime/Core/Private/Internationalization/LocalizationResourceTextSource.h
@@ -15,4 +15,26 @@ public:
 	virtual bool GetNativeCultureName(const ELocalizedTextSourceCategory InCategory, FString& OutNativeCultureName) override;
 	virtual void GetLocalizedCultureNames(const ELocalizationLoadFlags InLoadFlags, TSet<FString>& OutLocalizedCultureNames) override;
 	virtual void LoadLocalizedResources(const ELocalizationLoadFlags InLoadFlags, TArrayView<const FString> InPrioritizedCultures, FTextLocalizationResource& InOutNativeResource, FTextLocalizationResource& InOutLocalizedResource) override;
+
+	/**
+	 * Load the localized resources from the LocRes files for the given cultures at the given paths into the given maps (ELocalizationLoadFlags controls which resources should be loaded).
+	 */
+	void LoadLocalizedResourcesFromPaths(TArrayView<const FString> InPrioritizedNativePaths, TArrayView<const FString> InPrioritizedLocalizationPaths, TArrayView<const FString> InGameNativePaths, const ELocalizationLoadFlags InLoadFlags, TArrayView<const FString> InPrioritizedCultures, FTextLocalizationResource& InOutNativeResource, FTextLocalizationResource& InOutLocalizedResource) const;
+
+	/**
+	 * Register that a chunk has been loaded that may contain chunked LocRes data.
+	 */
+	void RegisterChunkId(const int32 InChunkId)
+	{
+		ChunkIds.Add(InChunkId);
+	}
+
+	/**
+	 * Get the list of localization targets that were chunked during cooking.
+	 */
+	static TArray<FString> GetChunkedLocalizationTargets();
+
+private:
+	/** Array of chunks that have been loaded that may contain chunked LocRes data */
+	TArray<int32> ChunkIds;
 };
diff --git a/Engine/Source/Runtime/Core/Private/Internationalization/PolyglotTextSource.cpp b/Engine/Source/Runtime/Core/Private/Internationalization/PolyglotTextSource.cpp
--- a/Engine/Source/Runtime/Core/Private/Internationalization/PolyglotTextSource.cpp
+++ b/Engine/Source/Runtime/Core/Private/Internationalization/PolyglotTextSource.cpp
@@ -48,9 +48,9 @@ void FPolyglotTextSource::GetLocalizedCultureNames(const ELocalizationLoadFlags
 	}
 }
 
-void FPolyglotTextSource::LoadLocalizedResources(const ELocalizationLoadFlags InLoadFlags, TArrayView<const FString> InPrioritizedCultures, FTextLocalizationResource& InOutNativeResource, FTextLocalizationResource& InOutLocalizedResource)
+void FPolyglotTextSource::AddPolyglotDataToResource(const FPolyglotTextData& InPolyglotTextData, const ELocalizationLoadFlags InLoadFlags, TArrayView<const FString> InPrioritizedCultures, FTextLocalizationResource& InOutNativeResource, FTextLocalizationResource& InOutLocalizedResource) const
 {
-	auto ShouldLoadLocalizedText = [InLoadFlags](const FPolyglotTextData& InPolyglotTextData) -> bool
+	auto ShouldLoadLocalizedText = [InLoadFlags, &InPolyglotTextData]() -> bool
 	{
 		switch (InPolyglotTextData.GetCategory())
 		{
@@ -67,7 +67,7 @@ void FPolyglotTextSource::LoadLocalizedResources(const ELocalizationLoadFlags In
 		return false;
 	};
 
-	auto GetLocalizedStringForPolyglotData = [](const FPolyglotTextData& InPolyglotTextData, TArrayView<const FString> InCulturesToCheck, FString& OutLocalizedString, int32* OutLocalizedPriority = nullptr) -> bool
+	auto GetLocalizedStringForPolyglotData = [&InPolyglotTextData](TArrayView<const FString> InCulturesToCheck, FString& OutLocalizedString, int32* OutLocalizedPriority = nullptr) -> bool
 	{
 		for (int32 CultureIndex = 0; CultureIndex < InCulturesToCheck.Num(); ++CultureIndex)
 		{
@@ -97,65 +97,80 @@ void FPolyglotTextSource::LoadLocalizedResources(const ELocalizationLoadFlags In
 	};
 
 	const int32 BaseResourcePriority = GetPriority() * -1; // Flip the priority as larger text source priorities are more important, but smaller text resource priorities are more important
+	const FString NativeCulture = InPolyglotTextData.ResolveNativeCulture();
 
-	for (const auto& PolyglotTextDataPair : PolyglotTextDataMap)
+	// We skip loading the native text if we're transitioning to the native culture as there's no extra work that needs to be done
+	if (ShouldLoadNative(InLoadFlags) && !InPrioritizedCultures.Contains(NativeCulture))
 	{
-		const FPolyglotTextData& PolyglotTextData = PolyglotTextDataPair.Value;
-		const FString NativeCulture = PolyglotTextData.ResolveNativeCulture();
+		FString LocalizedString;
+		if (GetLocalizedStringForPolyglotData(TArrayView<const FString>(&NativeCulture, 1), LocalizedString))
+		{
+			InOutNativeResource.AddEntry(
+				InPolyglotTextData.GetNamespace(),
+				InPolyglotTextData.GetKey(),
+				InPolyglotTextData.GetNativeString(),
+				LocalizedString,
+				BaseResourcePriority
+				);
+		}
+	}
 
-		// We skip loading the native text if we're transitioning to the native culture as there's no extra work that needs to be done
-		if (ShouldLoadNative(InLoadFlags) && !InPrioritizedCultures.Contains(NativeCulture))
+	if (ShouldLoadLocalizedText())
+	{
+		if (InPolyglotTextData.GetCategory() == ELocalizedTextSourceCategory::Game && ShouldLoadNativeGameData(InLoadFlags))
 		{
+			// The editor cheats and loads the native language's localizations for game data.
 			FString LocalizedString;
-			if (GetLocalizedStringForPolyglotData(PolyglotTextData, TArrayView<const FString>(&NativeCulture, 1), LocalizedString))
+			if (GetLocalizedStringForPolyglotData(TArrayView<const FString>(&NativeCulture, 1), LocalizedString))
 			{
-				InOutNativeResource.AddEntry(
-					PolyglotTextData.GetNamespace(),
-					PolyglotTextData.GetKey(),
-					PolyglotTextData.GetNativeString(),
+				InOutLocalizedResource.AddEntry(
+					InPolyglotTextData.GetNamespace(),
+					InPolyglotTextData.GetKey(),
+					InPolyglotTextData.GetNativeString(),
 					LocalizedString,
 					BaseResourcePriority
 					);
 			}
 		}
-
-		if (ShouldLoadLocalizedText(PolyglotTextData))
+		else
 		{
-			if (PolyglotTextData.GetCategory() == ELocalizedTextSourceCategory::Game && ShouldLoadNativeGameData(InLoadFlags))
-			{
-				// The editor cheats and loads the native language's localizations for game data.
-				FString LocalizedString;
-				if (GetLocalizedStringForPolyglotData(PolyglotTextData, TArrayView<const FString>(&NativeCulture, 1), LocalizedString))
-				{
-					InOutLocalizedResource.AddEntry(
-						PolyglotTextData.GetNamespace(),
-						PolyglotTextData.GetKey(),
-						PolyglotTextData.GetNativeString(),
-						LocalizedString,
-						BaseResourcePriority
-						);
-				}
-			}
-			else
+			// Find culture localization resource.
+			int32 LocalizedPriority = 0;
+			FString LocalizedString;
+			if (GetLocalizedStringForPolyglotData(InPrioritizedCultures, LocalizedString, &LocalizedPriority))
 			{
-				// Find culture localization resource.
-				int32 LocalizedPriority = 0;
-				FString LocalizedString;
-				if (GetLocalizedStringForPolyglotData(PolyglotTextData, InPrioritizedCultures, LocalizedString, &LocalizedPriority))
-				{
-					InOutLocalizedResource.AddEntry(
-						PolyglotTextData.GetNamespace(),
-						PolyglotTextData.GetKey(),
-						PolyglotTextData.GetNativeString(),
-						LocalizedString,
-						BaseResourcePriority + LocalizedPriority
-						);
-				}
+				InOutLocalizedResource.AddEntry(
+					InPolyglotTextData.GetNamespace(),
+					InPolyglotTextData.GetKey(),
+					InPolyglotTextData.GetNativeString(),
+					LocalizedString,
+					BaseResourcePriority + LocalizedPriority
+					);
 			}
 		}
 	}
 }
 
+void FPolyglotTextSource::LoadLocalizedResources(const ELocalizationLoadFlags InLoadFlags, TArrayView<const FString> InPrioritizedCultures, FTextLocalizationResource& InOutNativeResource, FTextLocalizationResource& InOutLocalizedResource)
+{
+	for (const auto& PolyglotTextDataPair : PolyglotTextDataMap)
+	{
+		const FPolyglotTextData& PolyglotTextData = PolyglotTextDataPair.Value;
+		AddPolyglotDataToResource(PolyglotTextData, InLoadFlags, InPrioritizedCultures, InOutNativeResource, InOutLocalizedResource);
+	}
+}
+
+EQueryLocalizedResourceResult FPolyglotTextSource::QueryLocalizedResource(const ELocalizationLoadFlags InLoadFlags, TArrayView<const FString> InPrioritizedCultures, const FTextId InTextId, FTextLocalizationResource& InOutNativeResource, FTextLocalizationResource& InOutLocalizedResource)
+{
+	const FLocKey Identity = FString::Printf(TEXT("%s::%s"), InTextId.GetNamespace().GetChars(), InTextId.GetKey().GetChars());
+	if (const FPolyglotTextData* PolyglotTextData = PolyglotTextDataMap.Find(Identity))
+	{
+		AddPolyglotDataToResource(*PolyglotTextData, InLoadFlags, InPrioritizedCultures, InOutNativeResource, InOutLocalizedResource);
+		return EQueryLocalizedResourceResult::Found;
+	}
+	return EQueryLocalizedResourceResult::NotFound;
+}
+
 void FPolyglotTextSource::RegisterPolyglotTextData(const FPolyglotTextData& InPolyglotTextData)
 {
 	check(InPolyglotTextData.IsValid());
diff --git a/Engine/Source/Runtime/Core/Private/Internationalization/PolyglotTextSource.h b/Engine/Source/Runtime/Core/Private/Internationalization/PolyglotTextSource.h
--- a/Engine/Source/Runtime/Core/Private/Internationalization/PolyglotTextSource.h
+++ b/Engine/Source/Runtime/Core/Private/Internationalization/PolyglotTextSource.h
@@ -18,6 +18,7 @@ public:
 	virtual bool GetNativeCultureName(const ELocalizedTextSourceCategory InCategory, FString& OutNativeCultureName) override;
 	virtual void GetLocalizedCultureNames(const ELocalizationLoadFlags InLoadFlags, TSet<FString>& OutLocalizedCultureNames) override;
 	virtual void LoadLocalizedResources(const ELocalizationLoadFlags InLoadFlags, TArrayView<const FString> InPrioritizedCultures, FTextLocalizationResource& InOutNativeResource, FTextLocalizationResource& InOutLocalizedResource) override;
+	virtual EQueryLocalizedResourceResult QueryLocalizedResource(const ELocalizationLoadFlags InLoadFlags, TArrayView<const FString> InPrioritizedCultures, const FTextId InTextId, FTextLocalizationResource& InOutNativeResource, FTextLocalizationResource& InOutLocalizedResource) override;
 
 	/**
 	 * Register a polyglot text data with this localized text source.
@@ -35,6 +36,11 @@ private:
 	 */
 	void UnregisterCultureNames(const FPolyglotTextData& InPolyglotTextData);
 
+	/**
+	 * Add the given polyglot data to the output localization resources (common logic to LoadLocalizedResources and QueryLocalizedResource).
+	 */
+	void AddPolyglotDataToResource(const FPolyglotTextData& InPolyglotTextData, const ELocalizationLoadFlags InLoadFlags, TArrayView<const FString> InPrioritizedCultures, FTextLocalizationResource& InOutNativeResource, FTextLocalizationResource& InOutLocalizedResource) const;
+
 	struct FCultureInfo
 	{
 		TMap<FString, int32> NativeCultures;
diff --git a/Engine/Source/Runtime/Core/Private/Internationalization/TextLocalizationManager.cpp b/Engine/Source/Runtime/Core/Private/Internationalization/TextLocalizationManager.cpp
--- a/Engine/Source/Runtime/Core/Private/Internationalization/TextLocalizationManager.cpp
+++ b/Engine/Source/Runtime/Core/Private/Internationalization/TextLocalizationManager.cpp
@@ -21,6 +21,7 @@
 #include "Stats/Stats.h"
 #include "Misc/ConfigCacheIni.h"
 #include "Misc/App.h"
+#include "Misc/CoreDelegates.h"
 #include "Templates/UniquePtr.h"
 
 DEFINE_LOG_CATEGORY_STATIC(LogTextLocalizationManager, Log, All);
@@ -393,11 +394,14 @@ FTextLocalizationManager::FTextLocalizationManager()
 	: bIsInitialized(false)
 	, SynchronizationObject()
 	, TextRevisionCounter(0)
+	, LocResTextSource(MakeShared<FLocalizationResourceTextSource>())
 	, PolyglotTextSource(MakeShared<FPolyglotTextSource>())
 {
 	const bool bRefreshResources = false;
-	RegisterTextSource(MakeShared<FLocalizationResourceTextSource>(), bRefreshResources);
+	RegisterTextSource(LocResTextSource.ToSharedRef(), bRefreshResources);
 	RegisterTextSource(PolyglotTextSource.ToSharedRef(), bRefreshResources);
+
+	FCoreDelegates::OnPakFileMounted.AddRaw(this, &FTextLocalizationManager::OnPakFileMounted);
 }
 
 void FTextLocalizationManager::DumpMemoryInfo()
@@ -836,6 +840,82 @@ void FTextLocalizationManager::RefreshResources()
 	LoadLocalizationResourcesForCulture(FInternationalization::Get().GetCurrentLanguage()->GetName(), LocLoadFlags);
 }
 
+void FTextLocalizationManager::OnPakFileMounted(const TCHAR* PakFilename, const int32 ChunkId)
+{
+	if (ChunkId == INDEX_NONE || ChunkId == 0)
+	{
+		// Skip non-chunked PAK files, and chunk 0 as that contains the standard localization data
+		return;
+	}
+
+	// Track this so that full resource refreshes (eg, changing culture) work as expected
+	LocResTextSource->RegisterChunkId(ChunkId);
+
+	if (!bIsInitialized)
+	{
+		// If we're not yet initialized then don't bother patching, as the full initialization path will load the data for this chunk
+		return;
+	}
+
+	// Load the resources from each target in this chunk
+	// Note: We only allow game localization targets to be chunked, and the layout is assumed to follow our standard pattern (as used by the localization dashboard and FLocTextHelper)
+	const TArray<FString> ChunkedLocalizationTargets = FLocalizationResourceTextSource::GetChunkedLocalizationTargets();
+	const TArray<FString> PrioritizedCultureNames = FInternationalization::Get().GetPrioritizedCultureNames(FInternationalization::Get().GetCurrentLanguage()->GetName());
+	const ELocalizationLoadFlags LocLoadFlags = ELocalizationLoadFlags::Game;
+	FTextLocalizationResource UnusedNativeResource;
+	FTextLocalizationResource LocalizedResource;
+	{
+		TArray<FString> PrioritizedLocalizationPaths;
+		for (const FString& LocalizationTarget : ChunkedLocalizationTargets)
+		{
+			const FString LocalizationTargetForChunk = TextLocalizationResourceUtil::GetLocalizationTargetNameForChunkId(LocalizationTarget, ChunkId);
+			PrioritizedLocalizationPaths.Add(FPaths::ProjectContentDir() / TEXT("Localization") / LocalizationTargetForChunk);
+		}
+		LocResTextSource->LoadLocalizedResourcesFromPaths(TArrayView<FString>(), PrioritizedLocalizationPaths, TArrayView<FString>(), LocLoadFlags, PrioritizedCultureNames, UnusedNativeResource, LocalizedResource);
+	}
+
+	// Allow any higher priority text sources to override the text loaded for the chunk (eg, to allow polyglot hot-fixes to take priority)
+	// Note: If any text sources don't support dynamic queries, then we must do a much slower full refresh instead :(
+	bool bNeedsFullRefresh = false;
+	{
+		// Copy the IDs array as QueryLocalizedResource can update the map
+		TArray<FTextId> ChunkTextIds;
+		LocalizedResource.Entries.GenerateKeyArray(ChunkTextIds);
+
+		for (const TSharedPtr<ILocalizedTextSource>& LocalizedTextSource : LocalizedTextSources)
+		{
+			if (LocalizedTextSource->GetPriority() <= LocResTextSource->GetPriority())
+			{
+				continue;
+			}
+
+			for (const FTextId& ChunkTextId : ChunkTextIds)
+			{
+				if (LocalizedTextSource->QueryLocalizedResource(LocLoadFlags, PrioritizedCultureNames, ChunkTextId, UnusedNativeResource, LocalizedResource) == EQueryLocalizedResourceResult::NotImplemented)
+				{
+					bNeedsFullRefresh = true;
+					break;
+				}
+			}
+
+			if (bNeedsFullRefresh)
+			{
+				break;
+			}
+		}
+	}
+
+	// Apply the new data
+	if (bNeedsFullRefresh)
+	{
+		RefreshResources();
+	}
+	else
+	{
+		UpdateFromLocalizations(MoveTemp(LocalizedResource), /*bDirtyTextRevision*/true);
+	}
+}
+
 void FTextLocalizationManager::OnCultureChanged()
 {
     if (!bIsInitialized)
diff --git a/Engine/Source/Runtime/Core/Private/Internationalization/TextLocalizationResource.cpp b/Engine/Source/Runtime/Core/Private/Internationalization/TextLocalizationResource.cpp
--- a/Engine/Source/Runtime/Core/Private/Internationalization/TextLocalizationResource.cpp
+++ b/Engine/Source/Runtime/Core/Private/Internationalization/TextLocalizationResource.cpp
@@ -671,5 +671,12 @@ const TArray<FString>& TextLocalizationResourceUtil::GetDisabledLocalizationTarg
 	return DisabledLocalizationTargets;
 }
 
+FString TextLocalizationResourceUtil::GetLocalizationTargetNameForChunkId(const FString& InLocalizationTargetName, const int32 InChunkId)
+{
+	return InChunkId == INDEX_NONE || InChunkId == 0
+		? InLocalizationTargetName
+		: FString::Printf(TEXT("%s_locchunk%d"), *InLocalizationTargetName, InChunkId);
+}
+
 #undef PRELOAD_LOCMETA_FILES
 #undef PRELOAD_LOCRES_FILES
diff --git a/Engine/Source/Runtime/Core/Public/Internationalization/ILocalizedTextSource.h b/Engine/Source/Runtime/Core/Public/Internationalization/ILocalizedTextSource.h
--- a/Engine/Source/Runtime/Core/Public/Internationalization/ILocalizedTextSource.h
+++ b/Engine/Source/Runtime/Core/Public/Internationalization/ILocalizedTextSource.h
@@ -5,6 +5,7 @@
 #include "CoreTypes.h"
 #include "Containers/ContainersFwd.h"
 #include "Containers/ArrayView.h"
+#include "Internationalization/TextKey.h"
 #include "Internationalization/LocalizedTextSourceTypes.h"
 #include "CoreGlobals.h"
 
@@ -43,6 +44,14 @@ public:
 	 */
 	virtual void LoadLocalizedResources(const ELocalizationLoadFlags InLoadFlags, TArrayView<const FString> InPrioritizedCultures, FTextLocalizationResource& InOutNativeResource, FTextLocalizationResource& InOutLocalizedResource) = 0;
 
+	/**
+	 * Query a localized resource from this localized text source for the given cultures and ID into the given maps (ELocalizationLoadFlags controls which resources should be queried).
+	 */
+	virtual EQueryLocalizedResourceResult QueryLocalizedResource(const ELocalizationLoadFlags InLoadFlags, TArrayView<const FString> InPrioritizedCultures, const FTextId InTextId, FTextLocalizationResource& InOutNativeResource, FTextLocalizationResource& InOutLocalizedResource)
+	{
+		return EQueryLocalizedResourceResult::NotImplemented;
+	}
+
 	/** Should we load native data based on the given load flags and environment? */
 	static FORCEINLINE bool ShouldLoadNative(const ELocalizationLoadFlags InLoadFlags)
 	{
diff --git a/Engine/Source/Runtime/Core/Public/Internationalization/LocalizedTextSourceTypes.h b/Engine/Source/Runtime/Core/Public/Internationalization/LocalizedTextSourceTypes.h
--- a/Engine/Source/Runtime/Core/Public/Internationalization/LocalizedTextSourceTypes.h
+++ b/Engine/Source/Runtime/Core/Public/Internationalization/LocalizedTextSourceTypes.h
@@ -16,6 +16,19 @@ enum class ELocalizedTextSourceCategory : uint8
 	Editor,
 };
 
+/**
+ * Result codes from calling QueryLocalizedResourceResult.
+ */
+enum class EQueryLocalizedResourceResult : uint8
+{
+	/** Indicates the query found a matching entry and added its result */
+	Found,
+	/** Indicates that the query failed to find a matching entry */
+	NotFound,
+	/** Indicates that the query failed as this text source doesn't support queries */
+	NotImplemented,
+};
+
 /**
  * Load flags used in localization initialization.
  */
diff --git a/Engine/Source/Runtime/Core/Public/Internationalization/TextLocalizationManager.h b/Engine/Source/Runtime/Core/Public/Internationalization/TextLocalizationManager.h
--- a/Engine/Source/Runtime/Core/Public/Internationalization/TextLocalizationManager.h
+++ b/Engine/Source/Runtime/Core/Public/Internationalization/TextLocalizationManager.h
@@ -227,6 +227,9 @@ public:
 	FTextRevisionChangedEvent OnTextRevisionChangedEvent;
 
 private:
+	/** Callback for when a PAK file is loaded. Loads any chunk specific localization resources. */
+	void OnPakFileMounted(const TCHAR* PakFilename, const int32 ChunkId);
+
 	/** Callback for changes in culture. Loads the new culture's localization resources. */
 	void OnCultureChanged();
 
@@ -251,6 +254,9 @@ private:
 	/** Array of registered localized text sources, sorted by priority (@see RegisterTextSource) */
 	TArray<TSharedPtr<ILocalizedTextSource>> LocalizedTextSources;
 
+	/** The LocRes text source (this is also added to LocalizedTextSources, but we keep a pointer to it directly so we can patch in chunked LocRes data at runtime) */
+	TSharedPtr<class FLocalizationResourceTextSource> LocResTextSource;
+
 	/** The polyglot text source (this is also added to LocalizedTextSources, but we keep a pointer to it directly so we can add new polyglot data to it at runtime) */
 	TSharedPtr<class FPolyglotTextSource> PolyglotTextSource;
 };
diff --git a/Engine/Source/Runtime/Core/Public/Internationalization/TextLocalizationResource.h b/Engine/Source/Runtime/Core/Public/Internationalization/TextLocalizationResource.h
--- a/Engine/Source/Runtime/Core/Public/Internationalization/TextLocalizationResource.h
+++ b/Engine/Source/Runtime/Core/Public/Internationalization/TextLocalizationResource.h
@@ -158,4 +158,9 @@ CORE_API TArray<FString> GetLocalizedCultureNames(const TArray<FString>& InLocal
  */
 CORE_API const TArray<FString>& GetDisabledLocalizationTargets();
 
+/**
+ * Get the name that the given localization target should have for the given chunk ID.
+ */
+CORE_API FString GetLocalizationTargetNameForChunkId(const FString& InLocalizationTargetName, const int32 InChunkId);
+
 }
