Changed the manifest list file's extension from 'manifest' to 'manifestlist' so it is no longer confused by the BuildPatchTool with real manifests.
#rnx

diff --git a/Engine/Source/Programs/AutomationTool/Scripts/CopyBuildToStagingDirectory.Automation.cs b/Engine/Source/Programs/AutomationTool/Scripts/CopyBuildToStagingDirectory.Automation.cs
--- a/Engine/Source/Programs/AutomationTool/Scripts/CopyBuildToStagingDirectory.Automation.cs
+++ b/Engine/Source/Programs/AutomationTool/Scripts/CopyBuildToStagingDirectory.Automation.cs
@@ -2759,7 +2759,7 @@ public partial class Project : CommandUtils
 		//Create the directory if it doesn't exist
 		InternalUtils.SafeCreateDirectory(Dir);
 
-		string FileName = CombinePaths(Dir, PlatformStr.ToLower() + ".manifest");
+		string FileName = CombinePaths(Dir, PlatformStr.ToLower() + ".manifestlist");
 		using (JsonWriter Writer = new JsonWriter(FileName))
 		{
 			Writer.WriteObjectStart();
