Adjust filtering settings for cloud content.  Switched IOS and Android to use monolithic process cloud content script.
Added support for generating additional backgrounddownloadsini files for builds that already exist.
#test preflight

diff --git a/Engine/Source/Programs/AutomationTool/AutomationUtils/MCPPublic.cs b/Engine/Source/Programs/AutomationTool/AutomationUtils/MCPPublic.cs
--- a/Engine/Source/Programs/AutomationTool/AutomationUtils/MCPPublic.cs
+++ b/Engine/Source/Programs/AutomationTool/AutomationUtils/MCPPublic.cs
@@ -141,6 +141,11 @@ namespace EpicGames.MCP.Automation
 		/// </summary>
 		Switch,
 
+		/// <summary>
+		/// Xbox One Platform
+		/// </summary>
+		XboxOne,
+
 
 	}
 
@@ -259,7 +264,26 @@ namespace EpicGames.MCP.Automation
 			{
 				return MCPPlatform.Linux;
 			}
-
+			else if (TargetPlatform == UnrealTargetPlatform.IOS)
+			{
+				return MCPPlatform.IOS;
+			}
+			else if (TargetPlatform == UnrealTargetPlatform.Android)
+			{
+				return MCPPlatform.Android;
+			}
+			else if (TargetPlatform == UnrealTargetPlatform.PS4)
+			{
+				return MCPPlatform.PS4;
+			}
+			else if (TargetPlatform == UnrealTargetPlatform.XboxOne)
+			{
+				return MCPPlatform.XboxOne;
+			}
+			else if (TargetPlatform == UnrealTargetPlatform.Switch)
+			{
+				return MCPPlatform.Switch;
+			}
 			throw new AutomationException("Platform {0} is not properly supported by the MCP backend yet", TargetPlatform);
         }
 
@@ -284,7 +308,26 @@ namespace EpicGames.MCP.Automation
 			{
 				return UnrealTargetPlatform.Linux;
 			}
-
+			else if (TargetPlatform == MCPPlatform.IOS)
+			{
+				return UnrealTargetPlatform.IOS;
+			}
+			else if (TargetPlatform == MCPPlatform.Android)
+			{
+				return UnrealTargetPlatform.Android;
+			}
+			else if (TargetPlatform == MCPPlatform.PS4)
+			{
+				return UnrealTargetPlatform.PS4;
+			}
+			else if (TargetPlatform == MCPPlatform.XboxOne)
+			{
+				return UnrealTargetPlatform.XboxOne;
+			}
+			else if (TargetPlatform == MCPPlatform.Switch)
+			{
+				return UnrealTargetPlatform.Switch;
+			}
 			throw new AutomationException("Platform {0} is not properly supported by the MCP backend yet", TargetPlatform);
         }
 
