Change code signing to use the machine store to find certs.

#rnx

diff --git a/Engine/Source/Programs/AutomationTool/AutomationUtils/CommandUtils.cs b/Engine/Source/Programs/AutomationTool/AutomationUtils/CommandUtils.cs
--- a/Engine/Source/Programs/AutomationTool/AutomationUtils/CommandUtils.cs
+++ b/Engine/Source/Programs/AutomationTool/AutomationUtils/CommandUtils.cs
@@ -2655,7 +2655,7 @@ namespace AutomationTool
 		/// <summary>
 		/// Whether to look in the machine store for signing certificates, rather than the user store.
 		/// </summary>
-		static bool bUseMachineStoreForCertificates = false;
+		static bool bUseMachineStoreForCertificates = true;
 
 		public enum SignatureType
 		{
