UAT: Staging: When creating the base release directory for a chunkinstall build, copy pak sig files if they exist

diff --git a/Engine/Source/Programs/AutomationTool/Scripts/CopyBuildToStagingDirectory.Automation.cs b/Engine/Source/Programs/AutomationTool/Scripts/CopyBuildToStagingDirectory.Automation.cs
--- a/Engine/Source/Programs/AutomationTool/Scripts/CopyBuildToStagingDirectory.Automation.cs
+++ b/Engine/Source/Programs/AutomationTool/Scripts/CopyBuildToStagingDirectory.Automation.cs
@@ -2041,13 +2041,17 @@ public partial class Project : CommandUtils
 					string ChunkInstallBasePath = CombinePaths(Params.ChunkInstallDirectory, SC.FinalCookPlatform);
 					string RawDataPath = CombinePaths(ChunkInstallBasePath, VersionString, PakName);
 					string RawDataPakPath = CombinePaths(RawDataPath, PakName + "-" + SC.FinalCookPlatform + PostFix + ".pak");
+					bool bPakFilesAreSigned = InternalUtils.SafeFileExists(Path.ChangeExtension(OutputLocation.FullName, ".sig"));
+
 					//copy the pak chunk to the raw data folder
-					if (InternalUtils.SafeFileExists(RawDataPakPath, true))
-					{
-						InternalUtils.SafeDeleteFile(RawDataPakPath, true);
-					}
+					InternalUtils.SafeDeleteFile(RawDataPakPath, true);
+					InternalUtils.SafeDeleteFile(Path.ChangeExtension(RawDataPakPath, ".sig"), true);
 					InternalUtils.SafeCreateDirectory(RawDataPath, true);
 					InternalUtils.SafeCopyFile(OutputLocation.FullName, RawDataPakPath);
+					if (bPakFilesAreSigned)
+					{
+						InternalUtils.SafeCopyFile(Path.ChangeExtension(OutputLocation.FullName, ".sig"), Path.ChangeExtension(RawDataPakPath, ".sig"), true);
+					}
 					InternalUtils.SafeDeleteFile(OutputLocation.FullName, true);
 
 					if (Params.IsGeneratingPatch)
@@ -2060,6 +2064,10 @@ public partial class Project : CommandUtils
 						// for distribution.
 						string SourceRawDataPakPath = CombinePaths(RawDataPath, PakName + "-" + SC.FinalCookPlatform + ".pak");
 						InternalUtils.SafeCopyFile(PatchSourceContentPath, SourceRawDataPakPath);
+						if (bPakFilesAreSigned)
+						{
+							InternalUtils.SafeCopyFile(PatchSourceContentPath, Path.ChangeExtension(SourceRawDataPakPath, ".sig"), true);
+						}
 					}
 
 					string BuildRoot = MakePathSafeToUseWithCommandLine(RawDataPath);
