Gauntlet - Improve handing of XB1 builds to address several issues -

* Paths to projects with spaces in the name now work correctly on XB1
* Loose builds for content-only projects no longer crash due to missing project arg (UE-78746, UE-78739)
* Packaged projects with version numbers in the name are correctly recognized (UE-85189)

diff --git a/Engine/Source/Programs/AutomationTool/Gauntlet/Framework/Gauntlet.Utils.cs b/Engine/Source/Programs/AutomationTool/Gauntlet/Framework/Gauntlet.Utils.cs
--- a/Engine/Source/Programs/AutomationTool/Gauntlet/Framework/Gauntlet.Utils.cs
+++ b/Engine/Source/Programs/AutomationTool/Gauntlet/Framework/Gauntlet.Utils.cs
@@ -1449,7 +1449,12 @@ namespace Gauntlet
 	{
 		public static bool MatchAndApplyGroups(string InContent, string RegEx, Action<string[]> InFunc)
 		{
-			Match M = Regex.Match(InContent, RegEx, RegexOptions.IgnoreCase);
+			return MatchAndApplyGroups(InContent, RegEx, RegexOptions.IgnoreCase, InFunc);
+		}
+
+		public static bool MatchAndApplyGroups(string InContent, string RegEx, RegexOptions Options, Action<string[]> InFunc)
+		{
+			Match M = Regex.Match(InContent, RegEx, Options);
 
 			IEnumerable<string> StringMatches = null;
 
@@ -1462,4 +1467,24 @@ namespace Gauntlet
 			return M.Success;
 		}
 	}
+
+	public static class DirectoryUtils
+	{
+		/// <summary>
+		/// Enumerate files from a given directory that pass the specified regex
+		/// </summary>
+		/// <param name="BaseDir">Base directory to search in</param>
+		/// <param name="Pattern">Pattern for matching files</param>
+		/// <param name="Option">Options for the search</param>
+		/// <returns>Sequence of file references</returns>
+		public static IEnumerable<string> FindFiles(string BaseDir, Regex Pattern)
+		{
+			IEnumerable<string> Files = System.IO.Directory.EnumerateFiles(BaseDir, "*");
+
+			Files = Files.Where(F => Pattern.IsMatch(F));
+
+			return Files.ToArray();
+		}
+	}
+
 }
diff --git a/Engine/Source/Programs/AutomationTool/Gauntlet/Unreal/Base/Gauntlet.UnrealAppConfig.cs b/Engine/Source/Programs/AutomationTool/Gauntlet/Unreal/Base/Gauntlet.UnrealAppConfig.cs
--- a/Engine/Source/Programs/AutomationTool/Gauntlet/Unreal/Base/Gauntlet.UnrealAppConfig.cs
+++ b/Engine/Source/Programs/AutomationTool/Gauntlet/Unreal/Base/Gauntlet.UnrealAppConfig.cs
@@ -54,5 +54,17 @@ namespace Gauntlet
 
 		// new system
 		public IBuild Build { get; set; }
+
+		/// <summary>
+		/// Constructor that sets some required values to defaults
+		/// </summary>
+		public UnrealAppConfig()
+		{
+			Name = "UnrealApp";
+			ProjectName = "UnknownProject";
+			CommandLine = "";
+			Configuration = UnrealTargetConfiguration.Development;
+			Sandbox = "Gauntlet";
+		}
 	}
 }
\ No newline at end of file
diff --git a/Engine/Source/Programs/AutomationTool/Scripts/SyncProject.Automation.cs b/Engine/Source/Programs/AutomationTool/Scripts/SyncProject.Automation.cs
--- a/Engine/Source/Programs/AutomationTool/Scripts/SyncProject.Automation.cs
+++ b/Engine/Source/Programs/AutomationTool/Scripts/SyncProject.Automation.cs
@@ -61,9 +61,15 @@ class SyncProject : BuildCommand
 
 		IEnumerable<string> ExplicitPaths = PathArg.Split(new char[] { ',',';' }, StringSplitOptions.RemoveEmptyEntries).Select(S => S.Trim());
 
-		if (CL == 0 && !ForceSync)
+		if (CL == 0)
 		{
-			throw new AutomationException("If using 0 CL to remove files -force must also be specified (and you may also want -projectonly)");
+			if (!ForceSync)
+			{
+				throw new AutomationException("If using 0 CL to remove files -force must also be specified");
+			}
+
+			// Cannot unsync the engine folder
+			ProjectOnly = true;
 		}
 		else if (CL == -1)
 		{
