[UAT]
 - Add "ExtraTargetsToStageWithClient" parameter which allows you to select other build targets to stage alongside a client build.
 - Added the notion of a "primary target" for a given target type, so a project can have many editor targets but know which was is the primary one that it should use for cooking
 - Fixed a typo in UAT param parsing where the check for the correct number of game targets was looking at the number of editor targets

diff --git a/Engine/Source/Programs/AutomationTool/AutomationUtils/ProjectParams.cs b/Engine/Source/Programs/AutomationTool/AutomationUtils/ProjectParams.cs
--- a/Engine/Source/Programs/AutomationTool/AutomationUtils/ProjectParams.cs
+++ b/Engine/Source/Programs/AutomationTool/AutomationUtils/ProjectParams.cs
@@ -692,6 +692,12 @@ namespace AutomationTool
 			}
 
 			this.LogWindow = GetParamValueIfNotSpecified(Command, LogWindow, this.LogWindow, "logwindow");
+			string ExtraTargetsToStageWithClientString = null;
+			ExtraTargetsToStageWithClientString = ParseParamValueIfNotSpecified(Command, ExtraTargetsToStageWithClientString, "ExtraTargetsToStageWithClient", null);
+			if (!string.IsNullOrEmpty(ExtraTargetsToStageWithClientString))
+			{
+				this.ExtraTargetsToStageWithClient = new ParamList<string>(ExtraTargetsToStageWithClientString.Split('+'));
+			}
 			this.Stage = GetParamValueIfNotSpecified(Command, Stage, this.Stage, "stage");
 			this.SkipStage = GetParamValueIfNotSpecified(Command, SkipStage, this.SkipStage, "skipstage");
 			if (this.SkipStage)
@@ -1694,6 +1700,11 @@ namespace AutomationTool
 		[Help("bundlename", "string to use as the bundle name when deploying to mobile device")]
         public string BundleName;
 
+		//<summary>
+		/// Stage: Specifies a list of extra targets that should be staged along with a client
+		/// </summary>
+		public ParamList<string> ExtraTargetsToStageWithClient = new ParamList<string>();
+
         /// <summary>
         /// On Windows, adds an executable to the root of the staging directory which checks for prerequisites being 
 		/// installed and launches the game with a path to the .uproject file.
@@ -1961,6 +1972,10 @@ namespace AutomationTool
 		{
 			return DetectedTargets.FindAll(Target => Target.Rules.Type == DesiredType).ConvertAll(Target => Target.TargetName);
 		}
+		private List<String> PrimaryTargetNamesOfType(TargetType DesiredType)
+		{
+			return DetectedTargets.FindAll(Target => (Target.Rules.Type == DesiredType && Target.Rules.IsPrimaryTarget)).ConvertAll(Target => Target.TargetName);
+		}
 
 		private String ChooseTarget(List<String> Targets, TargetType Type)
 		{
@@ -2100,7 +2115,7 @@ namespace AutomationTool
 				List<String> AvailableGameTargets = TargetNamesOfType(TargetType.Game);
 				List<String> AvailableClientTargets = TargetNamesOfType(TargetType.Client);
 				List<String> AvailableServerTargets = TargetNamesOfType(TargetType.Server);
-				List<String> AvailableEditorTargets = TargetNamesOfType(TargetType.Editor);
+				List<String> AvailableEditorTargets = PrimaryTargetNamesOfType(TargetType.Editor);
 
 				// That should cover all detected targets; Program targets are handled separately.
 				System.Diagnostics.Debug.Assert(DetectedTargets.Count == (AvailableGameTargets.Count + AvailableClientTargets.Count + AvailableServerTargets.Count + AvailableEditorTargets.Count));
@@ -2112,7 +2127,7 @@ namespace AutomationTool
 				}
 				else if (AvailableGameTargets.Count > 0)
 				{
-					if (AvailableEditorTargets.Count > 1)
+					if (AvailableGameTargets.Count > 1)
 					{
 						throw new AutomationException("There can be only one Game target per project.");
 					}
@@ -2124,7 +2139,7 @@ namespace AutomationTool
 				{
 					if (AvailableEditorTargets.Count > 1)
 					{
-						throw new AutomationException("There can be only one Editor target per project.");
+						throw new AutomationException("There can be only one primary editor target per project.");
 					}
 
 					EditorTarget = AvailableEditorTargets.First();
@@ -2204,6 +2219,11 @@ namespace AutomationTool
 					}
 
 					ClientCookedTargetsList = new ParamList<string>(GameTarget);
+					
+					if (ExtraTargetsToStageWithClient != null)
+					{
+						ClientCookedTargetsList.AddRange(ExtraTargetsToStageWithClient);
+					}
 				}
 				else
 				{
diff --git a/Engine/Source/Programs/UnrealBuildTool/Configuration/TargetRules.cs b/Engine/Source/Programs/UnrealBuildTool/Configuration/TargetRules.cs
--- a/Engine/Source/Programs/UnrealBuildTool/Configuration/TargetRules.cs
+++ b/Engine/Source/Programs/UnrealBuildTool/Configuration/TargetRules.cs
@@ -197,6 +197,11 @@ namespace UnrealBuildTool
 		/// </summary>
 		public global::UnrealBuildTool.TargetType Type = global::UnrealBuildTool.TargetType.Game;
 
+		/// <summary>
+		/// Indicates that this target is the primary one of its type. i.e. it's the primary editor target that should be used by UAT for cooking.
+		/// </summary>
+		public bool IsPrimaryTarget = true;
+
 		/// <summary>
 		/// Specifies the engine version to maintain backwards-compatible default build settings with (eg. DefaultSettingsVersion.Release_4_23, DefaultSettingsVersion.Release_4_24). Specify DefaultSettingsVersion.Latest to always
 		/// use defaults for the current engine version, at the risk of introducing build errors while upgrading.
