PakFileRules can not be disabled.

diff --git a/Engine/Source/Programs/AutomationTool/Scripts/CopyBuildToStagingDirectory.Automation.cs b/Engine/Source/Programs/AutomationTool/Scripts/CopyBuildToStagingDirectory.Automation.cs
--- a/Engine/Source/Programs/AutomationTool/Scripts/CopyBuildToStagingDirectory.Automation.cs
+++ b/Engine/Source/Programs/AutomationTool/Scripts/CopyBuildToStagingDirectory.Automation.cs
@@ -1428,6 +1428,9 @@ public partial class Project : CommandUtils
 		// Rather to allow overriding the chunk manifest, if false will only modify loose files
 		public bool bOverrideChunkManifest;
 
+		// Whether pak file rule is disabled
+		public bool bDisabled;
+
 		// List of pak files to use instead of manifest
 		public List<string> OverridePaks;
 	};
@@ -1555,6 +1558,7 @@ public partial class Project : CommandUtils
 			PakRules.Name = SectionName;
 			PakRulesConfig.TryGetValue(SectionName, "bExcludeFromPaks", out PakRules.bExcludeFromPaks);
 			PakRulesConfig.TryGetValue(SectionName, "bOverrideChunkManifest", out PakRules.bOverrideChunkManifest);
+			PakRulesConfig.TryGetValue(SectionName, "bDisabled", out PakRules.bDisabled);
 			string PakString;
 			PakRulesConfig.TryGetValue(SectionName, "OverridePaks", out PakString);
 
@@ -1613,7 +1617,7 @@ public partial class Project : CommandUtils
 		// Search in order, return on first match
 		foreach (var PakRules in RulesList)
 		{
-			if (PakRules.Filter.Matches(StagingFile.Key))
+			if (!PakRules.bDisabled && PakRules.Filter.Matches(StagingFile.Key))
 			{
 				if (ModifyPakList != null && ModifyPakList.Count > 0)
 				{
@@ -1637,7 +1641,22 @@ public partial class Project : CommandUtils
 				if (PakRules.OverridePaks != null && ModifyPakList != null && ChunkNameToDefinition != null)
 				{
 					ModifyPakList.Clear();
-					ModifyPakList.UnionWith(PakRules.OverridePaks.Select(x => ChunkNameToDefinition[x]));
+					ModifyPakList.UnionWith(PakRules.OverridePaks.Select(x =>
+					{
+						if (!ChunkNameToDefinition.ContainsKey(x))
+						{
+							LogInformation("With pak rules, {0} is moved to {1}", StagingFile.Key, x);
+							ChunkDefinition TargetChunk = new ChunkDefinition(x);
+							// TODO: Need to set encryptionGuid?
+							//TargetChunk.RequestedEncryptionKeyGuid = Chunk.RequestedEncryptionKeyGuid;
+							//TargetChunk.EncryptionKeyGuid = Chunk.EncryptionKeyGuid;
+							ChunkNameToDefinition.Add(x, TargetChunk);
+
+							return TargetChunk;
+						}
+
+						return ChunkNameToDefinition[x];
+					}));
 					//LogInformation("Setting pak assignment for file {0} to {1}", StagingFile.Key, string.Join(", ", PakRules.OverridePaks));
 				}
 				else if (bExcludeFromPaks)
