Fix staging fail with Audio3D plugin enabled
#rnx

diff --git a/Engine/Source/Programs/AutomationTool/AutomationUtils/DeploymentContext.cs b/Engine/Source/Programs/AutomationTool/AutomationUtils/DeploymentContext.cs
--- a/Engine/Source/Programs/AutomationTool/AutomationUtils/DeploymentContext.cs
+++ b/Engine/Source/Programs/AutomationTool/AutomationUtils/DeploymentContext.cs
@@ -427,11 +427,11 @@ public class DeploymentContext //: ProjectParams
 
 		// Build a list of restricted folder names. This will comprise all other restricted platforms, plus standard restricted folder names such as NoRedist, NotForLicensees, etc...
 		RestrictedFolderNames.UnionWith(PlatformExports.GetPlatformFolderNames());
-		foreach(UnrealTargetPlatform StagePlatform in StageTargetPlatform.GetStagePlatforms())
+		RestrictedFolderNames.UnionWith(RestrictedFolder.GetNames());
+		foreach (UnrealTargetPlatform StagePlatform in StageTargetPlatform.GetStagePlatforms())
 		{
 			RestrictedFolderNames.ExceptWith(PlatformExports.GetIncludedFolderNames(StagePlatform));
 		}
-		RestrictedFolderNames.UnionWith(RestrictedFolder.GetNames());
 		RestrictedFolderNames.Remove(StageTargetPlatform.IniPlatformType.ToString());
 
 		// Read the game config files
