Make InstallBundleConfig a hierarchical INI.
Fix FConfigCacheIni::GetSectionNames returning names in reverse order.
Make sure build scripts properly respect bundle order.

change 2/2 - everything not internal

diff --git a/Engine/Config/BaseInstallBundle.ini b/Engine/Config/BaseInstallBundle.ini
new file mode 100644
--- /dev/null
+++ b/Engine/Config/BaseInstallBundle.ini
@@ -0,0 +1 @@
+; Defines install bundles and options
\ No newline at end of file
diff --git a/Engine/Source/Programs/AutomationTool/AutomationUtils/BundleUtils.cs b/Engine/Source/Programs/AutomationTool/AutomationUtils/BundleUtils.cs
--- a/Engine/Source/Programs/AutomationTool/AutomationUtils/BundleUtils.cs
+++ b/Engine/Source/Programs/AutomationTool/AutomationUtils/BundleUtils.cs
@@ -28,27 +28,40 @@ namespace AutomationUtils.Automation
 			public List<string> FileRegex { get; set; }
 			public bool bFoundParent { get; set; }
 			public bool bContainsShaderLibrary { get; set; }
-			public int Priority { get; set; }
+			public int Order { get; set; }
 			public string ExecFileName { get; set; }
 		};
-		public static bool LoadBundleConfig(string BundleIniFile, ref Dictionary<string, BundleSettings> Bundles)
+		public static bool LoadBundleConfig(DirectoryReference ProjectDir, UnrealTargetPlatform Platform, out List<BundleSettings> Bundles)
 		{
-			if ( System.IO.File.Exists(BundleIniFile) == false )
+			Bundles = new List<BundleSettings>();
+
+			ConfigHierarchy BundleConfig = ConfigCache.ReadHierarchy(ConfigHierarchyType.InstallBundle, ProjectDir, Platform);
+			if (BundleConfig == null)
 			{
-				CommandUtils.LogWarning("Unable to find bundle config ini file  {0}", BundleIniFile);
+				CommandUtils.LogWarning("Unable to find bundle config ini file {0}", ConfigHierarchyType.InstallBundle.ToString());
 				return false;
 			}
 
-			FileReference BundleFileReference = new FileReference(BundleIniFile);
+			const string BundleDefinitionPrefix = "InstallBundleDefinition ";
 
-			ConfigHierarchy BundleConfig = new ConfigHierarchy(new ConfigFile[] { new ConfigFile(BundleFileReference) });
-			int PriorityCounter = 0;
 			foreach (string SectionName in BundleConfig.SectionNames)
 			{
+				if (!SectionName.StartsWith(BundleDefinitionPrefix))
+					continue;
+
 				BundleSettings Bundle = new BundleSettings();
-				Bundle.Name = SectionName;
-				Bundle.Priority = PriorityCounter;
-				++PriorityCounter;
+				Bundle.Name = SectionName.Substring(BundleDefinitionPrefix.Length);
+				{
+					int Order;
+					if(BundleConfig.GetInt32(SectionName, "Order", out Order))
+					{
+						Bundle.Order = Order;
+					}
+					else
+					{
+						Bundle.Order = int.MaxValue;
+					}
+				}
 				{
 					string ExecFileName;
 					BundleConfig.GetString(SectionName, "ExecFileName", out ExecFileName);
@@ -74,9 +87,12 @@ namespace AutomationUtils.Automation
 					Bundle.Tags = new List<string>();
 				}
 
-				Bundles.Add(SectionName, Bundle);
+				Bundles.Add(Bundle);
 			}
 
+			// Use OrderBy and not Sort because OrderBy is stable
+			Bundles = Bundles.OrderBy(Bundle => Bundle.Order).ToList();
+
 			return true;
 		}
 	}
diff --git a/Engine/Source/Programs/AutomationTool/Scripts/CopyBuildToStagingDirectory.Automation.cs b/Engine/Source/Programs/AutomationTool/Scripts/CopyBuildToStagingDirectory.Automation.cs
--- a/Engine/Source/Programs/AutomationTool/Scripts/CopyBuildToStagingDirectory.Automation.cs
+++ b/Engine/Source/Programs/AutomationTool/Scripts/CopyBuildToStagingDirectory.Automation.cs
@@ -1205,6 +1205,7 @@ public partial class Project : CommandUtils
 			case "input.ini":
 			case "scalability.ini":
 			case "runtimeoptions.ini":
+			case "installbundle.ini":
 				return true;
 			case "crypto.ini":
 			case "editor.ini":
diff --git a/Engine/Source/Programs/UnrealBuildTool/System/ConfigHierarchy.cs b/Engine/Source/Programs/UnrealBuildTool/System/ConfigHierarchy.cs
--- a/Engine/Source/Programs/UnrealBuildTool/System/ConfigHierarchy.cs
+++ b/Engine/Source/Programs/UnrealBuildTool/System/ConfigHierarchy.cs
@@ -44,6 +44,11 @@ namespace UnrealBuildTool
 		/// BaseEditorSettings.ini, DefaultEditorSettings.ini, etc...
 		/// </summary>
 		EditorSettings,
+
+		/// <summary>
+		/// BaseInstallBundle.ini, DefaultInstallBundle.ini, etc...
+		/// </summary>
+		InstallBundle,
 	}
 
 	/// <summary>
diff --git a/Engine/Source/Runtime/Core/Private/Misc/ConfigCacheIni.cpp b/Engine/Source/Runtime/Core/Private/Misc/ConfigCacheIni.cpp
--- a/Engine/Source/Runtime/Core/Private/Misc/ConfigCacheIni.cpp
+++ b/Engine/Source/Runtime/Core/Private/Misc/ConfigCacheIni.cpp
@@ -2591,8 +2591,7 @@ bool FConfigCacheIni::GetSectionNames( const FString& Filename, TArray<FString>&
 		out_SectionNames.Empty(Num());
 		for ( FConfigFile::TIterator It(*File); It; ++It )
 		{
-			// insert each item at the beginning of the array because TIterators return results in reverse order from which they were added
-			out_SectionNames.Insert(It.Key(),0);
+			out_SectionNames.Add(It.Key());
 
 			FCoreDelegates::OnConfigSectionNameRead.Broadcast(*Filename, *It.Key());
 		}
@@ -3822,6 +3821,8 @@ void FConfigCacheIni::InitializeConfigSystem()
 	FConfigCacheIni::LoadGlobalIniFile(GHardwareIni, TEXT("Hardware"));	
 	// Load runtime options
 	FConfigCacheIni::LoadGlobalIniFile(GRuntimeOptionsIni, TEXT("RuntimeOptions"));
+	// Load install bundle config
+	FConfigCacheIni::LoadGlobalIniFile(GInstallBundleIni, TEXT("InstallBundle"));
 	
 	// Load user game settings .ini, allowing merging. This also updates the user .ini if necessary.
 #if PLATFORM_PS4
diff --git a/Engine/Source/Runtime/Core/Private/Misc/CoreGlobals.cpp b/Engine/Source/Runtime/Core/Private/Misc/CoreGlobals.cpp
--- a/Engine/Source/Runtime/Core/Private/Misc/CoreGlobals.cpp
+++ b/Engine/Source/Runtime/Core/Private/Misc/CoreGlobals.cpp
@@ -221,6 +221,7 @@ FString				GInputIni;													/* Input ini filename */
 FString				GGameIni;													/* Game ini filename */
 FString				GGameUserSettingsIni;										/* User Game Settings ini filename */
 FString				GRuntimeOptionsIni;											/* Runtime Options ini filename */
+FString				GInstallBundleIni;											/* Install Bundle ini filename*/
 
 float					GNearClippingPlane				= 10.0f;				/* Near clipping plane */
 
diff --git a/Engine/Source/Runtime/Core/Public/CoreGlobals.h b/Engine/Source/Runtime/Core/Public/CoreGlobals.h
--- a/Engine/Source/Runtime/Core/Public/CoreGlobals.h
+++ b/Engine/Source/Runtime/Core/Public/CoreGlobals.h
@@ -303,6 +303,7 @@ extern CORE_API FString GInputIni;
 extern CORE_API FString GGameIni;
 extern CORE_API FString GGameUserSettingsIni;
 extern CORE_API FString GRuntimeOptionsIni;
+extern CORE_API FString GInstallBundleIni;
 
 extern CORE_API float GNearClippingPlane;
 
diff --git a/Engine/Source/Runtime/InstallBundleManager/Private/InstallBundleUtils.cpp b/Engine/Source/Runtime/InstallBundleManager/Private/InstallBundleUtils.cpp
--- a/Engine/Source/Runtime/InstallBundleManager/Private/InstallBundleUtils.cpp
+++ b/Engine/Source/Runtime/InstallBundleManager/Private/InstallBundleUtils.cpp
@@ -44,6 +44,12 @@ namespace InstallBundleUtil
 		return TEXT("");
 	}
 
+	const FString& GetInstallBundleSectionPrefix()
+	{
+		static FString Prefix(TEXT("InstallBundleDefinition ")); // trailing space intentional
+		return Prefix;
+	}
+
 	FName FInstallBundleManagerKeepAwake::Tag(TEXT("InstallBundleManagerKeepAwake"));
 	FName FInstallBundleManagerKeepAwake::TagWithRendering(TEXT("InstallBundleManagerKeepAwakeWithRendering"));
 
diff --git a/Engine/Source/Runtime/InstallBundleManager/Public/InstallBundleSourceInterface.h b/Engine/Source/Runtime/InstallBundleManager/Public/InstallBundleSourceInterface.h
--- a/Engine/Source/Runtime/InstallBundleManager/Public/InstallBundleSourceInterface.h
+++ b/Engine/Source/Runtime/InstallBundleManager/Public/InstallBundleSourceInterface.h
@@ -26,7 +26,6 @@ public:
 
 	// Called once by bundle manager after constructing the bundle source
 	virtual void Init(
-		TSharedRef<FConfigFile> InstallBundleConfig, 
 		TSharedRef<InstallBundleUtil::FContentRequestStatsMap> InRequestStats,
 		TSharedPtr<IAnalyticsProviderET> AnalyticsProvider,
 		TSharedPtr<InstallBundleUtil::PersistentStats::FPersistentStatContainerBase> PersistentStatsContainer) = 0;
diff --git a/Engine/Source/Runtime/InstallBundleManager/Public/InstallBundleUtils.h b/Engine/Source/Runtime/InstallBundleManager/Public/InstallBundleUtils.h
--- a/Engine/Source/Runtime/InstallBundleManager/Public/InstallBundleUtils.h
+++ b/Engine/Source/Runtime/InstallBundleManager/Public/InstallBundleUtils.h
@@ -17,6 +17,8 @@ namespace InstallBundleUtil
 
 	INSTALLBUNDLEMANAGER_API const TCHAR* GetInstallBundlePauseReason(EInstallBundlePauseFlags Flags);
 
+	INSTALLBUNDLEMANAGER_API const FString& GetInstallBundleSectionPrefix();
+
 	// It would really be nice to have these in core
 	template<class EnumType>
 	constexpr auto& CastAsUnderlying(EnumType &Type)
