Using InstallBundle.ini to generate chunks on PS4 if the PlayGo bundle source is enabled.

Automation changes

diff --git a/Engine/Source/Programs/AutomationTool/AutomationUtils/BundleUtils.cs b/Engine/Source/Programs/AutomationTool/AutomationUtils/BundleUtils.cs
--- a/Engine/Source/Programs/AutomationTool/AutomationUtils/BundleUtils.cs
+++ b/Engine/Source/Programs/AutomationTool/AutomationUtils/BundleUtils.cs
@@ -19,7 +19,6 @@ namespace AutomationUtils.Automation
 			bool bContainsShaderLibrary { get; }
 		}
 
-
 		public class BundleSettings : IReadOnlyBundleSettings
 		{
 			public string Name { get; set; }
@@ -30,26 +29,31 @@ namespace AutomationUtils.Automation
 			public bool bContainsShaderLibrary { get; set; }
 			public int Order { get; set; }
 			public string ExecFileName { get; set; }
-		};
-		public static bool LoadBundleConfig(DirectoryReference ProjectDir, UnrealTargetPlatform Platform, out List<BundleSettings> Bundles)
+		}
+
+		public static void LoadBundleConfig(DirectoryReference ProjectDir, UnrealTargetPlatform Platform, out List<BundleSettings> Bundles)
+		{
+			LoadBundleConfig<BundleSettings>(ProjectDir, Platform, out Bundles, delegate (BundleSettings Settings, ConfigHierarchy BundleConfig, string Section) { });
+		}
+
+		public static void LoadBundleConfig<TPlatformBundleSettings>(DirectoryReference ProjectDir, UnrealTargetPlatform Platform, 
+			out List<TPlatformBundleSettings> Bundles, 
+			Action<TPlatformBundleSettings, ConfigHierarchy, string> GetPlatformSettings) 
+			where TPlatformBundleSettings : BundleSettings, new()
 		{
-			Bundles = new List<BundleSettings>();
+			Bundles = new List<TPlatformBundleSettings>();
 
 			ConfigHierarchy BundleConfig = ConfigCache.ReadHierarchy(ConfigHierarchyType.InstallBundle, ProjectDir, Platform);
-			if (BundleConfig == null)
-			{
-				CommandUtils.LogWarning("Unable to find bundle config ini file {0}", ConfigHierarchyType.InstallBundle.ToString());
-				return false;
-			}
 
 			const string BundleDefinitionPrefix = "InstallBundleDefinition ";
+			string PlatformBundleDefinitionPrefix = "InstallBundleDefinition" + Platform.ToString() + " ";
 
 			foreach (string SectionName in BundleConfig.SectionNames)
 			{
 				if (!SectionName.StartsWith(BundleDefinitionPrefix))
 					continue;
 
-				BundleSettings Bundle = new BundleSettings();
+				TPlatformBundleSettings Bundle = new TPlatformBundleSettings();
 				Bundle.Name = SectionName.Substring(BundleDefinitionPrefix.Length);
 				{
 					int Order;
@@ -87,13 +91,13 @@ namespace AutomationUtils.Automation
 					Bundle.Tags = new List<string>();
 				}
 
+				GetPlatformSettings(Bundle, BundleConfig, PlatformBundleDefinitionPrefix + Bundle.Name);
+
 				Bundles.Add(Bundle);
 			}
 
 			// Use OrderBy and not Sort because OrderBy is stable
 			Bundles = Bundles.OrderBy(Bundle => Bundle.Order).ToList();
-
-			return true;
 		}
 	}
 }
