move over symstore/pdbstr copy fixes

diff --git a/Engine/Source/Programs/AutomationTool/Win/WinPlatform.Automation.cs b/Engine/Source/Programs/AutomationTool/Win/WinPlatform.Automation.cs
--- a/Engine/Source/Programs/AutomationTool/Win/WinPlatform.Automation.cs
+++ b/Engine/Source/Programs/AutomationTool/Win/WinPlatform.Automation.cs
@@ -395,6 +395,8 @@ public abstract class BaseWinPlatform : Platform
 		{
 			DateTime Start = DateTime.Now;
 			DirectoryReference TempSymStoreDir = DirectoryReference.Combine(RootDirectory, "Saved", "SymStore");
+			DirectoryReference.CreateDirectory(TempSymStoreDir);
+			DeleteDirectoryContents(TempSymStoreDir);
 
 			string TempFileName = Path.GetTempFileName();
 			try
