Install bundles - move platform specific bundle config into platform specific INIs.
Fix InstallBundleDefinition Startup Regex not being loaded correctly in C#
Fix GetChunkIndexForFileFromInstallBundles not returning the correct chunk

Checking in disabled until I can verify its OK with a preflight.

diff --git a/Engine/Source/Programs/AutomationTool/AutomationUtils/BundleUtils.cs b/Engine/Source/Programs/AutomationTool/AutomationUtils/BundleUtils.cs
--- a/Engine/Source/Programs/AutomationTool/AutomationUtils/BundleUtils.cs
+++ b/Engine/Source/Programs/AutomationTool/AutomationUtils/BundleUtils.cs
@@ -46,7 +46,6 @@ namespace AutomationUtils.Automation
 			ConfigHierarchy BundleConfig = ConfigCache.ReadHierarchy(ConfigHierarchyType.InstallBundle, ProjectDir, Platform);
 
 			const string BundleDefinitionPrefix = "InstallBundleDefinition ";
-			string PlatformBundleDefinitionPrefix = "InstallBundleDefinition" + Platform.ToString() + " ";
 
 			foreach (string SectionName in BundleConfig.SectionNames)
 			{
@@ -91,7 +90,7 @@ namespace AutomationUtils.Automation
 					Bundle.Tags = new List<string>();
 				}
 
-				GetPlatformSettings(Bundle, BundleConfig, PlatformBundleDefinitionPrefix + Bundle.Name);
+				GetPlatformSettings(Bundle, BundleConfig, BundleDefinitionPrefix + Bundle.Name);
 
 				Bundles.Add(Bundle);
 			}
