Make UAT use a config setting to decide which editor target is the default for use within automation scripts
Have UBT set the source target name as a define during compilation. For unique environments, embed that macro globally, but in shared environments just embed it into game modules.
Have the primary game module bind that define to a core delegate so engine systems can query it
Make LiveCodingModule pass the UBT target name to the UBT so that it doesn't have to guess which target to build
For agnostic executables (UE4Game, UE4Editor)  running content only projects, the delegate won't be bound, so revert back to type based recompile requests in live coding
Handle DTE string for VS2019 in the source code accessor module

diff --git a/Engine/Plugins/Developer/VisualStudioCodeSourceCodeAccess/Source/VisualStudioCodeSourceCodeAccess/VisualStudioCodeSourceCodeAccess.Build.cs b/Engine/Plugins/Developer/VisualStudioCodeSourceCodeAccess/Source/VisualStudioCodeSourceCodeAccess/VisualStudioCodeSourceCodeAccess.Build.cs
--- a/Engine/Plugins/Developer/VisualStudioCodeSourceCodeAccess/Source/VisualStudioCodeSourceCodeAccess/VisualStudioCodeSourceCodeAccess.Build.cs
+++ b/Engine/Plugins/Developer/VisualStudioCodeSourceCodeAccess/Source/VisualStudioCodeSourceCodeAccess/VisualStudioCodeSourceCodeAccess.Build.cs
@@ -28,6 +28,9 @@ namespace UnrealBuildTool.Rules
 				string DTEKey = null;
 				switch (Target.WindowsPlatform.Compiler)
 				{
+					case WindowsCompiler.VisualStudio2019:
+						DTEKey = "VisualStudio.DTE.16.0";
+						break;
 					case WindowsCompiler.VisualStudio2017:
 						DTEKey = "VisualStudio.DTE.15.0";
 						break;
diff --git a/Engine/Source/Developer/Windows/LiveCoding/Private/LiveCodingModule.cpp b/Engine/Source/Developer/Windows/LiveCoding/Private/LiveCodingModule.cpp
--- a/Engine/Source/Developer/Windows/LiveCoding/Private/LiveCodingModule.cpp
+++ b/Engine/Source/Developer/Windows/LiveCoding/Private/LiveCodingModule.cpp
@@ -307,10 +307,14 @@ bool FLiveCodingModule::StartLiveCoding()
 		LppRegisterProcessGroup(TCHAR_TO_ANSI(*ProcessGroup));
 
 		// Build the command line
-		FString Arguments;
-		Arguments += FString::Printf(TEXT("%s"), FPlatformMisc::GetUBTPlatform());
-		Arguments += FString::Printf(TEXT(" %s"), LexToString(FApp::GetBuildConfiguration()));
-		Arguments += FString::Printf(TEXT(" -TargetType=%s"), FPlatformMisc::GetUBTTarget());
+		FString KnownTargetName = FPlatformMisc::GetUBTTargetName();
+		FString Arguments = FString::Printf(TEXT("%s %s %s"),
+			*KnownTargetName,
+			FPlatformMisc::GetUBTPlatform(),
+			LexToString(FApp::GetBuildConfiguration()));
+
+		UE_LOG(LogLiveCoding, Display, TEXT("LiveCodingConsole Arguments: %s"), *Arguments);
+
 		if(SourceProject.Len() > 0)
 		{
 			Arguments += FString::Printf(TEXT(" -Project=\"%s\""), *FPaths::ConvertRelativePathToFull(SourceProject));
diff --git a/Engine/Source/Programs/AutomationTool/AutomationUtils/ProjectParams.cs b/Engine/Source/Programs/AutomationTool/AutomationUtils/ProjectParams.cs
--- a/Engine/Source/Programs/AutomationTool/AutomationUtils/ProjectParams.cs
+++ b/Engine/Source/Programs/AutomationTool/AutomationUtils/ProjectParams.cs
@@ -1972,10 +1972,6 @@ namespace AutomationTool
 		{
 			return DetectedTargets.FindAll(Target => Target.Rules.Type == DesiredType).ConvertAll(Target => Target.TargetName);
 		}
-		private List<String> PrimaryTargetNamesOfType(TargetType DesiredType)
-		{
-			return DetectedTargets.FindAll(Target => (Target.Rules.Type == DesiredType && Target.Rules.IsPrimaryTarget)).ConvertAll(Target => Target.TargetName);
-		}
 
 		private String ChooseTarget(List<String> Targets, TargetType Type)
 		{
@@ -2115,7 +2111,7 @@ namespace AutomationTool
 				List<String> AvailableGameTargets = TargetNamesOfType(TargetType.Game);
 				List<String> AvailableClientTargets = TargetNamesOfType(TargetType.Client);
 				List<String> AvailableServerTargets = TargetNamesOfType(TargetType.Server);
-				List<String> AvailableEditorTargets = PrimaryTargetNamesOfType(TargetType.Editor);
+				List<String> AvailableEditorTargets = TargetNamesOfType(TargetType.Editor);
 
 				// That should cover all detected targets; Program targets are handled separately.
 				System.Diagnostics.Debug.Assert(DetectedTargets.Count == (AvailableGameTargets.Count + AvailableClientTargets.Count + AvailableServerTargets.Count + AvailableEditorTargets.Count));
@@ -2137,12 +2133,26 @@ namespace AutomationTool
 
 				if (AvailableEditorTargets.Count > 0)
 				{
-					if (AvailableEditorTargets.Count > 1)
+					string DefaultEditorTarget;
+
+					if (EngineConfigs[BuildHostPlatform.Current.Platform].GetString("/Script/BuildSettings.BuildSettings", "DefaultEditorTarget", out DefaultEditorTarget))
 					{
-						throw new AutomationException("There can be only one primary editor target per project.");
+						if (!AvailableEditorTargets.Contains(DefaultEditorTarget))
+						{
+							throw new AutomationException(string.Format("A default editor target '{0}' was specified in engine.ini but does not exist", DefaultEditorTarget));
+						}
+
+						EditorTarget = DefaultEditorTarget;
 					}
+					else
+					{
+						if (AvailableEditorTargets.Count > 1)
+						{
+							throw new AutomationException("Project contains multiple editor targets but no default is set in engine.ini");
+						}
 
-					EditorTarget = AvailableEditorTargets.First();
+						EditorTarget = AvailableEditorTargets.First();
+					}
 				}
 
 				if (AvailableServerTargets.Count > 0 && (DedicatedServer || Cook || CookOnTheFly)) // only if server is needed
diff --git a/Engine/Source/Programs/UnrealBuildTool/Configuration/TargetRules.cs b/Engine/Source/Programs/UnrealBuildTool/Configuration/TargetRules.cs
--- a/Engine/Source/Programs/UnrealBuildTool/Configuration/TargetRules.cs
+++ b/Engine/Source/Programs/UnrealBuildTool/Configuration/TargetRules.cs
@@ -197,11 +197,6 @@ namespace UnrealBuildTool
 		/// </summary>
 		public global::UnrealBuildTool.TargetType Type = global::UnrealBuildTool.TargetType.Game;
 
-		/// <summary>
-		/// Indicates that this target is the primary one of its type. i.e. it's the primary editor target that should be used by UAT for cooking.
-		/// </summary>
-		public bool IsPrimaryTarget = true;
-
 		/// <summary>
 		/// Specifies the engine version to maintain backwards-compatible default build settings with (eg. DefaultSettingsVersion.Release_4_23, DefaultSettingsVersion.Release_4_24). Specify DefaultSettingsVersion.Latest to always
 		/// use defaults for the current engine version, at the risk of introducing build errors while upgrading.
diff --git a/Engine/Source/Programs/UnrealBuildTool/Configuration/UEBuildModuleCPP.cs b/Engine/Source/Programs/UnrealBuildTool/Configuration/UEBuildModuleCPP.cs
--- a/Engine/Source/Programs/UnrealBuildTool/Configuration/UEBuildModuleCPP.cs
+++ b/Engine/Source/Programs/UnrealBuildTool/Configuration/UEBuildModuleCPP.cs
@@ -1204,7 +1204,7 @@ namespace UnrealBuildTool
 				Result.Definitions.Add("UE_IS_ENGINE_MODULE=0");
 			}
 
-			// For game modules, set the define for the project name. This will be used by the IMPLEMENT_PRIMARY_GAME_MODULE macro.
+			// For game modules, set the define for the project and target names, which will be used by the IMPLEMENT_PRIMARY_GAME_MODULE macro.
 			if (!Rules.bTreatAsEngineModule)
 			{
 				// Make sure we don't set any define for a non-engine module that's under the engine directory (eg. UE4Game)
@@ -1212,6 +1212,7 @@ namespace UnrealBuildTool
 				{
 					string ProjectName = Target.ProjectFile.GetFileNameWithoutExtension();
 					Result.Definitions.Add(String.Format("UE_PROJECT_NAME={0}", ProjectName));
+					Result.Definitions.Add(String.Format("UE_TARGET_NAME={0}", Target.Name));
 				}
 			}
 
diff --git a/Engine/Source/Programs/UnrealBuildTool/System/RulesAssembly.cs b/Engine/Source/Programs/UnrealBuildTool/System/RulesAssembly.cs
--- a/Engine/Source/Programs/UnrealBuildTool/System/RulesAssembly.cs
+++ b/Engine/Source/Programs/UnrealBuildTool/System/RulesAssembly.cs
@@ -638,7 +638,7 @@ namespace UnrealBuildTool
 			foreach(KeyValuePair<string, FileReference> TargetPair in TargetNameToTargetFile)
 			{
 				TargetRules Rules = CreateTargetRulesInstance(TargetPair.Key + "Target", new TargetInfo(TargetPair.Key, Platform, Configuration, Architecture, ProjectFile, null));
-				if(Rules.Type == Type && Rules.IsPrimaryTarget)
+				if(Rules.Type == Type)
 				{
 					Matches.Add(TargetPair.Key);
 				}
diff --git a/Engine/Source/Runtime/Core/Private/GenericPlatform/GenericPlatformMisc.cpp b/Engine/Source/Runtime/Core/Private/GenericPlatform/GenericPlatformMisc.cpp
--- a/Engine/Source/Runtime/Core/Private/GenericPlatform/GenericPlatformMisc.cpp
+++ b/Engine/Source/Runtime/Core/Private/GenericPlatform/GenericPlatformMisc.cpp
@@ -1054,6 +1054,21 @@ const TCHAR* FGenericPlatformMisc::GetUBTTarget()
 	return TEXT(PREPROCESSOR_TO_STRING(UBT_COMPILED_TARGET));
 }
 
+/** The name of the UBT target that the current executable was built from. Defaults to the UE4 default target for this type to make content only projects work, 
+	but will be overridden by the primary game module if it exists */
+TCHAR GUBTTargetName[128] = TEXT("UE4" PREPROCESSOR_TO_STRING(UBT_COMPILED_TARGET));
+
+void FGenericPlatformMisc::SetUBTTargetName(const TCHAR* InTargetName)
+{
+	check(FCString::Strlen(InTargetName) < (UE_ARRAY_COUNT(GUBTTargetName) - 1));
+	FCString::Strcpy(GUBTTargetName, InTargetName);
+}
+
+const TCHAR* FGenericPlatformMisc::GetUBTTargetName()
+{
+	return GUBTTargetName;
+}
+
 const TCHAR* FGenericPlatformMisc::GetDefaultDeviceProfileName()
 {
 	return TEXT("Default");
diff --git a/Engine/Source/Runtime/Core/Public/GenericPlatform/GenericPlatformMisc.h b/Engine/Source/Runtime/Core/Public/GenericPlatform/GenericPlatformMisc.h
--- a/Engine/Source/Runtime/Core/Public/GenericPlatform/GenericPlatformMisc.h
+++ b/Engine/Source/Runtime/Core/Public/GenericPlatform/GenericPlatformMisc.h
@@ -921,8 +921,11 @@ public:
 
 	static const TCHAR* GetUBTTarget();
 
+	static void SetUBTTargetName(const TCHAR* InTargetName);
+	static const TCHAR* GetUBTTargetName();
+
 	/** 
-	 * Determines the shader format for the plarform
+	 * Determines the shader format for the platform
 	 *
 	 * @return	Returns the shader format to be used by that platform
 	 */
diff --git a/Engine/Source/Runtime/Core/Public/Modules/ModuleManager.h b/Engine/Source/Runtime/Core/Public/Modules/ModuleManager.h
--- a/Engine/Source/Runtime/Core/Public/Modules/ModuleManager.h
+++ b/Engine/Source/Runtime/Core/Public/Modules/ModuleManager.h
@@ -832,6 +832,15 @@ class FDefaultGameModuleImpl
 		} \
 	} GEncryptionKeyRegistration;
 
+#define IMPLEMENT_TARGET_NAME_REGISTRATION() \
+	struct FTargetNameRegistration \
+	{ \
+		FTargetNameRegistration() \
+		{ \
+			FPlatformMisc::SetUBTTargetName(TEXT(PREPROCESSOR_TO_STRING(UE_TARGET_NAME))); \
+		} \
+	} GTargetNameRegistration;
+
 #if IS_PROGRAM
 
 	#if IS_MONOLITHIC
@@ -883,6 +892,7 @@ class FDefaultGameModuleImpl
 			IMPLEMENT_LIVE_CODING_PROJECT() \
 			IMPLEMENT_SIGNING_KEY_REGISTRATION() \
 			IMPLEMENT_ENCRYPTION_KEY_REGISTRATION() \
+			IMPLEMENT_TARGET_NAME_REGISTRATION() \
 			IMPLEMENT_GAME_MODULE( ModuleImplClass, ModuleName ) \
 			PER_MODULE_BOILERPLATE
 
@@ -897,6 +907,7 @@ class FDefaultGameModuleImpl
 			IMPLEMENT_LIVE_CODING_PROJECT() \
 			IMPLEMENT_SIGNING_KEY_REGISTRATION() \
 			IMPLEMENT_ENCRYPTION_KEY_REGISTRATION() \
+			IMPLEMENT_TARGET_NAME_REGISTRATION() \
 			IMPLEMENT_GAME_MODULE( ModuleImplClass, ModuleName ) \
 			/* Implement the GIsGameAgnosticExe variable (See Core.h). */ \
 			bool GIsGameAgnosticExe = false;
@@ -907,6 +918,7 @@ class FDefaultGameModuleImpl
 
 	#define IMPLEMENT_PRIMARY_GAME_MODULE( ModuleImplClass, ModuleName, GameName ) \
 		/* Nothing special to do for modular builds.  The game name will be set via the command-line */ \
+		IMPLEMENT_TARGET_NAME_REGISTRATION() \
 		IMPLEMENT_GAME_MODULE( ModuleImplClass, ModuleName )
 #endif	//IS_MONOLITHIC
 
