Fixed issue with asserts not being reported if there was no callstack

diff --git a/Engine/Source/Programs/AutomationTool/Gauntlet/Unreal/Utils/Gauntlet.UnrealLogParser.cs b/Engine/Source/Programs/AutomationTool/Gauntlet/Unreal/Utils/Gauntlet.UnrealLogParser.cs
--- a/Engine/Source/Programs/AutomationTool/Gauntlet/Unreal/Utils/Gauntlet.UnrealLogParser.cs
+++ b/Engine/Source/Programs/AutomationTool/Gauntlet/Unreal/Utils/Gauntlet.UnrealLogParser.cs
@@ -522,8 +522,12 @@ namespace Gauntlet
 					if (Backtrace.Count > 0)
 					{
 						NewTrace.Callstack = Backtrace.ToArray();
-						Traces.Add(NewTrace);
 					}
+					else
+					{
+						NewTrace.Callstack = new[] { "Unable to parse callstack from log" };
+					}
+					Traces.Add(NewTrace);
 				}
 			}
 
