BuildGraph: Added a Horde export mode for an entire graph. Can be passed directly into the Horde API as a job definition.

#rnx

diff --git a/Engine/Source/Programs/AutomationTool/BuildGraph/BuildGraph.cs b/Engine/Source/Programs/AutomationTool/BuildGraph/BuildGraph.cs
--- a/Engine/Source/Programs/AutomationTool/BuildGraph/BuildGraph.cs
+++ b/Engine/Source/Programs/AutomationTool/BuildGraph/BuildGraph.cs
@@ -63,6 +63,7 @@ namespace AutomationTool
 	[Help("SkipTargetsWithoutTokens", "Excludes targets which we can't acquire tokens for, rather than failing")]
 	[Help("Preprocess=<FileName>", "Writes the preprocessed graph to the given file")]
 	[Help("Export=<FileName>", "Exports a JSON file containing the preprocessed build graph, for use as part of a build system")]
+	[Help("HordeExport=<FileName>", "Exports a JSON file containing the full build graph for use by Horde.")]
 	[Help("PublicTasksOnly", "Only include built-in tasks in the schema, excluding any other UAT modules")]
 	[Help("SharedStorageDir=<DirName>", "Sets the directory to use to transfer build products between agents in a build farm")]
 	[Help("SingleNode=<Name>", "Run only the given node. Intended for use on a build system after running with -Export.")]
@@ -81,6 +82,7 @@ namespace AutomationTool
 			string SchemaFileName = ParseParamValue("Schema", null);
 			string ImportSchemaFileName = ParseParamValue("ImportSchema", null);
 			string ExportFileName = ParseParamValue("Export", null);
+			string HordeExportFileName = ParseParamValue("HordeExport", null);
 			string PreprocessedFileName = ParseParamValue("Preprocess", null);
 			string SharedStorageDir = ParseParamValue("SharedStorageDir", null);
 			string SingleNodeName = ParseParamValue("SingleNode", null);
@@ -222,6 +224,12 @@ namespace AutomationTool
 				return ExitCode.Error_Unknown;
 			}
 
+			// Export the graph for Horde
+			if(HordeExportFileName != null)
+			{
+				Graph.ExportForHorde(new FileReference(HordeExportFileName));
+			}
+
 			// Create the temp storage handler
 			DirectoryReference RootDir = new DirectoryReference(CommandUtils.CmdEnv.LocalRoot);
 			TempStorage Storage = new TempStorage(RootDir, DirectoryReference.Combine(RootDir, "Engine", "Saved", "BuildGraph"), (SharedStorageDir == null)? null : new DirectoryReference(SharedStorageDir), bWriteToSharedStorage);
diff --git a/Engine/Source/Programs/AutomationTool/BuildGraph/Graph.cs b/Engine/Source/Programs/AutomationTool/BuildGraph/Graph.cs
--- a/Engine/Source/Programs/AutomationTool/BuildGraph/Graph.cs
+++ b/Engine/Source/Programs/AutomationTool/BuildGraph/Graph.cs
@@ -584,6 +584,46 @@ namespace AutomationTool
 			}
 		}
 
+		/// <summary>
+		/// Export the build graph to a Json file for parsing by Horde
+		/// </summary>
+		/// <param name="File">Output file to write</param>
+		public void ExportForHorde(FileReference File)
+		{
+			DirectoryReference.CreateDirectory(File.Directory);
+			using (JsonWriter JsonWriter = new JsonWriter(File.FullName))
+			{
+				JsonWriter.WriteArrayStart();
+				foreach (Agent Agent in Agents)
+				{
+					foreach (Node Node in Agent.Nodes)
+					{
+						JsonWriter.WriteObjectStart();
+						JsonWriter.WriteValue("Name", Node.Name);
+						JsonWriter.WriteValue("Group", Agent.Name);
+						JsonWriter.WriteValue("RunEarly", Node.bRunEarly);
+						JsonWriter.WriteValue("Exclusive", true);
+
+						JsonWriter.WriteArrayStart("InputDependencies");
+						foreach (string InputDependency in Node.GetDirectInputDependencies().Select(x => x.Name))
+						{
+							JsonWriter.WriteValue(InputDependency);
+						}
+						JsonWriter.WriteArrayEnd();
+
+						JsonWriter.WriteArrayStart("OrderDependencies");
+						foreach (string OrderDependency in Node.GetDirectOrderDependencies().Select(x => x.Name))
+						{
+							JsonWriter.WriteValue(OrderDependency);
+						}
+						JsonWriter.WriteArrayEnd();
+						JsonWriter.WriteObjectEnd();
+					}
+				}
+				JsonWriter.WriteArrayEnd();
+			}
+		}
+
 		/// <summary>
 		/// Print the contents of the graph
 		/// </summary>
