Fix CIS

#rnx

diff --git a/Engine/Source/Programs/AutomationTool/AutomationUtils/ProjectParams.cs b/Engine/Source/Programs/AutomationTool/AutomationUtils/ProjectParams.cs
--- a/Engine/Source/Programs/AutomationTool/AutomationUtils/ProjectParams.cs
+++ b/Engine/Source/Programs/AutomationTool/AutomationUtils/ProjectParams.cs
@@ -590,7 +590,7 @@ namespace AutomationTool
 			DLCName = ParseParamValueIfNotSpecified(Command, DLCName, "DLCName", String.Empty);
 			if(!String.IsNullOrEmpty(DLCName))
 			{
-				List<PluginInfo> CandidatePlugins = Plugins.ReadAvailablePlugins(CommandUtils.EngineDirectory, RawProjectPath, null);
+				List<PluginInfo> CandidatePlugins = Plugins.ReadAvailablePlugins(CommandUtils.EngineDirectory, DirectoryReference.FromFile(RawProjectPath), null);
 				PluginInfo DLCPlugin = CandidatePlugins.FirstOrDefault(x => String.Equals(x.Name, DLCName, StringComparison.InvariantCultureIgnoreCase));
 				if(DLCPlugin == null)
 				{
diff --git a/Engine/Source/Programs/AutomationTool/AutomationUtils/ProjectUtils.cs b/Engine/Source/Programs/AutomationTool/AutomationUtils/ProjectUtils.cs
--- a/Engine/Source/Programs/AutomationTool/AutomationUtils/ProjectUtils.cs
+++ b/Engine/Source/Programs/AutomationTool/AutomationUtils/ProjectUtils.cs
@@ -202,7 +202,7 @@ namespace AutomationTool
 			ProjectDescriptor Project = ProjectDescriptor.FromFile(RawProjectPath);
 
 			// find if there are any plugins enabled or disabled which differ from the default
-			List<PluginInfo> Plugins = UnrealBuildTool.Plugins.ReadAvailablePlugins(CommandUtils.EngineDirectory, RawProjectPath, Project.AdditionalPluginDirectories);
+			List<PluginInfo> Plugins = UnrealBuildTool.Plugins.ReadAvailablePlugins(CommandUtils.EngineDirectory, DirectoryReference.FromFile(RawProjectPath), Project.AdditionalPluginDirectories);
 			foreach (PluginInfo Plugin in Plugins)
 			{
 				bool bPluginEnabledForTarget = UnrealBuildTool.Plugins.IsPluginCompiledForTarget(Plugin, Project, Platform, Configuration, TargetType, bRequiresCookedData);
diff --git a/Engine/Source/Programs/AutomationTool/Scripts/ResaveProjectDescriptors.cs b/Engine/Source/Programs/AutomationTool/Scripts/ResaveProjectDescriptors.cs
--- a/Engine/Source/Programs/AutomationTool/Scripts/ResaveProjectDescriptors.cs
+++ b/Engine/Source/Programs/AutomationTool/Scripts/ResaveProjectDescriptors.cs
@@ -1,4 +1,4 @@
-﻿// Copyright 1998-2019 Epic Games, Inc. All Rights Reserved.
+// Copyright 1998-2019 Epic Games, Inc. All Rights Reserved.
 
 using System;
 using System.Collections.Generic;
@@ -58,7 +58,7 @@ namespace AutomationTool
 				ProjectDescriptor Descriptor;
 				try
 				{
-					Descriptor = new ProjectDescriptor(JsonObject.Parse(InputText));
+					Descriptor = new ProjectDescriptor(JsonObject.Parse(InputText), ProjectFile.Directory);
 				}
 				catch(JsonParseException Ex)
 				{
@@ -106,7 +106,7 @@ namespace AutomationTool
 				using (JsonWriter Writer = new JsonWriter(new StringWriter(Output)))
 				{
 					Writer.WriteObjectStart();
-					Descriptor.Write(Writer);
+					Descriptor.Write(Writer, ProjectFile.Directory);
 					Writer.WriteObjectEnd();
 				}
 
