BuildGraph: Make missing BuildGraph arguments an error, since they cause the build to fail.

#rnx

diff --git a/Engine/Source/Programs/AutomationTool/BuildGraph/Script.cs b/Engine/Source/Programs/AutomationTool/BuildGraph/Script.cs
--- a/Engine/Source/Programs/AutomationTool/BuildGraph/Script.cs
+++ b/Engine/Source/Programs/AutomationTool/BuildGraph/Script.cs
@@ -297,7 +297,7 @@ namespace AutomationTool
 			bool bInvalidArgument = false;
 			foreach(string InvalidArgumentName in Arguments.Keys.Except(Reader.Graph.Options.Select(x => x.Name), StringComparer.InvariantCultureIgnoreCase))
 			{
-				CommandUtils.LogWarning("Unknown argument '{0}' for '{1}'", InvalidArgumentName, File.FullName);
+				CommandUtils.LogError("Unknown argument '{0}' for '{1}'", InvalidArgumentName, File.FullName);
 				bInvalidArgument = true;
 			}
 			if(bInvalidArgument)
