Fixing error message for invalid platform so it actually contains the name.

#rnx

diff --git a/Engine/Source/Programs/AutomationTool/Scripts/BuildPluginCommand.Automation.cs b/Engine/Source/Programs/AutomationTool/Scripts/BuildPluginCommand.Automation.cs
--- a/Engine/Source/Programs/AutomationTool/Scripts/BuildPluginCommand.Automation.cs
+++ b/Engine/Source/Programs/AutomationTool/Scripts/BuildPluginCommand.Automation.cs
@@ -311,7 +311,7 @@ class BuildPlugin : BuildCommand
 					UnrealTargetPlatform TargetPlatform;
 					if (!UnrealTargetPlatform.TryParse(TargetPlatformName, out TargetPlatform))
 					{
-						throw new AutomationException("Unknown target platform '{0}' specified on command line");
+						throw new AutomationException("Unknown target platform '{0}' specified on command line", TargetPlatformName);
 					}
 					if(TargetPlatforms.Contains(TargetPlatform))
 					{
