Add directory mappings for mac device to allow nonstandard file gathering plus pushing replays etc to the device.

diff --git a/Engine/Source/Programs/AutomationTool/Gauntlet/Platform/Mac/Gauntlet.TargetDeviceMac.cs b/Engine/Source/Programs/AutomationTool/Gauntlet/Platform/Mac/Gauntlet.TargetDeviceMac.cs
--- a/Engine/Source/Programs/AutomationTool/Gauntlet/Platform/Mac/Gauntlet.TargetDeviceMac.cs
+++ b/Engine/Source/Programs/AutomationTool/Gauntlet/Platform/Mac/Gauntlet.TargetDeviceMac.cs
@@ -160,6 +160,17 @@ namespace Gauntlet
 			return "";
 		}
 
+		public void PopulateDirectoryMappings(string ProjectDir)
+		{
+			LocalDirectoryMappings.Add(EIntendedBaseCopyDirectory.Build, Path.Combine(ProjectDir, "Build"));
+			LocalDirectoryMappings.Add(EIntendedBaseCopyDirectory.Binaries, Path.Combine(ProjectDir, "Binaries"));
+			LocalDirectoryMappings.Add(EIntendedBaseCopyDirectory.Config, Path.Combine(ProjectDir, "Saved", "Config"));
+			LocalDirectoryMappings.Add(EIntendedBaseCopyDirectory.Content, Path.Combine(ProjectDir, "Content"));
+			LocalDirectoryMappings.Add(EIntendedBaseCopyDirectory.Demos, Path.Combine(ProjectDir, "Saved", "Demos"));
+			LocalDirectoryMappings.Add(EIntendedBaseCopyDirectory.Profiling, Path.Combine(ProjectDir, "Saved", "Profiling"));
+			LocalDirectoryMappings.Add(EIntendedBaseCopyDirectory.Saved, Path.Combine(ProjectDir, "Saved"));
+		}
+
 		protected IAppInstall InstallStagedBuild(UnrealAppConfig AppConfig, StagedBuild InBuild)
 		{
 			bool SkipDeploy = Globals.Params.ParseParam("SkipDeploy");
@@ -201,6 +212,10 @@ namespace Gauntlet
 			// Mac always forces this to stop logs and other artifacts going to different places
 			MacApp.CommandArguments += string.Format(" -userdir=\"{0}\"", UserDir);
 			MacApp.ArtifactPath = Path.Combine(UserDir, @"Saved");
+			if (LocalDirectoryMappings.Count == 0)
+			{
+				PopulateDirectoryMappings(BuildPath);
+			}
 
 			// temp - Mac doesn't support -userdir?
 			//MacApp.ArtifactPath = Path.Combine(Environment.GetFolderPath(Environment.SpecialFolder.Personal), "Library/Logs", AppConfig.ProjectName);
@@ -281,11 +296,16 @@ namespace Gauntlet
 			MacApp.CommandArguments += string.Format(" -userdir=\"{0}\"", UserDir);
 			MacApp.ArtifactPath = Path.Combine(UserDir, @"Saved");
 
+
 			// now turn the Foo.app into Foo/Content/MacOS/Foo
 			string AppPath = Path.GetDirectoryName(EditorBuild.ExecutablePath);
 			string FileName = Path.GetFileNameWithoutExtension(EditorBuild.ExecutablePath);
 			MacApp.ExecutablePath = Path.Combine(EditorBuild.ExecutablePath, "Contents", "MacOS", FileName);
 
+			if (LocalDirectoryMappings.Count == 0)
+			{
+				PopulateDirectoryMappings(AppPath);
+			}
 			return MacApp;
 		}
 
