Added a -ParallelGather option to the Localize UAT command

This will run the GatherText commandlets for the projects within a single batch in parallel rather than sequence, which can be advantageous for projects with a large number of static localization targets

diff --git a/Engine/Source/Programs/AutomationTool/Scripts/Localisation.Automation.cs b/Engine/Source/Programs/AutomationTool/Scripts/Localisation.Automation.cs
--- a/Engine/Source/Programs/AutomationTool/Scripts/Localisation.Automation.cs
+++ b/Engine/Source/Programs/AutomationTool/Scripts/Localisation.Automation.cs
@@ -23,6 +23,7 @@ using Tools.DotNETCommon;
 [Help("ExcludePlugins", "Optional comma separated list of plugins to exclude from the gather.")]
 [Help("IncludePlatforms", "Optional flag to include platforms from within the given UEProjectDirectory as part of the gather.")]
 [Help("AdditionalCommandletArguments", "Optional arguments to pass to the gather process.")]
+[Help("ParallelGather", "Run the gather processes for a single batch in parallel rather than sequence.")]
 class Localize : BuildCommand
 {
 	private class LocalizationBatch
@@ -128,6 +129,10 @@ class Localize : BuildCommand
 			AdditionalCommandletArguments = "";
 		}
 
+		var EnableParallelGather = ParseParam("ParallelGather");
+
+		var StartTime = DateTime.UtcNow;
+
 		var LocalizationBatches = new List<LocalizationBatch>();
 
 		// Add the static set of localization projects as a batch
@@ -200,7 +205,7 @@ class Localize : BuildCommand
 		Dictionary<string, byte[]> InitalPOFileHashes = null;
 		if (P4Enabled)
 		{
-			var ChangeListCommitMessage = "Localization Automation";
+			var ChangeListCommitMessage = String.Format("Localization Automation using CL {0}", P4Env.Changelist);
 			if (File.Exists(CombinePaths(CmdEnv.LocalRoot, @"Engine/Build/NotForLicensees/EpicInternal.txt")))
 			{
 				ChangeListCommitMessage += "\n#okforgithub ignore";
@@ -213,7 +218,7 @@ class Localize : BuildCommand
 		// Process each localization batch
 		foreach (var LocalizationBatch in LocalizationBatches)
 		{
-			ProcessLocalizationProjects(LocalizationBatch, PendingChangeList, UEProjectRoot, UEProjectName, LocalizationProviderName, LocalizationStepNames, AdditionalCommandletArguments);
+			ProcessLocalizationProjects(LocalizationBatch, PendingChangeList, UEProjectRoot, UEProjectName, LocalizationProviderName, LocalizationStepNames, AdditionalCommandletArguments, EnableParallelGather);
 		}
 
 		// Clean-up the changelist so it only contains the changed files, and then submit it (if we were asked to)
@@ -260,9 +265,12 @@ class Localize : BuildCommand
 				P4.Submit(PendingChangeList, out SubmittedChangeList);
 			}
 		}
+
+		var RunDuration = (DateTime.UtcNow - StartTime).TotalMilliseconds;
+		LogInformation("Localize command finished in {0} seconds", RunDuration / 1000);
 	}
 
-	private void ProcessLocalizationProjects(LocalizationBatch LocalizationBatch, int PendingChangeList, string UEProjectRoot, string UEProjectName, string LocalizationProviderName, IReadOnlyList<string> LocalizationSteps, string AdditionalCommandletArguments)
+	private void ProcessLocalizationProjects(LocalizationBatch LocalizationBatch, int PendingChangeList, string UEProjectRoot, string UEProjectName, string LocalizationProviderName, IReadOnlyList<string> LocalizationSteps, string AdditionalCommandletArguments, bool EnableParallelGather)
 	{
 		var EditorExe = CombinePaths(CmdEnv.LocalRoot, @"Engine/Binaries/Win64/UE4Editor-Cmd.exe");
 		var RootWorkingDirectory = CombinePaths(UEProjectRoot, LocalizationBatch.UEProjectDirectory);
@@ -320,8 +328,8 @@ class Localize : BuildCommand
 		}
 		EditorArguments += " -Unattended -LogLocalizationConflicts";
 
-		// Execute commandlet for each config in each project.
-		bool bLocCommandletFailed = false;
+		// Execute commandlet for all configs in each project, optionally running multiple commandlets in parallel.
+		var LocalizationGatherProcessResults = new List<IProcessResult>();
 		foreach (var ProjectInfo in ProjectInfos)
 		{
 			var LocalizationConfigFiles = new List<string>();
@@ -343,37 +351,62 @@ class Localize : BuildCommand
 					CommandletArguments += " " + AdditionalCommandletArguments;
 				}
 
-				string Arguments = String.Format("{0} -run=GatherText {1} {2}", ProjectArgument, EditorArguments, CommandletArguments);
-				LogInformation("Running localization commandlet: {0}", Arguments);
-				var StartTime = DateTime.UtcNow;
-				var RunResult = Run(EditorExe, Arguments, null, ERunOptions.Default | ERunOptions.NoLoggingOfRunCommand); // Disable logging of the run command as it will print the exit code which GUBP can pick up as an error (we do that ourselves below)
-				var RunDuration = (DateTime.UtcNow - StartTime).TotalMilliseconds;
-				LogInformation("Localization commandlet finished in {0}s", RunDuration / 1000);
-
-				if (RunResult.ExitCode != 0)
+				var CommandletRunOptions = ERunOptions.Default | ERunOptions.NoLoggingOfRunCommand; // Disable logging of the run command as it will print the exit code which GUBP can pick up as an error (we do that ourselves later)
+				if (EnableParallelGather)
 				{
-					LogWarning("The localization commandlet exited with code {0} which likely indicates a crash. It ran with the following arguments: '{1}'", RunResult.ExitCode, Arguments);
-					bLocCommandletFailed = true;
-					break; // We failed a step, so don't process any other steps in this config chain
+					CommandletRunOptions |= ERunOptions.NoWaitForExit;
+					EditorArguments += " -multiprocess";
 				}
+
+				var Arguments = String.Format("{0} -run=GatherText {1} {2}", ProjectArgument, EditorArguments, CommandletArguments);
+				LogInformation("Running localization commandlet for '{0}': {1}", ProjectInfo.ProjectName, Arguments);
+				LocalizationGatherProcessResults.Add(Run(EditorExe, Arguments, null, CommandletRunOptions));
+			}
+			else
+			{
+				LocalizationGatherProcessResults.Add(null);
 			}
 		}
 
-		if (LocalizationSteps.Contains("Upload") && LocProvider != null)
+		// Wait for each commandlet process to finish and report the result.
+		// This runs even for non-parallel execution to log the exit state of the process.
+		for (int ProjectIndex = 0; ProjectIndex < ProjectInfos.Count; ++ProjectIndex)
 		{
-			if (bLocCommandletFailed)
+			var ProjectInfo = ProjectInfos[ProjectIndex];
+			var RunResult = LocalizationGatherProcessResults[ProjectIndex];
+
+			if (RunResult != null)
 			{
-				LogWarning("Skipping upload to the localization provider due to an earlier commandlet failure.");
+				RunResult.WaitForExit();
+				if (RunResult.ExitCode == 0)
+				{
+					LogInformation("The localization commandlet for '{0}' exited with code 0.", ProjectInfo.ProjectName);
+				}
+				else
+				{
+					LogWarning("The localization commandlet for '{0}' exited with code {1} which likely indicates a crash.", ProjectInfo.ProjectName, RunResult.ExitCode);
+				}
 			}
-			else
+		}
+
+		if (LocalizationSteps.Contains("Upload") && LocProvider != null)
+		{
+			// Upload all text to our localization provider
+			for (int ProjectIndex = 0; ProjectIndex < ProjectInfos.Count; ++ProjectIndex)
 			{
-				// Upload all text to our localization provider
-				foreach (var ProjectInfo in ProjectInfos)
+				var ProjectInfo = ProjectInfos[ProjectIndex];
+				var RunResult = LocalizationGatherProcessResults[ProjectIndex];
+
+				if (RunResult != null && RunResult.ExitCode == 0)
 				{
 					// Recalculate the split platform paths before doing the upload, as the export may have changed them
 					ProjectInfo.ExportInfo.CalculateSplitPlatformNames(RootLocalizationTargetDirectory);
 					LocProvider.UploadProjectToLocalizationProvider(ProjectInfo.ProjectName, ProjectInfo.ExportInfo);
 				}
+				else
+				{
+					LogWarning("Skipping upload to the localization provider for '{0}' due to an earlier commandlet failure.", ProjectInfo.ProjectName);
+				}
 			}
 		}
 	}
