Suppress zip64 warnings when processing IPA's
#rnx

diff --git a/Engine/Source/Programs/AutomationTool/Gauntlet/Platform/IOS/Gauntlet.IOSBuildSource.cs b/Engine/Source/Programs/AutomationTool/Gauntlet/Platform/IOS/Gauntlet.IOSBuildSource.cs
--- a/Engine/Source/Programs/AutomationTool/Gauntlet/Platform/IOS/Gauntlet.IOSBuildSource.cs
+++ b/Engine/Source/Programs/AutomationTool/Gauntlet/Platform/IOS/Gauntlet.IOSBuildSource.cs
@@ -69,21 +69,24 @@ namespace Gauntlet
 		// This is limited to bulk builds which exceed 4 gigs, this method works around this, though we really do need to fix
 		internal static bool ExecuteIPAZipCommand(String Arguments, out String Output, String ShouldExist = "")
 		{
-			IProcessResult Result = ExecuteCommand("unzip", Arguments);
-			Output = Result.Output;
-
-			if (Result.ExitCode != 0)
+			using (new ScopedSuspendECErrorParsing())
 			{
-				if (!String.IsNullOrEmpty(ShouldExist))
+				IProcessResult Result = ExecuteCommand("unzip", Arguments);
+				Output = Result.Output;
+
+				if (Result.ExitCode != 0)
 				{
-					if (!File.Exists(ShouldExist) && !Directory.Exists(ShouldExist))
+					if (!String.IsNullOrEmpty(ShouldExist))
 					{
-						Log.Error(String.Format("unzip encountered an error or warning procesing IPA, possibly due to Zip64 issue, {0} missing", ShouldExist));
-						return false;
+						if (!File.Exists(ShouldExist) && !Directory.Exists(ShouldExist))
+						{
+							Log.Error(String.Format("unzip encountered an error or warning procesing IPA, possibly due to Zip64 issue, {0} missing", ShouldExist));
+							return false;
+						}
 					}
-				}
 
-				Log.Warning(String.Format("unzip encountered an error or warning procesing IPA, possibly due to Zip64 issue. Future steps may fail."));
+					Log.Info(String.Format("unzip encountered an issue procesing IPA, possibly due to Zip64. Future steps may fail."));
+				}
 			}
 			
 			return true;
