Expanded ParallelGather to run all projects from all batches in parallel, rather than just the projects within each batch

#rnx

diff --git a/Engine/Source/Programs/AutomationTool/Scripts/Localisation.Automation.cs b/Engine/Source/Programs/AutomationTool/Scripts/Localisation.Automation.cs
--- a/Engine/Source/Programs/AutomationTool/Scripts/Localisation.Automation.cs
+++ b/Engine/Source/Programs/AutomationTool/Scripts/Localisation.Automation.cs
@@ -42,6 +42,34 @@ class Localize : BuildCommand
 		public IReadOnlyList<string> LocalizationProjectNames { get; private set; }
 	};
 
+	private class LocalizationTask
+	{
+		public LocalizationTask(LocalizationBatch InBatch, string InUEProjectRoot, string InLocalizationProviderName, int InPendingChangeList, BuildCommand InCommand)
+		{
+			Batch = InBatch;
+			RootWorkingDirectory = CombinePaths(InUEProjectRoot, Batch.UEProjectDirectory);
+			RootLocalizationTargetDirectory = CombinePaths(InUEProjectRoot, Batch.LocalizationTargetDirectory);
+
+			// Try and find our localization provider
+			{
+				LocalizationProvider.LocalizationProviderArgs LocProviderArgs;
+				LocProviderArgs.RootWorkingDirectory = RootWorkingDirectory;
+				LocProviderArgs.RootLocalizationTargetDirectory = RootLocalizationTargetDirectory;
+				LocProviderArgs.RemoteFilenamePrefix = Batch.RemoteFilenamePrefix;
+				LocProviderArgs.Command = InCommand;
+				LocProviderArgs.PendingChangeList = InPendingChangeList;
+				LocProvider = LocalizationProvider.GetLocalizationProvider(InLocalizationProviderName, LocProviderArgs);
+			}
+		}
+
+		public LocalizationBatch Batch;
+		public string RootWorkingDirectory;
+		public string RootLocalizationTargetDirectory;
+		public LocalizationProvider LocProvider = null;
+		public List<ProjectInfo> ProjectInfos = new List<ProjectInfo>();
+		public List<IProcessResult> GatherProcessResults = new List<IProcessResult>();
+	};
+
 	public override void ExecuteBuild()
 	{
 		var UEProjectRoot = ParseParamValue("UEProjectRoot");
@@ -200,9 +228,8 @@ class Localize : BuildCommand
 			}
 		}
 
-		// Create a single changelist to use for all changes, and hash the current PO files on disk so we can work out whether they actually change
+		// Create a single changelist to use for all changes
 		int PendingChangeList = 0;
-		Dictionary<string, byte[]> InitalPOFileHashes = null;
 		if (P4Enabled)
 		{
 			var ChangeListCommitMessage = String.Format("Localization Automation using CL {0}", P4Env.Changelist);
@@ -212,203 +239,210 @@ class Localize : BuildCommand
 			}
 
 			PendingChangeList = P4.CreateChange(P4Env.Client, ChangeListCommitMessage);
-			InitalPOFileHashes = GetPOFileHashes(LocalizationBatches, UEProjectRoot);
 		}
 
-		// Process each localization batch
+		// Prepare to process each localization batch
+		var LocalizationTasks = new List<LocalizationTask>();
 		foreach (var LocalizationBatch in LocalizationBatches)
 		{
-			ProcessLocalizationProjects(LocalizationBatch, PendingChangeList, UEProjectRoot, UEProjectName, LocalizationProviderName, LocalizationStepNames, AdditionalCommandletArguments, EnableParallelGather);
-		}
+			var LocalizationTask = new LocalizationTask(LocalizationBatch, UEProjectRoot, LocalizationProviderName, PendingChangeList, this);
+			LocalizationTasks.Add(LocalizationTask);
 
-		// Clean-up the changelist so it only contains the changed files, and then submit it (if we were asked to)
-		if (P4Enabled)
-		{
-			// Revert any PO files that haven't changed aside from their header
+			// Make sure the Localization configs and content is up-to-date to ensure we don't get errors later on
+			if (P4Enabled)
 			{
-				var POFilesToRevert = new List<string>();
-
-				var CurrentPOFileHashes = GetPOFileHashes(LocalizationBatches, UEProjectRoot);
-				foreach (var CurrentPOFileHashPair in CurrentPOFileHashes)
-				{
-					byte[] InitialPOFileHash;
-					if (InitalPOFileHashes.TryGetValue(CurrentPOFileHashPair.Key, out InitialPOFileHash) && InitialPOFileHash.SequenceEqual(CurrentPOFileHashPair.Value))
-					{
-						POFilesToRevert.Add(CurrentPOFileHashPair.Key);
-					}
-				}
-
-				if (POFilesToRevert.Count > 0)
-				{
-					var P4RevertArgsFilename = CombinePaths(CmdEnv.LocalRoot, "Engine", "Intermediate", String.Format("LocalizationP4RevertArgs-{0}.txt", Guid.NewGuid().ToString()));
-
-					using (StreamWriter P4RevertArgsWriter = File.CreateText(P4RevertArgsFilename))
-					{
-						foreach (var POFileToRevert in POFilesToRevert)
-						{
-							P4RevertArgsWriter.WriteLine(POFileToRevert);
-						}
-					}
-
-					P4.LogP4(String.Format("-x{0} revert", P4RevertArgsFilename));
-					DeleteFile_NoExceptions(P4RevertArgsFilename);
-				}
+				LogInformation("Sync necessary content to head revision");
+				P4.Sync(P4Env.Branch + "/" + LocalizationTask.Batch.LocalizationTargetDirectory + "/Config/Localization/...");
+				P4.Sync(P4Env.Branch + "/" + LocalizationTask.Batch.LocalizationTargetDirectory + "/Content/Localization/...");
 			}
 
-			// Revert any other unchanged files
-			P4.RevertUnchanged(PendingChangeList);
-
-			// Submit that single changelist now
-			if (AllowSubmit)
+			// Generate the info we need to gather for each project
+			foreach (var ProjectName in LocalizationTask.Batch.LocalizationProjectNames)
 			{
-				int SubmittedChangeList;
-				P4.Submit(PendingChangeList, out SubmittedChangeList);
+				LocalizationTask.ProjectInfos.Add(GenerateProjectInfo(LocalizationTask.RootLocalizationTargetDirectory, ProjectName, LocalizationStepNames));
 			}
 		}
 
-		var RunDuration = (DateTime.UtcNow - StartTime).TotalMilliseconds;
-		LogInformation("Localize command finished in {0} seconds", RunDuration / 1000);
-	}
-
-	private void ProcessLocalizationProjects(LocalizationBatch LocalizationBatch, int PendingChangeList, string UEProjectRoot, string UEProjectName, string LocalizationProviderName, IReadOnlyList<string> LocalizationSteps, string AdditionalCommandletArguments, bool EnableParallelGather)
-	{
-		var EditorExe = CombinePaths(CmdEnv.LocalRoot, @"Engine/Binaries/Win64/UE4Editor-Cmd.exe");
-		var RootWorkingDirectory = CombinePaths(UEProjectRoot, LocalizationBatch.UEProjectDirectory);
-		var RootLocalizationTargetDirectory = CombinePaths(UEProjectRoot, LocalizationBatch.LocalizationTargetDirectory);
-
-		// Try and find our localization provider
-		LocalizationProvider LocProvider = null;
-		{
-			LocalizationProvider.LocalizationProviderArgs LocProviderArgs;
-			LocProviderArgs.RootWorkingDirectory = RootWorkingDirectory;
-			LocProviderArgs.RootLocalizationTargetDirectory = RootLocalizationTargetDirectory;
-			LocProviderArgs.RemoteFilenamePrefix = LocalizationBatch.RemoteFilenamePrefix;
-			LocProviderArgs.Command = this;
-			LocProviderArgs.PendingChangeList = PendingChangeList;
-			LocProvider = LocalizationProvider.GetLocalizationProvider(LocalizationProviderName, LocProviderArgs);
-		}
-
-		// Make sure the Localization configs and content is up-to-date to ensure we don't get errors later on
+		// Hash the current PO files on disk so we can work out whether they actually change
+		Dictionary<string, byte[]> InitalPOFileHashes = null;
 		if (P4Enabled)
 		{
-			LogInformation("Sync necessary content to head revision");
-			P4.Sync(P4Env.Branch + "/" + LocalizationBatch.LocalizationTargetDirectory + "/Config/Localization/...");
-			P4.Sync(P4Env.Branch + "/" + LocalizationBatch.LocalizationTargetDirectory + "/Content/Localization/...");
-		}
-
-		// Generate the info we need to gather for each project
-		var ProjectInfos = new List<ProjectInfo>();
-		foreach (var ProjectName in LocalizationBatch.LocalizationProjectNames)
-		{
-			ProjectInfos.Add(GenerateProjectInfo(RootLocalizationTargetDirectory, ProjectName, LocalizationSteps));
+			InitalPOFileHashes = GetPOFileHashes(LocalizationBatches, UEProjectRoot);
 		}
 
-		if (LocalizationSteps.Contains("Download") && LocProvider != null)
+		// Download the latest translations from our localization provider
+		if (LocalizationStepNames.Contains("Download"))
 		{
-			// Export all text from our localization provider
-			foreach (var ProjectInfo in ProjectInfos)
+			foreach (var LocalizationTask in LocalizationTasks)
 			{
-				LocProvider.DownloadProjectFromLocalizationProvider(ProjectInfo.ProjectName, ProjectInfo.ImportInfo);
+				if (LocalizationTask.LocProvider != null)
+				{
+					foreach (var ProjectInfo in LocalizationTask.ProjectInfos)
+					{
+						LocalizationTask.LocProvider.DownloadProjectFromLocalizationProvider(ProjectInfo.ProjectName, ProjectInfo.ImportInfo);
+					}
+				}
 			}
 		}
 
-		// Setup editor arguments for SCC.
-		string EditorArguments = String.Empty;
-		if (P4Enabled)
+		// Begin the gather command for each task
+		// These can run in parallel when ParallelGather is enabled
 		{
-			EditorArguments = String.Format("-SCCProvider={0} -P4Port={1} -P4User={2} -P4Client={3} -P4Passwd={4} -P4Changelist={5} -EnableSCC -DisableSCCSubmit", "Perforce", P4Env.ServerAndPort, P4Env.User, P4Env.Client, P4.GetAuthenticationToken(), PendingChangeList);
-		}
-		else
-		{
-			EditorArguments = String.Format("-SCCProvider={0}", "None");
-		}
-		if (IsBuildMachine)
-		{
-			EditorArguments += " -BuildMachine";
-		}
-		EditorArguments += " -Unattended -LogLocalizationConflicts";
+			var EditorExe = CombinePaths(CmdEnv.LocalRoot, @"Engine/Binaries/Win64/UE4Editor-Cmd.exe");
 
-		// Execute commandlet for all configs in each project, optionally running multiple commandlets in parallel.
-		var LocalizationGatherProcessResults = new List<IProcessResult>();
-		foreach (var ProjectInfo in ProjectInfos)
-		{
-			var LocalizationConfigFiles = new List<string>();
-			foreach (var LocalizationStep in ProjectInfo.LocalizationSteps)
+			// Set the common basic editor arguments
+			var EditorArguments = P4Enabled 
+				? String.Format("-SCCProvider=Perforce -P4Port={0} -P4User={1} -P4Client={2} -P4Passwd={3} -P4Changelist={4} -EnableSCC -DisableSCCSubmit", P4Env.ServerAndPort, P4Env.User, P4Env.Client, P4.GetAuthenticationToken(), PendingChangeList)
+				: "-SCCProvider=None";
+			if (IsBuildMachine)
 			{
-				if (LocalizationSteps.Contains(LocalizationStep.Name))
-				{
-					LocalizationConfigFiles.Add(LocalizationStep.LocalizationConfigFile);
-				}
+				EditorArguments += " -BuildMachine";
+			}
+			EditorArguments += " -Unattended -LogLocalizationConflicts";
+			if (EnableParallelGather)
+			{
+				EditorArguments += " -multiprocess";
+			}
+			if (!String.IsNullOrEmpty(AdditionalCommandletArguments))
+			{
+				EditorArguments += " " + AdditionalCommandletArguments;
 			}
 
-			if (LocalizationConfigFiles.Count > 0)
+			// Set the common process run options
+			var CommandletRunOptions = ERunOptions.Default | ERunOptions.NoLoggingOfRunCommand; // Disable logging of the run command as it will print the exit code which GUBP can pick up as an error (we do that ourselves later)
+			if (EnableParallelGather)
 			{
-				var ProjectArgument = String.IsNullOrEmpty(UEProjectName) ? "" : String.Format("\"{0}\"", Path.Combine(RootWorkingDirectory, String.Format("{0}.uproject", UEProjectName)));
-				var CommandletArguments = String.Format("-config=\"{0}\"", String.Join(";", LocalizationConfigFiles));
+				CommandletRunOptions |= ERunOptions.NoWaitForExit;
+			}
 
-				if (!String.IsNullOrEmpty(AdditionalCommandletArguments))
-				{
-					CommandletArguments += " " + AdditionalCommandletArguments;
-				}
+			foreach (var LocalizationTask in LocalizationTasks)
+			{
+				var ProjectArgument = String.IsNullOrEmpty(UEProjectName) ? "" : String.Format("\"{0}\"", Path.Combine(LocalizationTask.RootWorkingDirectory, String.Format("{0}.uproject", UEProjectName)));
 
-				var CommandletRunOptions = ERunOptions.Default | ERunOptions.NoLoggingOfRunCommand; // Disable logging of the run command as it will print the exit code which GUBP can pick up as an error (we do that ourselves later)
-				if (EnableParallelGather)
+				foreach (var ProjectInfo in LocalizationTask.ProjectInfos)
 				{
-					CommandletRunOptions |= ERunOptions.NoWaitForExit;
-					EditorArguments += " -multiprocess";
-				}
+					var LocalizationConfigFiles = new List<string>();
+					foreach (var LocalizationStep in ProjectInfo.LocalizationSteps)
+					{
+						if (LocalizationStepNames.Contains(LocalizationStep.Name))
+						{
+							LocalizationConfigFiles.Add(LocalizationStep.LocalizationConfigFile);
+						}
+					}
 
-				var Arguments = String.Format("{0} -run=GatherText {1} {2}", ProjectArgument, EditorArguments, CommandletArguments);
-				LogInformation("Running localization commandlet for '{0}': {1}", ProjectInfo.ProjectName, Arguments);
-				LocalizationGatherProcessResults.Add(Run(EditorExe, Arguments, null, CommandletRunOptions));
-			}
-			else
-			{
-				LocalizationGatherProcessResults.Add(null);
+					if (LocalizationConfigFiles.Count > 0)
+					{
+						var Arguments = String.Format("{0} -run=GatherText -config=\"{1}\" {2}", ProjectArgument, String.Join(";", LocalizationConfigFiles), EditorArguments);
+						LogInformation("Running localization commandlet for '{0}': {1}", ProjectInfo.ProjectName, Arguments);
+						LocalizationTask.GatherProcessResults.Add(Run(EditorExe, Arguments, null, CommandletRunOptions));
+					}
+					else
+					{
+						LocalizationTask.GatherProcessResults.Add(null);
+					}
+				}
 			}
 		}
 
 		// Wait for each commandlet process to finish and report the result.
 		// This runs even for non-parallel execution to log the exit state of the process.
-		for (int ProjectIndex = 0; ProjectIndex < ProjectInfos.Count; ++ProjectIndex)
+		foreach (var LocalizationTask in LocalizationTasks)
 		{
-			var ProjectInfo = ProjectInfos[ProjectIndex];
-			var RunResult = LocalizationGatherProcessResults[ProjectIndex];
-
-			if (RunResult != null)
+			for (int ProjectIndex = 0; ProjectIndex < LocalizationTask.ProjectInfos.Count; ++ProjectIndex)
 			{
-				RunResult.WaitForExit();
-				if (RunResult.ExitCode == 0)
+				var ProjectInfo = LocalizationTask.ProjectInfos[ProjectIndex];
+				var RunResult = LocalizationTask.GatherProcessResults[ProjectIndex];
+
+				if (RunResult != null)
 				{
-					LogInformation("The localization commandlet for '{0}' exited with code 0.", ProjectInfo.ProjectName);
+					RunResult.WaitForExit();
+					if (RunResult.ExitCode == 0)
+					{
+						LogInformation("The localization commandlet for '{0}' exited with code 0.", ProjectInfo.ProjectName);
+					}
+					else
+					{
+						LogWarning("The localization commandlet for '{0}' exited with code {1} which likely indicates a crash.", ProjectInfo.ProjectName, RunResult.ExitCode);
+					}
 				}
-				else
+			}
+		}
+
+		// Upload the latest sources to our localization provider
+		if (LocalizationStepNames.Contains("Upload"))
+		{
+			foreach (var LocalizationTask in LocalizationTasks)
+			{
+				if (LocalizationTask.LocProvider != null)
 				{
-					LogWarning("The localization commandlet for '{0}' exited with code {1} which likely indicates a crash.", ProjectInfo.ProjectName, RunResult.ExitCode);
+					// Upload all text to our localization provider
+					for (int ProjectIndex = 0; ProjectIndex < LocalizationTask.ProjectInfos.Count; ++ProjectIndex)
+					{
+						var ProjectInfo = LocalizationTask.ProjectInfos[ProjectIndex];
+						var RunResult = LocalizationTask.GatherProcessResults[ProjectIndex];
+
+						if (RunResult != null && RunResult.ExitCode == 0)
+						{
+							// Recalculate the split platform paths before doing the upload, as the export may have changed them
+							ProjectInfo.ExportInfo.CalculateSplitPlatformNames(LocalizationTask.RootLocalizationTargetDirectory);
+							LocalizationTask.LocProvider.UploadProjectToLocalizationProvider(ProjectInfo.ProjectName, ProjectInfo.ExportInfo);
+						}
+						else
+						{
+							LogWarning("Skipping upload to the localization provider for '{0}' due to an earlier commandlet failure.", ProjectInfo.ProjectName);
+						}
+					}
 				}
 			}
 		}
 
-		if (LocalizationSteps.Contains("Upload") && LocProvider != null)
+		// Clean-up the changelist so it only contains the changed files, and then submit it (if we were asked to)
+		if (P4Enabled)
 		{
-			// Upload all text to our localization provider
-			for (int ProjectIndex = 0; ProjectIndex < ProjectInfos.Count; ++ProjectIndex)
+			// Revert any PO files that haven't changed aside from their header
 			{
-				var ProjectInfo = ProjectInfos[ProjectIndex];
-				var RunResult = LocalizationGatherProcessResults[ProjectIndex];
+				var POFilesToRevert = new List<string>();
 
-				if (RunResult != null && RunResult.ExitCode == 0)
+				var CurrentPOFileHashes = GetPOFileHashes(LocalizationBatches, UEProjectRoot);
+				foreach (var CurrentPOFileHashPair in CurrentPOFileHashes)
 				{
-					// Recalculate the split platform paths before doing the upload, as the export may have changed them
-					ProjectInfo.ExportInfo.CalculateSplitPlatformNames(RootLocalizationTargetDirectory);
-					LocProvider.UploadProjectToLocalizationProvider(ProjectInfo.ProjectName, ProjectInfo.ExportInfo);
+					byte[] InitialPOFileHash;
+					if (InitalPOFileHashes.TryGetValue(CurrentPOFileHashPair.Key, out InitialPOFileHash) && InitialPOFileHash.SequenceEqual(CurrentPOFileHashPair.Value))
+					{
+						POFilesToRevert.Add(CurrentPOFileHashPair.Key);
+					}
 				}
-				else
+
+				if (POFilesToRevert.Count > 0)
 				{
-					LogWarning("Skipping upload to the localization provider for '{0}' due to an earlier commandlet failure.", ProjectInfo.ProjectName);
+					var P4RevertArgsFilename = CombinePaths(CmdEnv.LocalRoot, "Engine", "Intermediate", String.Format("LocalizationP4RevertArgs-{0}.txt", Guid.NewGuid().ToString()));
+
+					using (StreamWriter P4RevertArgsWriter = File.CreateText(P4RevertArgsFilename))
+					{
+						foreach (var POFileToRevert in POFilesToRevert)
+						{
+							P4RevertArgsWriter.WriteLine(POFileToRevert);
+						}
+					}
+
+					P4.LogP4(String.Format("-x{0} revert", P4RevertArgsFilename));
+					DeleteFile_NoExceptions(P4RevertArgsFilename);
 				}
 			}
+
+			// Revert any other unchanged files
+			P4.RevertUnchanged(PendingChangeList);
+
+			// Submit that single changelist now
+			if (AllowSubmit)
+			{
+				int SubmittedChangeList;
+				P4.Submit(PendingChangeList, out SubmittedChangeList);
+			}
 		}
+
+		var RunDuration = (DateTime.UtcNow - StartTime).TotalMilliseconds;
+		LogInformation("Localize command finished in {0} seconds", RunDuration / 1000);
 	}
 
 	private ProjectInfo GenerateProjectInfo(string RootWorkingDirectory, string ProjectName, IReadOnlyList<string> LocalizationStepNames)
