When staging multiple targets in a client build, detected if any of them are editor targets and force only development.

diff --git a/Engine/Source/Programs/AutomationTool/Scripts/CopyBuildToStagingDirectory.Automation.cs b/Engine/Source/Programs/AutomationTool/Scripts/CopyBuildToStagingDirectory.Automation.cs
--- a/Engine/Source/Programs/AutomationTool/Scripts/CopyBuildToStagingDirectory.Automation.cs
+++ b/Engine/Source/Programs/AutomationTool/Scripts/CopyBuildToStagingDirectory.Automation.cs
@@ -3372,6 +3372,26 @@ public partial class Project : CommandUtils
 		ParamList<string> ListToProcess = InDedicatedServer && (Params.Cook || Params.CookOnTheFly) ? Params.ServerCookedTargets : Params.ClientCookedTargets;
 		var ConfigsToProcess = InDedicatedServer && (Params.Cook || Params.CookOnTheFly) ? Params.ServerConfigsToBuild : Params.ClientConfigsToBuild;
 
+		List<Tuple<string, UnrealTargetConfiguration>> TargetAndConfigPairs = new List<Tuple<string, UnrealTargetConfiguration>>();
+
+		foreach (var Target in ListToProcess)
+		{
+			// If we are staging a client and have been asked to include editor targets, we currently only want to
+			// include a single Development editor target. Ideally, we should have shipping editor configs and then
+			// just include the requested configs for all targets
+			if (!InDedicatedServer && Params.EditorTargets.Contains(Target))
+			{
+				TargetAndConfigPairs.Add(new Tuple<string, UnrealTargetConfiguration>(Target, UnrealTargetConfiguration.Development));
+			}
+			else
+			{
+				foreach (var Config in ConfigsToProcess)
+				{
+					TargetAndConfigPairs.Add(new Tuple<string, UnrealTargetConfiguration>(Target, Config));
+				}
+			}
+		}
+
 		List<TargetPlatformDescriptor> PlatformsToStage = Params.ClientTargetPlatforms;
 		if (InDedicatedServer && (Params.Cook || Params.CookOnTheFly))
 		{
@@ -3393,17 +3413,16 @@ public partial class Project : CommandUtils
 
 			string PlatformName = StagePlatform.ToString();
 			string StageArchitecture = !String.IsNullOrEmpty(Params.SpecifiedArchitecture) ? Params.SpecifiedArchitecture : "";
-			foreach (var Target in ListToProcess)
+			foreach (var TargetAndConfig in TargetAndConfigPairs)
 			{
-				foreach (var Config in ConfigsToProcess)
+				string Target = TargetAndConfig.Item1;
+				UnrealTargetConfiguration Config = TargetAndConfig.Item2;
+				string Exe = Target;
+				if (Config != UnrealTargetConfiguration.Development)
 				{
-					string Exe = Target;
-					if (Config != UnrealTargetConfiguration.Development)
-					{
-						Exe = Target + "-" + PlatformName + "-" + Config.ToString() + StageArchitecture;
-					}
-					ExecutablesToStage.Add(Exe);
+					Exe = Target + "-" + PlatformName + "-" + Config.ToString() + StageArchitecture;
 				}
+				ExecutablesToStage.Add(Exe);
 			}
 
 			string StageDirectory = ((ShouldCreatePak(Params) || (Params.Stage)) || !String.IsNullOrEmpty(Params.StageDirectoryParam)) ? Params.BaseStageDirectory : "";
@@ -3412,63 +3431,62 @@ public partial class Project : CommandUtils
 			DirectoryReference ProjectDir = DirectoryReference.FromFile(Params.RawProjectPath);
 
 			List<StageTarget> TargetsToStage = new List<StageTarget>();
-			foreach (string Target in ListToProcess)
+			foreach (var TargetAndConfig in TargetAndConfigPairs)
 			{
-				foreach (UnrealTargetConfiguration Config in ConfigsToProcess)
-				{
-					DirectoryReference ReceiptBaseDir = Params.IsCodeBasedProject ? ProjectDir : EngineDir;
+				string Target = TargetAndConfig.Item1;
+				UnrealTargetConfiguration Config = TargetAndConfig.Item2;
+				DirectoryReference ReceiptBaseDir = Params.IsCodeBasedProject ? ProjectDir : EngineDir;
 
-					Platform PlatformInstance = Platform.Platforms[StagePlatform];
-					UnrealTargetPlatform[] SubPlatformsToStage = PlatformInstance.GetStagePlatforms();
+				Platform PlatformInstance = Platform.Platforms[StagePlatform];
+				UnrealTargetPlatform[] SubPlatformsToStage = PlatformInstance.GetStagePlatforms();
 
-					// if we are attempting to gathering multiple platforms, the files aren't required
-					bool bJustPackaging = Params.SkipStage && Params.Package;
-					bool bIsIterativeSharedCooking = Params.HasIterateSharedCookedBuild;
-					bool bRequireStagedFilesToExist = SubPlatformsToStage.Length == 1 && PlatformsToStage.Count == 1 && !bJustPackaging && !bIsIterativeSharedCooking;
+				// if we are attempting to gathering multiple platforms, the files aren't required
+				bool bJustPackaging = Params.SkipStage && Params.Package;
+				bool bIsIterativeSharedCooking = Params.HasIterateSharedCookedBuild;
+				bool bRequireStagedFilesToExist = SubPlatformsToStage.Length == 1 && PlatformsToStage.Count == 1 && !bJustPackaging && !bIsIterativeSharedCooking;
 
-					foreach (UnrealTargetPlatform ReceiptPlatform in SubPlatformsToStage)
+				foreach (UnrealTargetPlatform ReceiptPlatform in SubPlatformsToStage)
+				{
+					string Architecture = Params.SpecifiedArchitecture;
+					if (string.IsNullOrEmpty(Architecture))
 					{
-						string Architecture = Params.SpecifiedArchitecture;
-						if (string.IsNullOrEmpty(Architecture))
+						Architecture = "";
+						if (PlatformExports.IsPlatformAvailable(ReceiptPlatform))
 						{
-							Architecture = "";
-							if (PlatformExports.IsPlatformAvailable(ReceiptPlatform))
-							{
-								Architecture = PlatformExports.GetDefaultArchitecture(ReceiptPlatform, Params.RawProjectPath);
-							}
+							Architecture = PlatformExports.GetDefaultArchitecture(ReceiptPlatform, Params.RawProjectPath);
 						}
+					}
 
-						if (Params.IterateSharedBuildUsePrecompiledExe)
-						{
-							continue;
-						}
+					if (Params.IterateSharedBuildUsePrecompiledExe)
+					{
+						continue;
+					}
 
-						FileReference ReceiptFileName = TargetReceipt.GetDefaultPath(ReceiptBaseDir, Target, ReceiptPlatform, Config, Architecture);
-						if (!FileReference.Exists(ReceiptFileName))
+					FileReference ReceiptFileName = TargetReceipt.GetDefaultPath(ReceiptBaseDir, Target, ReceiptPlatform, Config, Architecture);
+					if (!FileReference.Exists(ReceiptFileName))
+					{
+						if (bRequireStagedFilesToExist)
 						{
-							if (bRequireStagedFilesToExist)
-							{
-								// if we aren't collecting multiple platforms, then it is expected to exist
-								throw new AutomationException(ExitCode.Error_MissingExecutable, "Stage Failed. Missing receipt '{0}'. Check that this target has been built.", ReceiptFileName);
-							}
-							else
-							{
-								// if it's multiple platforms, then allow missing receipts
-								continue;
-							}
-
+							// if we aren't collecting multiple platforms, then it is expected to exist
+							throw new AutomationException(ExitCode.Error_MissingExecutable, "Stage Failed. Missing receipt '{0}'. Check that this target has been built.", ReceiptFileName);
 						}
-
-						// Read the receipt for this target
-						TargetReceipt Receipt;
-						if (!TargetReceipt.TryRead(ReceiptFileName, out Receipt))
+						else
 						{
-							throw new AutomationException("Missing or invalid target receipt ({0})", ReceiptFileName);
+							// if it's multiple platforms, then allow missing receipts
+							continue;
 						}
 
-						// Convert the paths to absolute
-						TargetsToStage.Add(new StageTarget { Receipt = Receipt, RequireFilesExist = bRequireStagedFilesToExist });
 					}
+
+					// Read the receipt for this target
+					TargetReceipt Receipt;
+					if (!TargetReceipt.TryRead(ReceiptFileName, out Receipt))
+					{
+						throw new AutomationException("Missing or invalid target receipt ({0})", ReceiptFileName);
+					}
+
+					// Convert the paths to absolute
+					TargetsToStage.Add(new StageTarget { Receipt = Receipt, RequireFilesExist = bRequireStagedFilesToExist });
 				}
 			}
 
