Integrating LiveLink plugin changes from Fury

diff --git a/Engine/Source/Programs/AutomationTool/Scripts/BuildTarget.Automation.cs b/Engine/Source/Programs/AutomationTool/Scripts/BuildTarget.Automation.cs
--- a/Engine/Source/Programs/AutomationTool/Scripts/BuildTarget.Automation.cs
+++ b/Engine/Source/Programs/AutomationTool/Scripts/BuildTarget.Automation.cs
@@ -19,7 +19,7 @@ namespace AutomationTool
 	[Help("configuration=Development+Test", "Configurations to build, join multiple configurations using +")]
 	[Help("target=Editor+Game", "Targets to build, join multiple targets using +")]
 	[Help("notools", "Don't build any tools (UHT, ShaderCompiler, CrashReporter")]
-	class BuildTarget : BuildCommand
+	public class BuildTarget : BuildCommand
 	{
 		// exposed as a property so projects can derive and set this directly
 		public string ProjectName { get; set; }
@@ -46,6 +46,8 @@ namespace AutomationTool
 		{
 			string[] Arguments = this.Params;
 
+			string[] DefaultUE4Targets = new string[] { "editor", "game", "client", "server" };
+
 			ProjectName = ParseParamValue("project", ProjectName);
 			Targets = ParseParamValue("target", Targets);
 			Platforms = ParseParamValue("platform", Platforms);
@@ -94,6 +96,7 @@ namespace AutomationTool
 			}
 
 			FileReference ProjectFile = null;
+			DirectoryReference TargetSource = DirectoryReference.Combine(BuildCommand.EngineDirectory, "Source");
 
 			if (!string.IsNullOrEmpty(ProjectName))
 			{
@@ -104,48 +107,39 @@ namespace AutomationTool
 					throw new AutomationException("Unable to find uproject file for {0}", ProjectName);
 				}
 
-				string SourceDirectoryName = Path.Combine(ProjectFile.Directory.FullName, "Source");
-
-				if (Directory.Exists(SourceDirectoryName))
-				{
-					IEnumerable<string> TargetScripts = Directory.EnumerateFiles(SourceDirectoryName, "*.Target.cs");
-
-					foreach (string TargetName in TargetList)
-					{
-						string TargetScript = TargetScripts.Where(S => S.IndexOf(TargetName, StringComparison.OrdinalIgnoreCase) >= 0).FirstOrDefault();
-
-						if (TargetScript == null && (
-								TargetName.Equals("Client", StringComparison.OrdinalIgnoreCase) ||
-								TargetName.Equals("Game", StringComparison.OrdinalIgnoreCase)
-								)
-							)
-						{
-							// if there's no ProjectGame.Target.cs or ProjectClient.Target.cs then
-							// fallback to Project.Target.cs
-							TargetScript = TargetScripts.Where(S => S.IndexOf(ProjectName + ".", StringComparison.OrdinalIgnoreCase) >= 0).FirstOrDefault();
-						}
-
-						if (TargetScript == null)
-						{
-							throw new AutomationException("No Target.cs file for target {0} in project {1}", TargetName, ProjectName);
-						}
-
-						string FullName = Path.GetFileName(TargetScript);
-						TargetNames[TargetName] = Regex.Replace(FullName, ".Target.cs", "", RegexOptions.IgnoreCase);
-					}
-				}
+				TargetSource = DirectoryReference.Combine(ProjectFile.Directory, "Source");
 			}
 			else
 			{
-				Log.TraceWarning("No project specified, will build vanilla UE4 binaries");
+				Log.TraceInformation("No project specified, will build vanilla UE4 binaries");
 			}
 
-			// Handle content-only projects or when no project was specified
-			if (TargetNames.Keys.Count == 0)
-			{ 
-				foreach (string TargetName in TargetList)
+			IEnumerable<string> TargetScripts = DirectoryReference.EnumerateFiles(TargetSource, "*.Target.cs").Select(F => F.GetFileName());
+
+			foreach (string TargetName in TargetList)
+			{
+				string TargetScript = TargetScripts.Where(S => S.IndexOf(TargetName, StringComparison.OrdinalIgnoreCase) >= 0).FirstOrDefault();
+
+				if (TargetScript == null && (
+						TargetName.Equals("Client", StringComparison.OrdinalIgnoreCase) ||
+						TargetName.Equals("Game", StringComparison.OrdinalIgnoreCase)
+						)
+					)
+				{
+					// if there's no ProjectGame.Target.cs or ProjectClient.Target.cs then
+					// fallback to Project.Target.cs
+					TargetScript = TargetScripts.Where(S => S.IndexOf(ProjectName + ".", StringComparison.OrdinalIgnoreCase) >= 0).FirstOrDefault();
+				}
+
+				// IF we found a script for a target, then that is our build target (e.g. QAGame for Game)
+				if (TargetScript != null)
+				{
+					TargetNames[TargetName] = Regex.Replace(TargetScript, ".Target.cs", "", RegexOptions.IgnoreCase);
+				}
+				else
 				{
-					TargetNames[TargetName] = string.Format("UE4{0}", TargetName);
+					// else just build what they asked for
+					TargetNames[TargetName] = TargetName;
 				}
 			}
 
