Use the IsBuildMachine environment variable to set the IsBuildMachine flag, rather than inferring it from uebp_LOCAL_ROOT. The uebp_LOCAL_ROOT variable is defined for child UAT instances, which causes the IsBuildMachine flag to be set incorrectly.

diff --git a/Engine/Source/Programs/AutomationTool/AutomationUtils/Automation.cs b/Engine/Source/Programs/AutomationTool/AutomationUtils/Automation.cs
--- a/Engine/Source/Programs/AutomationTool/AutomationUtils/Automation.cs
+++ b/Engine/Source/Programs/AutomationTool/AutomationUtils/Automation.cs
@@ -421,7 +421,17 @@ AutomationTool.exe [-verbose] [-compileonly] [-p4] Command0 [-Arg0 -Arg1 -Arg2 .
 		{
 			// Initial check for local or build machine runs BEFORE we parse the command line (We need this value set
 			// in case something throws the exception while parsing the command line)
-			IsBuildMachine = !String.IsNullOrEmpty(Environment.GetEnvironmentVariable("uebp_LOCAL_ROOT")) || Arguments.Any(x => x.Equals("-BuildMachine", StringComparison.InvariantCultureIgnoreCase));
+			IsBuildMachine = Arguments.Any(x => x.Equals("-BuildMachine", StringComparison.InvariantCultureIgnoreCase));
+			if (!IsBuildMachine)
+			{
+				int Value;
+				if (int.TryParse(Environment.GetEnvironmentVariable("IsBuildMachine"), out Value) && Value != 0)
+				{
+					IsBuildMachine = true;
+				}
+			}
+			Log.TraceVerbose("IsBuildMachine={0}", IsBuildMachine);
+			Environment.SetEnvironmentVariable("IsBuildMachine", IsBuildMachine ? "1" : "0");
 
 			// Scan the command line for commands to execute.
 			var CommandsToExecute = new List<CommandInfo>();
@@ -431,9 +441,6 @@ AutomationTool.exe [-verbose] [-compileonly] [-p4] Command0 [-Arg0 -Arg1 -Arg2 .
 
 			// Get the path to the telemetry file, if present
 			string TelemetryFile = CommandUtils.ParseParamValue(Arguments, "-Telemetry");
-			
-			Log.TraceVerbose("IsBuildMachine={0}", IsBuildMachine);
-			Environment.SetEnvironmentVariable("IsBuildMachine", IsBuildMachine ? "1" : "0");
 
 			// should we kill processes on exit
 			ShouldKillProcesses = !GlobalCommandLine.NoKill;
