Create new artifact directories as part of starttest instead of stoptest to prevent A/B test runs on EC from all selecting the same target directory.

diff --git a/Engine/Source/Programs/AutomationTool/Gauntlet/Unreal/Base/Gauntlet.UnrealTestNode.cs b/Engine/Source/Programs/AutomationTool/Gauntlet/Unreal/Base/Gauntlet.UnrealTestNode.cs
--- a/Engine/Source/Programs/AutomationTool/Gauntlet/Unreal/Base/Gauntlet.UnrealTestNode.cs
+++ b/Engine/Source/Programs/AutomationTool/Gauntlet/Unreal/Base/Gauntlet.UnrealTestNode.cs
@@ -503,6 +503,10 @@ namespace Gauntlet
 
 			ReservedArtifcactPaths.Add(ArtifactPath);
 
+			// We need to create this directory at the start of the test rather than the end of the test - we are running into instances where multiple A/B tests
+			// on the same build are seeing the directory as non-existent and thinking it is safe to write to.
+			Directory.CreateDirectory(ArtifactPath);
+
 			// Launch the test
 			TestInstance = UnrealApp.LaunchSession();
 
@@ -638,6 +642,7 @@ namespace Gauntlet
 			try
 			{
 				Log.Info("Saving artifacts to {0}", ArtifactPath);
+				// run create dir again just in case the already made dir was cleaned up by another buildfarm job or something similar.
 				Directory.CreateDirectory(ArtifactPath);
 				Utils.SystemHelpers.MarkDirectoryForCleanup(ArtifactPath);
 
