- Fixed the call to UnrealPak with -platform to use Windows not Win64
#UE-88272

diff --git a/Engine/Source/Programs/AutomationTool/Scripts/CopyBuildToStagingDirectory.Automation.cs b/Engine/Source/Programs/AutomationTool/Scripts/CopyBuildToStagingDirectory.Automation.cs
--- a/Engine/Source/Programs/AutomationTool/Scripts/CopyBuildToStagingDirectory.Automation.cs
+++ b/Engine/Source/Programs/AutomationTool/Scripts/CopyBuildToStagingDirectory.Automation.cs
@@ -2199,7 +2199,7 @@ public partial class Project : CommandUtils
 					}
 
 					// pass the targetplatform so the index may be able to be frozen
-					AdditionalArgs += " -platform=" + SC.StageTargetPlatform.IniPlatformType;
+					AdditionalArgs += " -platform=" + ConfigHierarchy.GetIniPlatformName(SC.StageTargetPlatform.IniPlatformType);
 
 					Dictionary<string, string> UnrealPakResponseFile;
 					if (ShouldCreateIoStoreContainerFiles(Params, SC))
