- Moved AudioCaptureModuleName spec to a .ini setting, and updated all relevant .ini files
- Fixed .ini staging to follow the IniParent chain for PlatformExtensions

diff --git a/Engine/Config/Android/BaseAndroidEngine.ini b/Engine/Config/Android/BaseAndroidEngine.ini
--- a/Engine/Config/Android/BaseAndroidEngine.ini
+++ b/Engine/Config/Android/BaseAndroidEngine.ini
@@ -1,2 +1,5 @@
 [OnlineSubsystem]
-NativePlatformService=GooglePlay
\ No newline at end of file
+NativePlatformService=GooglePlay
+
+[Audio]
+AudioCaptureModuleName=AudioCaptureAndroid
diff --git a/Engine/Config/IOS/BaseIOSEngine.ini b/Engine/Config/IOS/BaseIOSEngine.ini
--- a/Engine/Config/IOS/BaseIOSEngine.ini
+++ b/Engine/Config/IOS/BaseIOSEngine.ini
@@ -1,2 +1,5 @@
 [OnlineSubsystem]
-NativePlatformService=IOS
\ No newline at end of file
+NativePlatformService=IOS
+
+[Audio]
+AudioCaptureModuleName=AudioCaptureAudioUnit
diff --git a/Engine/Config/Mac/BaseMacEngine.ini b/Engine/Config/Mac/BaseMacEngine.ini
new file mode 100644
--- /dev/null
+++ b/Engine/Config/Mac/BaseMacEngine.ini
@@ -0,0 +1,2 @@
+[Audio]
+AudioCaptureModuleName=AudioCaptureRtAudio
diff --git a/Engine/Config/Windows/BaseWindowsEngine.ini b/Engine/Config/Windows/BaseWindowsEngine.ini
new file mode 100644
--- /dev/null
+++ b/Engine/Config/Windows/BaseWindowsEngine.ini
@@ -0,0 +1,2 @@
+[Audio]
+AudioCaptureModuleName=AudioCaptureRtAudio
diff --git a/Engine/Plugins/Runtime/AudioCapture/Source/AudioCapture/AudioCapture.Build.cs b/Engine/Plugins/Runtime/AudioCapture/Source/AudioCapture/AudioCapture.Build.cs
--- a/Engine/Plugins/Runtime/AudioCapture/Source/AudioCapture/AudioCapture.Build.cs
+++ b/Engine/Plugins/Runtime/AudioCapture/Source/AudioCapture/AudioCapture.Build.cs
@@ -25,7 +25,7 @@ namespace UnrealBuildTool.Rules
             }
             else if (Target.Platform == UnrealTargetPlatform.PS4)
             {
-                PrivateDependencyModuleNames.Add("AudioCapturePS4Voice");
+                PrivateDependencyModuleNames.Add("AudioCaptureSony");
             }
             else if (Target.Platform == UnrealTargetPlatform.IOS)
             {
diff --git a/Engine/Plugins/Runtime/AudioCapture/Source/AudioCapture/Private/AudioCapture.cpp b/Engine/Plugins/Runtime/AudioCapture/Source/AudioCapture/Private/AudioCapture.cpp
--- a/Engine/Plugins/Runtime/AudioCapture/Source/AudioCapture/Private/AudioCapture.cpp
+++ b/Engine/Plugins/Runtime/AudioCapture/Source/AudioCapture/Private/AudioCapture.cpp
@@ -92,18 +92,12 @@ UAudioCapture* UAudioCaptureFunctionLibrary::CreateAudioCapture()
 
 void FAudioCaptureModule::StartupModule()
 {
-	// Load platform specific implementations for audio capture:
-#if PLATFORM_WINDOWS || PLATFORM_MAC || PLATFORM_XBOXONE
-	FModuleManager::Get().LoadModule(TEXT("AudioCaptureRtAudio"));
-#elif PLATFORM_IOS
-	FModuleManager::Get().LoadModule(TEXT("AudioCaptureAudioUnit"));
-#elif PLATFORM_PS4
-	FModuleManager::Get().LoadModule(TEXT("AudioCapturePS4Voice"));
-#elif PLATFORM_ANDROID
-	FModuleManager::Get().LoadModule(TEXT("AudioCaptureAndroid"));
-#elif PLATFORM_SWITCH
-	FModuleManager::Get().LoadModule(TEXT("AudioCaptureSwitch"));
-#endif
+	// Load platform specific implementations for audio capture (if specified in a .ini file)
+	FString AudioCaptureModuleName;
+	if (GConfig->GetString(TEXT("Audio"), TEXT("AudioCaptureModuleName"), AudioCaptureModuleName, GEngineIni))
+	{
+		FModuleManager::Get().LoadModule(*AudioCaptureModuleName);
+	}
 }
 
 void FAudioCaptureModule::ShutdownModule()
diff --git a/Engine/Source/Programs/AutomationTool/Scripts/CopyBuildToStagingDirectory.Automation.cs b/Engine/Source/Programs/AutomationTool/Scripts/CopyBuildToStagingDirectory.Automation.cs
--- a/Engine/Source/Programs/AutomationTool/Scripts/CopyBuildToStagingDirectory.Automation.cs
+++ b/Engine/Source/Programs/AutomationTool/Scripts/CopyBuildToStagingDirectory.Automation.cs
@@ -655,18 +655,17 @@ public partial class Project : CommandUtils
 				// Stage platform extension config files
 				{
 					string PlatformExtensionName = SC.StageTargetPlatform.PlatformType.ToString();
+					StagePlatformExtensionConfigFiles(SC, PlatformExtensionName);
 
-					DirectoryReference PlatformEngineConfigDir = DirectoryReference.Combine(SC.LocalRoot, "Engine", "Platforms", PlatformExtensionName, "Config");
-					DirectoryReference PlatformProjectConfigDir = DirectoryReference.Combine(SC.ProjectRoot, "Platforms", PlatformExtensionName, "Config");
-
-					if (DirectoryReference.Exists(PlatformEngineConfigDir))
+					// and any ini parents it has
+					DataDrivenPlatformInfo.ConfigDataDrivenPlatformInfo Info = DataDrivenPlatformInfo.GetDataDrivenInfoForPlatform(PlatformExtensionName);
+					if (Info != null && Info.IniParentChain != null)
 					{
-						StageConfigFiles(SC, PlatformEngineConfigDir, PlatformExtensionName);
-					}
-
-					if (DirectoryReference.Exists(PlatformProjectConfigDir))
-					{
-						StageConfigFiles(SC, PlatformProjectConfigDir, PlatformExtensionName);
+						// the IniParentChain
+						foreach (string ParentPlatform in Info.IniParentChain)
+						{
+							StagePlatformExtensionConfigFiles(SC, ParentPlatform);
+						}
 					}
 				}
 
@@ -1070,6 +1069,27 @@ public partial class Project : CommandUtils
 		}
 	}
 
+	/// <summary>
+	/// Looks for a potential PlatformExtension config directory
+	/// </summary>
+	/// <param name="SC">The staging context</param>
+	/// <param name="PlatformName">The name of the platform/platform parent</param>
+	static void StagePlatformExtensionConfigFiles(DeploymentContext SC, string PlatformName)
+	{
+		DirectoryReference PlatformEngineConfigDir = DirectoryReference.Combine(SC.LocalRoot, "Engine", "Platforms", PlatformName, "Config");
+		DirectoryReference PlatformProjectConfigDir = DirectoryReference.Combine(SC.ProjectRoot, "Platforms", PlatformName, "Config");
+
+		if (DirectoryReference.Exists(PlatformEngineConfigDir))
+		{
+			StageConfigFiles(SC, PlatformEngineConfigDir, PlatformName);
+
+		}
+		if (DirectoryReference.Exists(PlatformProjectConfigDir))
+		{
+			StageConfigFiles(SC, PlatformProjectConfigDir, PlatformName);
+		}
+	}
+
 	/// <summary>
 	/// Determines if an individual config file should be staged
 	/// </summary>
