Prevent shader bytecode and shader pipeline cache files from being compressed.

diff --git a/Engine/Source/Programs/AutomationTool/Scripts/CopyBuildToStagingDirectory.Automation.cs b/Engine/Source/Programs/AutomationTool/Scripts/CopyBuildToStagingDirectory.Automation.cs
--- a/Engine/Source/Programs/AutomationTool/Scripts/CopyBuildToStagingDirectory.Automation.cs
+++ b/Engine/Source/Programs/AutomationTool/Scripts/CopyBuildToStagingDirectory.Automation.cs
@@ -79,7 +79,7 @@ public partial class Project : CommandUtils
 			foreach (var Entry in ResponseFile)
 			{
 				string Line = String.Format("\"{0}\" \"{1}\"", Entry.Key, Entry.Value);
-				if (Compressed && !Path.GetExtension(Entry.Key).Contains(".mp4"))
+				if (Compressed && !Path.GetExtension(Entry.Key).Contains(".mp4") && !Path.GetExtension(Entry.Key).Contains("ushaderbytecode") && !Path.GetExtension(Entry.Key).Contains("upipelinecache"))
 				{
 					Line += " -compress";
 				}
