HTML5 - get launch onto current editor map working

#rn

diff --git a/Engine/Source/Programs/AutomationTool/HTML5/HTML5Platform.Automation.cs b/Engine/Source/Programs/AutomationTool/HTML5/HTML5Platform.Automation.cs
--- a/Engine/Source/Programs/AutomationTool/HTML5/HTML5Platform.Automation.cs
+++ b/Engine/Source/Programs/AutomationTool/HTML5/HTML5Platform.Automation.cs
@@ -176,18 +176,18 @@ public class HTML5Platform : Platform
 			BuildPath = LocalBuildPath;
 			TemplateFile = CombinePaths(BuildPath, "project_template.html");
 		}
-		GenerateFileFromTemplate(TemplateFile, ProjectBasename + ".html", Params, ConfigCache);
+		GenerateFileFromTemplate(TemplateFile, ProjectBasename + ".html", Params, SC, ConfigCache);
 
 		TemplateFile = CombinePaths(BuildPath, "project_template.js");
 		if ( File.Exists(TemplateFile) )
 		{
-			GenerateFileFromTemplate(TemplateFile, ProjectBasename + ".UE4.js", Params, ConfigCache);
+			GenerateFileFromTemplate(TemplateFile, ProjectBasename + ".UE4.js", Params, SC, ConfigCache);
 		}
 
 		TemplateFile = CombinePaths(BuildPath, "project_template.css");
 		if ( File.Exists(TemplateFile) )
 		{
-			GenerateFileFromTemplate(TemplateFile, ProjectBasename + ".css", Params, ConfigCache);
+			GenerateFileFromTemplate(TemplateFile, ProjectBasename + ".css", Params, SC, ConfigCache);
 		}
 
 
@@ -310,7 +310,7 @@ public class HTML5Platform : Platform
 		}
 	}
 
-	protected void GenerateFileFromTemplate(string InTemplateFile, string InOutputFile, ProjectParams Params, ConfigHierarchy ConfigCache)
+	protected void GenerateFileFromTemplate(string InTemplateFile, string InOutputFile, ProjectParams Params, DeploymentContext SC, ConfigHierarchy ConfigCache)
 	{
 		bool IsContentOnly = !Params.IsCodeBasedProject;
 		string ProjectConfiguration = Params.ClientConfigsToBuild[0].ToString();
@@ -368,6 +368,11 @@ public class HTML5Platform : Platform
 				if (LineStr.Contains("%UE4CMDLINE%"))
 				{
 					string ArgumentString = "'../../../" + Params.ShortProjectName + "/" + Params.ShortProjectName + ".uproject',";
+
+					if ( Params.MapToRun.Length > 0 )
+					{
+						ArgumentString += "'" + Params.MapToRun + "',";
+					}
 					ArgumentString += "'-stdout',"; // suppress double printing to console.log
 					LineStr = LineStr.Replace("%UE4CMDLINE%", ArgumentString);
 				}
