HTML5 - adjust AutomationTool for HTML5 to handle "alternative" project names

#jria UE-75713  HTML5 fails to Launch On or Open Packaged Game for Samples/Sandbox/Build/BuildTest.uproject - 404 Not Found on file BuildTest.js/BuildTest.wasm
#rn

diff --git a/Engine/Source/Programs/AutomationTool/HTML5/HTML5Platform.Automation.cs b/Engine/Source/Programs/AutomationTool/HTML5/HTML5Platform.Automation.cs
--- a/Engine/Source/Programs/AutomationTool/HTML5/HTML5Platform.Automation.cs
+++ b/Engine/Source/Programs/AutomationTool/HTML5/HTML5Platform.Automation.cs
@@ -1,4 +1,4 @@
-﻿// Copyright 1998-2019 Epic Games, Inc. All Rights Reserved.
+// Copyright 1998-2019 Epic Games, Inc. All Rights Reserved.
 using System;
 using System.IO;
 using System.Net.Http;
@@ -70,6 +70,11 @@ public class HTML5Platform : Platform
 		string UE4GameBasename = Path.Combine(PackagePath, _ProjectFilename);
 		string ProjectBasename = Path.Combine(PackagePath, Params.ShortProjectName + _ProjectNameExtra);
 
+		if (Params.IsCodeBasedProject && Params.HasClientCookedTargets)
+		{
+			UE4GameBasename = Params.ClientCookedTargets[0];
+		}
+
 		// ----------------------------------------
 		// packaging
 		if (HTMLPakAutomation.CanCreateMapPaks(Params))
@@ -315,7 +320,7 @@ public class HTML5Platform : Platform
 		bool IsContentOnly = !Params.IsCodeBasedProject;
 		string ProjectConfiguration = Params.ClientConfigsToBuild[0].ToString();
 
-		string UE4GameName = IsContentOnly ? "UE4Game" : Params.ShortProjectName;
+		string UE4GameName = IsContentOnly ? "UE4Game" : (Params.HasClientCookedTargets ? Params.ClientCookedTargets[0] : Params.ShortProjectName);
 		string ProjectName = Params.ShortProjectName;
 		if (ProjectConfiguration != "Development")
 		{
@@ -498,6 +503,10 @@ public class HTML5Platform : Platform
 		string UE4GameBasename = Path.GetFileNameWithoutExtension(Params.GetProjectExeForPlatform(UnrealTargetPlatform.HTML5).ToString());
 		string ProjectBasename = Params.ShortProjectName;
 		string ProjectConfiguration = Params.ClientConfigsToBuild[0].ToString();
+		if (Params.IsCodeBasedProject && Params.HasClientCookedTargets)
+		{
+			UE4GameBasename = Params.ClientCookedTargets[0];
+		}
 		if (ProjectConfiguration != "Development")
 		{
 			UE4GameBasename += "-HTML5-" + ProjectConfiguration;
