Fixed compile error for framework change. I used a feature of C# that isn't supported by our build stuff.

diff --git a/Engine/Source/Programs/AutomationTool/IOS/IOSPlatform.Automation.cs b/Engine/Source/Programs/AutomationTool/IOS/IOSPlatform.Automation.cs
--- a/Engine/Source/Programs/AutomationTool/IOS/IOSPlatform.Automation.cs
+++ b/Engine/Source/Programs/AutomationTool/IOS/IOSPlatform.Automation.cs
@@ -1,4 +1,4 @@
-// Copyright 1998-2019 Epic Games, Inc. All Rights Reserved.
+﻿// Copyright 1998-2019 Epic Games, Inc. All Rights Reserved.
 using System;
 using System.Collections.Generic;
 using System.Linq;
@@ -276,7 +276,8 @@ public class IOSPlatform : Platform
 		bool bIsUE4Game = !SC.IsCodeBasedProject;
 
 		FileReference ReceiptFileName = GetTargetReceiptFileName(Config, InExecutablePath, InEngineDir, InProjectDirectory, bIsUE4Game);
-		bool bIsReadSuccessful = TargetReceipt.TryRead(ReceiptFileName, out TargetReceipt Receipt);
+		TargetReceipt Receipt;
+		bool bIsReadSuccessful = TargetReceipt.TryRead(ReceiptFileName, out Receipt);
 
 		bool bIsBuiltAsFramework = false;
 		if (bIsReadSuccessful)
