Improving S3 CDN upload robustness with retry fixes and additions.
Adding timeouts to particular network activity that we believe may be hanging.

diff --git a/Engine/Source/Programs/AutomationTool/AutomationUtils/MCPPublic.cs b/Engine/Source/Programs/AutomationTool/AutomationUtils/MCPPublic.cs
--- a/Engine/Source/Programs/AutomationTool/AutomationUtils/MCPPublic.cs
+++ b/Engine/Source/Programs/AutomationTool/AutomationUtils/MCPPublic.cs
@@ -1472,13 +1472,10 @@ namespace EpicGames.MCP.Automation
 		/// <param name="ContentType">The MIME type of the file being uploaded. If left NULL, will be determined server-side by cloud provider.</param>
 		/// <param name="bOverwrite">If true, will overwrite an existing file.  If false, will throw an exception if the file exists.</param>
 		/// <param name="bMakePublic">Specifies whether the file should be made publicly readable.</param>
-        /// <param name="bQuiet">If set to true, all log output for the operation is supressed.</param>
+		/// <param name="bQuiet">If set to true, all log output for the operation is supressed.</param>
 		/// <param name="Metadata">If not null, key-value pairs of metadata to be applied to the object.</param>
 		/// <returns>A PostFileResult indicating whether the call was successful, and the URL to the uploaded file.</returns>
-		public PostFileResult PostFile(string Container, string Identifier, byte[] Contents, string ContentType = null, bool bOverwrite = true, bool bMakePublic = false, bool bQuiet = false, IDictionary<string, object> Metadata = null)
-		{
-			return PostFileAsync(Container, Identifier, Contents, ContentType, bOverwrite, bMakePublic, bQuiet, Metadata).Result;
-		}
+		abstract public PostFileResult PostFile(string Container, string Identifier, byte[] Contents, string ContentType = null, bool bOverwrite = true, bool bMakePublic = false, bool bQuiet = false, IDictionary<string, object> Metadata = null);
 
 		/// <summary>
 		/// Posts a file to the cloud storage provider asynchronously.
@@ -1503,13 +1500,10 @@ namespace EpicGames.MCP.Automation
 		/// <param name="ContentType">The MIME type of the file being uploaded. If left NULL, will be determined server-side by cloud provider.</param>
 		/// <param name="bOverwrite">If true, will overwrite an existing file.  If false, will throw an exception if the file exists.</param>
 		/// <param name="bMakePublic">Specifies whether the file should be made publicly readable.</param>
-        /// <param name="bQuiet">If set to true, all log output for the operation is supressed.</param>
+		/// <param name="bQuiet">If set to true, all log output for the operation is supressed.</param>
 		/// <param name="Metadata">If not null, key-value pairs of metadata to be applied to the object.</param>
 		/// <returns>A PostFileResult indicating whether the call was successful, and the URL to the uploaded file.</returns>
-		public PostFileResult PostFile(string Container, string Identifier, string SourceFilePath, string ContentType = null, bool bOverwrite = true, bool bMakePublic = false, bool bQuiet = false, IDictionary<string, object> Metadata = null)
-		{
-			return PostFileAsync(Container, Identifier, SourceFilePath, ContentType, bOverwrite, bMakePublic, bQuiet, Metadata).Result;
-		}
+		abstract public PostFileResult PostFile(string Container, string Identifier, string SourceFilePath, string ContentType = null, bool bOverwrite = true, bool bMakePublic = false, bool bQuiet = false, IDictionary<string, object> Metadata = null);
 
 		/// <summary>
 		/// Posts a file to the cloud storage provider asynchronously.
