Ignore duplicate lib*.dylib files produced as build artifacts. Hopefully fix installed build on Mac.

#rnx

diff --git a/Engine/Source/Programs/AutomationTool/BuildGraph/TempStorage.cs b/Engine/Source/Programs/AutomationTool/BuildGraph/TempStorage.cs
--- a/Engine/Source/Programs/AutomationTool/BuildGraph/TempStorage.cs
+++ b/Engine/Source/Programs/AutomationTool/BuildGraph/TempStorage.cs
@@ -1137,6 +1137,10 @@ namespace AutomationTool
 			{
 				return true;
 			}
+			if (FileName.StartsWith("lib", StringComparison.OrdinalIgnoreCase) && FileName.EndsWith(".dylib", StringComparison.OrdinalIgnoreCase))
+			{
+				return true;
+			}
 			return false;
 		}
 	}
