Sanitize device reservation messages as these are not fatal for jobs
#rnx

diff --git a/Engine/Source/Programs/AutomationTool/AutomationUtils/DeviceReservation.cs b/Engine/Source/Programs/AutomationTool/AutomationUtils/DeviceReservation.cs
--- a/Engine/Source/Programs/AutomationTool/AutomationUtils/DeviceReservation.cs
+++ b/Engine/Source/Programs/AutomationTool/AutomationUtils/DeviceReservation.cs
@@ -308,7 +308,7 @@ namespace AutomationTool.DeviceReservation
 			}
 			catch (Exception ex)
 			{
-				Console.WriteLine("Failed to delete device reservation: {0}", ex.Message);
+				Console.WriteLine("Failed to delete device reservation: {0}", SanitizeErrorMessage(ex.Message));
 			}
 		}
 
@@ -327,7 +327,7 @@ namespace AutomationTool.DeviceReservation
 			}
 			catch (Exception Ex)
 			{
-				Console.WriteLine("Failed to report device: {0} : {1}", DeviceName, Ex.Message);
+				Console.WriteLine("Failed to report device: {0} : {1}", DeviceName, SanitizeErrorMessage(Ex.Message));
 			}
 		}
 
