[IOS] Check if the bundle exists before trying to copy manifests.

- This prevents an error being displayed in the log before staging.

diff --git a/Engine/Source/Programs/AutomationTool/IOS/IOSPlatform.Automation.cs b/Engine/Source/Programs/AutomationTool/IOS/IOSPlatform.Automation.cs
--- a/Engine/Source/Programs/AutomationTool/IOS/IOSPlatform.Automation.cs
+++ b/Engine/Source/Programs/AutomationTool/IOS/IOSPlatform.Automation.cs
@@ -1243,13 +1243,27 @@ public class IOSPlatform : Platform
 				int EndPos = Contents.IndexOf("</string>", Pos);
 				BundleIdentifier = Contents.Substring(Pos, EndPos - Pos);
 			}
-			RunAndLog(CmdEnv, DeployServer, "Backup -file \"" + CombinePaths(Params.BaseStageDirectory, PlatformName, SC.GetUFSDeployedManifestFileName(null)) + "\" -file \"" + CombinePaths(Params.BaseStageDirectory, PlatformName, SC.GetNonUFSDeployedManifestFileName(null)) + "\"" + (String.IsNullOrEmpty(Params.DeviceNames[0]) ? "" : " -device " + Params.DeviceNames[0]) + " -bundle " + BundleIdentifier);
+			
+			LogInformation("Checking if bundle {0} is installed", BundleIdentifier);
+			string Output = CommandUtils.RunAndLog(DeployServer, "ListApplications" + (String.IsNullOrEmpty(Params.DeviceNames[0]) ? "" : " -device " + Params.DeviceNames[0]), GetRunAndLogOnlyName(CmdEnv, DeployServer, null), 0, null, ERunOptions.SpewIsVerbose);
+			bool bBundleIsInstalled = Output.Contains(string.Format("CFBundleIdentifier -> {0}{1}", BundleIdentifier, Environment.NewLine)); 
 
-			string[] ManifestFiles = Directory.GetFiles(CombinePaths(Params.BaseStageDirectory, PlatformName), "*_Manifest_UFS*.txt");
-			UFSManifests.AddRange(ManifestFiles);
+			if (bBundleIsInstalled)
+			{
+				LogInformation("Bundle {0} found, retrieving deployed manifests...", BundleIdentifier);
+
+				RunAndLog(CmdEnv, DeployServer, "Backup -file \"" + CombinePaths(Params.BaseStageDirectory, PlatformName, SC.GetUFSDeployedManifestFileName(null)) + "\" -file \"" + CombinePaths(Params.BaseStageDirectory, PlatformName, SC.GetNonUFSDeployedManifestFileName(null)) + "\"" + (String.IsNullOrEmpty(Params.DeviceNames[0]) ? "" : " -device " + Params.DeviceNames[0]) + " -bundle " + BundleIdentifier);
 
-			ManifestFiles = Directory.GetFiles(CombinePaths(Params.BaseStageDirectory, PlatformName), "*_Manifest_NonUFS*.txt");
-			NonUFSManifests.AddRange(ManifestFiles);
+				string[] ManifestFiles = Directory.GetFiles(CombinePaths(Params.BaseStageDirectory, PlatformName), "*_Manifest_UFS*.txt");
+				UFSManifests.AddRange(ManifestFiles);
+
+				ManifestFiles = Directory.GetFiles(CombinePaths(Params.BaseStageDirectory, PlatformName), "*_Manifest_NonUFS*.txt");
+				NonUFSManifests.AddRange(ManifestFiles);
+			}
+			else
+			{
+				LogInformation("Bundle {0} not found, skipping retrieving deployed manifests", BundleIdentifier);
+			}
 		}
 		catch (System.Exception)
 		{
diff --git a/Engine/Source/Programs/IOS/DeploymentServer/Program.cs b/Engine/Source/Programs/IOS/DeploymentServer/Program.cs
--- a/Engine/Source/Programs/IOS/DeploymentServer/Program.cs
+++ b/Engine/Source/Programs/IOS/DeploymentServer/Program.cs
@@ -331,6 +331,12 @@ namespace DeploymentServer
 								Writer.Flush();
 								break;
 
+							case "listapplications":
+								Console.SetOut(Writer);
+								DeploymentProxy.Deployer.ListApplications();
+								Writer.Flush();
+								break;
+
 							case "command":
 								if (Device.Length < 5)
 								{
@@ -1091,6 +1097,7 @@ namespace DeploymentServer
 				Console.WriteLine("\t install");
 				Console.WriteLine("\t enumerate");
 				Console.WriteLine("\t listdevices");
+				Console.WriteLine("\t listapplications");
 				Console.WriteLine("\t listentodevice");
 				Console.WriteLine("\t command");
 				Console.WriteLine("\t forward");
