Fix issue with install batch file detection for Android

#ue4
#android

diff --git a/Engine/Source/Programs/AutomationTool/Android/AndroidPlatform.Automation.cs b/Engine/Source/Programs/AutomationTool/Android/AndroidPlatform.Automation.cs
--- a/Engine/Source/Programs/AutomationTool/Android/AndroidPlatform.Automation.cs
+++ b/Engine/Source/Programs/AutomationTool/Android/AndroidPlatform.Automation.cs
@@ -295,6 +295,20 @@ public class AndroidPlatform : Platform
 		return (Ini.GetBool("/Script/AndroidRuntimeSettings.AndroidRuntimeSettings", "bSaveSymbols", out bSave) && bSave);
 	}
 
+	private bool GetEnableBundle(DeploymentContext SC)
+	{
+		ConfigHierarchy Ini = ConfigCache.ReadHierarchy(ConfigHierarchyType.Engine, DirectoryReference.FromFile(SC.RawProjectPath), SC.StageTargetPlatform.PlatformType);
+		bool bEnableBundle = false;
+		return (Ini.GetBool("/Script/AndroidRuntimeSettings.AndroidRuntimeSettings", "bEnableBundle", out bEnableBundle) && bEnableBundle);
+	}
+
+	private bool GetEnableUniversalAPK(DeploymentContext SC)
+	{
+		ConfigHierarchy Ini = ConfigCache.ReadHierarchy(ConfigHierarchyType.Engine, DirectoryReference.FromFile(SC.RawProjectPath), SC.StageTargetPlatform.PlatformType);
+		bool bEnableUniversalAPK = false;
+		return (Ini.GetBool("/Script/AndroidRuntimeSettings.AndroidRuntimeSettings", "bEnableUniversalAPK", out bEnableUniversalAPK) && bEnableUniversalAPK);
+	}
+
 	private Int64 GetMaxOBBSizeAllowed(DeploymentContext SC)
 	{
 		ConfigHierarchy Ini = ConfigCache.ReadHierarchy(ConfigHierarchyType.Engine, DirectoryReference.FromFile(SC.RawProjectPath), SC.StageTargetPlatform.PlatformType);
@@ -388,6 +402,8 @@ public class AndroidPlatform : Platform
 		bool bMakeSeparateApks = UnrealBuildTool.AndroidExports.ShouldMakeSeparateApks();
 		bool bBuildWithHiddenSymbolVisibility = BuildWithHiddenSymbolVisibility(SC);
 		bool bSaveSymbols = GetSaveSymbols(SC);
+		bool bEnableBundle = GetEnableBundle(SC);
+		bool bEnableUniversalAPK = GetEnableUniversalAPK(SC);
 
 		var Deploy = AndroidExports.CreateDeploymentHandler(Params.RawProjectPath, Params.ForcePackageData);
 		bool bPackageDataInsideApk = Deploy.GetPackageDataInsideApk();
@@ -535,7 +551,7 @@ public class AndroidPlatform : Platform
 			{
 				string ApkName = GetFinalApkName(Params, SC.StageExecutables[0], true, bMakeSeparateApks ? Architecture : "", bMakeSeparateApks ? GPUArchitecture : "");
 				string ApkBareName = GetFinalApkName(Params, SC.StageExecutables[0], true, "", "");
-				bool bHaveAPK = FileExists(ApkName);
+				bool bHaveAPK = !bEnableBundle;     // do not have a standard APK if bundle enabled
 				if (!SC.IsCodeBasedProject)
 				{
 					string UE4SOName = GetFinalApkName(Params, SC.StageExecutables[0], false, bMakeSeparateApks ? Architecture : "", bMakeSeparateApks ? GPUArchitecture : "");
@@ -608,17 +624,20 @@ public class AndroidPlatform : Platform
 				string APKBareNameWithoutExtension = Path.GetFileNameWithoutExtension(ApkBareName);
 				string UniversalApkName = Path.Combine(APKDirectory, APKNameWithoutExtension + "_universal.apk");
 				bool bHaveUniversal = false;
-				if (FileExists(UniversalApkName))
+				if (bEnableBundle && bEnableUniversalAPK)
 				{
-					bHaveUniversal = true;
-				}
-				else
-				{
-					UniversalApkName = Path.Combine(APKDirectory, APKBareNameWithoutExtension + "_universal.apk");
 					if (FileExists(UniversalApkName))
 					{
 						bHaveUniversal = true;
 					}
+					else
+					{
+						UniversalApkName = Path.Combine(APKDirectory, APKBareNameWithoutExtension + "_universal.apk");
+						if (FileExists(UniversalApkName))
+						{
+							bHaveUniversal = true;
+						}
+					}
 				}
 
 				//figure out which platforms we need to create install files for
