Explicitly delete temporary files instead of wiping the temp directory.

diff --git a/Engine/Source/Programs/AutomationTool/Win/WinPlatform.Automation.cs b/Engine/Source/Programs/AutomationTool/Win/WinPlatform.Automation.cs
--- a/Engine/Source/Programs/AutomationTool/Win/WinPlatform.Automation.cs
+++ b/Engine/Source/Programs/AutomationTool/Win/WinPlatform.Automation.cs
@@ -407,14 +407,18 @@ public abstract class BaseWinPlatform : Platform
 
 			List<FileReference> TemporaryFiles = new List<FileReference>();
 
-			foreach( FileReference File in FilesToAdd)
+			foreach(FileReference File in FilesToAdd)
 			{
-				FileReference TempFile = FileReference.Combine(TempDir, File.GetFileName());
+				// Don't add files over 2GB in size as cab container doesn't support that
+				if (File.ToFileInfo().Length < (2L * 1024 * 1024 * 1024))
+				{
+					FileReference TempFile = FileReference.Combine(TempDir, File.GetFileName());
 
-				CommandUtils.CopyFile(File.FullName, TempFile.FullName);
-				CommandUtils.SetFileAttributes(TempFile.FullName, ReadOnly: false);
+					CommandUtils.CopyFile(File.FullName, TempFile.FullName);
+					CommandUtils.SetFileAttributes(TempFile.FullName, ReadOnly: false);
 
-				TemporaryFiles.Add(TempFile);
+					TemporaryFiles.Add(TempFile);
+				}
 			}
 
 			// Process the temporary files
@@ -444,9 +448,13 @@ public abstract class BaseWinPlatform : Platform
 			}
 			finally
 			{
+				// clean up after ourselves
 				File.Delete(FileListFilename);
 				File.Delete(IndexListFilename);
-				Directory.Delete(TempDir.FullName);
+				foreach( FileReference TemporaryFile in TemporaryFiles)
+				{
+					File.Delete(TemporaryFile.FullName);
+				}
 			}
 
 			DateTime CompressDone = DateTime.Now;
