Ensure we clean-up the gather process after it's completed

#rnx

diff --git a/Engine/Source/Programs/AutomationTool/Scripts/Localisation.Automation.cs b/Engine/Source/Programs/AutomationTool/Scripts/Localisation.Automation.cs
--- a/Engine/Source/Programs/AutomationTool/Scripts/Localisation.Automation.cs
+++ b/Engine/Source/Programs/AutomationTool/Scripts/Localisation.Automation.cs
@@ -356,6 +356,9 @@ class Localize : BuildCommand
 				if (RunResult != null)
 				{
 					RunResult.WaitForExit();
+					RunResult.OnProcessExited();
+					RunResult.DisposeProcess();
+					
 					if (RunResult.ExitCode == 0)
 					{
 						LogInformation("The localization commandlet for '{0}' exited with code 0.", ProjectInfo.ProjectName);
