Only use the new symbol store process for XboxOne. Leave Windows as-is. Fixes errors caused by files being duplicated in symbol source list.

diff --git a/Engine/Source/Programs/AutomationTool/Win/WinPlatform.Automation.cs b/Engine/Source/Programs/AutomationTool/Win/WinPlatform.Automation.cs
--- a/Engine/Source/Programs/AutomationTool/Win/WinPlatform.Automation.cs
+++ b/Engine/Source/Programs/AutomationTool/Win/WinPlatform.Automation.cs
@@ -398,49 +398,19 @@ public abstract class BaseWinPlatform : Platform
 		if(FilesToAdd.Count > 0)
 		{
 			DateTime Start = DateTime.Now;
-
-			// A bug in Windows SDKs since 10.0.15063.0 causes symbols compressed with symstore's default (LZ) compression fail to load in VS
-			// Workaround is to compress them manually with makecab which uses MS-ZIP compression
-			// Source files need to be deleted after compression or symstore will use the uncompressed version
-			// Copy input files to a temporary directory so the originals don't get deleted by ProcessSymbols
-			DirectoryReference TempDir = new DirectoryReference(Path.GetTempPath());
-
-			List<FileReference> TemporaryFiles = new List<FileReference>();
-
-			foreach(FileReference File in FilesToAdd)
-			{
-				// Don't add files over 2GB in size as cab container doesn't support that
-				if (File.ToFileInfo().Length < (2L * 1024 * 1024 * 1024))
-				{
-					FileReference TempFile = FileReference.Combine(TempDir, File.GetFileName());
-
-					CommandUtils.CopyFile(File.FullName, TempFile.FullName);
-					CommandUtils.SetFileAttributes(TempFile.FullName, ReadOnly: false);
-
-					TemporaryFiles.Add(TempFile);
-				}
-			}
-
-			// Process the temporary files
 			DirectoryReference TempSymStoreDir = DirectoryReference.Combine(RootDirectory, "Saved", "SymStore");
-			DirectoryReference.CreateDirectory(TempSymStoreDir);
-			DeleteDirectoryContents(TempSymStoreDir);
 
-			string FileListFilename = Path.GetTempFileName();
-			string IndexListFilename = Path.GetTempFileName();
+			string TempFileName = Path.GetTempFileName();
 			try
 			{
-				File.WriteAllLines(FileListFilename, TemporaryFiles.Select(x => x.FullName), Encoding.ASCII);
+				File.WriteAllLines(TempFileName, FilesToAdd.Select(x => x.FullName), Encoding.ASCII);
 
-				// Process temporary files to temporary symbol store
+				// Copy everything to the temp symstore
 				ProcessStartInfo StartInfo = new ProcessStartInfo();
-
-				StartInfo.FileName = Path.Combine(CommandUtils.EngineDirectory.FullName, "Binaries", "Win64", "ProcessSymbols.bat");
-
-				StartInfo.Arguments = string.Format("\"{0}\" \"{1}\" \"{2}\" \"{3}\" \"{4}\" \"{5}\"", SymStoreExe.FullName, FileListFilename, IndexListFilename, TempDir.FullName, Product, TempSymStoreDir);
+				StartInfo.FileName = SymStoreExe.FullName;
+				StartInfo.Arguments = string.Format("add /f \"@{0}\" /s \"{1}\" /t \"{2}\" /compress", TempFileName, TempSymStoreDir, Product);
 				StartInfo.UseShellExecute = false;
 				StartInfo.CreateNoWindow = true;
-				StartInfo.WorkingDirectory = TempDir.FullName;
 				if (Utils.RunLocalProcessAndLogOutput(StartInfo) != 0)
 				{
 					return false;
@@ -448,15 +418,8 @@ public abstract class BaseWinPlatform : Platform
 			}
 			finally
 			{
-				// clean up after ourselves
-				File.Delete(FileListFilename);
-				File.Delete(IndexListFilename);
-				foreach( FileReference TemporaryFile in TemporaryFiles)
-				{
-					File.Delete(TemporaryFile.FullName);
-				}
+				File.Delete(TempFileName);
 			}
-
 			DateTime CompressDone = DateTime.Now;
 			LogInformation("Took {0}s to compress the symbol files", (CompressDone - Start).TotalSeconds);
 
