UAT

* Add a callback to the UAT project params so a script that uses buildcookrun can choose to modify the deployment context before staging
* Expose a function from the staging code which lets us remap staged file paths based on the project remapping configuration.
* Extend cooked target staging so that any programs or editor server targets are forced to Development configs, rather than what the master client/server targets request

diff --git a/Engine/Source/Programs/AutomationTool/AutomationUtils/ProjectParams.cs b/Engine/Source/Programs/AutomationTool/AutomationUtils/ProjectParams.cs
--- a/Engine/Source/Programs/AutomationTool/AutomationUtils/ProjectParams.cs
+++ b/Engine/Source/Programs/AutomationTool/AutomationUtils/ProjectParams.cs
@@ -1713,6 +1713,11 @@ namespace AutomationTool
 		/// </summary>
 		public ParamList<string> ExtraTargetsToStageWithClient = new ParamList<string>();
 
+        /// <summary>
+        /// Stage: Optional callback that a build script can use to modify a deployment context before it is applied
+        /// </summary>
+        public Action<ProjectParams, DeploymentContext> ModifyDeploymentContextCallback = null;
+
         /// <summary>
         /// On Windows, adds an executable to the root of the staging directory which checks for prerequisites being 
 		/// installed and launches the game with a path to the .uproject file.
diff --git a/Engine/Source/Programs/AutomationTool/Scripts/CopyBuildToStagingDirectory.Automation.cs b/Engine/Source/Programs/AutomationTool/Scripts/CopyBuildToStagingDirectory.Automation.cs
--- a/Engine/Source/Programs/AutomationTool/Scripts/CopyBuildToStagingDirectory.Automation.cs
+++ b/Engine/Source/Programs/AutomationTool/Scripts/CopyBuildToStagingDirectory.Automation.cs
@@ -938,6 +938,12 @@ public partial class Project : CommandUtils
 			}
 		}
 
+		// Allow the calling scripts to make modifications to the deployment context before we finalize it
+		if (Params.ModifyDeploymentContextCallback != null)
+		{
+			Params.ModifyDeploymentContextCallback(Params, SC);
+		}
+
 		// Apply all the directory mappings
 		SC.FilesToStage.UFSFiles = SC.FilesToStage.UFSFiles.ToDictionary(x => ApplyDirectoryRemap(SC, x.Key), x => x.Value);
 		SC.FilesToStage.NonUFSFiles = SC.FilesToStage.NonUFSFiles.ToDictionary(x => ApplyDirectoryRemap(SC, x.Key), x => x.Value);
@@ -1259,7 +1265,7 @@ public partial class Project : CommandUtils
 		}
 	}
 
-	static StagedFileReference ApplyDirectoryRemap(DeploymentContext SC, StagedFileReference InputFile)
+	public static StagedFileReference ApplyDirectoryRemap(DeploymentContext SC, StagedFileReference InputFile)
 	{
 		StagedFileReference CurrentFile = InputFile;
 		foreach (Tuple<StagedDirectoryReference, StagedDirectoryReference> RemapDirectory in SC.RemapDirectories)
@@ -3488,7 +3494,7 @@ public partial class Project : CommandUtils
 			// If we are staging a client and have been asked to include editor targets, we currently only want to
 			// include a single Development editor target. Ideally, we should have shipping editor configs and then
 			// just include the requested configs for all targets
-			if (!InDedicatedServer && Params.EditorTargets.Contains(Target))
+			if (Params.EditorTargets.Contains(Target) || Params.ProgramTargets.Contains(Target))
 			{
 				TargetAndConfigPairs.Add(new Tuple<string, UnrealTargetConfiguration>(Target, UnrealTargetConfiguration.Development));
 			}
