Fix issue wih blueprint projects not having a binaries folder, causes elemental demo to fail in Gauntlet documentation example steps

diff --git a/Engine/Source/Programs/AutomationTool/Gauntlet/Unreal/BuildSource/Gauntlet.UnrealBuildSource.cs b/Engine/Source/Programs/AutomationTool/Gauntlet/Unreal/BuildSource/Gauntlet.UnrealBuildSource.cs
--- a/Engine/Source/Programs/AutomationTool/Gauntlet/Unreal/BuildSource/Gauntlet.UnrealBuildSource.cs
+++ b/Engine/Source/Programs/AutomationTool/Gauntlet/Unreal/BuildSource/Gauntlet.UnrealBuildSource.cs
@@ -198,10 +198,9 @@ namespace Gauntlet
 					return false;
 				}
 				
-				// include binaries path for packaged builds
+				// include binaries path for packaged builds if it exists
 				string BinariesPath = Path.Combine(ProjectPath.Directory.FullName, "Binaries");
-
-				OutBuildPaths = new string[] { StagedPath, BinariesPath };
+				OutBuildPaths = Directory.Exists(BinariesPath) ? new string[] { StagedPath, BinariesPath } : new string[] { StagedPath };
 			}
 			else if (BuildDir.Name.Equals("editor", StringComparison.OrdinalIgnoreCase))
 			{
