UAT: Add a staging option to disable the use of existing pak files from the staging folder of other build targets, when we've force reuse of another target's cooked data.

diff --git a/Engine/Source/Programs/AutomationTool/AutomationUtils/ProjectParams.cs b/Engine/Source/Programs/AutomationTool/AutomationUtils/ProjectParams.cs
--- a/Engine/Source/Programs/AutomationTool/AutomationUtils/ProjectParams.cs
+++ b/Engine/Source/Programs/AutomationTool/AutomationUtils/ProjectParams.cs
@@ -280,6 +280,7 @@ namespace AutomationTool
             this.RunTimeoutSeconds = InParams.RunTimeoutSeconds;
 			this.Clean = InParams.Clean;
 			this.Pak = InParams.Pak;
+			this.TryUsePaksFromDifferentCookSource = InParams.TryUsePaksFromDifferentCookSource;
 			this.IoStore = InParams.IoStore;
 			this.SignPak = InParams.SignPak;
 			this.SignedPak = InParams.SignedPak;
@@ -455,6 +456,7 @@ namespace AutomationTool
 			bool? NoXGE = null,
 			bool? Package = null,
 			bool? Pak = null,
+			bool? TryUsePaksFromDifferentCookSource = null,
 			bool? IoStore = null,
 			bool? SkipIoStore = null,
 			bool? Prereqs = null,
@@ -629,6 +631,7 @@ namespace AutomationTool
 			}
 			this.PakAlignForMemoryMapping = GetParamValueIfNotSpecified(Command, PakAlignForMemoryMapping, this.PakAlignForMemoryMapping, "PakAlignForMemoryMapping");
 			this.Pak = GetParamValueIfNotSpecified(Command, Pak, this.Pak, "pak");
+			this.TryUsePaksFromDifferentCookSource = GetParamValueIfNotSpecified(Command, TryUsePaksFromDifferentCookSource, this.TryUsePaksFromDifferentCookSource, "TryUsePaksFromDifferentCookSource");
 			this.IoStore = GetParamValueIfNotSpecified(Command, IoStore, this.IoStore, "iostore");
 			this.SkipIoStore = GetParamValueIfNotSpecified(Command, IoStore, this.SkipIoStore, "skipiostore");
 			this.SkipPak = GetParamValueIfNotSpecified(Command, SkipPak, this.SkipPak, "skippak");
@@ -1190,6 +1193,13 @@ namespace AutomationTool
 		[Help("pak", "generate a pak file")]
 		public bool Pak { private set; get; }
 
+		/// <summary>
+		/// Stage: True if we should try and re-use pak files from another staged build when we've specified a different cook source platform
+		/// </summary>
+		[Help("pak", "reuse pak files from the alternate cook source folder, if specified")]
+		public bool TryUsePaksFromDifferentCookSource { get { return TryUsePaksFromDifferentCookSourceInternal; } private set { TryUsePaksFromDifferentCookSourceInternal = value; } }
+		private bool TryUsePaksFromDifferentCookSourceInternal = true;
+
 		/// <summary>
 		/// Shared: True if container file(s) should be generated with ZenPak.
 		/// </summary>
@@ -2818,6 +2828,7 @@ namespace AutomationTool
 				CommandUtils.LogLog("MapsToCook={0}", MapsToCook.ToString());
 				CommandUtils.LogLog("MapIniSectionsToCook={0}", MapIniSectionsToCook.ToString());
 				CommandUtils.LogLog("Pak={0}", Pak);
+				CommandUtils.LogLog("TryUsePaksFromDifferentCookSource={0}", TryUsePaksFromDifferentCookSource);
 				CommandUtils.LogLog("IoStore={0}", IoStore);
 				CommandUtils.LogLog("SkipIoStore={0}", SkipIoStore);
 				CommandUtils.LogLog("Package={0}", Package);
diff --git a/Engine/Source/Programs/AutomationTool/Scripts/CopyBuildToStagingDirectory.Automation.cs b/Engine/Source/Programs/AutomationTool/Scripts/CopyBuildToStagingDirectory.Automation.cs
--- a/Engine/Source/Programs/AutomationTool/Scripts/CopyBuildToStagingDirectory.Automation.cs
+++ b/Engine/Source/Programs/AutomationTool/Scripts/CopyBuildToStagingDirectory.Automation.cs
@@ -2088,7 +2088,7 @@ public partial class Project : CommandUtils
 
 			bool bCopiedExistingPak = false;
 
-			if (SC.StageTargetPlatform != SC.CookSourcePlatform)
+			if (SC.StageTargetPlatform != SC.CookSourcePlatform && Params.TryUsePaksFromDifferentCookSource)
 			{
 				// Check to see if we have an existing pak file we can use
 
