Fix typo in default project configuration.

#rnx

diff --git a/Engine/Source/Programs/AutomationTool/Linux/Linux.Automation.csproj b/Engine/Source/Programs/AutomationTool/Linux/Linux.Automation.csproj
--- a/Engine/Source/Programs/AutomationTool/Linux/Linux.Automation.csproj
+++ b/Engine/Source/Programs/AutomationTool/Linux/Linux.Automation.csproj
@@ -2,7 +2,7 @@
 <Project ToolsVersion="12.0" DefaultTargets="Build" xmlns="http://schemas.microsoft.com/developer/msbuild/2003">
   <Import Project="$(MSBuildExtensionsPath)\$(MSBuildToolsVersion)\Microsoft.Common.props" Condition="Exists('$(MSBuildExtensionsPath)\$(MSBuildToolsVersion)\Microsoft.Common.props')" />
   <PropertyGroup>
-    <Configuration Condition=" '$(Configuration)' == '' ">Developement</Configuration>
+    <Configuration Condition=" '$(Configuration)' == '' ">Development</Configuration>
     <Platform Condition=" '$(Platform)' == '' ">AnyCPU</Platform>
     <ProjectGuid>{EC395254-405B-4030-92CE-E346CC50848D}</ProjectGuid>
     <OutputType>Library</OutputType>
