Remove usages of "var"

#rnx

diff --git a/Engine/Source/Programs/AutomationTool/AutomationUtils/ScriptCompiler.cs b/Engine/Source/Programs/AutomationTool/AutomationUtils/ScriptCompiler.cs
--- a/Engine/Source/Programs/AutomationTool/AutomationUtils/ScriptCompiler.cs
+++ b/Engine/Source/Programs/AutomationTool/AutomationUtils/ScriptCompiler.cs
@@ -70,7 +70,7 @@ namespace AutomationTool
 				FindAndCompileScriptModules(ScriptsForProjectFileName, AdditionalScriptsFolders);
 			}
 
-			var ScriptAssemblies = new List<Assembly>();
+			List<Assembly> ScriptAssemblies = new List<Assembly>();
 			LoadPreCompiledScriptAssemblies(ScriptAssemblies);
 
 			// Setup platforms
@@ -79,11 +79,11 @@ namespace AutomationTool
 			// Instantiate all the automation classes for interrogation
 			Log.TraceVerbose("Creating commands.");
 			ScriptCommands = new Dictionary<string, Type>(StringComparer.InvariantCultureIgnoreCase);
-			foreach (var CompiledScripts in ScriptAssemblies)
+			foreach (Assembly CompiledScripts in ScriptAssemblies)
 			{
 				try
 				{
-					foreach (var ClassType in CompiledScripts.GetTypes())
+					foreach (Type ClassType in CompiledScripts.GetTypes())
 					{
 						if (ClassType.IsSubclassOf(typeof(BuildCommand)) && ClassType.IsAbstract == false)
 						{
@@ -106,7 +106,7 @@ namespace AutomationTool
 				}
 				catch (ReflectionTypeLoadException LoadEx)
 				{
-					foreach (var SubEx in LoadEx.LoaderExceptions)
+					foreach (Exception SubEx in LoadEx.LoaderExceptions)
 					{
 						Log.TraceWarning("Got type loader exception: {0}", SubEx.ToString());
 					}
@@ -121,8 +121,8 @@ namespace AutomationTool
 
 		private static void FindAndCompileScriptModules(string ScriptsForProjectFileName, List<string> AdditionalScriptsFolders)
 		{
-			var OldCWD = Environment.CurrentDirectory;
-			var UnrealBuildToolCWD = CommandUtils.CombinePaths(CommandUtils.CmdEnv.LocalRoot, "Engine", "Source");
+			string OldCWD = Environment.CurrentDirectory;
+			string UnrealBuildToolCWD = CommandUtils.CombinePaths(CommandUtils.CmdEnv.LocalRoot, "Engine", "Source");
 
 			Environment.CurrentDirectory = UnrealBuildToolCWD;
 
@@ -138,10 +138,10 @@ namespace AutomationTool
 				AllGameFolders = new List<DirectoryReference> { new DirectoryReference(Path.GetDirectoryName(ScriptsForProjectFileName)) };
 			}
 
-			var AllAdditionalScriptFolders = new List<DirectoryReference>(AdditionalScriptsFolders.Select(x => new DirectoryReference(x)));
-			foreach (var Folder in AllGameFolders)
+			List<DirectoryReference> AllAdditionalScriptFolders = AdditionalScriptsFolders.Select(x => new DirectoryReference(x)).ToList();
+			foreach (DirectoryReference Folder in AllGameFolders)
 			{
-				var GameBuildFolder = DirectoryReference.Combine(Folder, "Build");
+				DirectoryReference GameBuildFolder = DirectoryReference.Combine(Folder, "Build");
 				if (DirectoryReference.Exists(GameBuildFolder))
 				{
 					AllAdditionalScriptFolders.Add(GameBuildFolder);
@@ -150,9 +150,9 @@ namespace AutomationTool
 
 			Log.TraceVerbose("Discovering game folders.");
 
-			var DiscoveredModules = UnrealBuildTool.RulesCompiler.FindAllRulesSourceFiles(UnrealBuildTool.RulesCompiler.RulesFileType.AutomationModule, GameFolders: AllGameFolders, ForeignPlugins: null, AdditionalSearchPaths: AllAdditionalScriptFolders);
-			var ModulesToCompile = new List<string>(DiscoveredModules.Count);
-			foreach (var ModuleFilename in DiscoveredModules)
+			List<FileReference> DiscoveredModules = UnrealBuildTool.RulesCompiler.FindAllRulesSourceFiles(UnrealBuildTool.RulesCompiler.RulesFileType.AutomationModule, GameFolders: AllGameFolders, ForeignPlugins: null, AdditionalSearchPaths: AllAdditionalScriptFolders);
+			List<string> ModulesToCompile = new List<string>(DiscoveredModules.Count);
+			foreach (FileReference ModuleFilename in DiscoveredModules)
 			{
 				if (HostPlatform.Current.IsScriptModuleSupported(ModuleFilename.GetFileNameWithoutAnyExtensions()))
 				{
@@ -297,8 +297,8 @@ namespace AutomationTool
 				// Mono has an issue where arugments with semicolons or commas can't be passed through to
 				// as arguments so we need to manually construct a temp file with the list of modules
 				// see ([REDACTED_PROFILE_URL]/msbuild/issues/471)
-				var UATProjTemplate = CommandUtils.CombinePaths(CommandUtils.CmdEnv.LocalRoot, @"Engine\Source\Programs\AutomationTool\Scripts\UAT.proj");
-				var UATProjFile = Path.Combine(CommandUtils.CmdEnv.EngineSavedFolder, "UATTempProj.proj");
+				string UATProjTemplate = CommandUtils.CombinePaths(CommandUtils.CmdEnv.LocalRoot, @"Engine\Source\Programs\AutomationTool\Scripts\UAT.proj");
+				string UATProjFile = Path.Combine(CommandUtils.CmdEnv.EngineSavedFolder, "UATTempProj.proj");
 
 				string ProjContents = File.ReadAllText(UATProjTemplate);
 				ProjContents = ProjContents.Replace("$(Modules)", ModulesList);
@@ -308,9 +308,9 @@ namespace AutomationTool
 				
 				string MsBuildVerbosity = Log.OutputLevel >= LogEventType.Verbose ? "minimal" : "quiet";
 
-				var CmdLine = String.Format("\"{0}\" /p:Configuration={1} /verbosity:{2} /nologo", UATProjFile, BuildConfig, MsBuildVerbosity);
+				string CmdLine = String.Format("\"{0}\" /p:Configuration={1} /verbosity:{2} /nologo", UATProjFile, BuildConfig, MsBuildVerbosity);
 				// suppress the run command because it can be long and intimidating, making the logs around this code harder to read.
-				var Result = CommandUtils.Run(BuildTool, CmdLine, Options: CommandUtils.ERunOptions.Default | CommandUtils.ERunOptions.NoLoggingOfRunCommand | CommandUtils.ERunOptions.LoggingOfRunDuration);
+				IProcessResult Result = CommandUtils.Run(BuildTool, CmdLine, Options: CommandUtils.ERunOptions.Default | CommandUtils.ERunOptions.NoLoggingOfRunCommand | CommandUtils.ERunOptions.LoggingOfRunDuration);
 				if (Result.ExitCode != 0)
 				{
 					throw new AutomationException(String.Format("Failed to build \"{0}\":{1}{2}", UATProjFile, Environment.NewLine, Result.Output));
@@ -319,12 +319,12 @@ namespace AutomationTool
 			else
 			{
 				// Make sure DefaultScriptsDLLName is compiled first
-				var DefaultScriptsProjName = Path.ChangeExtension(DefaultScriptsDLLName, "csproj");
+				string DefaultScriptsProjName = Path.ChangeExtension(DefaultScriptsDLLName, "csproj");
 
 				// Primary modules must be built first
 				List<string> PrimaryModules = Modules.Where(M => M.IndexOf(DefaultScriptsProjName, StringComparison.InvariantCultureIgnoreCase) >= 0).ToList();
 
-				foreach (var ModuleName in PrimaryModules)
+				foreach (string ModuleName in PrimaryModules)
 				{
 					Log.TraceInformation("Building script module: {0}", ModuleName);
 					try
@@ -344,7 +344,7 @@ namespace AutomationTool
 				List<string> SecondaryModules = Modules.Where(M => !PrimaryModules.Contains(M)).ToList();
 
 				// Non-parallel method
-				foreach (var ModuleName in SecondaryModules)
+				foreach (string ModuleName in SecondaryModules)
 				{
 					Log.TraceInformation("Building script module: {0}", ModuleName);
 					try
@@ -396,7 +396,7 @@ namespace AutomationTool
 				ProjectFile, BuildConfig);
 
 			// Compile the project
-			var Result = CommandUtils.Run(CommandUtils.CmdEnv.MsBuildExe, CmdLine);
+			IProcessResult Result = CommandUtils.Run(CommandUtils.CmdEnv.MsBuildExe, CmdLine);
 
 			if (Result.ExitCode != 0)
 			{
@@ -419,13 +419,13 @@ namespace AutomationTool
 			CommandUtils.LogVerbose("Loading precompiled script DLLs");
 
 			bool DefaultScriptsDLLFound = false;
-			var ScriptsLocation = GetScriptAssemblyFolder();
+			string ScriptsLocation = GetScriptAssemblyFolder();
 			if (CommandUtils.DirectoryExists(ScriptsLocation))
 			{
-				var ScriptDLLFiles = Directory.GetFiles(ScriptsLocation, "*.Automation.dll", SearchOption.AllDirectories);
+				string[] ScriptDLLFiles = Directory.GetFiles(ScriptsLocation, "*.Automation.dll", SearchOption.AllDirectories);
 
 				CommandUtils.LogVerbose("Found {0} script DLL(s).", ScriptDLLFiles.Length);
-				foreach (var ScriptsDLLFilename in ScriptDLLFiles)
+				foreach (string ScriptsDLLFilename in ScriptDLLFiles)
 				{
 
 					if (!HostPlatform.Current.IsScriptModuleSupported(CommandUtils.GetFilenameWithoutAnyExtensions(ScriptsDLLFilename)))
@@ -437,7 +437,7 @@ namespace AutomationTool
 					CommandUtils.LogVerbose("Loading script DLL: {0}", ScriptsDLLFilename);
 					try
 					{
-						var Dll = AppDomain.CurrentDomain.Load(AssemblyName.GetAssemblyName(ScriptsDLLFilename));
+						Assembly Dll = AppDomain.CurrentDomain.Load(AssemblyName.GetAssemblyName(ScriptsDLLFilename));
 						OutScriptAssemblies.Add(Dll);
 						// Check if this is the default scripts DLL.
 						if (!DefaultScriptsDLLFound && String.Compare(Path.GetFileName(ScriptsDLLFilename), DefaultScriptsDLLName, true) == 0)
