Add support for UAT script projects using foreign engine builds, and for game projects to output binaries under their own project directory.

* UAT will now read each *.Automation.csproj file to determine the output path, and load from there rather than just loading anything under Engine\Binaries\AutomationScripts.
* When compiling *.Automation.csproj files, UAT will set the $(EngineDir) directory to the current engine directory, allowing the project to resolve assembly references to the correct location.
* When generating project files, UBT will create an *.Automation.csproj.props file next to the project containing the path to the engine directory, allowing the assembly to be compiled correctly from Visual Studio.

diff --git a/Engine/Source/Programs/AutomationTool/AutomationUtils/ScriptCompiler.cs b/Engine/Source/Programs/AutomationTool/AutomationUtils/ScriptCompiler.cs
--- a/Engine/Source/Programs/AutomationTool/AutomationUtils/ScriptCompiler.cs
+++ b/Engine/Source/Programs/AutomationTool/AutomationUtils/ScriptCompiler.cs
@@ -11,6 +11,7 @@ using System.Reflection;
 using System.Diagnostics;
 using UnrealBuildTool;
 using Tools.DotNETCommon;
+using System.Threading.Tasks;
 
 namespace AutomationTool
 {
@@ -43,7 +44,6 @@ namespace AutomationTool
 		{
 		}
 
-
 		/// <summary>
 		/// Finds and/or compiles all script files and assemblies.
 		/// </summary>
@@ -51,27 +51,32 @@ namespace AutomationTool
 		/// <param name="AdditionalScriptsFolders">Additional script fodlers to look for source files in.</param>
 		public void FindAndCompileAllScripts(string ScriptsForProjectFileName, List<string> AdditionalScriptsFolders)
 		{
-			bool DoCompile = false;
-			if (GlobalCommandLine.Compile)
-			{
-				if(CommandUtils.IsEngineInstalled())
-				{
-					CommandUtils.LogWarning("Ignoring -Compile argument because engine is installed.");
-				}
-				else
-				{
-					DoCompile = true;
-				}
-			}
+			// Find all the project files
+			List<FileReference> ProjectFiles = FindAutomationProjects(ScriptsForProjectFileName, AdditionalScriptsFolders);
+
+			// Get the default properties for compiling the projects
+			Dictionary<string, string> MsBuildProperties = new Dictionary<string, string>(StringComparer.OrdinalIgnoreCase);
+			MsBuildProperties.Add("Platform", "AnyCPU");
+			MsBuildProperties.Add("Configuration", BuildConfig);
+			MsBuildProperties.Add("EngineDir", CommandUtils.EngineDirectory.FullName);
+
+			// Read all the projects
+			CsProjectInfo[] Projects = new CsProjectInfo[ProjectFiles.Count];
+			Parallel.For(0, ProjectFiles.Count, Idx => Projects[Idx] = CsProjectInfo.Read(ProjectFiles[Idx], MsBuildProperties));
 
 			// Compile only if not disallowed.
-			if (DoCompile && !String.IsNullOrEmpty(CommandUtils.CmdEnv.MsBuildExe))
+			if (GlobalCommandLine.Compile && !String.IsNullOrEmpty(CommandUtils.CmdEnv.MsBuildExe))
 			{
-				FindAndCompileScriptModules(ScriptsForProjectFileName, AdditionalScriptsFolders);
+				List<CsProjectInfo> CompileProjects = new List<CsProjectInfo>(Projects);
+				if (CommandUtils.IsEngineInstalled())
+				{
+					CompileProjects.RemoveAll(x => x.ProjectPath.IsUnderDirectory(CommandUtils.EngineDirectory));
+				}
+				CompileAutomationProjects(CompileProjects, MsBuildProperties);
 			}
 
-			List<Assembly> ScriptAssemblies = new List<Assembly>();
-			LoadPreCompiledScriptAssemblies(ScriptAssemblies);
+			// Load everything
+			List<Assembly> ScriptAssemblies = LoadAutomationAssemblies(Projects);
 
 			// Setup platforms
 			Platform.InitializePlatforms(ScriptAssemblies.ToArray());
@@ -119,13 +124,8 @@ namespace AutomationTool
 			}
 		}
 
-		private static void FindAndCompileScriptModules(string ScriptsForProjectFileName, List<string> AdditionalScriptsFolders)
+		private static List<FileReference> FindAutomationProjects(string ScriptsForProjectFileName, List<string> AdditionalScriptsFolders)
 		{
-			string OldCWD = Environment.CurrentDirectory;
-			string UnrealBuildToolCWD = CommandUtils.CombinePaths(CommandUtils.CmdEnv.LocalRoot, "Engine", "Source");
-
-			Environment.CurrentDirectory = UnrealBuildToolCWD;
-
 			// Configure the rules compiler
 			// Get all game folders and convert them to build subfolders.
 			List<DirectoryReference> AllGameFolders;
@@ -151,22 +151,20 @@ namespace AutomationTool
 			Log.TraceVerbose("Discovering game folders.");
 
 			List<FileReference> DiscoveredModules = UnrealBuildTool.RulesCompiler.FindAllRulesSourceFiles(UnrealBuildTool.RulesCompiler.RulesFileType.AutomationModule, GameFolders: AllGameFolders, ForeignPlugins: null, AdditionalSearchPaths: AllAdditionalScriptFolders);
-			List<string> ModulesToCompile = new List<string>(DiscoveredModules.Count);
+			List<FileReference> ModulesToCompile = new List<FileReference>(DiscoveredModules.Count);
 			foreach (FileReference ModuleFilename in DiscoveredModules)
 			{
 				if (HostPlatform.Current.IsScriptModuleSupported(ModuleFilename.GetFileNameWithoutAnyExtensions()))
 				{
-					ModulesToCompile.Add(ModuleFilename.FullName);
+					ModulesToCompile.Add(ModuleFilename);
 				}
 				else
 				{
 					CommandUtils.LogVerbose("Script module {0} filtered by the Host Platform and will not be compiled.", ModuleFilename);
 				}
 			}
-			
-			CompileModules(ModulesToCompile);
 
-			Environment.CurrentDirectory = OldCWD;
+			return ModulesToCompile;
 		}
 
 		/// <summary>
@@ -263,16 +261,26 @@ namespace AutomationTool
 
 			return UpToDate;
 		}
-		
+
 		/// <summary>
-		/// Compiles all script modules.
+		/// Converts a set of MSBuild properties into command line arguments
 		/// </summary>
-		/// <param name="Modules">Module project filenames.</param>
-		private static void CompileModules(List<string> Modules)
+		/// <param name="Properties">The properties to set</param>
+		/// <returns>Command line arguments</returns>
+		private static string GetMsBuildPropertyArguments(Dictionary<string, string> Properties)
+		{
+			return String.Join(" ", Properties.Select(x => String.Format(" /p:{0}={1}", Utils.MakePathSafeToUseWithCommandLine(x.Key), Utils.MakePathSafeToUseWithCommandLine(x.Value))));
+		}
+
+		/// <summary>
+		/// Compiles all automation projects
+		/// </summary>
+		/// <param name="Projects">Projects to compile</param>
+		/// <param name="MsBuildProperties">Properties to set</param>
+		private static void CompileAutomationProjects(List<CsProjectInfo> Projects, Dictionary<string, string> MsBuildProperties)
 		{			
 			string DependencyFile = Path.Combine(CommandUtils.CmdEnv.EngineSavedFolder, "UATModuleHashes.xml");
-			
-			if (AreDependenciesUpToDate(Modules, DependencyFile) && !GlobalCommandLine.IgnoreDependencies)
+			if (AreDependenciesUpToDate(Projects.Select(x => x.ProjectPath.FullName), DependencyFile) && !GlobalCommandLine.IgnoreDependencies)
 			{
 				Log.TraceInformation("Dependencies are up to date. Skipping compile.");
 				return;
@@ -283,16 +291,15 @@ namespace AutomationTool
 			// clean old assemblies
 			CleanupScriptsAssemblies();
 
-			DateTime StartTime = DateTime.Now;
+			Stopwatch Timer = Stopwatch.StartNew();
 
 			string BuildTool = CommandUtils.CmdEnv.MsBuildExe;
 
 			// msbuild (standard on windows, in mono >=5.0 is preferred due to speed and parallel compilation)
 			bool UseParallelMsBuild = Path.GetFileNameWithoutExtension(BuildTool).ToLower() == "msbuild";
-
 			if (UseParallelMsBuild)
 			{
-				string ModulesList = string.Join(";", Modules);
+				string ProjectsList = string.Join(";", Projects.Select(x => x.ProjectPath));
 
 				// Mono has an issue where arugments with semicolons or commas can't be passed through to
 				// as arguments so we need to manually construct a temp file with the list of modules
@@ -301,14 +308,14 @@ namespace AutomationTool
 				string UATProjFile = Path.Combine(CommandUtils.CmdEnv.EngineSavedFolder, "UATTempProj.proj");
 
 				string ProjContents = File.ReadAllText(UATProjTemplate);
-				ProjContents = ProjContents.Replace("$(Modules)", ModulesList);
+				ProjContents = ProjContents.Replace("$(Modules)", ProjectsList);
 
 				Directory.CreateDirectory(Path.GetDirectoryName(UATProjFile));
 				File.WriteAllText(UATProjFile, ProjContents);
 				
 				string MsBuildVerbosity = Log.OutputLevel >= LogEventType.Verbose ? "minimal" : "quiet";
 
-				string CmdLine = String.Format("\"{0}\" /p:Configuration={1} /verbosity:{2} /nologo", UATProjFile, BuildConfig, MsBuildVerbosity);
+				string CmdLine = String.Format("\"{0}\" {1} /verbosity:{2} /nologo", UATProjFile, GetMsBuildPropertyArguments(MsBuildProperties), MsBuildVerbosity);
 				// suppress the run command because it can be long and intimidating, making the logs around this code harder to read.
 				IProcessResult Result = CommandUtils.Run(BuildTool, CmdLine, Options: CommandUtils.ERunOptions.Default | CommandUtils.ERunOptions.NoLoggingOfRunCommand | CommandUtils.ERunOptions.LoggingOfRunDuration);
 				if (Result.ExitCode != 0)
@@ -322,47 +329,42 @@ namespace AutomationTool
 				string DefaultScriptsProjName = Path.ChangeExtension(DefaultScriptsDLLName, "csproj");
 
 				// Primary modules must be built first
-				List<string> PrimaryModules = Modules.Where(M => M.IndexOf(DefaultScriptsProjName, StringComparison.InvariantCultureIgnoreCase) >= 0).ToList();
-
-				foreach (string ModuleName in PrimaryModules)
+				List<CsProjectInfo> PrimaryProjects = Projects.Where(M => M.ProjectPath.FullName.IndexOf(DefaultScriptsProjName, StringComparison.InvariantCultureIgnoreCase) >= 0).ToList();
+				foreach (CsProjectInfo PrimaryProject in PrimaryProjects)
 				{
-					Log.TraceInformation("Building script module: {0}", ModuleName);
+					Log.TraceInformation("Building script module: {0}", PrimaryProject.ProjectPath);
 					try
 					{
-						CompileScriptModule(ModuleName);
+						CompileAutomationProject(PrimaryProject.ProjectPath, MsBuildProperties);
 					}
 					catch (Exception Ex)
 					{
-						CommandUtils.LogError(LogUtils.FormatException(Ex));
-						throw new AutomationException("Failed to compile module {0}", ModuleName);
+						throw new AutomationException(Ex, "Failed to compile module {0}", PrimaryProject.ProjectPath);
 					}
 					break;
-
 				}
 
 				// Second pass, compile everything else
-				List<string> SecondaryModules = Modules.Where(M => !PrimaryModules.Contains(M)).ToList();
+				List<CsProjectInfo> SecondaryProjects = Projects.Where(M => !PrimaryProjects.Contains(M)).ToList();
 
 				// Non-parallel method
-				foreach (string ModuleName in SecondaryModules)
+				foreach (CsProjectInfo SecondaryProject in SecondaryProjects)
 				{
-					Log.TraceInformation("Building script module: {0}", ModuleName);
+					Log.TraceInformation("Building script module: {0}", SecondaryProject.ProjectPath);
 					try
 					{
-						CompileScriptModule(ModuleName);
+						CompileAutomationProject(SecondaryProject.ProjectPath, MsBuildProperties);
 					}
 					catch (Exception Ex)
 					{
-						CommandUtils.LogError(LogUtils.FormatException(Ex));
-						throw new AutomationException("Failed to compile module {0}", ModuleName);
+						throw new AutomationException(Ex, "Failed to compile module {0}", SecondaryProject.ProjectPath);
 					}
 				}
 			}			
 			
-			TimeSpan Duration = DateTime.Now - StartTime;
-			Log.TraceInformation("Compiled {0} modules in {1} secs", Modules.Count, Duration.TotalSeconds);
+			Log.TraceInformation("Compiled {0} modules in {1} secs", Projects.Count, Timer.Elapsed.TotalSeconds);
 
-			HashCollection NewHashes = HashModules(Modules);
+			HashCollection NewHashes = HashModules(Projects.Select(x => x.ProjectPath.FullName));
 
 			if (NewHashes == null)
 			{
@@ -381,19 +383,19 @@ namespace AutomationTool
 		/// </summary>
 		/// <param name="ProjectFile"></param>
 		/// <returns></returns>
-		private static bool CompileScriptModule(string ProjectFile)
+		private static bool CompileAutomationProject(FileReference ProjectFile, Dictionary<string, string> Properties)
 		{
-			if (!ProjectFile.EndsWith(".csproj", StringComparison.InvariantCultureIgnoreCase))
+			if (!ProjectFile.HasExtension(".csproj"))
 			{
 				throw new AutomationException(String.Format("Unable to build Project {0}. Not a valid .csproj file.", ProjectFile));
 			}
-			if (!CommandUtils.FileExists(ProjectFile))
+			if (!FileReference.Exists(ProjectFile))
 			{
 				throw new AutomationException(String.Format("Unable to build Project {0}. Project file not found.", ProjectFile));
 			}
 
-			string CmdLine = String.Format("\"{0}\" /verbosity:quiet /nologo /target:Build /property:Configuration={1} /property:Platform=AnyCPU /p:TreatWarningsAsErrors=false /p:NoWarn=\"612,618,672,1591\" /p:BuildProjectReferences=true",
-				ProjectFile, BuildConfig);
+			string CmdLine = String.Format("\"{0}\" /verbosity:quiet /nologo /target:Build {0} /p:TreatWarningsAsErrors=false /p:NoWarn=\"612,618,672,1591\" /p:BuildProjectReferences=true",
+				ProjectFile, GetMsBuildPropertyArguments(Properties));
 
 			// Compile the project
 			IProcessResult Result = CommandUtils.Run(CommandUtils.CmdEnv.MsBuildExe, CmdLine);
@@ -413,76 +415,58 @@ namespace AutomationTool
 		/// <summary>
 		/// Loads all precompiled assemblies (DLLs that end with *Scripts.dll).
 		/// </summary>
-		/// <param name="OutScriptAssemblies">List to store all loaded assemblies.</param>
-		private static void LoadPreCompiledScriptAssemblies(List<Assembly> OutScriptAssemblies)
+		/// <param name="Projects">Projects to load</param>
+		/// <returns>List of compiled assemblies</returns>
+		private static List<Assembly> LoadAutomationAssemblies(IEnumerable<CsProjectInfo> Projects)
 		{
-			CommandUtils.LogVerbose("Loading precompiled script DLLs");
-
-			bool DefaultScriptsDLLFound = false;
-			string ScriptsLocation = GetScriptAssemblyFolder();
-			if (CommandUtils.DirectoryExists(ScriptsLocation))
+			List<Assembly> Assemblies = new List<Assembly>();
+			foreach (CsProjectInfo Project in Projects)
 			{
-				string[] ScriptDLLFiles = Directory.GetFiles(ScriptsLocation, "*.Automation.dll", SearchOption.AllDirectories);
-
-				CommandUtils.LogVerbose("Found {0} script DLL(s).", ScriptDLLFiles.Length);
-				foreach (string ScriptsDLLFilename in ScriptDLLFiles)
+				// Get the output assembly name
+				FileReference AssemblyLocation;
+				if (!Project.TryGetOutputFile(out AssemblyLocation))
 				{
-
-					if (!HostPlatform.Current.IsScriptModuleSupported(CommandUtils.GetFilenameWithoutAnyExtensions(ScriptsDLLFilename)))
-					{
-						CommandUtils.LogVerbose("Script module {0} filtered by the Host Platform and will not be loaded.", ScriptsDLLFilename);
-						continue;
-					}
-					// Load the assembly into our app domain
-					CommandUtils.LogVerbose("Loading script DLL: {0}", ScriptsDLLFilename);
-					try
-					{
-						Assembly Dll = AppDomain.CurrentDomain.Load(AssemblyName.GetAssemblyName(ScriptsDLLFilename));
-						OutScriptAssemblies.Add(Dll);
-						// Check if this is the default scripts DLL.
-						if (!DefaultScriptsDLLFound && String.Compare(Path.GetFileName(ScriptsDLLFilename), DefaultScriptsDLLName, true) == 0)
-						{
-							DefaultScriptsDLLFound = true;
-						}
-					}
-					catch (Exception Ex)
-					{
-						throw new AutomationException("Failed to load script DLL: {0}: {1}", ScriptsDLLFilename, Ex.Message);
-					}
+					throw new AutomationException("Unable to get output file for {0}", Project.ProjectPath);
 				}
-			}
-			else
-			{
-				CommandUtils.LogError("Scripts folder {0} does not exist!", ScriptsLocation);
-			}
 
-			// The default scripts DLL is required!
-			if (!DefaultScriptsDLLFound)
-			{
-				throw new AutomationException("{0} was not found or could not be loaded, can't run scripts.", DefaultScriptsDLLName);
+				// Load the assembly into our app domain
+				CommandUtils.LogVerbose("Loading script DLL: {0}", AssemblyLocation);
+				try
+				{
+					Assembly Assembly = AppDomain.CurrentDomain.Load(AssemblyName.GetAssemblyName(AssemblyLocation.FullName));
+					Assemblies.Add(Assembly);
+				}
+				catch (Exception Ex)
+				{
+					throw new AutomationException("Failed to load script DLL: {0}: {1}", AssemblyLocation, Ex.Message);
+				}
 			}
+			return Assemblies;
 		}
 
 		private static void CleanupScriptsAssemblies()
 		{
-			CommandUtils.LogVerbose("Cleaning up script DLL folder");
-			CommandUtils.DeleteDirectory(GetScriptAssemblyFolder());
-
-			// Bug in PortalPublishingTool caused these DLLs to be copied into Engine/Binaries/DotNET. Delete any files left over.
-			DirectoryReference BinariesDir = DirectoryReference.Combine(CommandUtils.RootDirectory, "Engine", "Binaries", "DotNET");
-			foreach (FileReference FileToDelete in DirectoryReference.EnumerateFiles(BinariesDir, "*.automation.dll"))
+			if (!CommandUtils.IsEngineInstalled())
 			{
-				CommandUtils.DeleteFile(FileToDelete.FullName);
-			}
-			foreach (FileReference FileToDelete in DirectoryReference.EnumerateFiles(BinariesDir, "*.automation.pdb"))
-			{
-				CommandUtils.DeleteFile(FileToDelete.FullName);
+				CommandUtils.LogVerbose("Cleaning up script DLL folder");
+				CommandUtils.DeleteDirectory(GetScriptAssemblyFolder());
+
+				// Bug in PortalPublishingTool caused these DLLs to be copied into Engine/Binaries/DotNET. Delete any files left over.
+				DirectoryReference BinariesDir = DirectoryReference.Combine(CommandUtils.RootDirectory, "Engine", "Binaries", "DotNET");
+				foreach (FileReference FileToDelete in DirectoryReference.EnumerateFiles(BinariesDir, "*.automation.dll"))
+				{
+					CommandUtils.DeleteFile(FileToDelete.FullName);
+				}
+				foreach (FileReference FileToDelete in DirectoryReference.EnumerateFiles(BinariesDir, "*.automation.pdb"))
+				{
+					CommandUtils.DeleteFile(FileToDelete.FullName);
+				}
 			}
 		}
 
-		private static string GetScriptAssemblyFolder()
+		private static DirectoryReference GetScriptAssemblyFolder()
 		{
-			return CommandUtils.CombinePaths(CommandUtils.CmdEnv.LocalRoot, "Engine", "Binaries", "DotNET", "AutomationScripts");
+			return DirectoryReference.Combine(CommandUtils.EngineDirectory, "Binaries", "DotNET", "AutomationScripts");
 		}
 
 		public Dictionary<string, Type> Commands
diff --git a/Engine/Source/Programs/DotNETCommon/DotNETUtilities/CsProjectInfo.cs b/Engine/Source/Programs/DotNETCommon/DotNETUtilities/CsProjectInfo.cs
--- a/Engine/Source/Programs/DotNETCommon/DotNETUtilities/CsProjectInfo.cs
+++ b/Engine/Source/Programs/DotNETCommon/DotNETUtilities/CsProjectInfo.cs
@@ -55,6 +55,31 @@ namespace Tools.DotNETCommon
 			Properties = new Dictionary<string, string>(InProperties);
 		}
 
+		/// <summary>
+		/// Get the ouptut file for this project
+		/// </summary>
+		/// <param name="File">If successful, receives the assembly path</param>
+		/// <returns>True if the output file was found</returns>
+		public bool TryGetOutputFile(out FileReference File)
+		{
+			DirectoryReference OutputDir;
+			if(!TryGetOutputDir(out OutputDir))
+			{
+				File = null;
+				return false;
+			}
+
+			string AssemblyName;
+			if(!TryGetAssemblyName(out AssemblyName))
+			{
+				File = null;
+				return false;
+			}
+
+			File = FileReference.Combine(OutputDir, AssemblyName + ".dll");
+			return true;
+		}
+
 		/// <summary>
 		/// Resolve the project's output directory
 		/// </summary>
@@ -73,15 +98,33 @@ namespace Tools.DotNETCommon
 			}
 		}
 
+		/// <summary>
+		/// Resolve the project's output directory
+		/// </summary>
+		/// <param name="OutputDir">If successful, receives the output directory</param>
+		/// <returns>True if the output directory was found</returns>
+		public bool TryGetOutputDir(out DirectoryReference OutputDir)
+		{
+			string OutputPath;
+			if (Properties.TryGetValue("OutputPath", out OutputPath))
+			{
+				OutputDir = DirectoryReference.Combine(ProjectPath.Directory, OutputPath);
+				return true;
+			}
+			else
+			{
+				OutputDir = null;
+				return false;
+			}
+		}
+
 		/// <summary>
 		/// Returns the assembly name used by this project
 		/// </summary>
 		/// <returns></returns>
-		public string GetAssemblyName()
+		public bool TryGetAssemblyName(out string AssemblyName)
 		{
-			string Output = "";
-			Properties.TryGetValue("AssemblyName", out Output);
-			return Output;
+			return Properties.TryGetValue("AssemblyName", out AssemblyName);
 		}
 
 		/// <summary>
@@ -270,6 +313,22 @@ namespace Tools.DotNETCommon
 			}
 		}
 
+		/// <summary>
+		/// Reads project information for the given file.
+		/// </summary>
+		/// <param name="File">The project file to read</param>
+		/// <param name="Properties">Initial set of property values</param>
+		/// <returns>The parsed project info</returns>
+		public static CsProjectInfo Read(FileReference File, Dictionary<string, string> Properties)
+		{
+			CsProjectInfo Project;
+			if(!TryRead(File, Properties, out Project))
+			{
+				throw new Exception(String.Format("Unable to read '{0}'", File));
+			}
+			return Project;
+		}
+
 		/// <summary>
 		/// Attempts to read project information for the given file.
 		/// </summary>
@@ -904,9 +963,11 @@ namespace Tools.DotNETCommon
 		public static bool AddCsProjectInfo(this HashCollection Hasher, CsProjectInfo Project, HashCollection.HashType HashType)
 		{
 			// Get the output assembly and pdb file
-			DirectoryReference ProjectDirectory = Project.ProjectPath.Directory;
-			DirectoryReference OutputDir = Project.GetOutputDir(ProjectDirectory);
-			FileReference OutputFile = FileReference.Combine(OutputDir, Project.GetAssemblyName() + ".dll");
+			FileReference OutputFile;
+			if(!Project.TryGetOutputFile(out OutputFile))
+			{
+				throw new Exception(String.Format("Unable to get output file for {0}", Project.ProjectPath));
+			}
 			FileReference DebugFile = OutputFile.ChangeExtension("pdb");
 
 			// build a list of all input and output files from this module
diff --git a/Engine/Source/Programs/UnrealBuildTool/ProjectFiles/ProjectFileGenerator.cs b/Engine/Source/Programs/UnrealBuildTool/ProjectFiles/ProjectFileGenerator.cs
--- a/Engine/Source/Programs/UnrealBuildTool/ProjectFiles/ProjectFileGenerator.cs
+++ b/Engine/Source/Programs/UnrealBuildTool/ProjectFiles/ProjectFileGenerator.cs
@@ -433,8 +433,34 @@ namespace UnrealBuildTool
 					VCSharpProjectFile Project = new VCSharpProjectFile(ProjectFile);
 					Project.ShouldBuildForAllSolutionTargets = false;//true;
 					AddExistingProjectFile(Project, bForceDevelopmentConfiguration: true);
-                    AutomationProjectFiles.Add( Project );
+					AutomationProjectFiles.Add( Project );
 					Folder.ChildProjects.Add( Project );
+
+					if (!ProjectFile.IsUnderDirectory(UnrealBuildTool.EngineDirectory))
+					{
+						FileReference PropsFile = new FileReference(ProjectFile.FullName + ".props");
+						CreateAutomationProjectPropsFile(PropsFile);
+					}
+				}
+			}
+		}
+
+		/// <summary>
+		/// Creates a .props file next to each automation project which specifies the path to the engine directory
+		/// </summary>
+		/// <param name="PropsFile">The properties file path</param>
+		void CreateAutomationProjectPropsFile(FileReference PropsFile)
+		{
+			using (FileStream Stream = FileReference.Open(PropsFile, FileMode.Create, FileAccess.Write, FileShare.Read))
+			{
+				using (StreamWriter Writer = new StreamWriter(Stream, Encoding.UTF8))
+				{
+					Writer.WriteLine("<?xml version=\"1.0\" encoding=\"utf-8\"?>");
+					Writer.WriteLine("<Project ToolsVersion=\"Current\" xmlns=\"http://schemas.microsoft.com/developer/msbuild/2003\">");
+					Writer.WriteLine("\t<PropertyGroup>");
+					Writer.WriteLine("\t\t<EngineDir Condition=\"'$(EngineDir)' == ''\">{0}</EngineDir>", UnrealBuildTool.EngineDirectory);
+					Writer.WriteLine("\t</PropertyGroup>");
+					Writer.WriteLine("</Project>");
 				}
 			}
 		}
