Add support for content plugins in content-only projects. These were not being staged due to a unique target receipt for the project being generated.

diff --git a/Engine/Source/Programs/AutomationTool/Scripts/CopyBuildToStagingDirectory.Automation.cs b/Engine/Source/Programs/AutomationTool/Scripts/CopyBuildToStagingDirectory.Automation.cs
--- a/Engine/Source/Programs/AutomationTool/Scripts/CopyBuildToStagingDirectory.Automation.cs
+++ b/Engine/Source/Programs/AutomationTool/Scripts/CopyBuildToStagingDirectory.Automation.cs
@@ -433,6 +433,58 @@ public partial class Project : CommandUtils
 		}
 	}
 
+	/// <summary>
+	/// Figures out which content plugins are enabled for a content-only plugins. These will not be in the default target manifest, since it was compiled elsewhere.
+	/// </summary>
+	/// <param name="ProjectFile">The project being built</param>
+	/// <param name="Targets">List of targets being staged</param>
+	/// <returns>List of plugin files that should be staged</returns>
+	private static List<FileReference> GetContentPluginsForContentProject(FileReference ProjectFile, List<TargetReceipt> Targets)
+	{
+		ProjectDescriptor Project = ProjectDescriptor.FromFile(ProjectFile);
+		List<PluginInfo> AvailablePlugins = Plugins.ReadAvailablePlugins(CommandUtils.EngineDirectory, ProjectFile.Directory, null);
+
+		HashSet<FileReference> ContentPlugins = new HashSet<FileReference>();
+		foreach (TargetReceipt Target in Targets)
+		{
+			// Find all the specifically enabled plugins for this target
+			Dictionary<string, bool> EnabledPlugins = new Dictionary<string, bool>(StringComparer.OrdinalIgnoreCase);
+			if (Project.Plugins != null)
+			{
+				foreach (PluginReferenceDescriptor Reference in Project.Plugins)
+				{
+					bool bEnabled = false;
+					if (!Reference.IsEnabledForPlatform(Target.Platform) || !Reference.IsEnabledForTargetConfiguration(Target.Configuration) || !Reference.IsEnabledForTarget(Target.TargetType))
+					{
+						bEnabled = true;
+					}
+					EnabledPlugins[Reference.Name] = bEnabled;
+				}
+			}
+
+			// Add all the enabled-by-default plugins
+			foreach (PluginInfo AvailablePlugin in AvailablePlugins)
+			{
+				if (!EnabledPlugins.ContainsKey(AvailablePlugin.Name))
+				{
+					EnabledPlugins[AvailablePlugin.Name] = AvailablePlugin.Descriptor.SupportsTargetPlatform(Target.Platform) && AvailablePlugin.IsEnabledByDefault(!Project.DisableEnginePluginsByDefault);
+				}
+			}
+
+			// Add the enabled content-only plugins
+			foreach (PluginInfo AvailablePlugin in AvailablePlugins)
+			{
+				if (AvailablePlugin.Descriptor.SupportsTargetPlatform(Target.Platform) && AvailablePlugin.Descriptor.bCanContainContent && EnabledPlugins[AvailablePlugin.Name])
+				{
+					ContentPlugins.Add(AvailablePlugin.File);
+				}
+			}
+
+			ContentPlugins.ExceptWith(Target.RuntimeDependencies.Select(x => x.Path));
+		}
+		return ContentPlugins.ToList();
+	}
+
 	public static void CreateStagingManifest(ProjectParams Params, DeploymentContext SC)
 	{
 		if (!Params.Stage)
@@ -583,6 +635,13 @@ public partial class Project : CommandUtils
 				{
 					SC.StageRuntimeDependenciesFromReceipt(Target.Receipt, Target.RequireFilesExist, Params.UsePak(SC.StageTargetPlatform));
 				}
+
+				// Stage any content-only plugins for content-only projects. We don't have a custom executable for these.
+				if (!Params.IsCodeBasedProject)
+				{
+					List<FileReference> PluginFiles = GetContentPluginsForContentProject(Params.RawProjectPath, SC.StageTargets.ConvertAll(x => x.Receipt));
+					SC.StageFiles(StagedFileType.UFS, PluginFiles);
+				}
 			}
 
 			// move the UE4Commandline.txt file to the root of the stage
