Fixes Gauntlet support for 32-bit Android builds.

- Now use utils to pull build information instead of a custom regex.
- 32-bit builds are considered, but discarded if 64-bit versions are also present

diff --git a/Engine/Source/Programs/AutomationTool/Gauntlet/Platform/Android/Gauntlet.AndroidBuildSource.cs b/Engine/Source/Programs/AutomationTool/Gauntlet/Platform/Android/Gauntlet.AndroidBuildSource.cs
--- a/Engine/Source/Programs/AutomationTool/Gauntlet/Platform/Android/Gauntlet.AndroidBuildSource.cs
+++ b/Engine/Source/Programs/AutomationTool/Gauntlet/Platform/Android/Gauntlet.AndroidBuildSource.cs
@@ -25,13 +25,16 @@ namespace Gauntlet
 
 		public UnrealTargetPlatform Platform { get { return UnrealTargetPlatform.Android; } }
 
-		public AndroidBuild(UnrealTargetConfiguration InConfig, string InAndroidPackageName, string InApkPath, Dictionary<string, string> InFilesToInstall, BuildFlags InFlags)
+		public bool Is32Bit { get; protected set; }
+
+		public AndroidBuild(UnrealTargetConfiguration InConfig, string InAndroidPackageName, string InApkPath, Dictionary<string, string> InFilesToInstall, BuildFlags InFlags, bool InIs32Bit)
 		{
 			Configuration = InConfig;
 			AndroidPackageName = InAndroidPackageName;
 			SourceApkPath = InApkPath;
 			FilesToInstall = InFilesToInstall;
 			Flags = InFlags;
+			Is32Bit = InIs32Bit;
 		}
 
 		public bool CanSupportRole(UnrealTargetRole RoleType)
@@ -57,24 +60,24 @@ namespace Gauntlet
 			FileInfo[] InstallFiles = Di.GetFiles("Install_*");
 
 			foreach (FileInfo Fi in InstallFiles)
-			{
-				// todo 
-				if (Fi.FullName.Contains("armv7"))
+			{ 
+				bool PackageIs32Bit = Fi.FullName.Contains("armv7");
+	
+				UnrealTargetConfiguration UnrealConfig = UnrealHelpers.GetConfigurationFromExecutableName(InProjectName, Fi.FullName);
+				UnrealTargetRole UnrealRole = UnrealHelpers.GetRoleFromExecutableName(InProjectName, Fi.FullName);
+
+				if (UnrealConfig == UnrealTargetConfiguration.Unknown)
 				{
+					Log.Info("Skipping unrecognized build {0}", Fi.FullName);
 					continue;
 				}
 
-                // Unreal naming. In test & shipping the platform name and config will be includedn
-                string RegEx = "Install_(.+?)(client|game|server)-(Android)?-?(test|shipping)?-?(.+).bat";
-
-				Match Info = Regex.Match(Fi.Name, RegEx, RegexOptions.IgnoreCase);
-
 				bool TestInstall = Fi.Name.EndsWith("_TEST.bat", StringComparison.OrdinalIgnoreCase);
 				bool PatchInstall = Fi.Name.EndsWith("_Patch.bat", StringComparison.OrdinalIgnoreCase);
 
 				// filter out non-matching or test installation batch files
 				// test installation scripts are intended to be manually invoked
-				if (Info.Success == false || TestInstall || PatchInstall)
+				if (TestInstall || PatchInstall)
 				{
 					if (TestInstall || PatchInstall)
 					{
@@ -84,16 +87,6 @@ namespace Gauntlet
 					continue;
 				}
 
-				string TargetName = Info.Groups[1].ToString();
-				string TypeName = Info.Groups[2].ToString();
-				string ConfigName = Info.Groups[4].ToString();
-
-				if (string.IsNullOrEmpty(ConfigName))
-				{
-					ConfigName = "Development";
-				}
-
-
 				Log.Verbose("Pulling install data from {0}", Fi.FullName);
 
 				string AbsPath = Fi.Directory.FullName;
@@ -137,13 +130,6 @@ namespace Gauntlet
 					continue;
 				}
 
-				UnrealTargetConfiguration UnrealConfig;
-				if (Enum.TryParse(ConfigName, true, out UnrealConfig) == false)
-				{
-					Log.Warning("Could not determine Unreal Config type from '{0}'", ConfigName);
-					continue;
-				}
-
 				// Android builds are always packaged, and we can always replace the command line
 				BuildFlags Flags = BuildFlags.Packaged | BuildFlags.CanReplaceCommandLine;
 
@@ -161,12 +147,17 @@ namespace Gauntlet
 					Flags |= BuildFlags.NotBulk;
 				}
 
-				AndroidBuild NewBuild = new AndroidBuild(UnrealConfig, AndroidPackageName, SourceApkPath, FilesToInstall, Flags);
+				AndroidBuild NewBuild = new AndroidBuild(UnrealConfig, AndroidPackageName, SourceApkPath, FilesToInstall, Flags, PackageIs32Bit);
 
 				DiscoveredBuilds.Add(NewBuild);
 
 				Log.Verbose("Found {0} {1} build at {2}", UnrealConfig, ((Flags & BuildFlags.Bulk) == BuildFlags.Bulk) ? "(bulk)" : "(not bulk)", AbsPath);
+			}
 
+			// If we have both 32 and 64-bit builds, prefer 64-bit
+			if (DiscoveredBuilds.Where(B => B.Is32Bit == false).Any())
+			{
+				DiscoveredBuilds = DiscoveredBuilds.Where(B => B.Is32Bit = false).ToList();
 			}
 
 			return DiscoveredBuilds;
