Fix copying of automation DLLs to output directory. Now may exist within any subfolders.

#rnx

diff --git a/Engine/Source/Programs/AutomationTool/AutomationUtils/Automation.cs b/Engine/Source/Programs/AutomationTool/AutomationUtils/Automation.cs
--- a/Engine/Source/Programs/AutomationTool/AutomationUtils/Automation.cs
+++ b/Engine/Source/Programs/AutomationTool/AutomationUtils/Automation.cs
@@ -408,11 +408,6 @@ AutomationTool.exe [-verbose] [-compileonly] [-p4] Command0 [-Arg0 -Arg1 -Arg2 .
 			}
 		}
 
-		/// <summary>
-		/// Compiler with all scripts
-		/// </summary>
-		public static ScriptCompiler Compiler { get; set; }
-
 		/// <summary>
 		/// Main method.
 		/// </summary>
@@ -478,10 +473,9 @@ AutomationTool.exe [-verbose] [-compileonly] [-p4] Command0 [-Arg0 -Arg1 -Arg2 .
 			ProjectUtils.CleanupFolders();
 
 			// Compile scripts.
-			Compiler = new ScriptCompiler();
 			using(TelemetryStopwatch ScriptCompileStopwatch = new TelemetryStopwatch("ScriptCompile"))
 			{
-				Compiler.FindAndCompileAllScripts(OutScriptsForProjectFileName, AdditionalScriptsFolders);
+				ScriptCompiler.FindAndCompileAllScripts(OutScriptsForProjectFileName, AdditionalScriptsFolders);
 			}
 
 			if (GlobalCommandLine.CompileOnly)
@@ -492,18 +486,18 @@ AutomationTool.exe [-verbose] [-compileonly] [-p4] Command0 [-Arg0 -Arg1 -Arg2 .
 
 			if (GlobalCommandLine.List)
 			{
-				ListAvailableCommands(Compiler.Commands);
+				ListAvailableCommands(ScriptCompiler.Commands);
 				return ExitCode.Success;
 			}
 
 			if (GlobalCommandLine.Help)
 			{
-				DisplayHelp(CommandsToExecute, Compiler.Commands);
+				DisplayHelp(CommandsToExecute, ScriptCompiler.Commands);
 				return ExitCode.Success;
 			}
 
 			// Enable or disable P4 support
-			CommandUtils.InitP4Support(CommandsToExecute, Compiler.Commands);
+			CommandUtils.InitP4Support(CommandsToExecute, ScriptCompiler.Commands);
 			if (CommandUtils.P4Enabled)
 			{
 				Log.TraceLog("Setting up Perforce environment.");
@@ -512,7 +506,7 @@ AutomationTool.exe [-verbose] [-compileonly] [-p4] Command0 [-Arg0 -Arg1 -Arg2 .
 			}
 
 			// Find and execute commands.
-			ExitCode Result = Execute(CommandsToExecute, Compiler.Commands);
+			ExitCode Result = Execute(CommandsToExecute, ScriptCompiler.Commands);
 			if (TelemetryFile != null)
 			{
 				Directory.CreateDirectory(Path.GetDirectoryName(TelemetryFile));
diff --git a/Engine/Source/Programs/AutomationTool/AutomationUtils/ScriptCompiler.cs b/Engine/Source/Programs/AutomationTool/AutomationUtils/ScriptCompiler.cs
--- a/Engine/Source/Programs/AutomationTool/AutomationUtils/ScriptCompiler.cs
+++ b/Engine/Source/Programs/AutomationTool/AutomationUtils/ScriptCompiler.cs
@@ -30,9 +30,9 @@ namespace AutomationTool
 	/// <summary>
 	/// Compiles and loads script assemblies.
 	/// </summary>
-	public class ScriptCompiler
+	public static class ScriptCompiler
 	{
-		private Dictionary<string, Type> ScriptCommands;
+		private static Dictionary<string, Type> ScriptCommands;
 #if DEBUG
 		const string BuildConfig = "Debug";
 #else
@@ -40,16 +40,12 @@ namespace AutomationTool
 #endif
 		const string DefaultScriptsDLLName = "AutomationScripts.Automation.dll";
 
-		public ScriptCompiler()
-		{
-		}
-
 		/// <summary>
 		/// Finds and/or compiles all script files and assemblies.
 		/// </summary>
 		/// <param name="ScriptsForProjectFileName">Path to the current project. May be null, in which case we compile scripts for all projects.</param>
 		/// <param name="AdditionalScriptsFolders">Additional script fodlers to look for source files in.</param>
-		public void FindAndCompileAllScripts(string ScriptsForProjectFileName, List<string> AdditionalScriptsFolders)
+		public static void FindAndCompileAllScripts(string ScriptsForProjectFileName, List<string> AdditionalScriptsFolders)
 		{
 			// Find all the project files
 			List<FileReference> ProjectFiles = FindAutomationProjects(ScriptsForProjectFileName, AdditionalScriptsFolders);
@@ -76,15 +72,15 @@ namespace AutomationTool
 			}
 
 			// Load everything
-			List<Assembly> ScriptAssemblies = LoadAutomationAssemblies(Projects);
+			Assemblies = LoadAutomationAssemblies(Projects);
 
 			// Setup platforms
-			Platform.InitializePlatforms(ScriptAssemblies.ToArray());
+			Platform.InitializePlatforms(Assemblies.ToArray());
 
 			// Instantiate all the automation classes for interrogation
 			Log.TraceVerbose("Creating commands.");
 			ScriptCommands = new Dictionary<string, Type>(StringComparer.InvariantCultureIgnoreCase);
-			foreach (Assembly CompiledScripts in ScriptAssemblies)
+			foreach (Assembly CompiledScripts in Assemblies)
 			{
 				try
 				{
@@ -469,9 +465,15 @@ namespace AutomationTool
 			return DirectoryReference.Combine(CommandUtils.EngineDirectory, "Binaries", "DotNET", "AutomationScripts");
 		}
 
-		public Dictionary<string, Type> Commands
+		public static Dictionary<string, Type> Commands
 		{
 			get { return ScriptCommands; }
 		}
+
+		public static List<Assembly> Assemblies
+		{
+			get;
+			private set;
+		}
 	}
 }
diff --git a/Engine/Source/Programs/AutomationTool/AutomationUtils/UE4Build.cs b/Engine/Source/Programs/AutomationTool/AutomationUtils/UE4Build.cs
--- a/Engine/Source/Programs/AutomationTool/AutomationUtils/UE4Build.cs
+++ b/Engine/Source/Programs/AutomationTool/AutomationUtils/UE4Build.cs
@@ -1,4 +1,4 @@
-﻿// Copyright Epic Games, Inc. All Rights Reserved.
+// Copyright Epic Games, Inc. All Rights Reserved.
 using System;
 using System.Collections.Generic;
 using System.Text;
@@ -9,6 +9,7 @@ using System.Web.Script.Serialization;
 using UnrealBuildTool;
 using Tools.DotNETCommon;
 using System.Linq;
+using System.Reflection;
 
 namespace AutomationTool
 {
@@ -1314,14 +1315,7 @@ namespace AutomationTool
 		public void AddUATFilesToBuildProducts()
 		{
 			// Find all DLLs (scripts and their dependencies)
-			const string ScriptsPostfix = ".dll";
             var DotNetOutputLocation = CommandUtils.CombinePaths(CommandUtils.CmdEnv.LocalRoot, "Engine", "Binaries", "DotNET");
-			var UATScriptsLocation = CommandUtils.CombinePaths(DotNetOutputLocation, "AutomationScripts");
-			var UATScripts = Directory.GetFiles(UATScriptsLocation, "*" + ScriptsPostfix, SearchOption.AllDirectories);
-			if (CommandUtils.IsNullOrEmpty(UATScripts))
-			{
-				throw new UE4BuildException("No automation scripts found in {0}. Cannot add UAT files to the build products.", UATScriptsLocation);
-			}
 
 			var UATFiles = new List<string>(new string[] 
 					{
@@ -1345,8 +1339,9 @@ namespace AutomationTool
 			}
 
 			// All scripts are expected to exist in DotNET/AutomationScripts subfolder.
-			foreach (var UATScriptFilePath in UATScripts)
+			foreach (Assembly ScriptAssembly in ScriptCompiler.Assemblies)
 			{
+				string UATScriptFilePath = ScriptAssembly.Location;
 				if (!CommandUtils.FileExists_NoExceptions(UATScriptFilePath))
 				{
 					throw new UE4BuildException("Cannot add UAT to the build products because {0} does not exist.", UATScriptFilePath);
