Use original location when determining which files to copy for UAT.

#rnx

diff --git a/Engine/Source/Programs/AutomationTool/AutomationUtils/UE4Build.cs b/Engine/Source/Programs/AutomationTool/AutomationUtils/UE4Build.cs
--- a/Engine/Source/Programs/AutomationTool/AutomationUtils/UE4Build.cs
+++ b/Engine/Source/Programs/AutomationTool/AutomationUtils/UE4Build.cs
@@ -1341,7 +1341,7 @@ namespace AutomationTool
 			// All scripts are expected to exist in DotNET/AutomationScripts subfolder.
 			foreach (Assembly ScriptAssembly in ScriptCompiler.Assemblies)
 			{
-				string UATScriptFilePath = ScriptAssembly.Location;
+				string UATScriptFilePath = ScriptAssembly.GetOriginalLocation();
 				if (!CommandUtils.FileExists_NoExceptions(UATScriptFilePath))
 				{
 					throw new UE4BuildException("Cannot add UAT to the build products because {0} does not exist.", UATScriptFilePath);
