Fix enumeration of UAT build products.

#rnx

diff --git a/Engine/Source/Programs/AutomationTool/AutomationUtils/ScriptCompiler.cs b/Engine/Source/Programs/AutomationTool/AutomationUtils/ScriptCompiler.cs
--- a/Engine/Source/Programs/AutomationTool/AutomationUtils/ScriptCompiler.cs
+++ b/Engine/Source/Programs/AutomationTool/AutomationUtils/ScriptCompiler.cs
@@ -71,8 +71,28 @@ namespace AutomationTool
 				CompileAutomationProjects(CompileProjects, MsBuildProperties);
 			}
 
+			// Get all the build artifacts
+			BuildProducts = new HashSet<FileReference>();
+
+			HashSet<DirectoryReference> OutputDirs = new HashSet<DirectoryReference>();
+			OutputDirs.Add(DirectoryReference.Combine(CommandUtils.EngineDirectory, "Binaries", "DotNET")); // Don't want any artifacts from this directory (just AutomationTool.exe and AutomationScripts.dll)
+
+			foreach (CsProjectInfo Project in Projects)
+			{
+				DirectoryReference OutputDir;
+				if (!Project.TryGetOutputDir(out OutputDir))
+				{
+					throw new AutomationException("Unable to get output directory for {0}", Project.ProjectPath);
+				}
+
+				if (OutputDirs.Add(OutputDir))
+				{
+					BuildProducts.UnionWith(DirectoryReference.EnumerateFiles(OutputDir));
+				}
+			}
+
 			// Load everything
-			Assemblies = LoadAutomationAssemblies(Projects);
+			List<Assembly> Assemblies = LoadAutomationAssemblies(Projects);
 
 			// Setup platforms
 			Platform.InitializePlatforms(Assemblies.ToArray());
@@ -470,7 +490,7 @@ namespace AutomationTool
 			get { return ScriptCommands; }
 		}
 
-		public static List<Assembly> Assemblies
+		public static HashSet<FileReference> BuildProducts
 		{
 			get;
 			private set;
diff --git a/Engine/Source/Programs/AutomationTool/AutomationUtils/UE4Build.cs b/Engine/Source/Programs/AutomationTool/AutomationUtils/UE4Build.cs
--- a/Engine/Source/Programs/AutomationTool/AutomationUtils/UE4Build.cs
+++ b/Engine/Source/Programs/AutomationTool/AutomationUtils/UE4Build.cs
@@ -1339,14 +1339,13 @@ namespace AutomationTool
 			}
 
 			// All scripts are expected to exist in DotNET/AutomationScripts subfolder.
-			foreach (Assembly ScriptAssembly in ScriptCompiler.Assemblies)
+			foreach (FileReference BuildProduct in ScriptCompiler.BuildProducts)
 			{
-				string UATScriptFilePath = ScriptAssembly.GetOriginalLocation();
+				string UATScriptFilePath = BuildProduct.FullName;
 				if (!CommandUtils.FileExists_NoExceptions(UATScriptFilePath))
 				{
 					throw new UE4BuildException("Cannot add UAT to the build products because {0} does not exist.", UATScriptFilePath);
 				}
-
 				AddBuildProduct(UATScriptFilePath);
 			}
 		}
