Fix missing properties on command line to MSBuild when compiling automation projects.

diff --git a/Engine/Source/Programs/AutomationTool/AutomationUtils/ScriptCompiler.cs b/Engine/Source/Programs/AutomationTool/AutomationUtils/ScriptCompiler.cs
--- a/Engine/Source/Programs/AutomationTool/AutomationUtils/ScriptCompiler.cs
+++ b/Engine/Source/Programs/AutomationTool/AutomationUtils/ScriptCompiler.cs
@@ -410,7 +410,7 @@ namespace AutomationTool
 				throw new AutomationException(String.Format("Unable to build Project {0}. Project file not found.", ProjectFile));
 			}
 
-			string CmdLine = String.Format("\"{0}\" /verbosity:quiet /nologo /target:Build {0} /p:TreatWarningsAsErrors=false /p:NoWarn=\"612,618,672,1591\" /p:BuildProjectReferences=true",
+			string CmdLine = String.Format("\"{0}\" /verbosity:quiet /nologo /target:Build {1} /p:TreatWarningsAsErrors=false /p:NoWarn=\"612,618,672,1591\" /p:BuildProjectReferences=true",
 				ProjectFile, GetMsBuildPropertyArguments(Properties));
 
 			// Compile the project
