UAT: Fixed alternate cook source pak file ignoring

diff --git a/Engine/Source/Programs/AutomationTool/AutomationUtils/ProjectParams.cs b/Engine/Source/Programs/AutomationTool/AutomationUtils/ProjectParams.cs
--- a/Engine/Source/Programs/AutomationTool/AutomationUtils/ProjectParams.cs
+++ b/Engine/Source/Programs/AutomationTool/AutomationUtils/ProjectParams.cs
@@ -280,7 +280,7 @@ namespace AutomationTool
             this.RunTimeoutSeconds = InParams.RunTimeoutSeconds;
 			this.Clean = InParams.Clean;
 			this.Pak = InParams.Pak;
-			this.TryUsePaksFromDifferentCookSource = InParams.TryUsePaksFromDifferentCookSource;
+			this.IgnorePaksFromDifferentCookSource = InParams.IgnorePaksFromDifferentCookSource;
 			this.IoStore = InParams.IoStore;
 			this.SignPak = InParams.SignPak;
 			this.SignedPak = InParams.SignedPak;
@@ -456,7 +456,7 @@ namespace AutomationTool
 			bool? NoXGE = null,
 			bool? Package = null,
 			bool? Pak = null,
-			bool? TryUsePaksFromDifferentCookSource = null,
+			bool? IgnorePaksFromDifferentCookSource = null,
 			bool? IoStore = null,
 			bool? SkipIoStore = null,
 			bool? Prereqs = null,
@@ -631,7 +631,7 @@ namespace AutomationTool
 			}
 			this.PakAlignForMemoryMapping = GetParamValueIfNotSpecified(Command, PakAlignForMemoryMapping, this.PakAlignForMemoryMapping, "PakAlignForMemoryMapping");
 			this.Pak = GetParamValueIfNotSpecified(Command, Pak, this.Pak, "pak");
-			this.TryUsePaksFromDifferentCookSource = GetParamValueIfNotSpecified(Command, TryUsePaksFromDifferentCookSource, this.TryUsePaksFromDifferentCookSource, "TryUsePaksFromDifferentCookSource");
+			this.IgnorePaksFromDifferentCookSource = GetParamValueIfNotSpecified(Command, IgnorePaksFromDifferentCookSource, this.IgnorePaksFromDifferentCookSource, "IgnorePaksFromDifferentCookSource");
 			this.IoStore = GetParamValueIfNotSpecified(Command, IoStore, this.IoStore, "iostore");
 			this.SkipIoStore = GetParamValueIfNotSpecified(Command, IoStore, this.SkipIoStore, "skipiostore");
 			this.SkipPak = GetParamValueIfNotSpecified(Command, SkipPak, this.SkipPak, "skippak");
@@ -1194,11 +1194,10 @@ namespace AutomationTool
 		public bool Pak { private set; get; }
 
 		/// <summary>
-		/// Stage: True if we should try and re-use pak files from another staged build when we've specified a different cook source platform
+		/// Stage: True if we should disable trying to re-use pak files from another staged build when we've specified a different cook source platform
 		/// </summary>
-		[Help("pak", "reuse pak files from the alternate cook source folder, if specified")]
-		public bool TryUsePaksFromDifferentCookSource { get { return TryUsePaksFromDifferentCookSourceInternal; } private set { TryUsePaksFromDifferentCookSourceInternal = value; } }
-		private bool TryUsePaksFromDifferentCookSourceInternal = true;
+		[Help("pak", "disable reuse of pak files from the alternate cook source folder, if specified")]
+		public bool IgnorePaksFromDifferentCookSource { get; private set; }
 
 		/// <summary>
 		/// Shared: True if container file(s) should be generated with ZenPak.
@@ -2828,7 +2827,7 @@ namespace AutomationTool
 				CommandUtils.LogLog("MapsToCook={0}", MapsToCook.ToString());
 				CommandUtils.LogLog("MapIniSectionsToCook={0}", MapIniSectionsToCook.ToString());
 				CommandUtils.LogLog("Pak={0}", Pak);
-				CommandUtils.LogLog("TryUsePaksFromDifferentCookSource={0}", TryUsePaksFromDifferentCookSource);
+				CommandUtils.LogLog("IgnorePaksFromDifferentCookSource={0}", IgnorePaksFromDifferentCookSource);
 				CommandUtils.LogLog("IoStore={0}", IoStore);
 				CommandUtils.LogLog("SkipIoStore={0}", SkipIoStore);
 				CommandUtils.LogLog("Package={0}", Package);
diff --git a/Engine/Source/Programs/AutomationTool/Scripts/CopyBuildToStagingDirectory.Automation.cs b/Engine/Source/Programs/AutomationTool/Scripts/CopyBuildToStagingDirectory.Automation.cs
--- a/Engine/Source/Programs/AutomationTool/Scripts/CopyBuildToStagingDirectory.Automation.cs
+++ b/Engine/Source/Programs/AutomationTool/Scripts/CopyBuildToStagingDirectory.Automation.cs
@@ -2147,7 +2147,7 @@ public partial class Project : CommandUtils
 
 			bool bCopiedExistingPak = false;
 
-			if (SC.StageTargetPlatform != SC.CookSourcePlatform && Params.TryUsePaksFromDifferentCookSource)
+			if (SC.StageTargetPlatform != SC.CookSourcePlatform && !Params.IgnorePaksFromDifferentCookSource)
 			{
 				// Check to see if we have an existing pak file we can use
 
@@ -2200,7 +2200,7 @@ public partial class Project : CommandUtils
 
 			if (!bCopiedExistingPak)
 			{
-				if (FileReference.Exists(OutputLocation))
+				if (FileReference.Exists(OutputLocation) && !Params.IgnorePaksFromDifferentCookSource)
 				{
 					string UnrealPakResponseFileName = CombinePaths(CmdEnv.LogFolder, "PakList_" + OutputLocation.GetFileNameWithoutExtension() + ".txt");
 					if (File.Exists(UnrealPakResponseFileName) && FileReference.GetLastWriteTimeUtc(OutputLocation) > File.GetLastWriteTimeUtc(UnrealPakResponseFileName))
