Add timers for UAT startup.

#rnx

diff --git a/Engine/Source/Programs/AutomationTool/AutomationUtils/ScriptCompiler.cs b/Engine/Source/Programs/AutomationTool/AutomationUtils/ScriptCompiler.cs
--- a/Engine/Source/Programs/AutomationTool/AutomationUtils/ScriptCompiler.cs
+++ b/Engine/Source/Programs/AutomationTool/AutomationUtils/ScriptCompiler.cs
@@ -48,7 +48,9 @@ namespace AutomationTool
 		public static void FindAndCompileAllScripts(string ScriptsForProjectFileName, List<string> AdditionalScriptsFolders)
 		{
 			// Find all the project files
+			Stopwatch SearchTimer = Stopwatch.StartNew();
 			List<FileReference> ProjectFiles = FindAutomationProjects(ScriptsForProjectFileName, AdditionalScriptsFolders);
+			Log.TraceLog("Found {0} project files in {1:0.000}s", ProjectFiles.Count, SearchTimer.Elapsed.TotalSeconds);
 
 			// Get the default properties for compiling the projects
 			Dictionary<string, string> MsBuildProperties = new Dictionary<string, string>(StringComparer.OrdinalIgnoreCase);
@@ -57,8 +59,10 @@ namespace AutomationTool
 			MsBuildProperties.Add("EngineDir", CommandUtils.EngineDirectory.FullName);
 
 			// Read all the projects
+			Stopwatch ParsingTimer = Stopwatch.StartNew();
 			CsProjectInfo[] Projects = new CsProjectInfo[ProjectFiles.Count];
 			Parallel.For(0, ProjectFiles.Count, Idx => Projects[Idx] = CsProjectInfo.Read(ProjectFiles[Idx], MsBuildProperties));
+			Log.TraceLog("Parsed project files in {0:0.000}s", ParsingTimer.Elapsed.TotalSeconds);
 
 			// Compile only if not disallowed.
 			if (GlobalCommandLine.Compile && !String.IsNullOrEmpty(CommandUtils.CmdEnv.MsBuildExe))
@@ -92,7 +96,9 @@ namespace AutomationTool
 			}
 
 			// Load everything
+			Stopwatch LoadTimer = Stopwatch.StartNew();
 			List<Assembly> Assemblies = LoadAutomationAssemblies(Projects);
+			Log.TraceLog("Loaded assemblies in {0:0.000}s", LoadTimer.Elapsed.TotalSeconds);
 
 			// Setup platforms
 			Platform.InitializePlatforms(Assemblies.ToArray());
@@ -294,11 +300,13 @@ namespace AutomationTool
 		/// <param name="Projects">Projects to compile</param>
 		/// <param name="MsBuildProperties">Properties to set</param>
 		private static void CompileAutomationProjects(List<CsProjectInfo> Projects, Dictionary<string, string> MsBuildProperties)
-		{			
+		{
+			Stopwatch Timer = Stopwatch.StartNew();
+
 			string DependencyFile = Path.Combine(CommandUtils.CmdEnv.EngineSavedFolder, "UATModuleHashes.xml");
 			if (AreDependenciesUpToDate(Projects.Select(x => x.ProjectPath.FullName), DependencyFile) && !GlobalCommandLine.IgnoreDependencies)
 			{
-				Log.TraceInformation("Dependencies are up to date. Skipping compile.");
+				Log.TraceInformation("Dependencies are up to date ({0:0.000}s). Skipping compile.", Timer.Elapsed.TotalSeconds);
 				return;
 			}
 
@@ -307,8 +315,6 @@ namespace AutomationTool
 			// clean old assemblies
 			CleanupScriptsAssemblies();
 
-			Stopwatch Timer = Stopwatch.StartNew();
-
 			string BuildTool = CommandUtils.CmdEnv.MsBuildExe;
 
 			// msbuild (standard on windows, in mono >=5.0 is preferred due to speed and parallel compilation)
@@ -378,7 +384,7 @@ namespace AutomationTool
 				}
 			}			
 			
-			Log.TraceInformation("Compiled {0} modules in {1} secs", Projects.Count, Timer.Elapsed.TotalSeconds);
+			Log.TraceInformation("Compiled {0} modules in {1:0.000} secs", Projects.Count, Timer.Elapsed.TotalSeconds);
 
 			HashCollection NewHashes = HashModules(Projects.Select(x => x.ProjectPath.FullName));
 
