Added spew filter to P4 calls to really suppress spew when requested.

The issue here is that CommandUtils.Run only suppresses stdout when NoSpew is specified, so P4 functions like FileExistsInDepot which use p4 calls that write to stderr (e.g fstat of a path not in P4) still end up spewing.

diff --git a/Engine/Source/Programs/AutomationTool/AutomationUtils/P4Utils.cs b/Engine/Source/Programs/AutomationTool/AutomationUtils/P4Utils.cs
--- a/Engine/Source/Programs/AutomationTool/AutomationUtils/P4Utils.cs
+++ b/Engine/Source/Programs/AutomationTool/AutomationUtils/P4Utils.cs
@@ -1190,6 +1190,16 @@ namespace AutomationTool
 			}
 		}
 
+		/// <summary>
+		/// A filter that suppresses all output od stdout/stderr
+		/// </summary>
+		/// <param name="Message"></param>
+		/// <returns></returns>
+		static string NoSpewFilter(string Message)
+		{
+			return null;
+		}
+
 		/// <summary>
 		/// Shortcut to Run but with P4.exe as the program name.
 		/// </summary>
@@ -1197,14 +1207,17 @@ namespace AutomationTool
 		/// <param name="Input">Stdin</param>
 		/// <param name="AllowSpew">true for spew</param>
 		/// <returns>Exit code</returns>
-        public IProcessResult P4(string CommandLine, string Input = null, bool AllowSpew = true, bool WithClient = true, bool SpewIsVerbose = false)
+		public IProcessResult P4(string CommandLine, string Input = null, bool AllowSpew = true, bool WithClient = true, bool SpewIsVerbose = false)
 		{
 			CommandUtils.ERunOptions RunOptions = AllowSpew ? CommandUtils.ERunOptions.AllowSpew : CommandUtils.ERunOptions.NoLoggingOfRunCommand;
 			if( SpewIsVerbose )
 			{
 				RunOptions |= CommandUtils.ERunOptions.SpewIsVerbose;
 			}
-            return CommandUtils.Run(HostPlatform.Current.P4Exe, (WithClient ? GlobalOptions : GlobalOptionsWithoutClient) + CommandLine, Input, Options:RunOptions);
+
+			var SpewDelegate = AllowSpew ? null : new ProcessResult.SpewFilterCallbackType(NoSpewFilter);
+
+			return CommandUtils.Run(HostPlatform.Current.P4Exe, (WithClient ? GlobalOptions : GlobalOptionsWithoutClient) + CommandLine, Input, Options:RunOptions, null, SpewDelegate);
 		}
 
 		/// <summary>
