Fixed BuildTarget UAT script to use existing means of detecting the valid targets for a project rather than trusing the names of the .cs target files.

See https://udn.unrealengine.com/questions/555746/are-game-and-client-build-targets-the-same.html

diff --git a/Engine/Source/Programs/AutomationTool/Scripts/BuildTarget.Automation.cs b/Engine/Source/Programs/AutomationTool/Scripts/BuildTarget.Automation.cs
--- a/Engine/Source/Programs/AutomationTool/Scripts/BuildTarget.Automation.cs
+++ b/Engine/Source/Programs/AutomationTool/Scripts/BuildTarget.Automation.cs
@@ -39,22 +39,17 @@ namespace AutomationTool
 
 		public string UBTArgs { get; set; }
 
-		protected Dictionary<string, string> TargetNames { get; set; }
-
 		public BuildTarget()
 		{
 			Platforms = HostPlatform.Current.HostEditorPlatform.ToString();
 			Configurations = "Development";
 			UBTArgs = "";
-			TargetNames = new Dictionary<string, string>();
 		}
 
 		public override ExitCode Execute()
 		{
 			string[] Arguments = this.Params;
 
-			string[] DefaultUE4Targets = new string[] { "editor", "game", "client", "server" };
-
 			ProjectName = ParseParamValue("project", ProjectName);
 			Targets = ParseParamValue("target", Targets);
 			Platforms = ParseParamValue("platform", Platforms);
@@ -103,7 +98,11 @@ namespace AutomationTool
 			}
 
 			FileReference ProjectFile = null;
-			DirectoryReference TargetSource = DirectoryReference.Combine(BuildCommand.EngineDirectory, "Source");
+
+			// default UE4 targets
+			IEnumerable<string> AvailableTargets = new[] { "editor", "game", "client", "server" };
+
+			bool RequireProjectTargets = false;
 
 			if (!string.IsNullOrEmpty(ProjectName))
 			{
@@ -114,40 +113,61 @@ namespace AutomationTool
 					throw new AutomationException("Unable to find uproject file for {0}", ProjectName);
 				}
 
-				TargetSource = DirectoryReference.Combine(ProjectFile.Directory, "Source");
-			}
-			else
-			{
-				Log.TraceInformation("No project specified, will build vanilla UE4 binaries");
-			}
+				ProjectProperties Properties = ProjectUtils.GetProjectProperties(ProjectFile);
 
-			IEnumerable<string> TargetScripts = DirectoryReference.EnumerateFiles(TargetSource, "*.Target.cs").Select(F => F.GetFileName());
+				if (Properties.bIsCodeBasedProject)
+				{
+					RequireProjectTargets = true;
 
-			foreach (string TargetName in TargetList)
-			{
-				string TargetScript = TargetScripts.Where(S => S.IndexOf(TargetName, StringComparison.OrdinalIgnoreCase) >= 0).FirstOrDefault();
+					AvailableTargets = Properties.Targets.Select(T => T.Rules.Type.ToString());
 
-				if (TargetScript == null && (
-						TargetName.Equals("Client", StringComparison.OrdinalIgnoreCase) ||
-						TargetName.Equals("Game", StringComparison.OrdinalIgnoreCase)
-						)
-					)
-				{
-					// if there's no ProjectGame.Target.cs or ProjectClient.Target.cs then
-					// fallback to Project.Target.cs
-					TargetScript = TargetScripts.Where(S => S.IndexOf(ProjectName + ".", StringComparison.OrdinalIgnoreCase) >= 0).FirstOrDefault();
-				}
+					// go through the list of targets such as Editor, Client, Server etc and replace them with their real target names
+					List<string> ActualTargets = new List<string>();
+					foreach (string Target in TargetList)
+					{
+						// Could be Client, but could also be ShooterClient if that's what was explicitly asked for
+						if (Properties.Targets.Any(T => T.TargetName.Equals(Target, StringComparison.OrdinalIgnoreCase)))
+						{
+							ActualTargets.Add(Target);
+						}
+						else
+						{
+							string ProjectTarget = Properties.Targets.Where(T => T.Rules.Type.ToString().Equals(Target, StringComparison.OrdinalIgnoreCase)).Select(T => T.TargetName).FirstOrDefault();
+
+							if (ProjectTarget == null)
+							{
+								throw new AutomationException("The {0} is not a valid target for {1}", Target, ProjectName);
+							}
+
+							ActualTargets.Add(ProjectTarget);
+						}
+					}
 
-				// IF we found a script for a target, then that is our build target (e.g. QAGame for Game)
-				if (TargetScript != null)
-				{
-					TargetNames[TargetName] = Regex.Replace(TargetScript, ".Target.cs", "", RegexOptions.IgnoreCase);
+					TargetList = ActualTargets;
 				}
-				else
+			}
+
+			if (!RequireProjectTargets)
+			{
+				Log.TraceInformation("Will build vanilla UE4 targets");
+
+				List<string> ActualTargets = new List<string>();
+
+				foreach (string Target in TargetList)
 				{
-					// else just build what they asked for
-					TargetNames[TargetName] = TargetName;
+					// If they asked for editor, client etc then give them the UE version
+					if (AvailableTargets.Contains(Target.ToLower()))
+					{
+						ActualTargets.Add("UE4" + Target);
+					}
+					else
+					{
+						// or just build what they want and let later code figure out if that's valid. E.g. "UnrealPak"
+						ActualTargets.Add(Target);
+					}
 				}
+
+				TargetList = ActualTargets;
 			}
 
 			UE4Build Build = new UE4Build(this);
@@ -168,9 +188,7 @@ namespace AutomationTool
 
 			if (string.IsNullOrEmpty(EditorTarget) == false)
 			{
-				string TargetName = TargetNames[EditorTarget];
-
-				Agenda.AddTarget(TargetName, CurrentPlatform, UnrealTargetConfiguration.Development, ProjectFile, UBTArgs);
+				Agenda.AddTarget(EditorTarget, CurrentPlatform, UnrealTargetConfiguration.Development, ProjectFile, UBTArgs);
 
 				if (!NoTools)
 				{
@@ -184,8 +202,6 @@ namespace AutomationTool
 
 			foreach (string Target in OtherTargets)
 			{
-				string TargetName = TargetNames[Target];
-
 				bool IsServer = Target.EndsWith("Server", StringComparison.OrdinalIgnoreCase);
 
 				IEnumerable<UnrealTargetPlatform> PlatformsToBuild = IsServer ? new UnrealTargetPlatform[] { CurrentPlatform } : PlatformList;
@@ -194,7 +210,7 @@ namespace AutomationTool
 				{
 					foreach (UnrealTargetConfiguration Config in ConfigurationList)
 					{
-						Agenda.AddTarget(TargetName, Platform, Config, ProjectFile, UBTArgs);
+						Agenda.AddTarget(Target, Platform, Config, ProjectFile, UBTArgs);
 					}
 				}
 			}
