Added UAT project parameter validation when using -iostore without -pak.

diff --git a/Engine/Source/Programs/AutomationTool/AutomationUtils/ProjectParams.cs b/Engine/Source/Programs/AutomationTool/AutomationUtils/ProjectParams.cs
--- a/Engine/Source/Programs/AutomationTool/AutomationUtils/ProjectParams.cs
+++ b/Engine/Source/Programs/AutomationTool/AutomationUtils/ProjectParams.cs
@@ -2752,6 +2752,11 @@ namespace AutomationTool
 			{
 				throw new AutomationException("DiscVersion is required for generating remaster package.");
 			}*/
+
+			if ((IoStore && !SkipIoStore) && (!Pak || SkipPak))
+			{
+				throw new AutomationException("-iostore can only be used with -pak");
+			}
 		}
 
 		protected bool bLogged = false;
