Enable I/O dispatcher on dedicated servers. Requires
gc.UseDisregardForGCOnDedicatedServers=True.

diff --git a/Engine/Source/Programs/AutomationTool/Scripts/CopyBuildToStagingDirectory.Automation.cs b/Engine/Source/Programs/AutomationTool/Scripts/CopyBuildToStagingDirectory.Automation.cs
--- a/Engine/Source/Programs/AutomationTool/Scripts/CopyBuildToStagingDirectory.Automation.cs
+++ b/Engine/Source/Programs/AutomationTool/Scripts/CopyBuildToStagingDirectory.Automation.cs
@@ -3029,7 +3029,7 @@ public partial class Project : CommandUtils
 
 	private static bool ShouldCreateIoStoreContainerFiles(ProjectParams Params, DeploymentContext SC)
 	{
-		if (Params.CookOnTheFly || Params.NoClient)
+		if (Params.CookOnTheFly)
 		{
 			return false;
 		}
