Fix command quoting in Linux chmod command

Current code generates: 'chmod +x \"/epic/projects/foo/Intermediate/Staging/foo.sh\"'

This works with xbuild, but fails with msbuild. New code generates:

Run App:sh CommandLine-c 'chmod +x "/epic/projects/foo/Intermediate/Staging/foo.sh"'

Similar to Ben's CL 2106889 and CL 2112830 changes in Engine/Source/Programs/AutomationTool/BuildUtils.cs

diff --git a/Engine/Source/Programs/AutomationTool/Linux/LinuxPlatform.Automation.cs b/Engine/Source/Programs/AutomationTool/Linux/LinuxPlatform.Automation.cs
--- a/Engine/Source/Programs/AutomationTool/Linux/LinuxPlatform.Automation.cs
+++ b/Engine/Source/Programs/AutomationTool/Linux/LinuxPlatform.Automation.cs
@@ -135,7 +135,7 @@ public abstract class BaseLinuxPlatform : Platform
 
 		if (Utils.IsRunningOnMono)
 		{
-			var Result = CommandUtils.Run("sh", string.Format("-c 'chmod +x \\\"{0}\\\"'", IntermediateFile));
+			var Result = CommandUtils.Run("sh", string.Format("-c 'chmod +x \"{0}\"'", IntermediateFile.ToString().Replace("'", "'\"'\"'")));
 			if (Result.ExitCode != 0)
 			{
 				throw new AutomationException(string.Format("Failed to chmod \"{0}\"", IntermediateFile));
