Use the exe name rather then the project name for the bundle path and the booststrap exe need the full name not just the target name.  Commented out the code that was forcing the internal renaming (of exe, icon and info.plist) from target name to project name - I have no idea why that was there but left it in place for now in case there is a permutation I haven't seem - this fixes the display icon not showing correctly.

diff --git a/Engine/Source/Programs/AutomationTool/Mac/MacPlatform.Automation.cs b/Engine/Source/Programs/AutomationTool/Mac/MacPlatform.Automation.cs
--- a/Engine/Source/Programs/AutomationTool/Mac/MacPlatform.Automation.cs
+++ b/Engine/Source/Programs/AutomationTool/Mac/MacPlatform.Automation.cs
@@ -134,7 +134,8 @@ public class MacPlatform : Platform
 						}
 						else if (Params.IsCodeBasedProject)
 						{
-							BootstrapExeName = Target.Receipt.TargetName + ".app";
+							// We want Mac-Shipping etc in the bundle name
+							BootstrapExeName = Path.GetFileName(Executable.Path.FullName) + ".app";
 						}
 						else
 						{
@@ -326,7 +327,7 @@ public class MacPlatform : Platform
 		if (Params.CreateAppBundle)
 		{
 			string ExeName = SC.StageExecutables[0];
-			string BundlePath = CombinePaths(SC.ArchiveDirectory.FullName, SC.ShortProjectName + ".app");
+			string BundlePath = CombinePaths(SC.ArchiveDirectory.FullName, ExeName + ".app");
 
 			if (SC.bIsCombiningMultiplePlatforms)
 			{
@@ -385,6 +386,9 @@ public class MacPlatform : Platform
 				}
 			}
 
+			// This ends up renaming the Bundle, exe and entry in the info.plist to the Project Name rather than desired Target
+			// Commenting this out but leaving it here in the short term in case there is a permutation where its need.
+			/*
 			// Update executable name, icon and entry in Info.plist
 			string UE4GamePath = CombinePaths(BundlePath, "Contents", "MacOS", ExeName);
 			if (ExeName != SC.ShortProjectName && File.Exists(UE4GamePath))
@@ -421,6 +425,7 @@ public class MacPlatform : Platform
 				File.Delete(InfoPlistPath);
 				File.WriteAllText(InfoPlistPath, InfoPlistContents);
 			}
+			*/
 
 			if (!SC.bIsCombiningMultiplePlatforms)
 			{
