Made path to package/cooker order and output directory safe to use with command line for IoStore commandlet.

diff --git a/Engine/Source/Developer/IoStoreUtilities/Private/IoStoreUtilities.cpp b/Engine/Source/Developer/IoStoreUtilities/Private/IoStoreUtilities.cpp
--- a/Engine/Source/Developer/IoStoreUtilities/Private/IoStoreUtilities.cpp
+++ b/Engine/Source/Developer/IoStoreUtilities/Private/IoStoreUtilities.cpp
@@ -3085,15 +3085,22 @@ int32 CreateIoStoreContainerFiles(const TCHAR* CmdLine)
 
 			if (!ParsePakResponseFile(*ResponseFilePath, ContainerSpec.SourceFiles))
 			{
+				UE_LOG(LogIoStore, Error, TEXT("Failed to parse Pak response file '%s'"), *ResponseFilePath);
 				return -1;
 			}
 		}
 
-		UE_LOG(LogIoStore, Display, TEXT("Searching for cooked files..."));
+		UE_LOG(LogIoStore, Display, TEXT("Searching for cooked assets in folder '%s'"), *CookedDirectory);
 		FCookedFileStatMap CookedFileStatMap;
 		FCookedFileVisitor CookedFileVistor(CookedFileStatMap);
 		IFileManager::Get().IterateDirectoryStatRecursively(*CookedDirectory, CookedFileVistor);
-		UE_LOG(LogIoStore, Display, TEXT("Found '%d' files"), CookedFileStatMap.Num());
+		UE_LOG(LogIoStore, Display, TEXT("Found '%d' assets"), CookedFileStatMap.Num());
+
+		if (CookedFileStatMap.Num() == 0)
+		{
+			UE_LOG(LogIoStore, Error, TEXT("Failed find cooked assets in folder '%s'"), *CookedDirectory);
+			return -1;
+		}
 
 		int32 ReturnValue = CreateTarget(*OutputDirectory, *CookedDirectory, Containers, CookedFileStatMap, PackageOrderMap, CookerFileOpenOrderMap);
 		if (ReturnValue != 0)
diff --git a/Engine/Source/Programs/AutomationTool/Scripts/CopyBuildToStagingDirectory.Automation.cs b/Engine/Source/Programs/AutomationTool/Scripts/CopyBuildToStagingDirectory.Automation.cs
--- a/Engine/Source/Programs/AutomationTool/Scripts/CopyBuildToStagingDirectory.Automation.cs
+++ b/Engine/Source/Programs/AutomationTool/Scripts/CopyBuildToStagingDirectory.Automation.cs
@@ -2502,14 +2502,14 @@ public partial class Project : CommandUtils
 	private static void RunIoStore(ProjectParams Params, DeploymentContext SC, string CommandsFileName, FileReference PackageOrderFileLocation, FileReference CookerOpenOrderFileLocation)
 	{
 		FileReference OutputLocation = FileReference.Combine(SC.RuntimeRootDir, SC.RelativeProjectRootForStage.Name);
-		string CommandletParams = String.Format("-OutputDirectory={0} -CookedDirectory={1} -Commands={2}", OutputLocation.FullName, SC.PlatformCookDir, MakePathSafeToUseWithCommandLine(CommandsFileName));
+		string CommandletParams = String.Format("-OutputDirectory={0} -CookedDirectory={1} -Commands={2}", MakePathSafeToUseWithCommandLine(OutputLocation.FullName), MakePathSafeToUseWithCommandLine(SC.PlatformCookDir.ToString()), MakePathSafeToUseWithCommandLine(CommandsFileName));
 		if (PackageOrderFileLocation != null)
 		{
-			CommandletParams += String.Format(" -PackageOrder={0}", PackageOrderFileLocation.FullName);
+			CommandletParams += String.Format(" -PackageOrder={0}", MakePathSafeToUseWithCommandLine(PackageOrderFileLocation.FullName));
 		}
 		if (CookerOpenOrderFileLocation != null)
 		{
-			CommandletParams += String.Format(" -CookerOrder={0}", CookerOpenOrderFileLocation.FullName);
+			CommandletParams += String.Format(" -CookerOrder={0}", MakePathSafeToUseWithCommandLine(CookerOpenOrderFileLocation.FullName));
 		}
 		LogInformation("Running IoStore commandlet with arguments: {0}", CommandletParams);
 		RunCommandlet(SC.RawProjectPath, Params.UE4Exe, "IoStore", CommandletParams);
