Fixed compile warning that I introduced today in IOSPlatform.Automation.cs.

diff --git a/Engine/Source/Programs/AutomationTool/IOS/IOSPlatform.Automation.cs b/Engine/Source/Programs/AutomationTool/IOS/IOSPlatform.Automation.cs
--- a/Engine/Source/Programs/AutomationTool/IOS/IOSPlatform.Automation.cs
+++ b/Engine/Source/Programs/AutomationTool/IOS/IOSPlatform.Automation.cs
@@ -952,7 +952,6 @@ public class IOSPlatform : Platform
 
 		bool bLaunchscreenStoryboard = false;
 		ConfigHierarchy PlatformGameConfig;
-		bool bXCArchive = false;
 		if (Params.EngineConfigs.TryGetValue(SC.StageTargetPlatform.PlatformType, out PlatformGameConfig))
 		{
 			PlatformGameConfig.GetBool("/Script/IOSRuntimeSettings.IOSRuntimeSettings", "bLaunchscreenStoryboard", out bLaunchscreenStoryboard);
