BenchmarkBuild changes to explore DDC optimizations further.

Added NoDDC option for PIE tests.

ColdDDC option no longer nukes the local DDC, instead a temporary path is used.

Reworked PIE / Cook tests to rely more on a common base class functionality.

#tests ran BenchmarkBuild with different options

diff --git a/Engine/Source/Programs/AutomationTool/Scripts/BenchmarkBuild/BenchmarkBuild.Automation.cs b/Engine/Source/Programs/AutomationTool/Scripts/BenchmarkBuild/BenchmarkBuild.Automation.cs
--- a/Engine/Source/Programs/AutomationTool/Scripts/BenchmarkBuild/BenchmarkBuild.Automation.cs
+++ b/Engine/Source/Programs/AutomationTool/Scripts/BenchmarkBuild/BenchmarkBuild.Automation.cs
@@ -37,6 +37,7 @@ namespace AutomationTool.Benchmark
 	[Help("pie", "Launch the editor (only valid when -project is specified")]
 	[Help("pie-cold", "Launch the editor with a cold ddc")]
 	[Help("pie-hot", "Launch the editor with a hot ddc")]
+	[Help("pie-hot", "Launch the editor with a hot ddc")]
 	[Help("iterations=<n>", "How many times to perform each test)")]
 	[Help("wait=<n>", "How many seconds to wait between each test)")]
 	[Help("filename", "Name/path of file to write CSV results to. If empty the local machine name will be used")]
@@ -73,6 +74,7 @@ namespace AutomationTool.Benchmark
 			public bool DoPIETests = false;
 			public bool DoHotPIETests = false;
 			public bool DoColdPIETests = false;
+			public bool DoNoDDCPIETests = false;
 
 			// list of maps for each project
 			public IEnumerable<string> StartupMapList = Enumerable.Empty<string>();
@@ -84,8 +86,6 @@ namespace AutomationTool.Benchmark
 			public string ExtraCookArgs = "";
 			public string FileName = string.Format("{0}_Results.csv", Environment.MachineName);
 
-			public bool DoHintTests = false;
-
 			public void ParseParams(string[] InParams)
 			{
 				this.Params = InParams;
@@ -115,6 +115,7 @@ namespace AutomationTool.Benchmark
 				DoPIETests = AllThings | ParseParam("pie");
 				DoColdPIETests = AllThings | ParseParam("pie-cold");
 				DoHotPIETests = AllThings | ParseParam("pie-hot");
+				DoNoDDCPIETests = AllThings | ParseParam("pie-noddc");
 
 				// sanity
 				DoAcceleratedCompileTests = DoAcceleratedCompileTests && BenchmarkBuildTask.SupportsAcceleration;
@@ -125,8 +126,6 @@ namespace AutomationTool.Benchmark
 
 				FileName = ParseParamValue("filename", FileName);
 
-				DoHintTests = ParseParam("hints");
-
 				ExtraCookArgs = ParseParamValue("ExtraCookArgs", "");
 
 				// Parse the project arg
@@ -442,31 +441,20 @@ namespace AutomationTool.Benchmark
 
 			List<BenchmarkTaskBase> NewTasks = new List<BenchmarkTaskBase>();
 
-			bool HaveUE4Project = InOptions.ProjectsToTest.Where(P => P.ToLower() == "ue4").Any();
-			int ProjectIndex = InOptions.ProjectsToTest.ToList().FindIndex(P => P == InProjectName);
-
-			int MapIndex = HaveUE4Project ? ProjectIndex - 1 : ProjectIndex;
-			string ProjectMap = MapIndex < InOptions.StartupMapList.Count() ? InOptions.StartupMapList.ElementAt(ProjectIndex) : "";
-
-			if (ProjectMap.Equals("none", StringComparison.OrdinalIgnoreCase))
-			{
-				ProjectMap = "";
-			}
-
 			if (InOptions.DoPIETests)
 			{
 				EditorTaskOptions Options = InOptions.DoHotPIETests ? EditorTaskOptions.HotDDC : EditorTaskOptions.None;
-				NewTasks.Add(new BenchmarkRunEditorTask(InProjectName, Options, "-recordhints"));
+				NewTasks.Add(new BenchmarkRunEditorTask(InProjectName, Options));
 			}
 
 			if (InOptions.DoColdPIETests)
 			{
 				NewTasks.Add(new BenchmarkRunEditorTask(InProjectName, EditorTaskOptions.None | EditorTaskOptions.ColdDDC));
+			}
 
-				if (InOptions.DoHintTests)
-				{
-					NewTasks.Add(new BenchmarkRunEditorTask(InProjectName, EditorTaskOptions.None | EditorTaskOptions.ColdDDC, "-usehints"));
-				}
+			if (InOptions.DoNoDDCPIETests)
+			{
+				NewTasks.Add(new BenchmarkRunEditorTask(InProjectName, EditorTaskOptions.None | EditorTaskOptions.NoDDC));
 			}
 
 			return NewTasks;
diff --git a/Engine/Source/Programs/AutomationTool/Scripts/BenchmarkBuild/BenchmarkCookTask.cs b/Engine/Source/Programs/AutomationTool/Scripts/BenchmarkBuild/BenchmarkCookTask.cs
--- a/Engine/Source/Programs/AutomationTool/Scripts/BenchmarkBuild/BenchmarkCookTask.cs
+++ b/Engine/Source/Programs/AutomationTool/Scripts/BenchmarkBuild/BenchmarkCookTask.cs
@@ -16,24 +16,15 @@ namespace AutomationTool.Benchmark
 
 	class BenchmarkCookTask : BenchmarkEditorTaskBase
 	{
-		string			ProjectName;
-
-		FileReference	ProjectFile;
-
 		string			CookPlatformName;
 
-		EditorTaskOptions		Options;
-
 		string			CookArgs;
 
 		public BenchmarkCookTask(string InProject, UnrealTargetPlatform InPlatform, EditorTaskOptions InOptions, string InCookArgs)
+			: base(InProject, InOptions)
 		{
-			ProjectName = InProject;
-			Options = InOptions;
 			CookArgs = InCookArgs;
 
-			ProjectFile = ProjectUtils.FindProjectFileFromName(ProjectName);
-
 			var PlatformToCookPlatform = new Dictionary<UnrealTargetPlatform, string> {
 				{ UnrealTargetPlatform.Win64, "WindowsClient" },
 				{ UnrealTargetPlatform.Mac, "MacClient" },
@@ -50,17 +41,17 @@ namespace AutomationTool.Benchmark
 
 			TaskName = string.Format("Cook {0} {1}", InProject, CookPlatformName);
 
-			if (Options.HasFlag(EditorTaskOptions.NoDDC))
+			if (TaskOptions.HasFlag(EditorTaskOptions.NoDDC))
 			{
 				TaskModifiers.Add("noddc");
 			}
 
-			if (Options.HasFlag(EditorTaskOptions.NoShaderDDC))
+			if (TaskOptions.HasFlag(EditorTaskOptions.NoShaderDDC))
 			{
 				TaskModifiers.Add("noshaderddc");
 			}
 
-			if (Options.HasFlag(EditorTaskOptions.ColdDDC))
+			if (TaskOptions.HasFlag(EditorTaskOptions.ColdDDC))
 			{
 				TaskModifiers.Add("coldddc");
 			}
@@ -85,17 +76,19 @@ namespace AutomationTool.Benchmark
 			}
 
 			// Do a cook to make sure the remote ddc is warm?
-			if (Options.HasFlag(EditorTaskOptions.HotDDC))
+			if (TaskOptions.HasFlag(EditorTaskOptions.HotDDC))
 			{
 				// will throw an exception if it fails
 				CommandUtils.RunCommandlet(ProjectFile, "UE4Editor-Cmd.exe", "Cook", String.Format("-TargetPlatform={0} ", CookPlatformName));
 			}
 
-			if (Options.HasFlag(EditorTaskOptions.ColdDDC))
+			if (TaskOptions.HasFlag(EditorTaskOptions.ColdDDC))
 			{
-				DeleteLocalDDC(ProjectFile);
+				//DeleteLocalDDC(ProjectFile);
 			}
 
+			base.PerformPrequisites();
+
 			return true;
 		}
 
@@ -103,18 +96,18 @@ namespace AutomationTool.Benchmark
 		{
 			List<string> ExtraArgsList = new List<string>();
 			
-			if (Options.HasFlag(EditorTaskOptions.CookClient))
+			if (TaskOptions.HasFlag(EditorTaskOptions.CookClient))
 			{
 				ExtraArgsList.Add("client");
 			}
 
-			if (Options.HasFlag(EditorTaskOptions.NoShaderDDC))
+			if (TaskOptions.HasFlag(EditorTaskOptions.NoShaderDDC))
 			{
 				ExtraArgsList.Add("noshaderddc");
 				//ExtraArgs.Add("noxgeshadercompile");
 			}
 
-			if (Options.HasFlag(EditorTaskOptions.NoDDC))
+			if (TaskOptions.HasFlag(EditorTaskOptions.NoDDC))
 			{
 				ExtraArgsList.Add("ddc=noshared");
 			}
diff --git a/Engine/Source/Programs/AutomationTool/Scripts/BenchmarkBuild/BenchmarkRunEditorTask.cs b/Engine/Source/Programs/AutomationTool/Scripts/BenchmarkBuild/BenchmarkRunEditorTask.cs
--- a/Engine/Source/Programs/AutomationTool/Scripts/BenchmarkBuild/BenchmarkRunEditorTask.cs
+++ b/Engine/Source/Programs/AutomationTool/Scripts/BenchmarkBuild/BenchmarkRunEditorTask.cs
@@ -13,35 +13,25 @@ namespace AutomationTool.Benchmark
 {
 	class BenchmarkRunEditorTask : BenchmarkEditorTaskBase
 	{
-		FileReference ProjectFile = null;
-
-		string Project = "";
 		string EditorArgs = "";
 
-		EditorTaskOptions Options;
-
-		public BenchmarkRunEditorTask(string InProject, EditorTaskOptions InOptions = EditorTaskOptions.None, string InEditorArgs="")
+		public BenchmarkRunEditorTask(string InProject, EditorTaskOptions InOptions, string InEditorArgs="")
+			: base(InProject, InOptions)
 		{
-			Options = InOptions;
 			EditorArgs = InEditorArgs;
 
-			if (!InProject.Equals("UE4", StringComparison.OrdinalIgnoreCase))
-			{
-				Project = InProject;
-				ProjectFile = ProjectUtils.FindProjectFileFromName(InProject);
-			}
 
-			if (Options.HasFlag(EditorTaskOptions.NoDDC))
+			if (TaskOptions.HasFlag(EditorTaskOptions.NoDDC))
 			{
 				TaskModifiers.Add("noddc");
 			}
 
-			if (Options.HasFlag(EditorTaskOptions.NoShaderDDC))
+			if (TaskOptions.HasFlag(EditorTaskOptions.NoShaderDDC))
 			{
 				TaskModifiers.Add("noshaderddc");
 			}
 
-			if (Options.HasFlag(EditorTaskOptions.ColdDDC))
+			if (TaskOptions.HasFlag(EditorTaskOptions.ColdDDC))
 			{
 				TaskModifiers.Add("coldddc");
 			}
@@ -76,31 +66,35 @@ namespace AutomationTool.Benchmark
 		{
 			// build editor
 			BuildTarget Command = new BuildTarget();
-			Command.ProjectName = Project;
+			Command.ProjectName = ProjectFile != null ? ProjectFile.GetFileNameWithoutAnyExtensions() : null;
 			Command.Platforms = BuildHostPlatform.Current.Platform.ToString();
 			Command.Targets = "Editor";
-			Command.NoTools = true;
+			Command.NoTools = false;
 
 			if (Command.Execute() != ExitCode.Success)
 			{
 				return false;
 
 			}
-			if (Options.HasFlag(EditorTaskOptions.ColdDDC))
+
+			if (TaskOptions.HasFlag(EditorTaskOptions.ColdDDC))
 			{
-				DeleteLocalDDC(ProjectFile);
+				//DeleteLocalDDC(ProjectFile);
 			}
 
 			// if they want a hot DDC then do the test one time with no timing
-			if (Options.HasFlag(EditorTaskOptions.HotDDC))
+			if (TaskOptions.HasFlag(EditorTaskOptions.HotDDC))
 			{
 				RunEditorAndWaitForMapLoad();
 			}
 
+			base.PerformPrequisites();
+
 			return true;
 		}
 
 		static IProcessResult CurrentProcess = null;
+		static DateTime LastStdOutTime = DateTime.Now;
 		//static bool TestCompleted = false;
 
 		/// <summary>
@@ -112,10 +106,15 @@ namespace AutomationTool.Benchmark
 		{
 			if (CurrentProcess != null)
 			{
-				if (Message.Contains("TEST COMPLETE"))
+				lock (CurrentProcess)
 				{
-					Log.TraceInformation("Automation test reported as complete.");
-					//TestCompleted = true;
+					if (Message.Contains("TEST COMPLETE"))
+					{
+						Log.TraceInformation("Automation test reported as complete.");
+						//TestCompleted = true;
+					}
+
+					LastStdOutTime = DateTime.Now;
 				}
 			}
 			return Message;
@@ -127,11 +126,17 @@ namespace AutomationTool.Benchmark
 			string EditorPath = HostPlatform.Current.GetUE4ExePath("UE4Editor.exe");
 			string Arguments = string.Format("{0} {1} -execcmds=\"automation runtest System.Maps.PIE;Quit\" -stdout -AllowStdOutLogVerbosity -unattended", ProjectArg, EditorArgs);
 
+			if (TaskOptions.HasFlag(EditorTaskOptions.NoDDC))
+			{
+				Arguments += (" -ddc=noshared");
+			}
+
 			var RunOptions = CommandUtils.ERunOptions.AllowSpew | CommandUtils.ERunOptions.NoWaitForExit;
 
 			var SpewDelegate = new ProcessResult.SpewFilterCallbackType(EndOnMapCheckFilter);
 
 			//TestCompleted = false;
+			LastStdOutTime = DateTime.Now;
 			CurrentProcess = CommandUtils.Run(EditorPath, Arguments, Options: RunOptions, SpewFilterCallback: SpewDelegate);
 
 			DateTime StartTime = DateTime.Now;
@@ -142,10 +147,13 @@ namespace AutomationTool.Benchmark
 			{
 				Thread.Sleep(5 * 1000);
 
-				if ((DateTime.Now - StartTime).TotalMinutes >= MaxWaitMins)
+				lock (CurrentProcess)
 				{
-					Log.TraceError("Gave up waiting for task after {0} minutes", MaxWaitMins);
-					CurrentProcess.ProcessObject.Kill();
+					if ((LastStdOutTime - StartTime).TotalMinutes >= MaxWaitMins)
+					{
+						Log.TraceError("Gave up waiting for task after {0} minutes of no output", MaxWaitMins);
+						CurrentProcess.ProcessObject.Kill();
+					}
 				}
 			}
 
diff --git a/Engine/Source/Programs/AutomationTool/Scripts/BenchmarkBuild/BenchmarkTaskBase.cs b/Engine/Source/Programs/AutomationTool/Scripts/BenchmarkBuild/BenchmarkTaskBase.cs
--- a/Engine/Source/Programs/AutomationTool/Scripts/BenchmarkBuild/BenchmarkTaskBase.cs
+++ b/Engine/Source/Programs/AutomationTool/Scripts/BenchmarkBuild/BenchmarkTaskBase.cs
@@ -7,6 +7,7 @@ using System.Linq;
 using System.Text;
 using System.Threading.Tasks;
 using Tools.DotNETCommon;
+using UnrealBuildTool;
 
 namespace AutomationTool.Benchmark
 {
@@ -45,6 +46,11 @@ namespace AutomationTool.Benchmark
 		/// </summary>
 		virtual protected bool PerformPrequisites() { return true;  }
 
+		/// <summary>
+		/// Perform post-task cleanup
+		/// </summary>
+		virtual protected void PerformCleanup() { }
+
 		/// <summary>
 		/// Perform the actual task that is measured
 		/// </summary>
@@ -101,6 +107,15 @@ namespace AutomationTool.Benchmark
 			{
 				Log.TraceError("{0} failed. {1}", GetFullTaskName(), FailureString);
 			}
+
+			try
+			{
+				PerformCleanup();
+			}
+			catch (Exception Ex)
+			{
+				Log.TraceError("Cleanup of {0} failed. {1}", GetFullTaskName(), Ex);
+			}
 		}
 
 		/// <summary>
@@ -148,7 +163,29 @@ namespace AutomationTool.Benchmark
 
 	abstract class BenchmarkEditorTaskBase : BenchmarkTaskBase
 	{
-		protected void DeleteLocalDDC(FileReference InProjectFile)
+		protected EditorTaskOptions TaskOptions;
+
+		protected FileReference ProjectFile = null;
+
+		protected string ProjectName
+		{
+			get
+			{
+				return ProjectFile == null ? null : ProjectFile.GetFileNameWithoutAnyExtensions();
+			}
+		}
+
+		protected BenchmarkEditorTaskBase(string InProject, EditorTaskOptions InTaskOptions)
+		{
+			TaskOptions = InTaskOptions;
+
+			if (!InProject.Equals("UE4", StringComparison.OrdinalIgnoreCase))
+			{
+				ProjectFile = ProjectUtils.FindProjectFileFromName(InProject);
+			}
+		}
+
+		private void DeleteLocalDDC(FileReference InProjectFile)
 		{
 			List<DirectoryReference> DirsToClear = new List<DirectoryReference>();
 
@@ -181,5 +218,59 @@ namespace AutomationTool.Benchmark
 				}
 			}
 		}
+
+		private DirectoryReference PrivateDDCPath = null;
+		private string OldLocalDCCPath = null;
+
+		private string GetLocalDDCPathEnvironmentKey()
+		{
+			string EnvVar = "UE-LocalDataCachePath";
+
+			if (BuildHostPlatform.Current.Platform != UnrealTargetPlatform.Win64)
+			{
+				EnvVar = EnvVar.Replace("-", "_");
+			}
+
+			return EnvVar;
+		}
+
+		protected override bool PerformPrequisites()
+		{
+			if (TaskOptions.HasFlag(EditorTaskOptions.ColdDDC))
+			{
+				PrivateDDCPath = DirectoryReference.Combine(CommandUtils.EngineDirectory, "BenchmarkDerivedDataCache");
+
+				string Key = GetLocalDDCPathEnvironmentKey();
+
+				OldLocalDCCPath = Environment.GetEnvironmentVariable(Key);
+				Environment.SetEnvironmentVariable(Key, PrivateDDCPath.FullName);
+
+				DirectoryReference BootDDC = DirectoryReference.Combine(ProjectFile.Directory, "DerivedDataCache");
+
+				try
+				{
+					DirectoryReference.Delete(BootDDC, true);
+				}
+				catch
+				{
+
+				}
+			}
+
+			return base.PerformPrequisites();
+		}
+
+		protected override void PerformCleanup()
+		{
+			if (PrivateDDCPath !=  null)
+			{
+				if (DirectoryReference.Exists(PrivateDDCPath))
+				{
+					DirectoryReference.Delete(PrivateDDCPath, true);
+				}
+
+				Environment.SetEnvironmentVariable(GetLocalDDCPathEnvironmentKey(), OldLocalDCCPath);
+			}
+		}
 	}
 }
