BenchmarkTest - fixed intermediate results not being cleared from a cancelled test

diff --git a/Engine/Source/Programs/AutomationTool/Scripts/BenchmarkBuild/BenchmarkTaskBase.cs b/Engine/Source/Programs/AutomationTool/Scripts/BenchmarkBuild/BenchmarkTaskBase.cs
--- a/Engine/Source/Programs/AutomationTool/Scripts/BenchmarkBuild/BenchmarkTaskBase.cs
+++ b/Engine/Source/Programs/AutomationTool/Scripts/BenchmarkBuild/BenchmarkTaskBase.cs
@@ -240,6 +240,11 @@ namespace AutomationTool.Benchmark
 			{
 				PrivateDDCPath = DirectoryReference.Combine(CommandUtils.EngineDirectory, "BenchmarkDerivedDataCache");
 
+				if (DirectoryReference.Exists(PrivateDDCPath))
+				{
+					DirectoryReference.Delete(PrivateDDCPath, true);
+				}
+
 				string Key = GetLocalDDCPathEnvironmentKey();
 
 				OldLocalDCCPath = Environment.GetEnvironmentVariable(Key);
