Add support to BuildCMakeLib to allow a fully overriden source directory.

- This is useful when the source for a ThirdParty lib is not contained in the engine.

diff --git a/Engine/Source/Programs/AutomationTool/Scripts/BuildCMakeLib.Automation.cs b/Engine/Source/Programs/AutomationTool/Scripts/BuildCMakeLib.Automation.cs
--- a/Engine/Source/Programs/AutomationTool/Scripts/BuildCMakeLib.Automation.cs
+++ b/Engine/Source/Programs/AutomationTool/Scripts/BuildCMakeLib.Automation.cs
@@ -18,6 +18,7 @@ using UnrealBuildTool;
 [Help("Builds a third party library using the CMake build system.")]
 [Help("TargetLib", "Specify the target library to build.")]
 [Help("TargetLibVersion", "Specify the target library version to build.")]
+[Help("TargetLibSourcePath", "Override the path to source, if external to the engine. (eg. -TargetLibSourcePath=path). Default is empty.")]
 [Help("TargetPlatform", "Specify the name of the target platform to build (eg. -TargetPlatform=IOS).")]
 [Help("TargetConfigs", "Specify a list of configurations to build, separated by '+' characters (eg. -TargetConfigs=release+debug). Default is release+debug.")]
 [Help("LibOutputPath", "Override the path to output the libs to. (eg. -LibOutputPath=lib). Default is empty.")]
@@ -35,6 +36,7 @@ public sealed class BuildCMakeLib : BuildCommand
 	{
 		public string Name = "";
 		public string Version = "";
+		public string SourcePath = "";
 		public string LibOutputPath = "";
 		public string CMakeProjectIncludeFile = "";
 		public string CMakeAdditionalArguments = "";
@@ -53,7 +55,18 @@ public sealed class BuildCMakeLib : BuildCommand
 		};
 
 		public static DirectoryReference ThirdPartySourceDirectory = DirectoryReference.Combine(RootDirectory, "Engine", "Source", "ThirdParty");
-		public DirectoryReference LibSourceDirectory { get { return DirectoryReference.Combine(ThirdPartySourceDirectory, Name, Version); } }
+
+		public DirectoryReference GetLibSourceDirectory() 
+		{ 
+			if (string.IsNullOrEmpty(SourcePath))
+			{ 
+				return DirectoryReference.Combine(ThirdPartySourceDirectory, Name, Version); 
+			}
+			else 
+			{
+				return new DirectoryReference(SourcePath);
+			} 
+		}
 
 		public override string ToString() => Name;
 	}
@@ -129,7 +142,7 @@ public sealed class BuildCMakeLib : BuildCommand
 				if (!FileReference.Exists(FileReference.Combine(CMakeDirectory, "CMakeLists.txt")))
 				{
 					// If not available then check the lib source root
-					CMakeDirectory = DirectoryReference.Combine(GetTargetLibBaseRootDirectory(TargetLib));
+					CMakeDirectory = TargetLib.GetLibSourceDirectory();
 				}
 			}
 
@@ -247,7 +260,7 @@ public sealed class BuildCMakeLib : BuildCommand
 			string AdditionalArgs = GetAdditionalCMakeArguments(TargetLib, TargetConfiguration);
 			if (AdditionalArgs != null)
 			{
-				Args += AdditionalArgs;
+				Args += AdditionalArgs.Replace("${TARGET_CONFIG}", TargetConfiguration ?? "");
 			}
 
 			return Args;
@@ -439,6 +452,7 @@ public sealed class BuildCMakeLib : BuildCommand
 
 		TargetLib.Name = ParseParamValue("TargetLib", "");
 		TargetLib.Version = ParseParamValue("TargetLibVersion", "");
+		TargetLib.SourcePath = ParseParamValue("TargetLibSourcePath", "");
 		TargetLib.LibOutputPath = ParseParamValue("LibOutputPath", "");
 		TargetLib.CMakeProjectIncludeFile = ParseParamValue("CMakeProjectIncludeFile", "");
 		TargetLib.CMakeAdditionalArguments = ParseParamValue("CMakeAdditionalArguments", "");
