[NO NON-PII COMMIT MESSAGE]

diff --git a/Engine/Source/Programs/AutomationTool/Scripts/BuildPluginCommand.Automation.cs b/Engine/Source/Programs/AutomationTool/Scripts/BuildPluginCommand.Automation.cs
--- a/Engine/Source/Programs/AutomationTool/Scripts/BuildPluginCommand.Automation.cs
+++ b/Engine/Source/Programs/AutomationTool/Scripts/BuildPluginCommand.Automation.cs
@@ -20,6 +20,7 @@ using System.Text;
 class BuildPlugin : BuildCommand
 {
 	const string AndroidArchitectures = "armv7+arm64";
+	const string HoloLensArchitecture = "arm64";
 
 	public override void ExecuteBuild()
 	{
@@ -203,7 +204,12 @@ class BuildPlugin : BuildCommand
 			{
 				Arguments += String.Format(" -architectures={0}", AndroidArchitectures);
 			}
-			if(!String.IsNullOrEmpty(InAdditionalArgs))
+			else if (Platform == UnrealTargetPlatform.HoloLens)
+			{
+				Arguments += String.Format(" -Architecture={0}", HoloLensArchitecture);
+			}
+
+			if (!String.IsNullOrEmpty(InAdditionalArgs))
 			{
 				Arguments += InAdditionalArgs;
 			}
