Fixed option typo in benchmark script

diff --git a/Engine/Source/Programs/AutomationTool/Scripts/BenchmarkBuild/BenchmarkBuild.Automation.cs b/Engine/Source/Programs/AutomationTool/Scripts/BenchmarkBuild/BenchmarkBuild.Automation.cs
--- a/Engine/Source/Programs/AutomationTool/Scripts/BenchmarkBuild/BenchmarkBuild.Automation.cs
+++ b/Engine/Source/Programs/AutomationTool/Scripts/BenchmarkBuild/BenchmarkBuild.Automation.cs
@@ -484,13 +484,13 @@ namespace AutomationTool.Benchmark
 				}
 
 				// no shaders in the ddc
-				if (InOptions.DDCOptions.HasFlag(DDCTaskOptions.ColdDDC))
+				if (InOptions.DDCOptions.HasFlag(DDCTaskOptions.NoShaderDDC))
 				{
 					NewTasks.Add(new BenchmarkRunEditorTask(InProjectFile, DDCTaskOptions.NoShaderDDC, DefaultExtraArgs));
 				}
 
 				// no ddc!
-				if (InOptions.DDCOptions.HasFlag(DDCTaskOptions.ColdDDC))
+				if (InOptions.DDCOptions.HasFlag(DDCTaskOptions.NoDDC))
 				{
 					NewTasks.Add(new BenchmarkRunEditorTask(InProjectFile, DDCTaskOptions.NoDDC, DefaultExtraArgs));
 				}
