Fixed issue with BuildTarget script needlessly recompiling UBT

diff --git a/Engine/Source/Programs/AutomationTool/Scripts/BuildTarget.Automation.cs b/Engine/Source/Programs/AutomationTool/Scripts/BuildTarget.Automation.cs
--- a/Engine/Source/Programs/AutomationTool/Scripts/BuildTarget.Automation.cs
+++ b/Engine/Source/Programs/AutomationTool/Scripts/BuildTarget.Automation.cs
@@ -206,10 +206,10 @@ namespace AutomationTool
 				if (!NoTools)
 				{
 					Agenda.AddTarget("UnrealPak", CurrentPlatform, UnrealTargetConfiguration.Development, ProjectFile, UBTArgs);
-					Agenda.AddTarget("ShaderCompileWorker", CurrentPlatform, UnrealTargetConfiguration.Development, null, UBTArgs);
-					Agenda.AddTarget("UnrealLightmass", CurrentPlatform, UnrealTargetConfiguration.Development, null, UBTArgs);
-					Agenda.AddTarget("CrashReportClient", CurrentPlatform, UnrealTargetConfiguration.Shipping, null, UBTArgs);
-					Agenda.AddTarget("CrashReportClientEditor", CurrentPlatform, UnrealTargetConfiguration.Shipping, null, UBTArgs);
+					Agenda.AddTarget("ShaderCompileWorker", CurrentPlatform, UnrealTargetConfiguration.Development, ProjectFile, UBTArgs);
+					Agenda.AddTarget("UnrealLightmass", CurrentPlatform, UnrealTargetConfiguration.Development, ProjectFile, UBTArgs);
+					Agenda.AddTarget("CrashReportClient", CurrentPlatform, UnrealTargetConfiguration.Shipping, ProjectFile, UBTArgs);
+					Agenda.AddTarget("CrashReportClientEditor", CurrentPlatform, UnrealTargetConfiguration.Shipping, ProjectFile, UBTArgs);
 				}
 			}
 
