Benchmark script -

Added option to select map to PIE to
Fixed issue with clients being compiled when not specified

diff --git a/Engine/Source/Programs/AutomationTool/Scripts/BenchmarkBuild/BenchmarkBuild.Automation.cs b/Engine/Source/Programs/AutomationTool/Scripts/BenchmarkBuild/BenchmarkBuild.Automation.cs
--- a/Engine/Source/Programs/AutomationTool/Scripts/BenchmarkBuild/BenchmarkBuild.Automation.cs
+++ b/Engine/Source/Programs/AutomationTool/Scripts/BenchmarkBuild/BenchmarkBuild.Automation.cs
@@ -32,11 +32,12 @@ namespace AutomationTool.Benchmark
 	[Help("cores=X+Y+Z", "Do noxge builds with these processor counts (default is Environment.ProcessorCount)")]
 	[Help("cook", "Do a cook for the specified platform")]
 	[Help("pie", "Launch the editor (only valid when -project is specified")]
+	[Help("map", "Map to PIE with (only valid when using a single project")]
 	[Help("warmddc", "Cook / PIE with a warm DDC")]
 	[Help("hotddc", "Cook / PIE with a hot local DDC (an untimed pre-run is performed)")]
 	[Help("coldddc", "Cook / PIE with a cold local DDC (a temporary folder is used)")]
 	[Help("noshaderddc", "Cook / PIE with no shaders in the DDC")]
-	[Help("noddc", "Cook / PIE with no ddc (local or shared)")]
+	[Help("nosharedddc", "Cook / PIE with no shared ddc")]
 	[Help("iterations=<n>", "How many times to perform each test)")]
 	[Help("wait=<n>", "How many seconds to wait between each test)")]
 	[Help("filename", "Name/path of file to write CSV results to. If empty the local machine name will be used")]
@@ -69,6 +70,8 @@ namespace AutomationTool.Benchmark
 			// editor startup tests
 			public bool DoPIETests = false;
 
+			public string MapName = "";
+
 			// misc
 			public int Iterations = 1;
 			public bool NoClean = false;
@@ -102,12 +105,13 @@ namespace AutomationTool.Benchmark
 				
 				// editor startup tests
 				DoPIETests = AllThings | ParseParam("pie");
+				MapName = ParseParamValue("map", "");
 
 				// DDC options
 				DDCOptions |= ParseParam("warmddc") ? DDCTaskOptions.WarmDDC : DDCTaskOptions.None;
 				DDCOptions |= ParseParam("hotddc") ? DDCTaskOptions.HotDDC : DDCTaskOptions.None;
 				DDCOptions |= ParseParam("coldddc") ? DDCTaskOptions.ColdDDC : DDCTaskOptions.None;
-				DDCOptions |= ParseParam("noddc") ? DDCTaskOptions.NoDDC : DDCTaskOptions.None;
+				DDCOptions |= ParseParam("nosharedddc") ? DDCTaskOptions.NoSharedDDC : DDCTaskOptions.None;
 				DDCOptions |= ParseParam("noshaderddc") ? DDCTaskOptions.NoShaderDDC : DDCTaskOptions.None;
 
 				// sanity
@@ -208,18 +212,27 @@ namespace AutomationTool.Benchmark
 				}
 
 				// do startup tests
-				Tasks.AddRange(AddPIETests(ProjectFile, Options));
+				if (Options.DoPIETests)
+				{
+					Tasks.AddRange(AddPIETests(ProjectFile, Options));
+				}
 
 				foreach (var ClientPlatform in Options.PlatformsToTest)
 				{
 					// build a client if the project supports it
 					string TargetName = ProjectSupportsClientBuild(ProjectFile) ? "Client" : "Game";
 
-					// do build tests
-					Tasks.AddRange(AddBuildTests(ProjectFile, ClientPlatform, TargetName, Options));
+					if (Options.DoBuildClientTests)
+					{
+						// do build tests
+						Tasks.AddRange(AddBuildTests(ProjectFile, ClientPlatform, TargetName, Options));
+					}
 
 					// do cook tests
-					Tasks.AddRange(AddCookTests(ProjectFile, ClientPlatform, Options));
+					if (Options.DoCookTests)
+					{
+						Tasks.AddRange(AddCookTests(ProjectFile, ClientPlatform, Options));
+					}
 				}
 			}
 
@@ -383,9 +396,9 @@ namespace AutomationTool.Benchmark
 					NewTasks.Add(new BenchmarkCookTask(InProjectFile, InPlatform, CookClient, DDCTaskOptions.NoShaderDDC, InOptions.CookArgs));
 				}
 
-				if (InOptions.DDCOptions.HasFlag(DDCTaskOptions.NoDDC))
+				if (InOptions.DDCOptions.HasFlag(DDCTaskOptions.NoSharedDDC))
 				{
-					NewTasks.Add(new BenchmarkCookTask(InProjectFile, InPlatform, CookClient, DDCTaskOptions.NoDDC, InOptions.CookArgs));
+					NewTasks.Add(new BenchmarkCookTask(InProjectFile, InPlatform, CookClient, DDCTaskOptions.NoSharedDDC, InOptions.CookArgs));
 				}
 			}
 
@@ -403,47 +416,17 @@ namespace AutomationTool.Benchmark
 
 			string DefaultExtraArgs = InOptions.PIEArgs;
 
-			// specify a map to PIE (only works with one project on the command line)
-			string MapName = ParseParamValue("piemap", "");
-			if (!string.IsNullOrEmpty(MapName))
+			if (!string.IsNullOrEmpty(InOptions.MapName))
 			{
-				DefaultExtraArgs += " -map=" + MapName;
+				DefaultExtraArgs += " -map=" + InOptions.MapName;
 			}
-						
-			// TEMP for testing DDC perf
-			bool DoHints = ParseParam("hints");
-			bool DoMinSharedFetch = ParseParam("ddcminfetch");
-			bool DoFakeMisses = ParseParam("ddcfakemisses");
-
 			
 			if (InOptions.DoPIETests)
 			{
 				// if no options assume warm
 				if (InOptions.DDCOptions == DDCTaskOptions.None || InOptions.DDCOptions.HasFlag(DDCTaskOptions.WarmDDC))
 				{
-					// TEMP TO RECORD DDC HINTS!
-					string FinalArgs = DefaultExtraArgs;
-					if (DoHints)
-					{
-						FinalArgs += " -recordhints";
-					}
-					// END TEMP
-
-					NewTasks.Add(new BenchmarkRunEditorTask(InProjectFile, DDCTaskOptions.WarmDDC, FinalArgs));
-
-					// TEMP to test misses
-					if (DoFakeMisses)
-					{
-						IEnumerable<int> MissPercent = new int[] { 5, 10 };
-
-						foreach (int Size in MissPercent)
-						{
-							string Args = string.Format("{0} -DDCLOCALMISSPERCENT={1}", FinalArgs, Size);
-
-							NewTasks.Add(new BenchmarkRunEditorTask(InProjectFile, DDCTaskOptions.None | DDCTaskOptions.ColdDDC, Args));
-						}
-					}
-					// END TEMP
+					NewTasks.Add(new BenchmarkRunEditorTask(InProjectFile, DDCTaskOptions.WarmDDC, DefaultExtraArgs));
 				}
 
 				// hot ddc
@@ -456,31 +439,6 @@ namespace AutomationTool.Benchmark
 				if (InOptions.DDCOptions.HasFlag(DDCTaskOptions.ColdDDC))
 				{
 					NewTasks.Add(new BenchmarkRunEditorTask(InProjectFile, DDCTaskOptions.ColdDDC, DefaultExtraArgs));
-
-					// TEMP CODE TO test DDC sizes and hints
-					if (DoMinSharedFetch)
-					{
-						IEnumerable<int> FileSizes = new int[] { 64, 256 };
-						foreach (int Size in FileSizes)
-						{
-							string Args = string.Format("{0} -ddcminsharedfetch={1}", DefaultExtraArgs, Size);
-				
-							NewTasks.Add(new BenchmarkRunEditorTask(InProjectFile, DDCTaskOptions.None | DDCTaskOptions.ColdDDC, Args));		
-
-							if (DoHints)
-							{
-								NewTasks.Add(new BenchmarkRunEditorTask(InProjectFile, DDCTaskOptions.None | DDCTaskOptions.ColdDDC, Args + " -usehints"));
-							}
-						}
-					}
-					else
-					{
-						if (DoHints)
-						{
-							NewTasks.Add(new BenchmarkRunEditorTask(InProjectFile, DDCTaskOptions.None | DDCTaskOptions.ColdDDC, "-usehints"));
-						}
-					}
-					// END TEMP
 				}
 
 				// no shaders in the ddc
@@ -490,9 +448,9 @@ namespace AutomationTool.Benchmark
 				}
 
 				// no ddc!
-				if (InOptions.DDCOptions.HasFlag(DDCTaskOptions.NoDDC))
+				if (InOptions.DDCOptions.HasFlag(DDCTaskOptions.NoSharedDDC))
 				{
-					NewTasks.Add(new BenchmarkRunEditorTask(InProjectFile, DDCTaskOptions.NoDDC, DefaultExtraArgs));
+					NewTasks.Add(new BenchmarkRunEditorTask(InProjectFile, DDCTaskOptions.NoSharedDDC, DefaultExtraArgs));
 				}
 			}
 
@@ -505,7 +463,6 @@ namespace AutomationTool.Benchmark
 		/// </summary>
 		void WriteCSVResults(string InFileName, List<BenchmarkTaskBase> InTasks, Dictionary<BenchmarkTaskBase, List<TimeSpan>> InResults)
 		{
-
 			Log.TraceInformation("Writing results to {0}", InFileName);
 
 			try
diff --git a/Engine/Source/Programs/AutomationTool/Scripts/BenchmarkBuild/BenchmarkCookTask.cs b/Engine/Source/Programs/AutomationTool/Scripts/BenchmarkBuild/BenchmarkCookTask.cs
--- a/Engine/Source/Programs/AutomationTool/Scripts/BenchmarkBuild/BenchmarkCookTask.cs
+++ b/Engine/Source/Programs/AutomationTool/Scripts/BenchmarkBuild/BenchmarkCookTask.cs
@@ -84,7 +84,7 @@ namespace AutomationTool.Benchmark
 				//ExtraArgs.Add("noxgeshadercompile");
 			}
 
-			if (TaskOptions.HasFlag(DDCTaskOptions.NoDDC))
+			if (TaskOptions.HasFlag(DDCTaskOptions.NoSharedDDC))
 			{
 				ExtraArgsList.Add("ddc=noshared");
 			}
diff --git a/Engine/Source/Programs/AutomationTool/Scripts/BenchmarkBuild/BenchmarkRunEditorTask.cs b/Engine/Source/Programs/AutomationTool/Scripts/BenchmarkBuild/BenchmarkRunEditorTask.cs
--- a/Engine/Source/Programs/AutomationTool/Scripts/BenchmarkBuild/BenchmarkRunEditorTask.cs
+++ b/Engine/Source/Programs/AutomationTool/Scripts/BenchmarkBuild/BenchmarkRunEditorTask.cs
@@ -72,13 +72,22 @@ namespace AutomationTool.Benchmark
 			return Message;
 		}
 
+		private static string MakeValidFileName(string name)
+		{
+			string invalidChars = System.Text.RegularExpressions.Regex.Escape(new string(System.IO.Path.GetInvalidFileNameChars()));
+			string invalidRegStr = string.Format(@"([{0}]*\.+$)|([{0}]+)", invalidChars);
+
+			return System.Text.RegularExpressions.Regex.Replace(name, invalidRegStr, "_");
+		}
+
 		protected bool RunEditorAndWaitForMapLoad()
 		{
 			string ProjectArg = ProjectFile != null ? ProjectFile.ToString() : "";
 			string EditorPath = HostPlatform.Current.GetUE4ExePath("UE4Editor.exe");
-			string Arguments = string.Format("{0} {1} -execcmds=\"automation runtest System.Maps.PIE;Quit\" -stdout -AllowStdOutLogVerbosity -unattended", ProjectArg, EditorArgs);
+			string LogArg = string.Format("-log={0}.log", MakeValidFileName(GetFullTaskName()).Replace(" ", "_"));
+			string Arguments = string.Format("{0} {1} -execcmds=\"automation runtest System.Maps.PIE;Quit\" -stdout -AllowStdOutLogVerbosity -unattended {2}", ProjectArg, EditorArgs, LogArg);
 
-			if (TaskOptions.HasFlag(DDCTaskOptions.NoDDC))
+			if (TaskOptions.HasFlag(DDCTaskOptions.NoSharedDDC))
 			{
 				Arguments += (" -ddc=noshared");
 			}
diff --git a/Engine/Source/Programs/AutomationTool/Scripts/BenchmarkBuild/BenchmarkTaskBase.cs b/Engine/Source/Programs/AutomationTool/Scripts/BenchmarkBuild/BenchmarkTaskBase.cs
--- a/Engine/Source/Programs/AutomationTool/Scripts/BenchmarkBuild/BenchmarkTaskBase.cs
+++ b/Engine/Source/Programs/AutomationTool/Scripts/BenchmarkBuild/BenchmarkTaskBase.cs
@@ -161,9 +161,11 @@ namespace AutomationTool.Benchmark
 		None = 0,
 		WarmDDC = 1 << 0,
 		ColdDDC = 1 << 1,
-		NoDDC = 1 << 2,
+		NoSharedDDC = 1 << 2,
 		NoShaderDDC = 1 << 3,
 		HotDDC = 1 << 4,
+
+		KeepMemoryDDC = 1 << 5,
 	}
 
 	abstract class BenchmarkEditorTaskBase : BenchmarkTaskBase
@@ -203,7 +205,7 @@ namespace AutomationTool.Benchmark
 				TaskModifiers.Add("hotddc");
 			}
 
-			if (TaskOptions.HasFlag(DDCTaskOptions.NoDDC))
+			if (TaskOptions.HasFlag(DDCTaskOptions.NoSharedDDC))
 			{
 				TaskModifiers.Add("noddc");
 			}
@@ -213,87 +215,68 @@ namespace AutomationTool.Benchmark
 				TaskModifiers.Add("noshaderddc");
 			}
 
-			if (!string.IsNullOrEmpty(EditorArgs))
-			{
-				TaskModifiers.Add(EditorArgs);
-			}
-		}
-
-		private void DeleteLocalDDC(FileReference InProjectFile)
-		{
-			List<DirectoryReference> DirsToClear = new List<DirectoryReference>();
-
-			DirectoryReference ProjectDir = InProjectFile.Directory;
-
-			DirsToClear.Add(DirectoryReference.Combine(ProjectDir, "Saved"));
-			DirsToClear.Add(DirectoryReference.Combine(CommandUtils.EngineDirectory, "DerivedDataCache"));
-			DirsToClear.Add(DirectoryReference.Combine(ProjectDir, "DerivedDataCache"));
-
-			string LocalDDC = Environment.GetEnvironmentVariable("UE-LocalDataCachePath");
-
-			if (!string.IsNullOrEmpty(LocalDDC) && Directory.Exists(LocalDDC))
+			if (TaskOptions.HasFlag(DDCTaskOptions.KeepMemoryDDC))
 			{
-				DirsToClear.Add(new DirectoryReference(LocalDDC));
+				TaskModifiers.Add("withbootddc");
 			}
 
-			foreach (var Dir in DirsToClear)
+			if (!string.IsNullOrEmpty(EditorArgs))
 			{
-				try
-				{
-					if (DirectoryReference.Exists(Dir))
-					{
-						Log.TraceInformation("Removing {0}", Dir);
-						DirectoryReference.Delete(Dir, true);
-					}
-				}
-				catch (Exception Ex)
-				{
-					Log.TraceWarning("Failed to remove path {0}. {1}", Dir.FullName, Ex.Message);
-				}
+				TaskModifiers.Add(EditorArgs);
 			}
 		}
 
-		private DirectoryReference PrivateDDCPath = null;
-		private string OldLocalDCCPath = null;
+		private Dictionary<string, string> StoredEnvVars = new Dictionary<string, string>();
+		private List<DirectoryReference> CachePaths = new List<DirectoryReference>();
 
-		private string GetLocalDDCPathEnvironmentKey()
+		private string GetXPlatformEnvironmentKey(string InKey)
 		{
-			string EnvVar = "UE-LocalDataCachePath";
-
+			// Mac uses _ in place of -
 			if (BuildHostPlatform.Current.Platform != UnrealTargetPlatform.Win64)
 			{
-				EnvVar = EnvVar.Replace("-", "_");
+				InKey = InKey.Replace("-", "_");
 			}
 
-			return EnvVar;
+			return InKey;
 		}
 
 		protected override bool PerformPrequisites()
 		{
 			if (TaskOptions.HasFlag(DDCTaskOptions.ColdDDC))
 			{
-				PrivateDDCPath = DirectoryReference.Combine(CommandUtils.EngineDirectory, "BenchmarkDerivedDataCache");
+				StoredEnvVars.Clear();
+				CachePaths.Clear();
 
-				if (DirectoryReference.Exists(PrivateDDCPath))
+				// We put our temp DDC paths in here
+				DirectoryReference BasePath = DirectoryReference.Combine(CommandUtils.EngineDirectory, "BenchmarkDDC");
+
+				IEnumerable<string> DDCEnvVars = new string[] { GetXPlatformEnvironmentKey("UE-BootDataCachePath"), GetXPlatformEnvironmentKey("UE-LocalDataCachePath") };
+				
+				if (TaskOptions.HasFlag(DDCTaskOptions.KeepMemoryDDC))
 				{
-					DirectoryReference.Delete(PrivateDDCPath, true);
+					DDCEnvVars = DDCEnvVars.Where(E => !E.Contains("UE-Boot"));
 				}
 
-				string Key = GetLocalDDCPathEnvironmentKey();
+				// get all current environment vars and set them to our temp dir
+				foreach (var Key in DDCEnvVars)
+				{
+					// save current key
+					StoredEnvVars.Add(Key, Environment.GetEnvironmentVariable(Key));
 
-				OldLocalDCCPath = Environment.GetEnvironmentVariable(Key);
-				Environment.SetEnvironmentVariable(Key, PrivateDDCPath.FullName);
+					// create a new dir for this key
+					DirectoryReference Dir = DirectoryReference.Combine(BasePath, Key);
 
-				DirectoryReference BootDDC = DirectoryReference.Combine(ProjectFile.Directory, "DerivedDataCache");
+					if (DirectoryReference.Exists(Dir))
+					{
+						DirectoryReference.Delete(Dir, true);
+					}
 
-				try
-				{
-					DirectoryReference.Delete(BootDDC, true);
-				}
-				catch
-				{
+					DirectoryReference.CreateDirectory(Dir);
 
-				}
+					// save this dir and set it as the env var
+					CachePaths.Add(Dir);
+					Environment.SetEnvironmentVariable(Key, Dir.FullName);
+				}				
 			}
 
 			return base.PerformPrequisites();
@@ -301,15 +284,19 @@ namespace AutomationTool.Benchmark
 
 		protected override void PerformCleanup()
 		{
-			if (PrivateDDCPath !=  null)
+			// restore keys
+			foreach (var KV in StoredEnvVars)
 			{
-				if (DirectoryReference.Exists(PrivateDDCPath))
-				{
-					DirectoryReference.Delete(PrivateDDCPath, true);
-				}
+				Environment.SetEnvironmentVariable(KV.Key, KV.Value);
+			}
 
-				Environment.SetEnvironmentVariable(GetLocalDDCPathEnvironmentKey(), OldLocalDCCPath);
+			foreach (var Dir in CachePaths)
+			{
+				DirectoryReference.Delete(Dir, true);
 			}
+
+			CachePaths.Clear();
+			StoredEnvVars.Clear();
 		}
 	}
 }
