Benchmark Script

Fixed issue when building vanilla UE4 with no project

diff --git a/Engine/Source/Programs/AutomationTool/Scripts/BenchmarkBuild/BenchmarkBuild.Automation.cs b/Engine/Source/Programs/AutomationTool/Scripts/BenchmarkBuild/BenchmarkBuild.Automation.cs
--- a/Engine/Source/Programs/AutomationTool/Scripts/BenchmarkBuild/BenchmarkBuild.Automation.cs
+++ b/Engine/Source/Programs/AutomationTool/Scripts/BenchmarkBuild/BenchmarkBuild.Automation.cs
@@ -23,7 +23,7 @@ namespace AutomationTool.Benchmark
 	[Help("all", "Run all the things (except noddc)")]
 	[Help("allcompile", "Run all the compile things")]
 	[Help("editor", "Build an editor for compile tests")]
-	[Help("client", "Build a client for comple tests (see -platform)")]	
+	[Help("client", "Build a client for comple tests (see -platform)")]
 	[Help("platform=<p1+p2>", "Specify the platform(s) to use for client compilation/cooking, if empty the local platform be used if -client or -cook is specified")]
 	[Help("xge", "Do a compile with XGE / FASTBuild")]
 	[Help("noxge", "Do a compile without XGE / FASTBuild")]
@@ -66,7 +66,7 @@ namespace AutomationTool.Benchmark
 
 			// cooking
 			public bool DoCookTests = false;
-			
+
 			// editor startup tests
 			public bool DoPIETests = false;
 
@@ -102,7 +102,7 @@ namespace AutomationTool.Benchmark
 
 				// cooking
 				DoCookTests = AllThings | ParseParam("cook");
-				
+
 				// editor startup tests
 				DoPIETests = AllThings | ParseParam("pie");
 				MapName = ParseParamValue("map", "");
@@ -178,12 +178,12 @@ namespace AutomationTool.Benchmark
 
 						CoresForLocalJobs = ProcessorList.Select(P => Convert.ToInt32(P));
 					}
-				}				
+				}
 			}
 		}
 
 		public BenchmarkBuild()
-		{ 
+		{
 		}
 
 		public override ExitCode Execute()
@@ -195,13 +195,13 @@ namespace AutomationTool.Benchmark
 
 			Dictionary<BenchmarkTaskBase, List<TimeSpan>> Results = new Dictionary<BenchmarkTaskBase, List<TimeSpan>>();
 
-			for (int ProjectIndex = 0; ProjectIndex <  Options.ProjectsToTest.Count(); ProjectIndex++)
+			for (int ProjectIndex = 0; ProjectIndex < Options.ProjectsToTest.Count(); ProjectIndex++)
 			{
 				string Project = Options.ProjectsToTest.ElementAt(ProjectIndex);
 
 				FileReference ProjectFile = ProjectUtils.FindProjectFileFromName(Project);
 
-				if (ProjectFile == null)
+				if (ProjectFile == null && !Project.Equals("UE4", StringComparison.OrdinalIgnoreCase))
 				{
 					throw new AutomationException("Could not find project file for {0}", Project);
 				}
@@ -257,7 +257,7 @@ namespace AutomationTool.Benchmark
 				{
 					foreach (var Task in Tasks)
 					{
-						Log.TraceInformation("Starting task {0} (Pass {1})", Task.GetFullTaskName(), i+1);
+						Log.TraceInformation("Starting task {0} (Pass {1})", Task.GetFullTaskName(), i + 1);
 
 						Task.Run();
 
@@ -305,7 +305,7 @@ namespace AutomationTool.Benchmark
 						var AvgTime = new TimeSpan(TaskTimes.Sum(T => T.Ticks) / TaskTimes.Count());
 
 						AvgTimeString = string.Format(" (Avg: {0})", AvgTime.ToString(@"hh\:mm\:ss"));
-					}					
+					}
 
 					Log.TraceInformation("Task {0}:\t{1}{2}", Task.GetFullTaskName(), TimeString, AvgTimeString);
 				}
@@ -313,7 +313,7 @@ namespace AutomationTool.Benchmark
 
 				TimeSpan Elapsed = DateTime.Now - StartTime;
 
-				Log.TraceInformation("Total benchmark time: {0}",  Elapsed.ToString(@"hh\:mm\:ss"));
+				Log.TraceInformation("Total benchmark time: {0}", Elapsed.ToString(@"hh\:mm\:ss"));
 
 				WriteCSVResults(Options.FileName, Tasks, Results);
 			}
@@ -420,7 +420,7 @@ namespace AutomationTool.Benchmark
 			{
 				DefaultExtraArgs += " -map=" + InOptions.MapName;
 			}
-			
+
 			if (InOptions.DoPIETests)
 			{
 				// if no options assume warm
