Add x64 architecture for Hololens (fixes packaging)

diff --git a/Engine/Source/Programs/AutomationTool/Scripts/BuildPluginCommand.Automation.cs b/Engine/Source/Programs/AutomationTool/Scripts/BuildPluginCommand.Automation.cs
--- a/Engine/Source/Programs/AutomationTool/Scripts/BuildPluginCommand.Automation.cs
+++ b/Engine/Source/Programs/AutomationTool/Scripts/BuildPluginCommand.Automation.cs
@@ -20,7 +20,7 @@ using System.Text;
 class BuildPlugin : BuildCommand
 {
 	const string AndroidArchitectures = "armv7+arm64";
-	const string HoloLensArchitecture = "arm64";
+	const string HoloLensArchitectures = "arm64+x64";
 
 	public override void ExecuteBuild()
 	{
@@ -206,7 +206,7 @@ class BuildPlugin : BuildCommand
 			}
 			else if (Platform == UnrealTargetPlatform.HoloLens)
 			{
-				Arguments += String.Format(" -Architecture={0}", HoloLensArchitecture);
+				Arguments += String.Format(" -architectures={0}", HoloLensArchitectures);
 			}
 
 			if (!String.IsNullOrEmpty(InAdditionalArgs))
