Also upload non USF files using manifest via deployment server.  Note: Metal libs files can't be in the UFS manifest as they get memory mapped at runtime - unless we unpacked them to a known location on first launch but that's a whole different ball game.

diff --git a/Engine/Source/Programs/AutomationTool/IOS/IOSPlatform.Automation.cs b/Engine/Source/Programs/AutomationTool/IOS/IOSPlatform.Automation.cs
--- a/Engine/Source/Programs/AutomationTool/IOS/IOSPlatform.Automation.cs
+++ b/Engine/Source/Programs/AutomationTool/IOS/IOSPlatform.Automation.cs
@@ -1357,6 +1357,12 @@ public class IOSPlatform : Platform
 		{
 			// push over the changed files
 			RunAndLog(CmdEnv, DeployServer, "Deploy -manifest \"" + SC.GetUFSDeploymentDeltaPath(Params.Devices.Count == 0 ? null : Params.DeviceNames[0]) + "\"" + (Params.Devices.Count == 0 ? "" : " -device " + Params.DeviceNames[0]) + AdditionalCommandline + " -bundle " + BundleIdentifier);
+
+			// IPA for launch-on doesn't include these files - they are a delta and still need uploading
+			if(bNeedsIPA)
+			{
+				RunAndLog(CmdEnv, DeployServer, "Deploy -manifest \"" + SC.GetNonUFSDeploymentDeltaPath(Params.Devices.Count == 0 ? null : Params.DeviceNames[0]) + "\"" + (Params.Devices.Count == 0 ? "" : " -device " + Params.DeviceNames[0]) + AdditionalCommandline + " -bundle " + BundleIdentifier);
+			}
 		}
 		Directory.SetCurrentDirectory(CurrentDir);
 		PrintRunTime();
