Update small OBB limit to 1 GiB

#ue4
#android

diff --git a/Engine/Source/Programs/AutomationTool/Android/AndroidPlatform.Automation.cs b/Engine/Source/Programs/AutomationTool/Android/AndroidPlatform.Automation.cs
--- a/Engine/Source/Programs/AutomationTool/Android/AndroidPlatform.Automation.cs
+++ b/Engine/Source/Programs/AutomationTool/Android/AndroidPlatform.Automation.cs
@@ -14,8 +14,8 @@ using Tools.DotNETCommon;
 
 public class AndroidPlatform : Platform
 {
-	// Maximum allowed OBB size (500 MiB, 2 GiB or 4 GiB based on project settings)
-	private const Int64 SmallOBBSizeAllowed = 524288000;
+	// Maximum allowed OBB size (1 GiB, 2 GiB or 4 GiB based on project settings)
+	private const Int64 SmallOBBSizeAllowed = 1073741824;
 	private const Int64 NormalOBBSizeAllowed = 2147483648;
 	private const Int64 MaxOBBSizeAllowed = 4294967296;
 
@@ -452,7 +452,7 @@ public class AndroidPlatform : Platform
 			OBBNeedsUpdate = true;
 		}
 		Int64 OBBSizeAllowed = GetMaxOBBSizeAllowed(SC);
-		string LimitString = (OBBSizeAllowed < NormalOBBSizeAllowed) ? "500 MiB" : ((OBBSizeAllowed < MaxOBBSizeAllowed) ? "2 GiB" : "4 GiB");
+		string LimitString = (OBBSizeAllowed < NormalOBBSizeAllowed) ? "1 GiB" : ((OBBSizeAllowed < MaxOBBSizeAllowed) ? "2 GiB" : "4 GiB");
 
 		if (!OBBNeedsUpdate)
 		{
diff --git a/Engine/Source/Runtime/Android/AndroidRuntimeSettings/Classes/AndroidRuntimeSettings.h b/Engine/Source/Runtime/Android/AndroidRuntimeSettings/Classes/AndroidRuntimeSettings.h
--- a/Engine/Source/Runtime/Android/AndroidRuntimeSettings/Classes/AndroidRuntimeSettings.h
+++ b/Engine/Source/Runtime/Android/AndroidRuntimeSettings/Classes/AndroidRuntimeSettings.h
@@ -237,7 +237,7 @@ public:
 	UPROPERTY(GlobalConfig, EditAnywhere, Category = "APK Packaging", Meta = (DisplayName = "Disable verify OBB on first start/update."))
 	bool bDisableVerifyOBBOnStartUp;
 
-	// If checked, OBB is limited to 500 MiB.
+	// If checked, OBB is limited to 1 GiB.
 	UPROPERTY(GlobalConfig, EditAnywhere, Category = "APK Packaging", Meta = (DisplayName = "Force small OBB files."))
 	bool bForceSmallOBBFiles;
 
