Ensure both x64 and Arm64 are targets for Hololens

diff --git a/Engine/Source/Programs/AutomationTool/Scripts/BuildPluginCommand.Automation.cs b/Engine/Source/Programs/AutomationTool/Scripts/BuildPluginCommand.Automation.cs
--- a/Engine/Source/Programs/AutomationTool/Scripts/BuildPluginCommand.Automation.cs
+++ b/Engine/Source/Programs/AutomationTool/Scripts/BuildPluginCommand.Automation.cs
@@ -20,7 +20,7 @@ using System.Text;
 class BuildPlugin : BuildCommand
 {
 	const string AndroidArchitectures = "armv7+arm64";
-	const string HoloLensArchitecture = "arm64";
+	const string HoloLensArchitecture = "arm64+x64";
 
 	public override void ExecuteBuild()
 	{
