Benchmark build now captures logs with verbosity.

used safer means of deleting old directories

diff --git a/Engine/Source/Programs/AutomationTool/Scripts/BenchmarkBuild/BenchmarkRunEditorTask.cs b/Engine/Source/Programs/AutomationTool/Scripts/BenchmarkBuild/BenchmarkRunEditorTask.cs
--- a/Engine/Source/Programs/AutomationTool/Scripts/BenchmarkBuild/BenchmarkRunEditorTask.cs
+++ b/Engine/Source/Programs/AutomationTool/Scripts/BenchmarkBuild/BenchmarkRunEditorTask.cs
@@ -85,7 +85,7 @@ namespace AutomationTool.Benchmark
 			string ProjectArg = ProjectFile != null ? ProjectFile.ToString() : "";
 			string EditorPath = HostPlatform.Current.GetUE4ExePath("UE4Editor.exe");
 			string LogArg = string.Format("-log={0}.log", MakeValidFileName(GetFullTaskName()).Replace(" ", "_"));
-			string Arguments = string.Format("{0} {1} -execcmds=\"automation runtest System.Maps.PIE;Quit\" -stdout -AllowStdOutLogVerbosity -unattended {2}", ProjectArg, EditorArgs, LogArg);
+			string Arguments = string.Format("{0} {1} -execcmds=\"automation runtest System.Maps.PIE;Quit\" -stdout -FullStdOutLogOutput -unattended {2}", ProjectArg, EditorArgs, LogArg);
 
 			if (TaskOptions.HasFlag(DDCTaskOptions.NoSharedDDC))
 			{
diff --git a/Engine/Source/Programs/AutomationTool/Scripts/BenchmarkBuild/BenchmarkTaskBase.cs b/Engine/Source/Programs/AutomationTool/Scripts/BenchmarkBuild/BenchmarkTaskBase.cs
--- a/Engine/Source/Programs/AutomationTool/Scripts/BenchmarkBuild/BenchmarkTaskBase.cs
+++ b/Engine/Source/Programs/AutomationTool/Scripts/BenchmarkBuild/BenchmarkTaskBase.cs
@@ -292,7 +292,7 @@ namespace AutomationTool.Benchmark
 
 			foreach (var Dir in CachePaths)
 			{
-				DirectoryReference.Delete(Dir, true);
+				CommandUtils.DeleteDirectory_NoExceptions(Dir.FullName);
 			}
 
 			CachePaths.Clear();
