#Gauntlet Support DebugGame and Debug editor builds

diff --git a/Engine/Source/Programs/AutomationTool/Gauntlet/Unreal/BuildSource/Gauntlet.StagedBuildSource.cs b/Engine/Source/Programs/AutomationTool/Gauntlet/Unreal/BuildSource/Gauntlet.StagedBuildSource.cs
--- a/Engine/Source/Programs/AutomationTool/Gauntlet/Unreal/BuildSource/Gauntlet.StagedBuildSource.cs
+++ b/Engine/Source/Programs/AutomationTool/Gauntlet/Unreal/BuildSource/Gauntlet.StagedBuildSource.cs
@@ -16,7 +16,7 @@ namespace Gauntlet
 	{
 		public UnrealTargetPlatform Platform { get { return BuildHostPlatform.Current.Platform; } }
 
-		public UnrealTargetConfiguration Configuration { get { return UnrealTargetConfiguration.Development; } }
+		public UnrealTargetConfiguration Configuration { get; protected set; }
 
 		public BuildFlags Flags { get { return BuildFlags.CanReplaceCommandLine | BuildFlags.Loose; } }
 
@@ -24,9 +24,10 @@ namespace Gauntlet
 
 		public string ExecutablePath { get; protected set; }
 
-		public EditorBuild(string InExecutablePath)
+		public EditorBuild(string InExecutablePath, UnrealTargetConfiguration InConfig)
 		{
 			ExecutablePath = InExecutablePath;
+			Configuration = InConfig;
 		}
 	}
 
diff --git a/Engine/Source/Programs/AutomationTool/Gauntlet/Unreal/BuildSource/Gauntlet.UnrealBuildSource.cs b/Engine/Source/Programs/AutomationTool/Gauntlet/Unreal/BuildSource/Gauntlet.UnrealBuildSource.cs
--- a/Engine/Source/Programs/AutomationTool/Gauntlet/Unreal/BuildSource/Gauntlet.UnrealBuildSource.cs
+++ b/Engine/Source/Programs/AutomationTool/Gauntlet/Unreal/BuildSource/Gauntlet.UnrealBuildSource.cs
@@ -274,11 +274,11 @@ namespace Gauntlet
 			}
 
 			// Editor?
-			IBuild EditorBuild = CreateEditorBuild(UnrealPath);
+			List<EditorBuild> EditorBuilds = CreateEditorBuilds(UnrealPath);
 
-			if (EditorBuild != null)
+			if (EditorBuilds.Count > 0)
 			{
-				BuildList.Add(EditorBuild);
+				BuildList.AddRange(EditorBuilds);
 			}
 			else
 			{
@@ -543,7 +543,15 @@ namespace Gauntlet
 
 			if (TargetType.UsesEditor())
 			{
-				ExePath = string.Format("Engine/Binaries/{0}/UE4Editor{1}", BuildHostPlatform.Current.Platform, Platform.GetExeExtension(TargetPlatform));
+				string ExeFileName = "UE4Editor";
+				if (TargetConfiguration != UnrealTargetConfiguration.Development)
+				{
+					ExeFileName += string.Format("-{0}-{1}", TargetPlatform.ToString(), TargetConfiguration.ToString());
+				}
+
+				ExeFileName += Platform.GetExeExtension(TargetPlatform);
+
+				ExePath = string.Format("Engine/Binaries/{0}/{1}", BuildHostPlatform.Current.Platform, ExeFileName);
 			}
 			else
 			{
@@ -688,24 +696,26 @@ namespace Gauntlet
 			return PlatformPath;
 		}
 
-		EditorBuild CreateEditorBuild(DirectoryReference InUnrealPath)
+		List<EditorBuild> CreateEditorBuilds(DirectoryReference InUnrealPath)
 		{
+			List<EditorBuild> EditorBuildList = new List<EditorBuild>();
 			if (InUnrealPath == null)
 			{
-				return null;
+				return EditorBuildList;
 			}
 
 			// check for the editor
-			string EditorExe = Path.Combine(InUnrealPath.FullName, GetRelativeExecutablePath(UnrealTargetRole.Editor, BuildHostPlatform.Current.Platform, UnrealTargetConfiguration.Development));
-
-			if (!Utils.SystemHelpers.ApplicationExists(EditorExe))
+			List<UnrealTargetConfiguration> ConfigsToCheck = new List<UnrealTargetConfiguration>() { UnrealTargetConfiguration.Development, UnrealTargetConfiguration.DebugGame, UnrealTargetConfiguration.Debug };
+			foreach (UnrealTargetConfiguration Config in ConfigsToCheck)
 			{
-				return null;
+				string EditorExe = Path.Combine(InUnrealPath.FullName, GetRelativeExecutablePath(UnrealTargetRole.Editor, BuildHostPlatform.Current.Platform, Config));
+				if (Utils.SystemHelpers.ApplicationExists(EditorExe))
+				{
+					EditorBuildList.Add(new EditorBuild(EditorExe, Config));
+				}
 			}
 
-			EditorBuild NewBuild = new EditorBuild(EditorExe);
-
-			return NewBuild;
+			return EditorBuildList;
 		}
 	}
 }
diff --git a/Engine/Source/Programs/AutomationTool/Gauntlet/Unreal/RunUnreal.cs b/Engine/Source/Programs/AutomationTool/Gauntlet/Unreal/RunUnreal.cs
--- a/Engine/Source/Programs/AutomationTool/Gauntlet/Unreal/RunUnreal.cs
+++ b/Engine/Source/Programs/AutomationTool/Gauntlet/Unreal/RunUnreal.cs
@@ -223,7 +223,10 @@ namespace Gauntlet
 						{
 							Role.Type = UnrealTargetRole.EditorGame;
 							Role.Platform = DefaultPlatform;
-							Role.Configuration = UnrealTargetConfiguration.Development;
+							if (Role.Configuration > UnrealTargetConfiguration.Development)
+							{
+								Role.Configuration = UnrealTargetConfiguration.Development;
+							}
 						}
 					}
 					else if (Type.IsServer())
@@ -235,7 +238,10 @@ namespace Gauntlet
 						{
 							Role.Type = UnrealTargetRole.EditorServer;
 							Role.Platform = DefaultPlatform;
-							Role.Configuration = UnrealTargetConfiguration.Development;
+							if (Role.Configuration > UnrealTargetConfiguration.Development)
+							{
+								Role.Configuration = UnrealTargetConfiguration.Development;
+							}
 						}
 					}
 
