Need to keep some of old code path I removed for BluePrint projects while still handling correctly named CPP project application packages from the last change.

diff --git a/Engine/Source/Programs/AutomationTool/Mac/MacPlatform.Automation.cs b/Engine/Source/Programs/AutomationTool/Mac/MacPlatform.Automation.cs
--- a/Engine/Source/Programs/AutomationTool/Mac/MacPlatform.Automation.cs
+++ b/Engine/Source/Programs/AutomationTool/Mac/MacPlatform.Automation.cs
@@ -327,7 +327,7 @@ public class MacPlatform : Platform
 		if (Params.CreateAppBundle)
 		{
 			string ExeName = SC.StageExecutables[0];
-			string BundlePath = CombinePaths(SC.ArchiveDirectory.FullName, ExeName + ".app");
+			string BundlePath = SC.IsCodeBasedProject ? CombinePaths(SC.ArchiveDirectory.FullName, ExeName + ".app") : BundlePath = CombinePaths(SC.ArchiveDirectory.FullName, SC.ShortProjectName + ".app");
 
 			if (SC.bIsCombiningMultiplePlatforms)
 			{
@@ -386,12 +386,9 @@ public class MacPlatform : Platform
 				}
 			}
 
-			// This ends up renaming the Bundle, exe and entry in the info.plist to the Project Name rather than desired Target
-			// Commenting this out but leaving it here in the short term in case there is a permutation where its need.
-			/*
 			// Update executable name, icon and entry in Info.plist
 			string UE4GamePath = CombinePaths(BundlePath, "Contents", "MacOS", ExeName);
-			if (ExeName != SC.ShortProjectName && File.Exists(UE4GamePath))
+			if (!SC.IsCodeBasedProject && ExeName != SC.ShortProjectName && File.Exists(UE4GamePath))
 			{
 				string GameExePath = CombinePaths(BundlePath, "Contents", "MacOS", SC.ShortProjectName);
 				File.Delete(GameExePath);
@@ -425,7 +422,6 @@ public class MacPlatform : Platform
 				File.Delete(InfoPlistPath);
 				File.WriteAllText(InfoPlistPath, InfoPlistContents);
 			}
-			*/
 
 			if (!SC.bIsCombiningMultiplePlatforms)
 			{
