Fixing Switch screenshot filename for output

diff --git a/Engine/Source/Programs/AutomationTool/Gauntlet/Unreal/Base/Gauntlet.UnrealPGONode.cs b/Engine/Source/Programs/AutomationTool/Gauntlet/Unreal/Base/Gauntlet.UnrealPGONode.cs
--- a/Engine/Source/Programs/AutomationTool/Gauntlet/Unreal/Base/Gauntlet.UnrealPGONode.cs
+++ b/Engine/Source/Programs/AutomationTool/Gauntlet/Unreal/Base/Gauntlet.UnrealPGONode.cs
@@ -148,15 +148,18 @@ namespace Gauntlet
 					string ImageOutputPath = Path.Combine(ScreenshotDirectory, ImageTimestamp.ToString().Replace(':', '-') + ".jpg");
 					ImageUtils.ResaveImageAsJpgWithScaleAndQuality(Path.Combine(ScreenshotDirectory, ImageFilename), ImageOutputPath, ScreenshotScale, ScreenshotQuality);
 				}
-				catch
+				catch(Exception e)
 				{
-					// Just ignore errors.
+					Log.Warning("Got Exception Renaming PGO image: {0}", e.ToString());
 				}
 				finally
 				{
 					// Delete the temporary image file
 					try { File.Delete(Path.Combine(ScreenshotDirectory, ImageFilename)); }
-					catch { }
+					catch (Exception e)
+					{
+						Log.Warning("Got Exception Deleting temp iamge: {0}", e.ToString());
+					}
 				}
 			}
 
