Refactor commandline for gauntlet to make it far easier for test writers to interact with.

Please keep an eye out for issues with strange commandlines in tests over the next couple of days. I have run extensive testing on this but not exhaustive.

diff --git a/Engine/Source/Programs/AutomationTool/Gauntlet/Unreal/Base/Gauntlet.UnrealAppConfig.cs b/Engine/Source/Programs/AutomationTool/Gauntlet/Unreal/Base/Gauntlet.UnrealAppConfig.cs
--- a/Engine/Source/Programs/AutomationTool/Gauntlet/Unreal/Base/Gauntlet.UnrealAppConfig.cs
+++ b/Engine/Source/Programs/AutomationTool/Gauntlet/Unreal/Base/Gauntlet.UnrealAppConfig.cs
@@ -41,11 +41,37 @@ namespace Gauntlet
         /// Files to copy over to the device, plus what type of files they are.
         /// </summary>
         public List<UnrealFileToCopy> FilesToCopy { get; set; }
-        
-        /// <summary>
-        /// Arguments for this instance
-        /// </summary>
-        public string CommandLine { get; set; }
+
+		/// <summary>
+		/// Arguments for this instance
+		/// </summary>
+		public string CommandLine
+		{
+			get
+			{
+				if (CommandLineParams == null)
+				{
+					CommandLineParams = new GauntletCommandLine();
+				}
+				return CommandLineParams.GenerateFullCommandLine();
+			}
+			set
+			{
+				if (CommandLineParams == null)
+				{
+					CommandLineParams = new GauntletCommandLine();
+				}
+
+				CommandLineParams.ClearCommandLine();
+
+				CommandLineParams.AddRawCommandline(value, false);
+			}
+		}
+		/// <summary>
+		/// Dictionary of commandline arguments that are turned into a commandline at the end.
+		/// For flags, leave the value set to null.
+		/// </summary>
+		public GauntletCommandLine CommandLineParams;
 
 		/// <summary>
 		/// Sandbox that we'd like to install this instance in
diff --git a/Engine/Source/Programs/AutomationTool/Gauntlet/Unreal/Base/Gauntlet.UnrealSession.cs b/Engine/Source/Programs/AutomationTool/Gauntlet/Unreal/Base/Gauntlet.UnrealSession.cs
--- a/Engine/Source/Programs/AutomationTool/Gauntlet/Unreal/Base/Gauntlet.UnrealSession.cs
+++ b/Engine/Source/Programs/AutomationTool/Gauntlet/Unreal/Base/Gauntlet.UnrealSession.cs
@@ -57,7 +57,35 @@ namespace Gauntlet
 		/// <summary>
 		/// Command line that this role will use
 		/// </summary>
-		public string CommandLine;
+		public string CommandLine
+		{
+			get
+			{
+				if (CommandLineParams == null)
+				{
+					CommandLineParams = new GauntletCommandLine();
+				}
+				return CommandLineParams.GenerateFullCommandLine();
+			}
+			set
+			{
+				if (CommandLineParams == null)
+				{
+					CommandLineParams = new GauntletCommandLine();
+				}
+
+				CommandLineParams.ClearCommandLine();
+
+				CommandLineParams.AddRawCommandline(value, false);
+			}
+		}
+		/// <summary>
+		/// Dictionary of commandline arguments that are turned into a commandline at the end.
+		/// For flags, leave the value set to null. Generated from the Test Role's commandline object
+		/// and modified as needed. Then passed through to the AppConfig in UnrealBuildSource.cs
+		/// </summary>
+		public GauntletCommandLine CommandLineParams { get; set; }
+
 
 		/// <summary>
 		/// Map override to use on a server in case we don't want them all running the same map.
@@ -165,6 +193,7 @@ namespace Gauntlet
 
 			Options = InOptions;
             FilesToCopy = new List<UnrealFileToCopy>();
+			CommandLineParams = new GauntletCommandLine();
 			RoleModifier = ERoleModifier.None;
         }
 
diff --git a/Engine/Source/Programs/AutomationTool/Gauntlet/Unreal/Base/Gauntlet.UnrealTestConfiguration.cs b/Engine/Source/Programs/AutomationTool/Gauntlet/Unreal/Base/Gauntlet.UnrealTestConfiguration.cs
--- a/Engine/Source/Programs/AutomationTool/Gauntlet/Unreal/Base/Gauntlet.UnrealTestConfiguration.cs
+++ b/Engine/Source/Programs/AutomationTool/Gauntlet/Unreal/Base/Gauntlet.UnrealTestConfiguration.cs
@@ -5,10 +5,216 @@ using System.Collections.Generic;
 using AutomationTool;
 using UnrealBuildTool;
 using System.Linq;
+using System.Text.RegularExpressions;
 
 namespace Gauntlet
 {
 
+	public class GauntletCommandLine
+	{
+		public string Project;
+		public string GameMap;
+		private Dictionary<string, object> Params;
+
+		public GauntletCommandLine()
+		{
+			Project = string.Empty;
+			GameMap = string.Empty;
+			Params = new Dictionary<string, object>(StringComparer.OrdinalIgnoreCase);
+		}
+
+		/// <summary>
+		/// Breaks down a raw commandline and adds it to the commandline dictionary.
+		/// Will override current set values in the dictionary when conflicts arise.
+		/// </summary>
+		/// <param name="InRawCommandline"></param>
+		public void AddRawCommandline(string InRawCommandline, bool bOverrideExistingValues = true)
+		{
+			// turn Name(p1,etc) into a collection of Name|(p1,etc) groups
+			MatchCollection Matches = Regex.Matches(InRawCommandline, "-(?<option>\\-?[\\w\\d.:\\[\\]\\/\\\\]+)(=(?<value>(\"([^\"]*)\")|(\\S+)))?");
+
+			foreach (Match M in Matches)
+			{
+				if (M.Groups["option"] == null || string.IsNullOrWhiteSpace(M.Groups["option"].ToString()))
+				{
+					continue;
+				}
+				if (bOverrideExistingValues)
+				{
+					Add(M.Groups["option"].ToString().Trim(), string.IsNullOrWhiteSpace(M.Groups["value"].ToString()) ? null : M.Groups["value"]);
+				}
+				else
+				{
+					AddUnique(M.Groups["option"].ToString().Trim(), string.IsNullOrWhiteSpace(M.Groups["value"].ToString()) ? null : M.Groups["value"]);
+				}
+			}
+		}
+
+		/// <summary>
+		/// Add a new value to the commandline, returning false if the value already exists on the commandline and would be set
+		/// to something other than what is passed in. Execcmds passed in here will still append to an existing value.
+		/// </summary>
+		/// <param name="ParamName"></param>
+		/// <param name="ParamVal"></param>
+		/// <returns>Whether the value passed in is now the expected value of the commandline. Returns false if we failed to set.</returns>
+		public bool AddUnique(string ParamName, object ParamVal = null)
+		{
+			if (Params.ContainsKey(ParamName))
+			{
+				if (ParamName.ToLower() == "execcmds" && ParamVal != null)
+				{
+					AddOrAppendParamValue(ParamName, ParamVal.ToString());
+					return true;
+				}
+
+				if (Params[ParamName] != null)
+				{
+					return (Params[ParamName] == ParamVal);
+				}
+				else
+				{
+					Params[ParamName] = ParamVal;
+					return true;
+				}
+			}
+			Params.Add(ParamName, ParamVal);
+			return true;	
+		}
+
+		/// <summary>
+		/// Add new param value (adding with a value of null makes it a flag instead of a param)
+		/// Overrides current value of the param if one exists.
+		/// Execcmds will be appended to instead of overridden.
+		/// </summary>
+		/// <param name="ParamName"></param>
+		/// <param name="ParamVal"></param>
+		public void Add(string ParamName, object ParamVal = null) 
+		{
+			if (Params.ContainsKey(ParamName))
+			{
+				if (ParamName.ToLower() == "execcmds" && ParamVal != null)
+				{
+					AddOrAppendParamValue(ParamName, ParamVal.ToString());
+					return;
+				}
+
+				if (ParamVal == null && Params[ParamName] != null)
+				{
+					Gauntlet.Log.Info(string.Format("Ignored attempt to convert param {0} from a param with value to a flag", ParamName));
+					return;
+				}
+				else if (ParamVal != null)
+				{
+					if (Params[ParamName] == null)
+					{
+						Gauntlet.Log.Info(string.Format("Converting param {0} from a flag to a param with value {1}.", ParamName, ParamVal));
+					}
+					else if (Params[ParamName] != null && ParamVal.ToString().Trim() != Params[ParamName].ToString().Trim())
+					{
+						Gauntlet.Log.Info(string.Format("Overriding value of param {0} from {1} to {2}", ParamName, Params[ParamName], ParamVal));
+					}
+					Params[ParamName] = ParamVal;
+				}
+			} 			
+			else
+			{
+				Params.Add(ParamName, ParamVal);
+			}
+		}
+		/// <summary>
+		/// Add a new param, or append passed in value to it if a value already exists,
+		/// delimited with the Delimiter passed in.
+		/// </summary>
+		/// <param name="ParamName"></param>
+		/// <param name="ParamVal"></param>
+		/// <param name="DelimiterToUse"></param>
+		public void AddOrAppendParamValue(string ParamName, string ParamVal, string DelimiterToUse = ",") 
+		{
+			if (Params.ContainsKey(ParamName))
+			{
+				if (Params[ParamName] == null)
+				{
+					Params[ParamName] = ParamVal;
+				}
+				else if (Params[ParamName].ToString() != ParamVal) 
+				{
+					Params[ParamName] = Params[ParamName].ToString() + DelimiterToUse + ParamVal;
+					Params[ParamName] = Params[ParamName].ToString().Replace("\"", "");
+				}
+			}
+			else
+			{
+				Params.Add(ParamName, ParamVal);
+			}
+		}
+		/// <summary>
+		/// Remove the param or flag by the passed-in name if it exists.
+		/// </summary>
+		/// <param name="ParamName"></param>
+		public void RemoveParam(string ParamName) 
+		{
+			if ( Params.ContainsKey(ParamName))
+			{
+				Params.Remove(ParamName);
+			}
+		}
+
+		/// <summary>
+		/// Returns the current value of the param if it has one. Flags and nonexistent 
+		/// entries will return null.
+		/// </summary>
+		/// <param name="ParamName"></param>
+		/// <returns></returns>
+		public object GetParamValue(string ParamName) 
+		{
+			if (!Params.ContainsKey(ParamName))
+			{
+				return null;
+			}
+			return Params[ParamName];
+		}
+
+		public bool HasParam(string ParamName)
+		{
+			return (Params != null && Params.ContainsKey(ParamName));
+		}
+
+		/// <summary>
+		/// Wipe out the entire set of passed in params.
+		/// </summary>
+		public void ClearCommandLine()
+		{
+			Params.Clear();
+		}
+
+		/// <summary>
+		/// Cobbles together the finalized commandline from the passed-in values. Wraps param values w/ spaces
+		/// in quotes etc.
+		/// This is the function we use, as well, when accessing the Role's commandline itself.
+		/// </summary>
+		/// <returns></returns>
+		public string GenerateFullCommandLine()
+		{
+			string FinalCommandline = string.Format("{0} {1} ", Project, GameMap);
+			foreach (string Key in Params.Keys)
+			{
+				string CurrentArgument;
+				if (Params[Key] != null && !string.IsNullOrWhiteSpace(Params[Key].ToString()))
+				{
+					CurrentArgument = string.Format("-{0}={1}", Key, 
+						(Params[Key].ToString().Contains(' ') && !Params[Key].ToString().Contains('\"'))
+						? string.Format("\"{0}\"", Params[Key]) : Params[Key]);
+				}
+				else
+				{
+					CurrentArgument = string.Format("-{0}", Key);
+				}
+				FinalCommandline = string.Format("{0} {1} ", FinalCommandline, CurrentArgument);
+			}
+			return FinalCommandline;
+		}
+	}
+
     public enum EWindowMode
     {
         /// <summary>
@@ -74,6 +280,7 @@ namespace Gauntlet
             FilesToCopy = new List<UnrealFileToCopy>();
 			AdditionalArtifactDirectories = new List<EIntendedBaseCopyDirectory>();
             RoleType = ERoleModifier.None;
+			CommandLineParams = new GauntletCommandLine();
 		}
 
         public ERoleModifier RoleType { get; set; }
@@ -91,7 +298,36 @@ namespace Gauntlet
 		/// <summary>
 		/// Command line or this role
 		/// </summary>
-		public string CommandLine { get; set; }
+		public string CommandLine
+		{
+			get
+			{
+				if (CommandLineParams == null)
+				{
+					CommandLineParams = new GauntletCommandLine();
+				}
+				return CommandLineParams.GenerateFullCommandLine();
+			}
+			set
+			{
+				if (CommandLineParams == null)
+				{
+					CommandLineParams = new GauntletCommandLine();
+				}
+
+				CommandLineParams.ClearCommandLine();
+
+				CommandLineParams.AddRawCommandline(value);
+			}
+		}
+
+
+		/// <summary>
+		/// Dictionary of commandline arguments that are turned into a commandline at the end.
+		/// For flags, leave the value set to null. Created and then passed through to the Session Role's Commandline Object
+		/// in UnrealTestNode.cs
+		/// </summary>
+		public GauntletCommandLine CommandLineParams { get; set; }
 
 		/// <summary>
 		/// Controllers for this role
diff --git a/Engine/Source/Programs/AutomationTool/Gauntlet/Unreal/Base/Gauntlet.UnrealTestNode.cs b/Engine/Source/Programs/AutomationTool/Gauntlet/Unreal/Base/Gauntlet.UnrealTestNode.cs
--- a/Engine/Source/Programs/AutomationTool/Gauntlet/Unreal/Base/Gauntlet.UnrealTestNode.cs
+++ b/Engine/Source/Programs/AutomationTool/Gauntlet/Unreal/Base/Gauntlet.UnrealTestNode.cs
@@ -108,7 +108,6 @@ namespace Gauntlet
 		/// </summary>
 		/// 
 		protected Version TestVersion;
-		
 
 		/// <summary>
 		/// Path to the directory that logs and other artifacts are copied to after the test run.
@@ -366,9 +365,9 @@ namespace Gauntlet
 
 					UnrealSessionRole SessionRole = new UnrealSessionRole(RoleContext.Type, SessionPlatform, RoleContext.Configuration, TestRole.CommandLine);
 
-					SessionRole.RoleModifier = TestRole.RoleType;
+					SessionRole.CommandLineParams = TestRole.CommandLineParams;
+ 					SessionRole.RoleModifier = TestRole.RoleType;
 					SessionRole.Constraint = UseContextConstraint ? Context.Constraint : new UnrealTargetConstraint(SessionPlatform);
-					
 					Log.Verbose("Created SessionRole {0} from RoleContext {1} (RoleType={2})", SessionRole, RoleContext, TypesToRoles.Key);
 
 					// TODO - this can all / mostly go into UnrealTestConfiguration.ApplyToConfig
@@ -381,18 +380,17 @@ namespace Gauntlet
 					else
 					{
 						// start with anything from our context
-						SessionRole.CommandLine = RoleContext.ExtraArgs;
+						SessionRole.CommandLine += RoleContext.ExtraArgs;
 
 						// did the test ask for anything?
 						if (string.IsNullOrEmpty(TestRole.CommandLine) == false)
 						{
-							SessionRole.CommandLine += " " + TestRole.CommandLine;
+							SessionRole.CommandLine += TestRole.CommandLine;
 						}
 
 						// add controllers
-						SessionRole.CommandLine += TestRole.Controllers.Count > 0 ?
-							string.Format(" -gauntlet=\"{0}\"", string.Join(",", TestRole.Controllers)) 
-							: " -gauntlet";
+						SessionRole.CommandLineParams.Add("gauntlet", 
+							TestRole.Controllers.Count > 0 ? string.Join(",", TestRole.Controllers) : null);				
 
 						if (PassThroughArgs.Count() > 0)
 						{
diff --git a/Engine/Source/Programs/AutomationTool/Gauntlet/Unreal/BuildSource/Gauntlet.UnrealBuildSource.cs b/Engine/Source/Programs/AutomationTool/Gauntlet/Unreal/BuildSource/Gauntlet.UnrealBuildSource.cs
--- a/Engine/Source/Programs/AutomationTool/Gauntlet/Unreal/BuildSource/Gauntlet.UnrealBuildSource.cs
+++ b/Engine/Source/Programs/AutomationTool/Gauntlet/Unreal/BuildSource/Gauntlet.UnrealBuildSource.cs
@@ -395,6 +395,7 @@ namespace Gauntlet
 			Config.Platform = Role.Platform;
 			Config.Configuration = Role.Configuration;
 			Config.CommandLine = "";
+			Config.CommandLineParams = Role.CommandLineParams;
             Config.FilesToCopy = new List<UnrealFileToCopy>();
 
 			// new system of retrieving and encapsulating the info needed to install/launch. Android & Mac
@@ -416,7 +417,7 @@ namespace Gauntlet
 
 			if (string.IsNullOrEmpty(Role.CommandLine) == false)
 			{
-				Config.CommandLine += " " + Role.CommandLine;
+				Config.CommandLine += Role.CommandLine;
 			}
 
 			// Cleanup the commandline
@@ -431,11 +432,11 @@ namespace Gauntlet
 				// add in -game or -server
 				if (Role.RoleType.IsClient())
 				{
-					Config.CommandLine = "-game " + Config.CommandLine;
+					Config.CommandLineParams.Add("game");
 				}
 				else if (Role.RoleType.IsServer())
 				{
-					Config.CommandLine = "-server " + Config.CommandLine;
+					Config.CommandLineParams.Add("server");
 				}
 
 				string ProjectParam = ProjectPath.FullName;
@@ -445,16 +446,14 @@ namespace Gauntlet
 				{
 					ProjectParam = string.Format("../../../{0}/{0}.uproject", ProjectName);
 				}
-
-				// project must be first
-				Config.CommandLine = String.Format("\"{0}\"", ProjectParam) + " " + Config.CommandLine;
+				Config.CommandLineParams.Project = ProjectParam;
+				Config.CommandLine = Config.CommandLineParams.GenerateFullCommandLine();
 			}
 
             if (Role.FilesToCopy != null)
             {
                 Config.FilesToCopy = Role.FilesToCopy;
             }
-			
 			return Config;
 		}
 
diff --git a/Engine/Source/Programs/AutomationTool/Gauntlet/Unreal/Game/UE4Game.TestConfig.cs b/Engine/Source/Programs/AutomationTool/Gauntlet/Unreal/Game/UE4Game.TestConfig.cs
--- a/Engine/Source/Programs/AutomationTool/Gauntlet/Unreal/Game/UE4Game.TestConfig.cs
+++ b/Engine/Source/Programs/AutomationTool/Gauntlet/Unreal/Game/UE4Game.TestConfig.cs
@@ -36,7 +36,7 @@ namespace UE4Game
 				|| (AppConfig.ProcessType.IsClient() && RoleCount(UnrealTargetRole.Server) == 0))
 				{
 					// must be the first argument!
-					AppConfig.CommandLine = ConfigRole.MapOverride + " " + AppConfig.CommandLine;
+					AppConfig.CommandLineParams.GameMap = ConfigRole.MapOverride;
 				}
 			}
 			else if (string.IsNullOrEmpty(Map) == false)
@@ -45,8 +45,7 @@ namespace UE4Game
 				if (AppConfig.ProcessType.IsServer()
 				|| (AppConfig.ProcessType.IsClient() && RoleCount(UnrealTargetRole.Server) == 0))
 				{
-					// must be the first argument!
-					AppConfig.CommandLine = Map + " " + AppConfig.CommandLine;
+					AppConfig.CommandLineParams.GameMap = Map;
 				}
 			}			
 		}
