Handle case when the *.so is attempted to be stripped into debug mulitple times

diff --git a/Engine/Source/Programs/AutomationTool/Scripts/BuildPhysX.Automation.cs b/Engine/Source/Programs/AutomationTool/Scripts/BuildPhysX.Automation.cs
--- a/Engine/Source/Programs/AutomationTool/Scripts/BuildPhysX.Automation.cs
+++ b/Engine/Source/Programs/AutomationTool/Scripts/BuildPhysX.Automation.cs
@@ -926,6 +926,7 @@ class BuildPhysX_Linux : BuildPhysX.MakefileTargetPlatform
 	public BuildPhysX_Linux(string Architecture = "x86_64-unknown-linux-gnu")
 	{
 		this.Architecture = Architecture;
+		this.GeneratedDebugSymbols = new Dictionary<string, bool>();
 	}
 
 	private static DirectoryReference DumpSymsPath = DirectoryReference.Combine(RootDirectory, "Engine/Binaries/Linux/dump_syms");
@@ -945,6 +946,9 @@ class BuildPhysX_Linux : BuildPhysX.MakefileTargetPlatform
 	public override string TargetBuildPlatform => "linux";
 	public override string FriendlyName => Platform.ToString() + "-" + Architecture;
 
+	// Only split debug symbols for *.so we have not already do so for
+	private Dictionary<string, bool> GeneratedDebugSymbols;
+
 	public override bool SupportsTargetLib(BuildPhysX.PhysXTargetLib Library)
 	{
 		bool b64BitX86 = Architecture.StartsWith("x86_64");
@@ -1031,7 +1035,7 @@ class BuildPhysX_Linux : BuildPhysX.MakefileTargetPlatform
 		FileReference ObjcopyPath = GetObjCopyPath();
 
 		// Linux does not have a great way to split the debug file from the *.so so lets do it now as well as grab the sym file
-		foreach (FileReference SOFile in EnumerateOutputFiles(OutputBinaryDirectory, string.Format("*{0}.{1}", BuildSuffix[TargetConfiguration],DynamicLibraryExtension), TargetLib))
+		foreach (FileReference SOFile in EnumerateOutputFiles(OutputBinaryDirectory, string.Format("*{0}.{1}", BuildSuffix[TargetConfiguration], DynamicLibraryExtension), TargetLib))
 		{
 			string ExeSuffix = "";
 			if (BuildHostPlatform.Current.Platform.IsInGroup(UnrealPlatformGroup.Windows))
@@ -1064,33 +1068,38 @@ class BuildPhysX_Linux : BuildPhysX.MakefileTargetPlatform
 			// Clean up the Temp *.psym file, as they are no longer needed
 			InternalUtils.SafeDeleteFile(PSymbolFile.FullName);
 
-			// objcopy --strip-all sofile.so sofile_stripped
-			StartInfo.FileName = ObjcopyPath.FullName + ExeSuffix;
-			StartInfo.Arguments = "--strip-all " +
-				SOFile.FullName + " " +
-				StrippedFile.FullName;
-
-			LogInformation("Running: '{0} {1}'", StartInfo.FileName, StartInfo.Arguments);
-			Utils.RunLocalProcessAndLogOutput(StartInfo);
-
-			// objcopy --only-keep-debug sofile.so sofile.debug
-			StartInfo.FileName = ObjcopyPath.FullName + ExeSuffix;
-			StartInfo.Arguments = "--only-keep-debug " +
-				SOFile.FullName + " " +
-				DebugFile.FullName;
-
-			LogInformation("Running: '{0} {1}'", StartInfo.FileName, StartInfo.Arguments);
-			Utils.RunLocalProcessAndLogOutput(StartInfo);
-
-			// objcopy --add-gnu-debuglink=sofile.debug sofile_stripped sofile.so
-			StartInfo.FileName = ObjcopyPath.FullName + ExeSuffix;
-			StartInfo.Arguments = "--add-gnu-debuglink=" +
-				DebugFile.FullName    + " " +
-				StrippedFile.FullName + " " +
-				SOFile.FullName;
-
-			LogInformation("Running: '{0} {1}'", StartInfo.FileName, StartInfo.Arguments);
-			Utils.RunLocalProcessAndLogOutput(StartInfo);
+			if (!GeneratedDebugSymbols.ContainsKey(SOFile.FullName))
+			{
+				// objcopy --strip-all sofile.so sofile_stripped
+				StartInfo.FileName = ObjcopyPath.FullName + ExeSuffix;
+				StartInfo.Arguments = "--strip-all " +
+					SOFile.FullName + " " +
+					StrippedFile.FullName;
+
+				LogInformation("Running: '{0} {1}'", StartInfo.FileName, StartInfo.Arguments);
+				Utils.RunLocalProcessAndLogOutput(StartInfo);
+
+				// objcopy --only-keep-debug sofile.so sofile.debug
+				StartInfo.FileName = ObjcopyPath.FullName + ExeSuffix;
+				StartInfo.Arguments = "--only-keep-debug " +
+					SOFile.FullName + " " +
+					DebugFile.FullName;
+
+				LogInformation("Running: '{0} {1}'", StartInfo.FileName, StartInfo.Arguments);
+				Utils.RunLocalProcessAndLogOutput(StartInfo);
+
+				// objcopy --add-gnu-debuglink=sofile.debug sofile_stripped sofile.so
+				StartInfo.FileName = ObjcopyPath.FullName + ExeSuffix;
+				StartInfo.Arguments = "--add-gnu-debuglink=" +
+					DebugFile.FullName + " " +
+					StrippedFile.FullName + " " +
+					SOFile.FullName;
+
+				LogInformation("Running: '{0} {1}'", StartInfo.FileName, StartInfo.Arguments);
+				Utils.RunLocalProcessAndLogOutput(StartInfo);
+
+				GeneratedDebugSymbols.Add(SOFile.FullName, true);
+			}
 
 			// Clean up the Temp *_stripped file, as they are no longer needed
 			InternalUtils.SafeDeleteFile(StrippedFile.FullName);
