Benchmark script - clean S3DDC files when clean benchmarking

#tests ran script with and without clean

diff --git a/Engine/Source/Programs/AutomationTool/Scripts/BenchmarkBuild/BenchmarkTaskBase.cs b/Engine/Source/Programs/AutomationTool/Scripts/BenchmarkBuild/BenchmarkTaskBase.cs
--- a/Engine/Source/Programs/AutomationTool/Scripts/BenchmarkBuild/BenchmarkTaskBase.cs
+++ b/Engine/Source/Programs/AutomationTool/Scripts/BenchmarkBuild/BenchmarkTaskBase.cs
@@ -277,6 +277,10 @@ namespace AutomationTool.Benchmark
 					CachePaths.Add(Dir);
 					Environment.SetEnvironmentVariable(Key, Dir.FullName);
 				}				
+
+				// remove S3 files
+				DirectoryReference S3DDC = DirectoryReference.Combine(ProjectFile.Directory, "Saved", "S3DDC");
+				CommandUtils.DeleteDirectory_NoExceptions(S3DDC.FullName);
 			}
 
 			return base.PerformPrequisites();
