Fixed benchmark script so singlecompile doesn't do a rebuild when -noxge is used

diff --git a/Engine/Source/Programs/AutomationTool/Scripts/BenchmarkBuild/BenchmarkBuild.Automation.cs b/Engine/Source/Programs/AutomationTool/Scripts/BenchmarkBuild/BenchmarkBuild.Automation.cs
--- a/Engine/Source/Programs/AutomationTool/Scripts/BenchmarkBuild/BenchmarkBuild.Automation.cs
+++ b/Engine/Source/Programs/AutomationTool/Scripts/BenchmarkBuild/BenchmarkBuild.Automation.cs
@@ -134,6 +134,8 @@ namespace AutomationTool.Benchmark
 
 				if (DoBuildEditorTests)
 				{
+					BuildOptions NoAndSingleCompileOptions = BuildOptions.None;
+
 					if (DoXGE)
 					{
 						Tasks.Add(new BenchmarkBuildTask(Project, "Editor", EditorPlatform, CleanFlag));
@@ -142,13 +144,14 @@ namespace AutomationTool.Benchmark
 					if (DoNoXGE)
 					{
 						Tasks.Add(new BenchmarkBuildTask(Project, "Editor", EditorPlatform, CleanFlag | BuildOptions.NoXGE));
+						NoAndSingleCompileOptions |= BuildOptions.NoXGE;
 					}
 
 					if (DoNoCompile)
 					{
 						
 						// note, don't clean since we build normally then build a single file
-						Tasks.Add(new BenchmarkNopCompileTask(Project, "Editor", EditorPlatform, BuildOptions.None));
+						Tasks.Add(new BenchmarkNopCompileTask(Project, "Editor", EditorPlatform, NoAndSingleCompileOptions));
 					}
 
 					if (DoSingleCompile)
@@ -156,7 +159,7 @@ namespace AutomationTool.Benchmark
 						FileReference SourceFile = FindProjectSourceFile(Project);
 
 						// note, don't clean since we build normally then build again
-						Tasks.Add(new BenchmarkSingleCompileTask(Project, "Editor", EditorPlatform, SourceFile, BuildOptions.None));
+						Tasks.Add(new BenchmarkSingleCompileTask(Project, "Editor", EditorPlatform, SourceFile, NoAndSingleCompileOptions));
 					}
 				}
 
@@ -167,6 +170,8 @@ namespace AutomationTool.Benchmark
 
 					foreach (var ClientPlatform in ClientPlatforms)
 					{
+						BuildOptions NoAndSingleCompileOptions = BuildOptions.None;
+
 						if (DoXGE)
 						{
 							Tasks.Add(new BenchmarkBuildTask(Project, TargetName, ClientPlatform, CleanFlag));
@@ -175,12 +180,13 @@ namespace AutomationTool.Benchmark
 						if (DoNoXGE)
 						{
 							Tasks.Add(new BenchmarkBuildTask(Project, TargetName, ClientPlatform, CleanFlag | BuildOptions.NoXGE));
+							NoAndSingleCompileOptions |= BuildOptions.NoXGE;
 						}
 
 						if (DoNoCompile)
 						{
 							// note, don't clean since we build normally then build again
-							Tasks.Add(new BenchmarkNopCompileTask(Project, TargetName, ClientPlatform, BuildOptions.NoXGE));
+							Tasks.Add(new BenchmarkNopCompileTask(Project, TargetName, ClientPlatform, NoAndSingleCompileOptions));
 						}
 
 						if (DoSingleCompile)
@@ -188,7 +194,7 @@ namespace AutomationTool.Benchmark
 							FileReference SourceFile = FindProjectSourceFile(Project);
 
 							// note, don't clean since we build normally then build a single file
-							Tasks.Add(new BenchmarkSingleCompileTask(Project, TargetName, ClientPlatform, SourceFile, BuildOptions.None));
+							Tasks.Add(new BenchmarkSingleCompileTask(Project, TargetName, ClientPlatform, SourceFile, NoAndSingleCompileOptions));
 						}
 					}
 				}
