Fixed noxge in benchmark script

diff --git a/Engine/Source/Programs/AutomationTool/Scripts/BenchmarkBuild/BenchmarkBuild.Automation.cs b/Engine/Source/Programs/AutomationTool/Scripts/BenchmarkBuild/BenchmarkBuild.Automation.cs
--- a/Engine/Source/Programs/AutomationTool/Scripts/BenchmarkBuild/BenchmarkBuild.Automation.cs
+++ b/Engine/Source/Programs/AutomationTool/Scripts/BenchmarkBuild/BenchmarkBuild.Automation.cs
@@ -268,6 +268,7 @@ namespace AutomationTool.Benchmark
 
 						WriteCSVResults(FileName);
 
+						Log.TraceInformation("Waiting {0} secs until next task {1}", TimeBetweenTasks);
 						Thread.Sleep(TimeBetweenTasks * 1000);
 					}
 				}
diff --git a/Engine/Source/Programs/AutomationTool/Scripts/BenchmarkBuild/BenchmarkBuildTests.cs b/Engine/Source/Programs/AutomationTool/Scripts/BenchmarkBuild/BenchmarkBuildTests.cs
--- a/Engine/Source/Programs/AutomationTool/Scripts/BenchmarkBuild/BenchmarkBuildTests.cs
+++ b/Engine/Source/Programs/AutomationTool/Scripts/BenchmarkBuild/BenchmarkBuildTests.cs
@@ -78,6 +78,7 @@ namespace AutomationTool.Benchmark
 
 				Command.UBTArgs += " -" + Arg;
 				TaskModifiers.Add(Arg);
+				Command.Params = new[] { Arg };	// need to also pass it to this
 			}
 			else
 			{
diff --git a/Engine/Source/Programs/AutomationTool/Scripts/BuildTarget.Automation.cs b/Engine/Source/Programs/AutomationTool/Scripts/BuildTarget.Automation.cs
--- a/Engine/Source/Programs/AutomationTool/Scripts/BuildTarget.Automation.cs
+++ b/Engine/Source/Programs/AutomationTool/Scripts/BuildTarget.Automation.cs
@@ -174,11 +174,11 @@ namespace AutomationTool
 
 				if (!NoTools)
 				{
-					Agenda.AddTarget("UnrealPak", CurrentPlatform, UnrealTargetConfiguration.Development, ProjectFile);
-					Agenda.AddTarget("ShaderCompileWorker", CurrentPlatform, UnrealTargetConfiguration.Development);
-					Agenda.AddTarget("UnrealLightmass", CurrentPlatform, UnrealTargetConfiguration.Development);
-					Agenda.AddTarget("CrashReportClient", CurrentPlatform, UnrealTargetConfiguration.Shipping);
-					Agenda.AddTarget("CrashReportClientEditor", CurrentPlatform, UnrealTargetConfiguration.Shipping);
+					Agenda.AddTarget("UnrealPak", CurrentPlatform, UnrealTargetConfiguration.Development, ProjectFile, UBTArgs);
+					Agenda.AddTarget("ShaderCompileWorker", CurrentPlatform, UnrealTargetConfiguration.Development, null, UBTArgs);
+					Agenda.AddTarget("UnrealLightmass", CurrentPlatform, UnrealTargetConfiguration.Development, null, UBTArgs);
+					Agenda.AddTarget("CrashReportClient", CurrentPlatform, UnrealTargetConfiguration.Shipping, null, UBTArgs);
+					Agenda.AddTarget("CrashReportClientEditor", CurrentPlatform, UnrealTargetConfiguration.Shipping, null, UBTArgs);
 				}
 			}
 
