Fixed copy & pasta with format string

diff --git a/Engine/Source/Programs/AutomationTool/Scripts/BenchmarkBuild/BenchmarkBuild.Automation.cs b/Engine/Source/Programs/AutomationTool/Scripts/BenchmarkBuild/BenchmarkBuild.Automation.cs
--- a/Engine/Source/Programs/AutomationTool/Scripts/BenchmarkBuild/BenchmarkBuild.Automation.cs
+++ b/Engine/Source/Programs/AutomationTool/Scripts/BenchmarkBuild/BenchmarkBuild.Automation.cs
@@ -268,7 +268,7 @@ namespace AutomationTool.Benchmark
 
 						WriteCSVResults(FileName);
 
-						Log.TraceInformation("Waiting {0} secs until next task {1}", TimeBetweenTasks);
+						Log.TraceInformation("Waiting {0} secs until next task", TimeBetweenTasks);
 						Thread.Sleep(TimeBetweenTasks * 1000);
 					}
 				}
