Added CPU count to benchmark results

diff --git a/Engine/Source/Programs/AutomationTool/Scripts/BenchmarkBuild/BenchmarkBuild.Automation.cs b/Engine/Source/Programs/AutomationTool/Scripts/BenchmarkBuild/BenchmarkBuild.Automation.cs
--- a/Engine/Source/Programs/AutomationTool/Scripts/BenchmarkBuild/BenchmarkBuild.Automation.cs
+++ b/Engine/Source/Programs/AutomationTool/Scripts/BenchmarkBuild/BenchmarkBuild.Automation.cs
@@ -334,8 +334,8 @@ namespace AutomationTool.Benchmark
 			{
 				List<string> Lines = new List<string>();
 
-				// first line is machine name,,Iteration 1, Iteration 2 etc
-				string FirstLine = string.Format("{0},", Environment.MachineName);
+				// first line is machine name,CPU count,Iteration 1, Iteration 2 etc
+				string FirstLine = string.Format("{0},{1}", Environment.MachineName, Environment.ProcessorCount);
 
 				if (Tasks.Count() > 0)
 				{
