Benchmark Script

Fixed issue when building vanilla UE4 with no project

diff --git a/Engine/Source/Programs/AutomationTool/Scripts/BenchmarkBuild/BenchmarkBuild.Automation.cs b/Engine/Source/Programs/AutomationTool/Scripts/BenchmarkBuild/BenchmarkBuild.Automation.cs
--- a/Engine/Source/Programs/AutomationTool/Scripts/BenchmarkBuild/BenchmarkBuild.Automation.cs
+++ b/Engine/Source/Programs/AutomationTool/Scripts/BenchmarkBuild/BenchmarkBuild.Automation.cs
@@ -201,7 +201,7 @@ namespace AutomationTool.Benchmark
 
 				FileReference ProjectFile = ProjectUtils.FindProjectFileFromName(Project);
 
-				if (ProjectFile == null)
+				if (ProjectFile == null && !Project.Equals("UE4", StringComparison.OrdinalIgnoreCase))
 				{
 					throw new AutomationException("Could not find project file for {0}", Project);
 				}
