Harden device reservation renewal and deletion
#rnx

diff --git a/Engine/Source/Programs/AutomationTool/AutomationUtils/DeviceReservation.cs b/Engine/Source/Programs/AutomationTool/AutomationUtils/DeviceReservation.cs
--- a/Engine/Source/Programs/AutomationTool/AutomationUtils/DeviceReservation.cs
+++ b/Engine/Source/Programs/AutomationTool/AutomationUtils/DeviceReservation.cs
@@ -79,27 +79,38 @@ namespace AutomationTool.DeviceReservation
 			{
 				try
 				{
-					ActiveReservation.Renew(ReservationBaseUri, ReserveTime);
-					RetryCurrent = 0;
-					RenewTimeCurrent = RenewTime;					
+					// Renew the reservation on the backend
+					if (RetryCurrent <= RenewRetryMax)
+					{
+						ActiveReservation.Renew(ReservationBaseUri, ReserveTime);
+						RetryCurrent = 0;
+						RenewTimeCurrent = RenewTime;
+					}
 				}
-				catch (AutomationException ex)
+				catch (Exception ex)
 				{
-					// @todo: finer grain on renew error from service
+					// There was an exception, warn if we've exceeded retry
 					if (RetryCurrent == RenewRetryMax)
 					{
-						throw new AutomationException(ex, "Reserveration renew exception.");
-					}
-					else
-					{											
-						// try again
-						RetryCurrent++;
-						RenewTimeCurrent = RenewRetryTime;
+						Console.WriteLine("Warning: Reservation renewal returned bad status: {0}", ex);
 					}
+
+					// try again
+					RetryCurrent++;
+					RenewTimeCurrent = RenewRetryTime;
 				}				
 			}
 
-			ActiveReservation.Delete(ReservationBaseUri);
+			// Delete reservation on server, if the web request fails backend has logic to cleanup reservations
+			try
+			{
+				ActiveReservation.Delete(ReservationBaseUri);
+			}
+			catch (Exception ex)
+			{
+				Console.WriteLine("Warning: Reservation delete returned bad status: {0}", ex);
+			}
+
 		}
 
 		private void StopAutoRenew()
