Sanitize all device reservation output for EC parsing
#rnx

diff --git a/Engine/Source/Programs/AutomationTool/AutomationUtils/DeviceReservation.cs b/Engine/Source/Programs/AutomationTool/AutomationUtils/DeviceReservation.cs
--- a/Engine/Source/Programs/AutomationTool/AutomationUtils/DeviceReservation.cs
+++ b/Engine/Source/Programs/AutomationTool/AutomationUtils/DeviceReservation.cs
@@ -87,12 +87,12 @@ namespace AutomationTool.DeviceReservation
 						RenewTimeCurrent = RenewTime;
 					}
 				}
-				catch (Exception ex)
+				catch (Exception Ex)
 				{
 					// There was an exception, warn if we've exceeded retry
 					if (RetryCurrent == RenewRetryMax)
 					{
-						Console.WriteLine("Warning: Reservation renewal returned bad status: {0}", ex);
+						Utils.Log(string.Format("Warning: Reservation renewal returned bad status: {0}", Ex.Message));
 					}
 
 					// try again
@@ -106,9 +106,9 @@ namespace AutomationTool.DeviceReservation
 			{
 				ActiveReservation.Delete(ReservationBaseUri);
 			}
-			catch (Exception ex)
+			catch (Exception Ex)
 			{
-				Console.WriteLine("Warning: Reservation delete returned bad status: {0}", ex);
+				Utils.Log(string.Format("Warning: Reservation delete returned bad status: {0}", Ex.Message));
 			}
 
 		}
@@ -137,6 +137,31 @@ namespace AutomationTool.DeviceReservation
 
 	public static class Utils
 	{
+
+		static public string SanitizeErrorMessage(string Message)
+		{
+			string[] TriggersSrc = { "Warning:", "Error:", "Exception:" };
+			string[] TriggersDst = { "Warn1ng:", "Err0r:", "Except10n:" };
+
+			for (int Index = 0; Index < TriggersSrc.Length; ++Index)
+			{
+				if (Message.IndexOf(TriggersSrc[Index], StringComparison.OrdinalIgnoreCase) != -1)
+				{
+					Message = Regex.Replace(Message, TriggersSrc[Index], TriggersDst[Index], RegexOptions.IgnoreCase);
+				}
+			}
+
+			return Message;
+		}
+
+		static public void Log(string Message)
+		{
+			Console.WriteLine("<-- Suspend Log Parsing -->");
+			Console.WriteLine(SanitizeErrorMessage(Message));
+			Console.WriteLine("<-- Resume Log Parsing -->");
+
+		}
+
 		public static Uri AppendPath(this Uri BaseUri, string NewPath)
 		{
 			var Builder = new UriBuilder(BaseUri);
@@ -217,22 +242,6 @@ namespace AutomationTool.DeviceReservation
 			public string PoolID;
 		}
 
-		private static string SanitizeErrorMessage(string Message)
-		{
-			string[] TriggersSrc = { "Warning:", "Error:", "Exception:" };
-			string[] TriggersDst = { "Warn1ng:", "Err0r:", "Except10n:" };
-
-			for (int Index = 0; Index < TriggersSrc.Length; ++Index)
-			{
-				if (Message.IndexOf(TriggersSrc[Index], StringComparison.OrdinalIgnoreCase) != -1)
-				{
-					Message = Regex.Replace(Message, TriggersSrc[Index], TriggersDst[Index], RegexOptions.IgnoreCase);
-				}
-			}
-
-			return Message;
-		}
-
 		public static Reservation Create(Uri BaseUri, string[] DeviceTypes, TimeSpan Duration, int RetryMax = 5, string PoolID = "")
 		{
 			bool bFirst = true;
@@ -265,7 +274,7 @@ namespace AutomationTool.DeviceReservation
 				}
 				catch (WebException WebEx)
 				{
-					Console.WriteLine(String.Format("WebException on reservation request: {0} : {1}", SanitizeErrorMessage(WebEx.Message), WebEx.Status));
+					Utils.Log(String.Format("WebException on reservation request: {0} : {1}", WebEx.Message, WebEx.Status));
 
 					if (RetryCount == RetryMax)
 					{
@@ -292,7 +301,7 @@ namespace AutomationTool.DeviceReservation
 				catch (Exception Ex)
 				{
 					UnknownException = Ex;
-					Console.WriteLine("Device reservation unsuccessful: {0}", SanitizeErrorMessage(UnknownException.Message));
+					Utils.Log(string.Format("Device reservation unsuccessful: {0}", UnknownException.Message));
 				}
 			}
 		}
@@ -319,7 +328,7 @@ namespace AutomationTool.DeviceReservation
 			}
 			catch (Exception ex)
 			{
-				Console.WriteLine("Failed to delete device reservation: {0}", SanitizeErrorMessage(ex.Message));
+				Utils.Log(string.Format("Failed to delete device reservation: {0}", ex.Message));
 			}
 		}
 
@@ -334,11 +343,11 @@ namespace AutomationTool.DeviceReservation
 			{
 				Uri BaseUri = new Uri(InBaseUri);
 				Utils.InvokeAPI(BaseUri.AppendPath("api/v1/deviceerror/" + DeviceName), "PUT");
-				Console.WriteLine("Reported device problem: {0} : {1}", DeviceName, Error);
+				Utils.Log(string.Format("Reported device problem: {0} : {1}", DeviceName, Error));
 			}
 			catch (Exception Ex)
 			{
-				Console.WriteLine("Failed to report device: {0} : {1}", DeviceName, SanitizeErrorMessage(Ex.Message));
+				Utils.Log(string.Format("Failed to report device: {0} : {1}", DeviceName, Ex.Message));
 			}
 		}
 
