Benchmark Tests

- made noxge apply to shader compiles
- added option to test backends
- tweaked some cleanup stuff.

diff --git a/Engine/Source/Programs/AutomationTool/Scripts/BenchmarkBuild/BenchmarkBuild.Automation.cs b/Engine/Source/Programs/AutomationTool/Scripts/BenchmarkBuild/BenchmarkBuild.Automation.cs
--- a/Engine/Source/Programs/AutomationTool/Scripts/BenchmarkBuild/BenchmarkBuild.Automation.cs
+++ b/Engine/Source/Programs/AutomationTool/Scripts/BenchmarkBuild/BenchmarkBuild.Automation.cs
@@ -37,7 +37,7 @@ namespace AutomationTool.Benchmark
 	[Help("hotddc", "Cook / PIE with a hot local DDC (an untimed pre-run is performed)")]
 	[Help("coldddc", "Cook / PIE with a cold local DDC (a temporary folder is used)")]
 	[Help("noshaderddc", "Cook / PIE with no shaders in the DDC")]
-	[Help("nosharedddc", "Cook / PIE with no shared ddc")]
+	[Help("ddc=", "Type of DDC to use. E.g. ddc=noshared, ddc=noshared+shared")]
 	[Help("iterations=<n>", "How many times to perform each test)")]
 	[Help("wait=<n>", "How many seconds to wait between each test)")]
 	[Help("filename", "Name/path of file to write CSV results to. If empty the local machine name will be used")]
@@ -70,17 +70,19 @@ namespace AutomationTool.Benchmark
 			// editor startup tests
 			public bool DoPIETests = false;
 
-			public string MapName = "";
+			public IEnumerable<string> MapList = Enumerable.Empty<string>();
 
 			// misc
 			public int Iterations = 1;
 			public bool NoClean = false;
 			public int TimeBetweenTasks = 0;
+			
 			public string CookArgs = "";
 			public string PIEArgs = "";
 			public string FileName = string.Format("{0}_Results.csv", Environment.MachineName);
 
 			public DDCTaskOptions DDCOptions = DDCTaskOptions.None;
+			public IEnumerable<string> DCCTypes = Enumerable.Empty<string>();
 
 			public void ParseParams(string[] InParams)
 			{
@@ -105,14 +107,13 @@ namespace AutomationTool.Benchmark
 
 				// editor startup tests
 				DoPIETests = AllThings | ParseParam("pie");
-				MapName = ParseParamValue("map", "");
-
+				
 				// DDC options
 				DDCOptions |= ParseParam("warmddc") ? DDCTaskOptions.WarmDDC : DDCTaskOptions.None;
 				DDCOptions |= ParseParam("hotddc") ? DDCTaskOptions.HotDDC : DDCTaskOptions.None;
 				DDCOptions |= ParseParam("coldddc") ? DDCTaskOptions.ColdDDC : DDCTaskOptions.None;
-				DDCOptions |= ParseParam("nosharedddc") ? DDCTaskOptions.NoSharedDDC : DDCTaskOptions.None;
 				DDCOptions |= ParseParam("noshaderddc") ? DDCTaskOptions.NoShaderDDC : DDCTaskOptions.None;
+				DDCOptions |= ParseParam("noxge") ? DDCTaskOptions.NoXGE : DDCTaskOptions.None;
 
 				// sanity
 				DoAcceleratedCompileTests = DoAcceleratedCompileTests && BenchmarkBuildTask.SupportsAcceleration;
@@ -179,6 +180,26 @@ namespace AutomationTool.Benchmark
 						CoresForLocalJobs = ProcessorList.Select(P => Convert.ToInt32(P));
 					}
 				}
+
+				// parse ddc args
+				{
+					string Arg = ParseParamValue("ddc", "");
+
+					if (!string.IsNullOrEmpty(Arg))
+					{
+						DCCTypes = Arg.Split(new[] { '+', ',' }, StringSplitOptions.RemoveEmptyEntries);
+					}
+				}
+
+				// parse map args
+				{
+					string Arg = ParseParamValue("map", "");
+
+					if (!string.IsNullOrEmpty(Arg))
+					{
+						MapList = Arg.Split(new[] { '+', ',' }, StringSplitOptions.RemoveEmptyEntries);
+					}
+				}
 			}
 		}
 
@@ -407,7 +428,7 @@ namespace AutomationTool.Benchmark
 
 		IEnumerable<BenchmarkTaskBase> AddPIETests(FileReference InProjectFile, BenchmarkOptions InOptions)
 		{
-			if (InProjectFile == null)
+			if (InProjectFile == null || !InOptions.DoPIETests)
 			{
 				return Enumerable.Empty<BenchmarkTaskBase>();
 			}
@@ -416,44 +437,68 @@ namespace AutomationTool.Benchmark
 
 			string DefaultExtraArgs = InOptions.PIEArgs;
 
-			if (!string.IsNullOrEmpty(InOptions.MapName))
+			List<string> MapsToTest = InOptions.MapList.ToList();
+
+			if (!MapsToTest.Any())
 			{
-				DefaultExtraArgs += " -map=" + InOptions.MapName;
+				MapsToTest.Add("");
 			}
 
-			if (InOptions.DoPIETests)
+			List<string> BackendsToTest = InOptions.DCCTypes.ToList();
+
+			if (!BackendsToTest.Any())
 			{
-				// if no options assume warm
-				if (InOptions.DDCOptions == DDCTaskOptions.None || InOptions.DDCOptions.HasFlag(DDCTaskOptions.WarmDDC))
+				BackendsToTest.Add("");
+			}
+		
+			foreach (var Map in MapsToTest)
+			{
+				foreach (string Backend in BackendsToTest)
 				{
-					NewTasks.Add(new BenchmarkRunEditorTask(InProjectFile, DDCTaskOptions.WarmDDC, DefaultExtraArgs));
-				}
+					string FinalArgs = DefaultExtraArgs;
 
-				// hot ddc
-				if (InOptions.DDCOptions.HasFlag(DDCTaskOptions.HotDDC))
-				{
-					NewTasks.Add(new BenchmarkRunEditorTask(InProjectFile, DDCTaskOptions.HotDDC, DefaultExtraArgs));
-				}
+					if (!string.IsNullOrEmpty(Map))
+					{
+						FinalArgs += " -map=" + Map;
+					}
 
-				// cold ddc
-				if (InOptions.DDCOptions.HasFlag(DDCTaskOptions.ColdDDC))
-				{
-					NewTasks.Add(new BenchmarkRunEditorTask(InProjectFile, DDCTaskOptions.ColdDDC, DefaultExtraArgs));
-				}
+					if (!string.IsNullOrEmpty(Backend))
+					{
+						FinalArgs += string.Format(" -ddc={0}", Backend);
+					}
 
-				// no shaders in the ddc
-				if (InOptions.DDCOptions.HasFlag(DDCTaskOptions.NoShaderDDC))
-				{
-					NewTasks.Add(new BenchmarkRunEditorTask(InProjectFile, DDCTaskOptions.NoShaderDDC, DefaultExtraArgs));
-				}
+					// if no options assume warm
+					if (InOptions.DDCOptions == DDCTaskOptions.None || InOptions.DDCOptions.HasFlag(DDCTaskOptions.WarmDDC))
+					{
+						NewTasks.Add(new BenchmarkRunEditorTask(InProjectFile, DDCTaskOptions.WarmDDC, FinalArgs));
+					}
 
-				// no ddc!
-				if (InOptions.DDCOptions.HasFlag(DDCTaskOptions.NoSharedDDC))
-				{
-					NewTasks.Add(new BenchmarkRunEditorTask(InProjectFile, DDCTaskOptions.NoSharedDDC, DefaultExtraArgs));
+					// hot ddc
+					if (InOptions.DDCOptions.HasFlag(DDCTaskOptions.HotDDC))
+					{
+						NewTasks.Add(new BenchmarkRunEditorTask(InProjectFile, DDCTaskOptions.HotDDC, FinalArgs));
+					}
+
+					// cold ddc
+					if (InOptions.DDCOptions.HasFlag(DDCTaskOptions.ColdDDC))
+					{
+						NewTasks.Add(new BenchmarkRunEditorTask(InProjectFile, DDCTaskOptions.ColdDDC, FinalArgs));
+					}
+
+					// no shaders in the ddc
+					if (InOptions.DDCOptions.HasFlag(DDCTaskOptions.NoShaderDDC))
+					{
+						NewTasks.Add(new BenchmarkRunEditorTask(InProjectFile, DDCTaskOptions.NoShaderDDC, FinalArgs));
+					}
+
+					// no ddc!
+					if (InOptions.DDCOptions.HasFlag(DDCTaskOptions.NoSharedDDC))
+					{
+						NewTasks.Add(new BenchmarkRunEditorTask(InProjectFile, DDCTaskOptions.NoSharedDDC, FinalArgs));
+					}
 				}
 			}
-
+		
 			return NewTasks;
 		}
 
diff --git a/Engine/Source/Programs/AutomationTool/Scripts/BenchmarkBuild/BenchmarkCookTask.cs b/Engine/Source/Programs/AutomationTool/Scripts/BenchmarkBuild/BenchmarkCookTask.cs
--- a/Engine/Source/Programs/AutomationTool/Scripts/BenchmarkBuild/BenchmarkCookTask.cs
+++ b/Engine/Source/Programs/AutomationTool/Scripts/BenchmarkBuild/BenchmarkCookTask.cs
@@ -81,7 +81,6 @@ namespace AutomationTool.Benchmark
 			if (TaskOptions.HasFlag(DDCTaskOptions.NoShaderDDC))
 			{
 				ExtraArgsList.Add("noshaderddc");
-				//ExtraArgs.Add("noxgeshadercompile");
 			}
 
 			if (TaskOptions.HasFlag(DDCTaskOptions.NoSharedDDC))
@@ -89,6 +88,11 @@ namespace AutomationTool.Benchmark
 				ExtraArgsList.Add("ddc=noshared");
 			}
 
+			if (TaskOptions.HasFlag(DDCTaskOptions.NoXGE))
+			{
+				ExtraArgsList.Add("noxgeshadercompile");
+			}
+
 			string ExtraArgs = "";
 
 			if (ExtraArgsList.Any())
diff --git a/Engine/Source/Programs/AutomationTool/Scripts/BenchmarkBuild/BenchmarkRunEditorTask.cs b/Engine/Source/Programs/AutomationTool/Scripts/BenchmarkBuild/BenchmarkRunEditorTask.cs
--- a/Engine/Source/Programs/AutomationTool/Scripts/BenchmarkBuild/BenchmarkRunEditorTask.cs
+++ b/Engine/Source/Programs/AutomationTool/Scripts/BenchmarkBuild/BenchmarkRunEditorTask.cs
@@ -37,7 +37,7 @@ namespace AutomationTool.Benchmark
 			// if they want a hot DDC then do the test one time with no timing
 			if (TaskOptions.HasFlag(DDCTaskOptions.HotDDC))
 			{
-				RunEditorAndWaitForMapLoad();
+				RunEditorAndWaitForMapLoad(true);
 			}
 
 			base.PerformPrequisites();
@@ -80,7 +80,7 @@ namespace AutomationTool.Benchmark
 			return System.Text.RegularExpressions.Regex.Replace(name, invalidRegStr, "_");
 		}
 
-		protected bool RunEditorAndWaitForMapLoad()
+		protected bool RunEditorAndWaitForMapLoad(bool bIsWarming)
 		{
 			string ProjectArg = ProjectFile != null ? ProjectFile.ToString() : "";
 			string EditorPath = HostPlatform.Current.GetUE4ExePath("UE4Editor.exe");
@@ -92,6 +92,16 @@ namespace AutomationTool.Benchmark
 				Arguments += (" -ddc=noshared");
 			}
 
+			if (!bIsWarming && TaskOptions.HasFlag(DDCTaskOptions.NoShaderDDC))
+			{
+				Arguments += (" -noshaderddc");
+			}
+
+			if (TaskOptions.HasFlag(DDCTaskOptions.NoXGE))
+			{
+				Arguments += (" -noxgeshadercompile");
+			}
+
 			var RunOptions = CommandUtils.ERunOptions.AllowSpew | CommandUtils.ERunOptions.NoWaitForExit;
 
 			var SpewDelegate = new ProcessResult.SpewFilterCallbackType(EndOnMapCheckFilter);
@@ -126,7 +136,7 @@ namespace AutomationTool.Benchmark
 
 		protected override bool PerformTask()
 		{
-			return RunEditorAndWaitForMapLoad();
+			return RunEditorAndWaitForMapLoad(false);
 		}
 	}
 }
diff --git a/Engine/Source/Programs/AutomationTool/Scripts/BenchmarkBuild/BenchmarkTaskBase.cs b/Engine/Source/Programs/AutomationTool/Scripts/BenchmarkBuild/BenchmarkTaskBase.cs
--- a/Engine/Source/Programs/AutomationTool/Scripts/BenchmarkBuild/BenchmarkTaskBase.cs
+++ b/Engine/Source/Programs/AutomationTool/Scripts/BenchmarkBuild/BenchmarkTaskBase.cs
@@ -164,8 +164,9 @@ namespace AutomationTool.Benchmark
 		NoSharedDDC = 1 << 2,
 		NoShaderDDC = 1 << 3,
 		HotDDC = 1 << 4,
+		NoXGE = 1 << 5,			// don't use XGE for shader compilation
 
-		KeepMemoryDDC = 1 << 5,
+		KeepMemoryDDC = 1 << 6,
 	}
 
 	abstract class BenchmarkEditorTaskBase : BenchmarkTaskBase
@@ -276,7 +277,11 @@ namespace AutomationTool.Benchmark
 					// save this dir and set it as the env var
 					CachePaths.Add(Dir);
 					Environment.SetEnvironmentVariable(Key, Dir.FullName);
-				}				
+				}
+
+				// remove project files
+				DirectoryReference ProjectDDC = DirectoryReference.Combine(ProjectFile.Directory, "DerivedDataCache");
+				CommandUtils.DeleteDirectory_NoExceptions(ProjectDDC.FullName);
 
 				// remove S3 files
 				DirectoryReference S3DDC = DirectoryReference.Combine(ProjectFile.Directory, "Saved", "S3DDC");
