Fix slow UAT startup due to populating metadata for every project in branch.

diff --git a/Engine/Source/Programs/AutomationTool/AutomationUtils/ProjectUtils.cs b/Engine/Source/Programs/AutomationTool/AutomationUtils/ProjectUtils.cs
--- a/Engine/Source/Programs/AutomationTool/AutomationUtils/ProjectUtils.cs
+++ b/Engine/Source/Programs/AutomationTool/AutomationUtils/ProjectUtils.cs
@@ -930,7 +930,20 @@ namespace AutomationTool
         {
             public string GameName;
             public FileReference FilePath;
-            public ProjectProperties Properties;
+
+			private ProjectProperties CachedProperties;
+			
+			public ProjectProperties Properties
+			{
+				get
+				{
+					if(CachedProperties == null)
+					{
+						CachedProperties = ProjectUtils.GetProjectProperties(FilePath);
+					}
+					return CachedProperties;
+				}
+			}
 
             public BranchUProject(FileReference ProjectFile)
             {
@@ -943,113 +956,46 @@ namespace AutomationTool
                 {
                     throw new AutomationException("Could not resolve relative path corrctly {0} -> {1} which doesn't exist.", ProjectFile, FilePath);
                 }
-
-                Properties = ProjectUtils.GetProjectProperties(FilePath);
-
-
-
             }
+
             public BranchUProject()
             {
                 GameName = "UE4";
-                Properties = ProjectUtils.GetProjectProperties(null);
                 if (!Properties.Targets.Exists(Target => Target.Rules.Type == TargetType.Editor))
                 {
                     throw new AutomationException("Base UE4 project did not contain an editor target.");
                 }
             }
-            public void Dump(List<UnrealTargetPlatform> InHostPlatforms)
-            {
-                CommandUtils.LogVerbose("    ShortName:    " + GameName);
-				CommandUtils.LogVerbose("      FilePath          : " + FilePath);
-				CommandUtils.LogVerbose("      bIsCodeBasedProject  : " + (Properties.bIsCodeBasedProject ? "YES" : "NO"));
-                foreach (UnrealTargetPlatform HostPlatform in InHostPlatforms)
-                {
-					CommandUtils.LogVerbose("      For Host : " + HostPlatform.ToString());
-					CommandUtils.LogVerbose("          Targets {0}:", Properties.Targets.Count);
-                    foreach (SingleTargetProperties ThisTarget in Properties.Targets)
-                    {
-						CommandUtils.LogVerbose("            TargetName          : " + ThisTarget.TargetName);
-						CommandUtils.LogVerbose("              Type          : " + ThisTarget.Rules.Type);
-						CommandUtils.LogVerbose("              bUsesSteam  : " + (ThisTarget.Rules.bUsesSteam ? "YES" : "NO"));
-						CommandUtils.LogVerbose("              bUsesCEF3   : " + (ThisTarget.Rules.bUsesCEF3 ? "YES" : "NO"));
-						CommandUtils.LogVerbose("              bUsesSlate  : " + (ThisTarget.Rules.bUsesSlate ? "YES" : "NO"));
-                    }
-					CommandUtils.LogVerbose("      Programs {0}:", Properties.Programs.Count);
-                    foreach (SingleTargetProperties ThisTarget in Properties.Programs)
-                    {
-						CommandUtils.LogVerbose("            TargetName                    : " + ThisTarget.TargetName);
-                    }
-                }
-            }
-        };
-        public BranchUProject BaseEngineProject = null;
-        public List<BranchUProject> CodeProjects = new List<BranchUProject>();
-        public List<BranchUProject> NonCodeProjects = new List<BranchUProject>();
+        }
 
-		public IEnumerable<BranchUProject> AllProjects
-		{
-			get { return CodeProjects.Union(NonCodeProjects); }
-		}
+        public BranchUProject BaseEngineProject = null;
+		public List<BranchUProject> AllProjects = new List<BranchUProject>();
 
-        public BranchInfo(List<UnrealTargetPlatform> InHostPlatforms)
+        public BranchInfo()
         {
             BaseEngineProject = new BranchUProject();
 
-            IEnumerable<FileReference> AllProjects = UnrealBuildTool.NativeProjects.EnumerateProjectFiles();
-			using(TelemetryStopwatch SortProjectsStopwatch = new TelemetryStopwatch("SortProjects"))
+            IEnumerable<FileReference> ProjectFiles = UnrealBuildTool.NativeProjects.EnumerateProjectFiles();
+			foreach (FileReference InfoEntry in ProjectFiles)
 			{
-				foreach (FileReference InfoEntry in AllProjects)
-				{
-					BranchUProject UProject = new BranchUProject(InfoEntry);
-					if (UProject.Properties.bIsCodeBasedProject)
-					{
-						CodeProjects.Add(UProject);
-					}
-					else
-					{
-						NonCodeProjects.Add(UProject);
-						// the base project uses BlankProject if it really needs a .uproject file
-						if (BaseEngineProject.FilePath == null && UProject.GameName == "BlankProject")
-						{
-							BaseEngineProject.FilePath = UProject.FilePath;
-						}
-					}
-				}
+				AllProjects.Add(new BranchUProject(InfoEntry));
 			}
- /*           if (String.IsNullOrEmpty(BaseEngineProject.FilePath))
-            {
-                throw new AutomationException("All branches must have the blank project /Samples/Sandbox/BlankProject");
-            }*/
 
-			using(TelemetryStopwatch ProjectDumpStopwatch = new TelemetryStopwatch("Project Dump"))
+			CommandUtils.LogVerbose("  {0} projects:", AllProjects.Count);
+			foreach (BranchUProject Proj in AllProjects)
 			{
-				CommandUtils.LogVerbose("  Base Engine:");
-				BaseEngineProject.Dump(InHostPlatforms);
-
-				CommandUtils.LogVerbose("  {0} Code projects:", CodeProjects.Count);
-				foreach (BranchUProject Proj in CodeProjects)
-				{
-					Proj.Dump(InHostPlatforms);
-				}
-				CommandUtils.LogVerbose("  {0} Non-Code projects:", NonCodeProjects.Count);
-				foreach (BranchUProject Proj in NonCodeProjects)
-				{
-					Proj.Dump(InHostPlatforms);
-				}
+				CommandUtils.LogLog(" {0}: {1}", Proj.GameName, Proj.FilePath);
 			}
         }
 
+		public BranchInfo(List<UnrealTargetPlatform> HostPlatforms)
+			: this()
+		{
+		}
+
         public BranchUProject FindGame(string GameName)
         {
-            foreach (BranchUProject Proj in CodeProjects)
-            {
-                if (Proj.GameName.Equals(GameName, StringComparison.InvariantCultureIgnoreCase))
-                {
-                    return Proj;
-                }
-            }
-            foreach (BranchUProject Proj in NonCodeProjects)
+			foreach (BranchUProject Proj in AllProjects)
             {
                 if (Proj.GameName.Equals(GameName, StringComparison.InvariantCultureIgnoreCase))
                 {
@@ -1058,6 +1004,7 @@ namespace AutomationTool
             }
             return null;
         }
+
 		public BranchUProject FindGameChecked(string GameName)
 		{
 			BranchUProject Project = FindGame(GameName);
@@ -1067,32 +1014,5 @@ namespace AutomationTool
 			}
 			return Project;
 		}
-        public SingleTargetProperties FindProgram(string ProgramName)
-        {
-            foreach (SingleTargetProperties Proj in BaseEngineProject.Properties.Programs)
-            {
-                if (Proj.TargetName.Equals(ProgramName, StringComparison.InvariantCultureIgnoreCase))
-                {
-                    return Proj;
-                }
-            }
-
-			foreach (BranchUProject CodeProj in CodeProjects)
-			{
-				foreach (SingleTargetProperties Proj in CodeProj.Properties.Programs)
-				{
-					if (Proj.TargetName.Equals(ProgramName, StringComparison.InvariantCultureIgnoreCase))
-					{
-						return Proj;
-					}
-				}
-			}
-
-            SingleTargetProperties Result = new SingleTargetProperties();
-            Result.TargetName = ProgramName;
-            Result.Rules = null;
-            return Result;
-        }
-    };
-
+    }
 }
