Fix enabling editor-only plugins in content-only projects causing errors at load time. These plugins are not actually required at runtime, but the runtime doesn't know that - the project file still references them. We have to include the .uplugin file for the runtime to make that determination.

diff --git a/Engine/Source/Programs/AutomationTool/Scripts/CopyBuildToStagingDirectory.Automation.cs b/Engine/Source/Programs/AutomationTool/Scripts/CopyBuildToStagingDirectory.Automation.cs
--- a/Engine/Source/Programs/AutomationTool/Scripts/CopyBuildToStagingDirectory.Automation.cs
+++ b/Engine/Source/Programs/AutomationTool/Scripts/CopyBuildToStagingDirectory.Automation.cs
@@ -434,17 +434,18 @@ public partial class Project : CommandUtils
 	}
 
 	/// <summary>
-	/// Figures out which content plugins are enabled for a content-only plugins. These will not be in the default target manifest, since it was compiled elsewhere.
+	/// Figures out which plugins are enabled for a content-only plugins. This may include content-only plugins and editor plugins, which will not be in the default 
+	/// target manifest, since it was compiled elsewhere.
 	/// </summary>
 	/// <param name="ProjectFile">The project being built</param>
 	/// <param name="Targets">List of targets being staged</param>
 	/// <returns>List of plugin files that should be staged</returns>
-	private static List<FileReference> GetContentPluginsForContentProject(FileReference ProjectFile, List<TargetReceipt> Targets)
+	private static List<FileReference> GetPluginsForContentProject(FileReference ProjectFile, List<TargetReceipt> Targets)
 	{
 		ProjectDescriptor Project = ProjectDescriptor.FromFile(ProjectFile);
-		List<PluginInfo> AvailablePlugins = Plugins.ReadAvailablePlugins(CommandUtils.EngineDirectory, ProjectFile.Directory, null);
+		List<PluginInfo> AvailablePlugins = UnrealBuildTool.Plugins.ReadAvailablePlugins(CommandUtils.EngineDirectory, ProjectFile.Directory, null);
 
-		HashSet<FileReference> ContentPlugins = new HashSet<FileReference>();
+		HashSet<FileReference> Plugins = new HashSet<FileReference>();
 		foreach (TargetReceipt Target in Targets)
 		{
 			// Find all the specifically enabled plugins for this target
@@ -454,7 +455,7 @@ public partial class Project : CommandUtils
 				foreach (PluginReferenceDescriptor Reference in Project.Plugins)
 				{
 					bool bEnabled = false;
-					if (!Reference.IsEnabledForPlatform(Target.Platform) || !Reference.IsEnabledForTargetConfiguration(Target.Configuration) || !Reference.IsEnabledForTarget(Target.TargetType))
+					if (Reference.IsEnabledForPlatform(Target.Platform) && Reference.IsEnabledForTargetConfiguration(Target.Configuration) && Reference.IsEnabledForTarget(Target.TargetType))
 					{
 						bEnabled = true;
 					}
@@ -471,18 +472,28 @@ public partial class Project : CommandUtils
 				}
 			}
 
-			// Add the enabled content-only plugins
+			// Add the enabled plugins
 			foreach (PluginInfo AvailablePlugin in AvailablePlugins)
 			{
-				if (AvailablePlugin.Descriptor.SupportsTargetPlatform(Target.Platform) && AvailablePlugin.Descriptor.bCanContainContent && EnabledPlugins[AvailablePlugin.Name])
+				if (AvailablePlugin.Descriptor.SupportsTargetPlatform(Target.Platform) && EnabledPlugins[AvailablePlugin.Name])
 				{
-					ContentPlugins.Add(AvailablePlugin.File);
+					Plugins.Add(AvailablePlugin.File);
 				}
 			}
 
-			ContentPlugins.ExceptWith(Target.RuntimeDependencies.Select(x => x.Path));
+			// Exclude all the plugins that would have been considered for the base target (editor only plugins that are enabled by default)
+			foreach (PluginInfo AvailablePlugin in AvailablePlugins)
+			{
+				if (AvailablePlugin.LoadedFrom == PluginLoadedFrom.Engine && !AvailablePlugin.Descriptor.bInstalled && AvailablePlugin.IsEnabledByDefault(!Project.DisableEnginePluginsByDefault))
+				{
+					Plugins.Remove(AvailablePlugin.File);
+				}
+			}
+
+			// Exclude all the existing plugin descriptors
+			Plugins.ExceptWith(Target.RuntimeDependencies.Select(x => x.Path));
 		}
-		return ContentPlugins.ToList();
+		return Plugins.ToList();
 	}
 
 	public static void CreateStagingManifest(ProjectParams Params, DeploymentContext SC)
@@ -639,7 +650,7 @@ public partial class Project : CommandUtils
 				// Stage any content-only plugins for content-only projects. We don't have a custom executable for these.
 				if (!Params.IsCodeBasedProject)
 				{
-					List<FileReference> PluginFiles = GetContentPluginsForContentProject(Params.RawProjectPath, SC.StageTargets.ConvertAll(x => x.Receipt));
+					List<FileReference> PluginFiles = GetPluginsForContentProject(Params.RawProjectPath, SC.StageTargets.ConvertAll(x => x.Receipt));
 					SC.StageFiles(StagedFileType.UFS, PluginFiles);
 				}
 			}
