Fixed two issues in SyncProject script.

1) Typo in checking against the Count() of matching target types in a project. If there's only one match (e.g. a single Editor target) that's what we want.

2) In the case where there are multiple matches we try to pick the one prefixed with the project name. This failed when the user supplied a full path for the project argument as the name was not being taken from the project file.

#tests ran standalone BuildTarget script with full project path to EngineTest and via SyncProject script

diff --git a/Engine/Source/Programs/AutomationTool/Scripts/BuildTarget.Automation.cs b/Engine/Source/Programs/AutomationTool/Scripts/BuildTarget.Automation.cs
--- a/Engine/Source/Programs/AutomationTool/Scripts/BuildTarget.Automation.cs
+++ b/Engine/Source/Programs/AutomationTool/Scripts/BuildTarget.Automation.cs
@@ -106,13 +106,17 @@ namespace AutomationTool
 
 			if (!string.IsNullOrEmpty(ProjectName))
 			{
-				ProjectFile = ProjectUtils.FindProjectFileFromName(ProjectName);
+				// find the project
+				ProjectFile = ProjectUtils.FindProjectFileFromName(ProjectName);				
 
 				if (ProjectFile == null)
 				{
 					throw new AutomationException("Unable to find uproject file for {0}", ProjectName);
 				}
 
+				// user may have passed a full path, we just want the name
+				ProjectName = ProjectFile.GetFileNameWithoutAnyExtensions();
+
 				ProjectProperties Properties = ProjectUtils.GetProjectProperties(ProjectFile);
 
 				if (Properties.bIsCodeBasedProject)
@@ -142,7 +146,7 @@ namespace AutomationTool
 
 							string ProjectTarget;
 
-							if (MatchingTargetTypes.Count() == 0)
+							if (MatchingTargetTypes.Count() == 1)
 							{
 								ProjectTarget = MatchingTargetTypes.First();
 							}
@@ -152,6 +156,11 @@ namespace AutomationTool
 								ProjectTarget = MatchingTargetTypes.Where(T => string.CompareOrdinal(T, 0, ProjectName, 0, 1) == 0).FirstOrDefault();
 							}
 
+							if (ProjectTarget == null)
+							{
+								throw new AutomationException("Unable to find project target for {0}", Target);
+							}
+
 							ActualTargets.Add(ProjectTarget);
 						}
 					}
