Deprecate the BranchInfo constructor that takes a list of host platforms as a parameter.

diff --git a/Engine/Source/Programs/AutomationTool/AutomationUtils/CommandletUtils.cs b/Engine/Source/Programs/AutomationTool/AutomationUtils/CommandletUtils.cs
--- a/Engine/Source/Programs/AutomationTool/AutomationUtils/CommandletUtils.cs
+++ b/Engine/Source/Programs/AutomationTool/AutomationUtils/CommandletUtils.cs
@@ -493,7 +493,7 @@ namespace AutomationTool
 			}
 			else
 			{
-				var Branch = new BranchInfo(new List<UnrealTargetPlatform> { UnrealBuildTool.BuildHostPlatform.Current.Platform });
+				var Branch = new BranchInfo();
 				var GameProj = Branch.FindGame(OriginalProjectName);
 				if (GameProj != null)
 				{
diff --git a/Engine/Source/Programs/AutomationTool/AutomationUtils/ProjectUtils.cs b/Engine/Source/Programs/AutomationTool/AutomationUtils/ProjectUtils.cs
--- a/Engine/Source/Programs/AutomationTool/AutomationUtils/ProjectUtils.cs
+++ b/Engine/Source/Programs/AutomationTool/AutomationUtils/ProjectUtils.cs
@@ -917,14 +917,6 @@ namespace AutomationTool
 
     public class BranchInfo
     {
-
-        public static List<TargetType> MonolithicKinds = new List<TargetType>
-        {
-            TargetType.Game,
-            TargetType.Client,
-            TargetType.Server,
-        };
-
 		[DebuggerDisplay("{GameName}")]
         public class BranchUProject
         {
@@ -957,24 +949,12 @@ namespace AutomationTool
                     throw new AutomationException("Could not resolve relative path corrctly {0} -> {1} which doesn't exist.", ProjectFile, FilePath);
                 }
             }
-
-            public BranchUProject()
-            {
-                GameName = "UE4";
-                if (!Properties.Targets.Exists(Target => Target.Rules.Type == TargetType.Editor))
-                {
-                    throw new AutomationException("Base UE4 project did not contain an editor target.");
-                }
-            }
         }
 
-        public BranchUProject BaseEngineProject = null;
 		public List<BranchUProject> AllProjects = new List<BranchUProject>();
 
         public BranchInfo()
         {
-            BaseEngineProject = new BranchUProject();
-
             IEnumerable<FileReference> ProjectFiles = UnrealBuildTool.NativeProjects.EnumerateProjectFiles();
 			foreach (FileReference InfoEntry in ProjectFiles)
 			{
@@ -988,6 +968,7 @@ namespace AutomationTool
 			}
         }
 
+		[Obsolete("BranchInfo no longer requires HostPlatforms passed to constructor")]
 		public BranchInfo(List<UnrealTargetPlatform> HostPlatforms)
 			: this()
 		{
diff --git a/Engine/Source/Programs/AutomationTool/Scripts/BuildCookRun.Automation.cs b/Engine/Source/Programs/AutomationTool/Scripts/BuildCookRun.Automation.cs
--- a/Engine/Source/Programs/AutomationTool/Scripts/BuildCookRun.Automation.cs
+++ b/Engine/Source/Programs/AutomationTool/Scripts/BuildCookRun.Automation.cs
@@ -320,7 +320,7 @@ public class BuildCookRun : BuildCommand
 					}
 					else
 					{
-						var Branch = new BranchInfo(new List<UnrealTargetPlatform> { UnrealBuildTool.BuildHostPlatform.Current.Platform });
+						var Branch = new BranchInfo();
 						var GameProj = Branch.FindGame(OriginalProjectName);
 						if (GameProj != null)
 						{
diff --git a/Engine/Source/Programs/AutomationTool/Scripts/CopySharedCookedBuild.Automation.cs b/Engine/Source/Programs/AutomationTool/Scripts/CopySharedCookedBuild.Automation.cs
--- a/Engine/Source/Programs/AutomationTool/Scripts/CopySharedCookedBuild.Automation.cs
+++ b/Engine/Source/Programs/AutomationTool/Scripts/CopySharedCookedBuild.Automation.cs
@@ -42,7 +42,7 @@ class CopySharedCookedBuild : BuildCommand
 			}
 			else
 			{
-				var Branch = new BranchInfo(new List<UnrealTargetPlatform> { UnrealBuildTool.BuildHostPlatform.Current.Platform });
+				var Branch = new BranchInfo();
 				var GameProj = Branch.FindGame(OriginalProjectName);
 				if (GameProj != null)
 				{
diff --git a/Engine/Source/Programs/AutomationTool/Scripts/CryptoKeys.Automation.cs b/Engine/Source/Programs/AutomationTool/Scripts/CryptoKeys.Automation.cs
--- a/Engine/Source/Programs/AutomationTool/Scripts/CryptoKeys.Automation.cs
+++ b/Engine/Source/Programs/AutomationTool/Scripts/CryptoKeys.Automation.cs
@@ -159,7 +159,7 @@ namespace AutomationTool
 						}
 						else
 						{
-							var Branch = new BranchInfo(new List<UnrealTargetPlatform> { UnrealBuildTool.BuildHostPlatform.Current.Platform });
+							var Branch = new BranchInfo();
 							var GameProj = Branch.FindGame(OriginalProjectName);
 							if (GameProj != null)
 							{
diff --git a/Engine/Source/Programs/AutomationTool/Scripts/RebuildHLODCommand.Automation.cs b/Engine/Source/Programs/AutomationTool/Scripts/RebuildHLODCommand.Automation.cs
--- a/Engine/Source/Programs/AutomationTool/Scripts/RebuildHLODCommand.Automation.cs
+++ b/Engine/Source/Programs/AutomationTool/Scripts/RebuildHLODCommand.Automation.cs
@@ -405,7 +405,7 @@ namespace AutomationScripts.Automation
 					}
 					else
 					{
-						var Branch = new BranchInfo(new List<UnrealTargetPlatform> { UnrealBuildTool.BuildHostPlatform.Current.Platform });
+						var Branch = new BranchInfo();
 						var GameProj = Branch.FindGame(OriginalProjectName);
 						if (GameProj != null)
 						{
diff --git a/Engine/Source/Programs/AutomationTool/Scripts/RebuildLightMapsCommand.Automation.cs b/Engine/Source/Programs/AutomationTool/Scripts/RebuildLightMapsCommand.Automation.cs
--- a/Engine/Source/Programs/AutomationTool/Scripts/RebuildLightMapsCommand.Automation.cs
+++ b/Engine/Source/Programs/AutomationTool/Scripts/RebuildLightMapsCommand.Automation.cs
@@ -344,7 +344,7 @@ namespace AutomationScripts.Automation
 					}
 					else
 					{
-						var Branch = new BranchInfo(new List<UnrealTargetPlatform> { UnrealBuildTool.BuildHostPlatform.Current.Platform });
+						var Branch = new BranchInfo();
 						var GameProj = Branch.FindGame(OriginalProjectName);
 						if (GameProj != null)
 						{
diff --git a/Engine/Source/Programs/AutomationTool/Scripts/ResavePackagesCommand.Automation.cs b/Engine/Source/Programs/AutomationTool/Scripts/ResavePackagesCommand.Automation.cs
--- a/Engine/Source/Programs/AutomationTool/Scripts/ResavePackagesCommand.Automation.cs
+++ b/Engine/Source/Programs/AutomationTool/Scripts/ResavePackagesCommand.Automation.cs
@@ -339,7 +339,7 @@ namespace AutomationScripts.Automation
 					}
 					else
 					{
-						var Branch = new BranchInfo(new List<UnrealTargetPlatform> { UnrealBuildTool.BuildHostPlatform.Current.Platform });
+						var Branch = new BranchInfo();
 						var GameProj = Branch.FindGame(OriginalProjectName);
 						if (GameProj != null)
 						{
diff --git a/Engine/Source/Programs/AutomationTool/Scripts/Tests.Automation.cs b/Engine/Source/Programs/AutomationTool/Scripts/Tests.Automation.cs
--- a/Engine/Source/Programs/AutomationTool/Scripts/Tests.Automation.cs
+++ b/Engine/Source/Programs/AutomationTool/Scripts/Tests.Automation.cs
@@ -1428,10 +1428,7 @@ class DumpBranch : BuildCommand
 	{
 		LogInformation("************************* DumpBranch");
 
-		var HostPlatforms = new List<UnrealTargetPlatform>();
-		HostPlatforms.Add(UnrealTargetPlatform.Win64);
-		HostPlatforms.Add(UnrealTargetPlatform.Mac);
-		new BranchInfo(HostPlatforms);
+		new BranchInfo();
 	}
 }
 
