Fixed misnamed log channel used to detect activity.

Added -args / -execcmds support to gauntlet in addition to -clientargs / -clientexeccmds when you need to make sure all instances of a test get an arg.

diff --git a/Engine/Source/Programs/AutomationTool/Gauntlet/Unreal/RunUnreal.cs b/Engine/Source/Programs/AutomationTool/Gauntlet/Unreal/RunUnreal.cs
--- a/Engine/Source/Programs/AutomationTool/Gauntlet/Unreal/RunUnreal.cs
+++ b/Engine/Source/Programs/AutomationTool/Gauntlet/Unreal/RunUnreal.cs
@@ -186,16 +186,28 @@ namespace Gauntlet
 						RequestedConfiguration = (UnrealTargetConfiguration)Enum.Parse(typeof(UnrealTargetConfiguration), ConfigString, true);
 					}
 
-					// look for -clientargs= and -editorclient etc
-					Role.ExtraArgs = Globals.Params.ParseValue(Type.ToString() + "Args", "");
+					// look for-args= and then -clientargs= and -editorargs etc
+					Role.ExtraArgs = Globals.Params.ParseValue("Args", "");
+					string ExtraRoleArgs = Globals.Params.ParseValue(Type.ToString() + "Args", "");
+
+					if (!string.IsNullOrEmpty(ExtraRoleArgs))
+					{
+						Role.ExtraArgs += ExtraRoleArgs;
+					}
 
 					// look for -clientexeccmds=, -editorexeccmds= etc, these are separate from clientargs for sanity
-					string ExecCmds = Globals.Params.ParseValue(Type.ToString() + "ExecCmds", "");
+					string ExecCmds = Globals.Params.ParseValue("ExecCmds", "");
+					string RoleExecCmds= Globals.Params.ParseValue(Type.ToString() + "ExecCmds", "");
 					if (!string.IsNullOrEmpty(ExecCmds))
 					{
 						Role.ExtraArgs += string.Format(" -ExecCmds=\"{0}\"", ExecCmds);
 					}
 
+					if (!string.IsNullOrEmpty(RoleExecCmds))
+					{
+						Role.ExtraArgs += string.Format(" -ExecCmds=\"{0}\"", RoleExecCmds);
+					}
+
 					bool UsesEditor = EditorForAllRoles || Globals.Params.ParseParam("Editor" + Type.ToString());
 
 					if (UsesEditor)
