Invalidate target.cs cache for project after syncing. This ensures any changes (including those files existing) are considered when performing any build steps.

#tests Ran SyncProject for QAGame and EngineTest with only the project file synced.

diff --git a/Engine/Source/Programs/AutomationTool/Scripts/SyncProject.Automation.cs b/Engine/Source/Programs/AutomationTool/Scripts/SyncProject.Automation.cs
--- a/Engine/Source/Programs/AutomationTool/Scripts/SyncProject.Automation.cs
+++ b/Engine/Source/Programs/AutomationTool/Scripts/SyncProject.Automation.cs
@@ -287,6 +287,13 @@ class SyncProject : BuildCommand
 			{
 				Log.TraceVerbose("Building Editor for {0}", ProjectArgForEditor);
 
+				// Invalidate the location of the Target.cs files incase they were synced
+				if (ProjectFile != null)
+				{
+					string TargetSourceDir = CommandUtils.CombinePaths(Path.GetDirectoryName(ProjectFile.FullName), "Source");
+					RulesCompiler.InvalidateRulesFileCache(TargetSourceDir);
+				}
+
 				BuildEditor BuildCmd = new BuildEditor();
 				BuildCmd.Clean = ParseParam("clean");
 				BuildCmd.ProjectName = ProjectArgForEditor;
