BuildTarget script - print valid target names when there's a mismatch to assist diagnosis / bug reporting.

diff --git a/Engine/Source/Programs/AutomationTool/Scripts/BuildTarget.Automation.cs b/Engine/Source/Programs/AutomationTool/Scripts/BuildTarget.Automation.cs
--- a/Engine/Source/Programs/AutomationTool/Scripts/BuildTarget.Automation.cs
+++ b/Engine/Source/Programs/AutomationTool/Scripts/BuildTarget.Automation.cs
@@ -141,7 +141,8 @@ namespace AutomationTool
 
 							if (!MatchingTargetTypes.Any())
 							{
-								throw new AutomationException("The {0} is not a valid target for {1}", Target, ProjectName);
+								string ValidTargets = string.Join(",", Properties.Targets.Select(T => T.Rules.Type.ToString()));
+								throw new AutomationException("'{0}' is not a valid target for {1}. Targets are {2}", Target, ProjectName, ValidTargets);
 							}
 
 							string ProjectTarget;
