Add support for some basic documentation being able to be built into gauntlet nodes. Includes support for describing commandline arguments as well as sample commandlines.

diff --git a/Engine/Source/Programs/AutomationTool/Gauntlet/Framework/Base/Gauntlet.BaseTestNode.cs b/Engine/Source/Programs/AutomationTool/Gauntlet/Framework/Base/Gauntlet.BaseTestNode.cs
--- a/Engine/Source/Programs/AutomationTool/Gauntlet/Framework/Base/Gauntlet.BaseTestNode.cs
+++ b/Engine/Source/Programs/AutomationTool/Gauntlet/Framework/Base/Gauntlet.BaseTestNode.cs
@@ -166,6 +166,13 @@ namespace Gauntlet
 		/// <returns></returns>
 		public virtual void SetContext(ITestContext InContext)
 		{
-		}	
+		}
+
+		/// <summary>
+		/// Output all defined commandline information for this test to the gauntlet window and exit test early.
+		/// </summary>
+		public virtual void DisplayCommandlineHelp()
+		{ 
+		}
 	}
 }
diff --git a/Engine/Source/Programs/AutomationTool/Gauntlet/Framework/Base/Gauntlet.TestNode.cs b/Engine/Source/Programs/AutomationTool/Gauntlet/Framework/Base/Gauntlet.TestNode.cs
--- a/Engine/Source/Programs/AutomationTool/Gauntlet/Framework/Base/Gauntlet.TestNode.cs
+++ b/Engine/Source/Programs/AutomationTool/Gauntlet/Framework/Base/Gauntlet.TestNode.cs
@@ -147,7 +147,11 @@ namespace Gauntlet
 		/// </summary>
 		/// <param name="Node"></param>
 		/// <returns></returns>
-		void CleanupTest();	
-		
+		void CleanupTest();
+
+		/// <summary>
+		/// Output all defined commandline information for this test to the gauntlet window and exit test early.
+		/// </summary>
+		void DisplayCommandlineHelp();
 	}
 }
\ No newline at end of file
diff --git a/Engine/Source/Programs/AutomationTool/Gauntlet/SelfTest/Framework/Gauntlet.SelfTest.OrderOfOpsTest.cs b/Engine/Source/Programs/AutomationTool/Gauntlet/SelfTest/Framework/Gauntlet.SelfTest.OrderOfOpsTest.cs
--- a/Engine/Source/Programs/AutomationTool/Gauntlet/SelfTest/Framework/Gauntlet.SelfTest.OrderOfOpsTest.cs
+++ b/Engine/Source/Programs/AutomationTool/Gauntlet/SelfTest/Framework/Gauntlet.SelfTest.OrderOfOpsTest.cs
@@ -151,5 +151,13 @@ namespace Gauntlet.SelfTest
 		{
 			SetNewStage(TestStages.CleanupTest);
 		}
+
+		/// <summary>
+		/// Output all defined commandline information for this test to the gauntlet window and exit test early.
+		/// </summary>
+		public virtual void DisplayCommandlineHelp()
+		{
+
+		}
 	}
 }
diff --git a/Engine/Source/Programs/AutomationTool/Gauntlet/SelfTest/Gauntlet.SelfTest.BaseNode.cs b/Engine/Source/Programs/AutomationTool/Gauntlet/SelfTest/Gauntlet.SelfTest.BaseNode.cs
--- a/Engine/Source/Programs/AutomationTool/Gauntlet/SelfTest/Gauntlet.SelfTest.BaseNode.cs
+++ b/Engine/Source/Programs/AutomationTool/Gauntlet/SelfTest/Gauntlet.SelfTest.BaseNode.cs
@@ -138,6 +138,14 @@ namespace Gauntlet.SelfTest
 			MarkComplete(TestFailures.Count == 0 ? TestResult.Passed : TestResult.Failed);
 		}
 
+		/// <summary>
+		/// Output all defined commandline information for this test to the gauntlet window and exit test early.
+		/// </summary>
+		public virtual void DisplayCommandlineHelp()
+		{
+
+		}
+
 		// Self test interface
 		public abstract void TickTest();
 
diff --git a/Engine/Source/Programs/AutomationTool/Gauntlet/Unreal/Base/Gauntlet.UnrealTestNode.cs b/Engine/Source/Programs/AutomationTool/Gauntlet/Unreal/Base/Gauntlet.UnrealTestNode.cs
--- a/Engine/Source/Programs/AutomationTool/Gauntlet/Unreal/Base/Gauntlet.UnrealTestNode.cs
+++ b/Engine/Source/Programs/AutomationTool/Gauntlet/Unreal/Base/Gauntlet.UnrealTestNode.cs
@@ -14,7 +14,73 @@ using System.Security.Cryptography;
 
 namespace Gauntlet
 {
+	public class GauntletParamDescription
+	{
+		/// <summary>
+		/// The param name that is passed in on the commandline
+		/// In -PARAMNAME or -PARAMNAME=val format
+		/// </summary>
+		public string ParamName;
+
+		/// <summary>
+		/// Is this required for the test to run?
+		/// </summary>
+		public bool Required;
+
+		/// <summary>		
+		/// Very brief desc of what to pass in -
+		/// Will show up in -param=<InputFormat> format.
+		/// Ex - a value of "Map To Use" would show up as -param=<Map To use>
+		/// Leave blank for a param that is just a flag.
+		/// </summary>
+		public string InputFormat;
 
+		/// <summary>
+		/// Helpful description for what this Parameter or flag represents and what can be passed in.
+		/// </summary>
+		public string ParamDesc;
+
+		/// <summary>
+		/// If you would like to provide a sample input for this field, do so here. Will show up as (ex: SampleInput) at the end of the param description
+		/// </summary>
+		public string SampleInput;
+
+		/// <summary>
+		///  If this param has a default value, put it here. Will show ups as (default: DefaultValue)
+		/// </summary>
+		public string DefaultValue;
+
+		/// <summary>
+		/// Whether this is a Test-specific param or a generic gauntlet param.
+		/// </summary>
+		public bool TestSpecificParam;
+
+		public GauntletParamDescription()
+		{
+			TestSpecificParam = true;
+		}
+
+		public override string ToString()
+		{
+			string ParamFormat = ParamName;
+			if (!string.IsNullOrEmpty(InputFormat))
+			{
+				ParamFormat += "=" + InputFormat;
+			}
+			string DefaultFormat = "";
+			if (!string.IsNullOrEmpty(DefaultValue))
+			{
+				DefaultFormat = string.Format("(default: {0}) ", DefaultValue);
+			}
+			string SampleFormat = "";
+			if (!string.IsNullOrEmpty(SampleInput))
+			{
+				SampleFormat = string.Format(" (ex: {0})", SampleInput);
+			}
+			ParamFormat = string.Format("{0}:\t\t{1}{2}{3}{4}", ParamFormat, Required ? "*Required* " : "", DefaultFormat, ParamDesc, SampleInput);
+			return ParamFormat;
+		}
+	}
 	public abstract class UnrealTestNode<TConfigClass> : BaseTest, IDisposable
 		where TConfigClass : UnrealTestConfiguration, new()
 	{
@@ -138,6 +204,21 @@ namespace Gauntlet
 		// artifact paths that have been used in this run
 		static protected HashSet<string> ReservedArtifcactPaths = new HashSet<string>();
 
+		/// <summary>
+		/// Help doc-style list of parameters supported by this test. List can be divided into test-specific and general arguments.
+		/// </summary>
+		public List<GauntletParamDescription> SupportedParameters = new List<GauntletParamDescription>();
+
+		/// <summary>
+		/// Optional list of provided commandlines to be displayed to users who want to look at test help docs.
+		/// Key should be the commandline to use, value should be the description for that commandline.
+		/// </summary>
+		protected List<KeyValuePair<string, string>> SampleCommandlines = new List<KeyValuePair<string, string>>();
+
+		public void AddSampleCommandline(string Commandline, string Description)
+		{
+			SampleCommandlines.Add(new KeyValuePair<string, string>(Commandline, Description));
+		}
 		public UnrealTestNode(UnrealTestContext InContext)
 		{
 			Context = InContext;
@@ -150,6 +231,7 @@ namespace Gauntlet
 			NumPasses = 0;
 			TestVersion = new Version("1.0.0");
 			ArtifactPath = string.Empty;
+			PopulateCommandlineInfo();
 		}
 
 		 ~UnrealTestNode()
@@ -525,6 +607,51 @@ namespace Gauntlet
 			return TestInstance != null;
 		}
 
+		public virtual void PopulateCommandlineInfo()
+		{
+			SupportedParameters.Add(new GauntletParamDescription()
+			{
+				ParamName = "nomcp",
+				TestSpecificParam = false,
+				ParamDesc = "Run test without an mcp backend",
+				DefaultValue = "false"
+			});
+			SupportedParameters.Add(new GauntletParamDescription()
+			{
+				ParamName = "ResX",
+				InputFormat = "1280",
+				Required = false,
+				TestSpecificParam = false,
+				ParamDesc = "Horizontal resolution for the game client.",
+				DefaultValue = "1280"
+			});
+			SupportedParameters.Add(new GauntletParamDescription()
+			{
+				ParamName = "ResY",
+				InputFormat = "720",
+				Required = false,
+				TestSpecificParam = false,
+				ParamDesc = "Vertical resolution for the game client.",
+				DefaultValue = "720"
+			});
+			SupportedParameters.Add(new GauntletParamDescription()
+			{
+				ParamName = "FailOnEnsure",
+				Required = false,
+				TestSpecificParam = false,
+				ParamDesc = "Consider the test a fail if we encounter ensures."
+			});
+			SupportedParameters.Add(new GauntletParamDescription()
+			{
+				ParamName = "MaxDuration",
+				InputFormat = "600",
+				Required = false,
+				TestSpecificParam = false,
+				ParamDesc = "Time in seconds for test to run before it is timed out",
+				DefaultValue = "Test-defined"
+			});
+		}
+
 		/// <summary>
 		/// Cleanup all resources
 		/// </summary>
@@ -722,13 +849,62 @@ namespace Gauntlet
 		/// Whether creating the test report failed
 		/// </summary>
 		public bool CreateReportFailed { get; protected set; }
-		 
+
+		/// <summary>
+		/// Display all of the defined commandline information for this test.
+		/// Will display generic gauntlet params as well if -help=full is passed in.
+		/// </summary>
+		public override void DisplayCommandlineHelp()
+		{
+			Log.Info(string.Format("--- Available Commandline Parameters for {0} ---", Name));
+			Log.Info("--------------------");
+			Log.Info("TEST-SPECIFIC PARAMS");
+			Log.Info("--------------------");
+			foreach (GauntletParamDescription ParamDesc in SupportedParameters)
+			{
+				if (ParamDesc.TestSpecificParam)
+				{
+					Log.Info(ParamDesc.ToString());
+				}
+			}
+			if (Context.TestParams.ParseParam("listallargs"))
+			{
+				Log.Info("\n--------------");
+				Log.Info("GENERIC PARAMS");
+				Log.Info("--------------");
+				foreach (GauntletParamDescription ParamDesc in SupportedParameters)
+				{
+					if (!ParamDesc.TestSpecificParam)
+					{
+						Log.Info(ParamDesc.ToString());
+					}
+				}
+			}
+			else
+			{
+				Log.Info("\nIf you need information on base Gauntlet arguments, use -listallargs\n\n");
+			}
+			if (SampleCommandlines.Count > 0)
+			{
+				Log.Info("\n-------------------");
+				Log.Info("SAMPLE COMMANDLINES");
+				Log.Info("-------------------");
+				foreach (KeyValuePair<string, string> SampleCommandline in SampleCommandlines)
+				{
+					Log.Info(SampleCommandline.Key);
+					Log.Info("");
+					Log.Info(SampleCommandline.Value);
+					Log.Info("-------------------");
+				}
+			}
+
+		}
 
 		/// <summary>
 		/// Optional function that is called on test completion and gives an opportunity to create a report
 		/// </summary>
 		/// <param name="Result"></param>
-		/// <param name="Contex"></param>
+		/// <param name="Context"></param>
 		/// <param name="Build"></param>
 		public virtual void CreateReport(TestResult Result, UnrealTestContext Context, UnrealBuildSource Build, IEnumerable<UnrealRoleArtifacts> Artifacts, string ArtifactPath)
 		{
diff --git a/Engine/Source/Programs/AutomationTool/Gauntlet/Unreal/RunUnreal.cs b/Engine/Source/Programs/AutomationTool/Gauntlet/Unreal/RunUnreal.cs
--- a/Engine/Source/Programs/AutomationTool/Gauntlet/Unreal/RunUnreal.cs
+++ b/Engine/Source/Programs/AutomationTool/Gauntlet/Unreal/RunUnreal.cs
@@ -397,7 +397,14 @@ namespace Gauntlet
 				// This will throw if the test cannot be created
 				ITestNode NewTest = Utils.TestConstructor.ConstructTest<ITestNode, UnrealTestContext>(Test.TestName, TestContext, Namespaces);
 
-				NodeList.Add(NewTest);
+				if (CombinedParams.ParseParam("listargs") || CombinedParams.ParseParam("listallargs"))
+				{
+					NewTest.DisplayCommandlineHelp();
+				}
+				else
+				{
+					NodeList.Add(NewTest);
+				}
 			}
 
 			return NodeList;
