Re-adding ability to add explicit commandline args to the commandline of a gauntlet test for non UE-compliant args.

diff --git a/Engine/Source/Programs/AutomationTool/Gauntlet/Unreal/Base/Gauntlet.UnrealTestConfiguration.cs b/Engine/Source/Programs/AutomationTool/Gauntlet/Unreal/Base/Gauntlet.UnrealTestConfiguration.cs
--- a/Engine/Source/Programs/AutomationTool/Gauntlet/Unreal/Base/Gauntlet.UnrealTestConfiguration.cs
+++ b/Engine/Source/Programs/AutomationTool/Gauntlet/Unreal/Base/Gauntlet.UnrealTestConfiguration.cs
@@ -15,11 +15,17 @@ namespace Gauntlet
 		public string Project;
 		public string GameMap;
 		private Dictionary<string, object> Params;
+		/// <summary>
+		/// String containing extra commandline args that do not conform to UE commandline arg specs.
+		/// Please do not use this for for standard flags.
+		/// </summary>
+		public string AdditionalExplicitCommandLineArgs;
 
 		public GauntletCommandLine()
 		{
 			Project = string.Empty;
 			GameMap = string.Empty;
+			AdditionalExplicitCommandLineArgs = string.Empty;
 			Params = new Dictionary<string, object>(StringComparer.OrdinalIgnoreCase);
 		}
 
@@ -211,6 +217,10 @@ namespace Gauntlet
 				}
 				FinalCommandline = string.Format("{0} {1} ", FinalCommandline, CurrentArgument);
 			}
+			if (!string.IsNullOrEmpty(AdditionalExplicitCommandLineArgs))
+			{
+				FinalCommandline = string.Format("{0} {1}", FinalCommandline, AdditionalExplicitCommandLineArgs);
+			}
 			return FinalCommandline;
 		}
 	}
diff --git a/Engine/Source/Programs/AutomationTool/Gauntlet/Unreal/BuildSource/Gauntlet.UnrealBuildSource.cs b/Engine/Source/Programs/AutomationTool/Gauntlet/Unreal/BuildSource/Gauntlet.UnrealBuildSource.cs
--- a/Engine/Source/Programs/AutomationTool/Gauntlet/Unreal/BuildSource/Gauntlet.UnrealBuildSource.cs
+++ b/Engine/Source/Programs/AutomationTool/Gauntlet/Unreal/BuildSource/Gauntlet.UnrealBuildSource.cs
@@ -447,7 +447,6 @@ namespace Gauntlet
 					ProjectParam = string.Format("../../../{0}/{0}.uproject", ProjectName);
 				}
 				Config.CommandLineParams.Project = ProjectParam;
-				Config.CommandLine = Config.CommandLineParams.GenerateFullCommandLine();
 			}
 
             if (Role.FilesToCopy != null)
