Automatically delete DDC filter files after 5 days.

diff --git a/Engine/Source/Programs/AutomationTool/Scripts/UploadDDCToAWS.cs b/Engine/Source/Programs/AutomationTool/Scripts/UploadDDCToAWS.cs
--- a/Engine/Source/Programs/AutomationTool/Scripts/UploadDDCToAWS.cs
+++ b/Engine/Source/Programs/AutomationTool/Scripts/UploadDDCToAWS.cs
@@ -194,7 +194,8 @@ namespace AutomationTool
 					HashSet<string> Paths = new HashSet<string>();
 					foreach (FileInfo FilterFile in FilterDir.ToDirectoryInfo().EnumerateFiles("*.txt"))
 					{
-						if (FilterFile.LastWriteTimeUtc > DateTime.UtcNow - TimeSpan.FromDays(3))
+						TimeSpan Age = DateTime.UtcNow - FilterFile.LastWriteTimeUtc;
+						if (Age < TimeSpan.FromDays(3))
 						{
 							Log.TraceInformation("Reading {0}", FilterFile.FullName);
 
@@ -208,6 +209,19 @@ namespace AutomationTool
 								}
 							}
 						}
+						else if(Age > TimeSpan.FromDays(5))
+						{
+							try
+							{
+								Log.TraceInformation("Deleting {0}", FilterFile.FullName);
+								FilterFile.Delete();
+							}
+							catch (Exception Ex)
+							{
+								Log.TraceWarning("Unable to delete: {0}", Ex.Message);
+								Log.TraceLog(ExceptionUtils.FormatExceptionDetails(Ex));
+							}
+						}
 					}
 					Log.TraceInformation("Found {0:n0} files", Paths.Count);
 
