Fix exception trying to print errors for missing parameters.

#rnx

diff --git a/Engine/Source/Programs/AutomationTool/AutomationUtils/BuildCommand.cs b/Engine/Source/Programs/AutomationTool/AutomationUtils/BuildCommand.cs
--- a/Engine/Source/Programs/AutomationTool/AutomationUtils/BuildCommand.cs
+++ b/Engine/Source/Programs/AutomationTool/AutomationUtils/BuildCommand.cs
@@ -63,7 +63,7 @@ namespace AutomationTool
 			string Value = ParseOptionalStringParam(Param);
 			if(Value == null)
 			{
-				throw new AutomationException("Missing -{0}=... parameter");
+				throw new AutomationException("Missing -{0}=... parameter", Param);
 			}
 			return Value;
 		}
@@ -129,7 +129,7 @@ namespace AutomationTool
 			DirectoryReference Value = ParseOptionalDirectoryReferenceParam(Param);
 			if(Value == null)
 			{
-				throw new AutomationException("Missing -{0}=... parameter");
+				throw new AutomationException("Missing -{0}=... parameter", Param);
 			}
 			return Value;
 		}
@@ -167,7 +167,7 @@ namespace AutomationTool
 			Nullable<T> Value = ParseOptionalEnumParam<T>(Param);
 			if(!Value.HasValue)
 			{
-				throw new AutomationException("Missing -{0}=... parameter");
+				throw new AutomationException("Missing -{0}=... parameter", Param);
 			}
 			return Value.Value;
 		}
