#Gauntlet Disable warning about a long filename. PersistentDownloadDir was moved to the saved folder which is copied to the automation archive folder after a test, but it contains many files with deep paths and is not really necessary to archive.

diff --git a/Engine/Source/Programs/AutomationTool/Gauntlet/Framework/Gauntlet.Utils.cs b/Engine/Source/Programs/AutomationTool/Gauntlet/Framework/Gauntlet.Utils.cs
--- a/Engine/Source/Programs/AutomationTool/Gauntlet/Framework/Gauntlet.Utils.cs
+++ b/Engine/Source/Programs/AutomationTool/Gauntlet/Framework/Gauntlet.Utils.cs
@@ -1147,7 +1147,7 @@ namespace Gauntlet
 									Copied = true;
 
 									// Filter out some known unneeded files which can cause this warning, and log the message instead
-									string[] Blacklist = new string[]{ "UE4CC-XboxOne" };
+									string[] Blacklist = new string[]{ "UE4CC-XboxOne", "PersistentDownloadDir" };
 									string Message = string.Format("Long path file copy failed with {0}.  Please verify that this file is not required.", ex.Message);
 									if ( Blacklist.FirstOrDefault(B => { return SourcePath.IndexOf(B, StringComparison.OrdinalIgnoreCase) >= 0; }) == null)
 									{
