UAT: Add a fast path for detecting which client to use if the P4CLIENT variable is set to a valid workspace. Also remove warning for executing a command with P4 functionality.

diff --git a/Engine/Source/Programs/AutomationTool/AutomationUtils/P4Environment.cs b/Engine/Source/Programs/AutomationTool/AutomationUtils/P4Environment.cs
--- a/Engine/Source/Programs/AutomationTool/AutomationUtils/P4Environment.cs
+++ b/Engine/Source/Programs/AutomationTool/AutomationUtils/P4Environment.cs
@@ -368,9 +368,26 @@ namespace AutomationTool
 		{
 			CommandUtils.LogVerbose("uebp_CLIENT not set, detecting current client...");
 
-			var MatchingClients = new List<P4ClientInfo>();
+			// Check the default client. If it matches we can save any guess work.
+			IProcessResult Result = CommandUtils.Run(HostPlatform.Current.P4Exe, "set -q P4CLIENT", null, CommandUtils.ERunOptions.NoLoggingOfRunCommand);
+			if (Result.ExitCode == 0)
+			{
+				const string KeyName = "P4CLIENT=";
+				if (Result.Output.StartsWith(KeyName))
+				{
+					string ClientName = Result.Output.Substring(KeyName.Length).Trim();
+					P4ClientInfo ClientInfo = Connection.GetClientInfo(ClientName, true);
+					if (Connection.IsValidClientForFile(ClientInfo, UATLocation))
+					{
+						return ClientInfo;
+					}
+				}
+			}
+
+			// Otherwise search for all clients that match
+			List<P4ClientInfo> MatchingClients = new List<P4ClientInfo>();
 			P4ClientInfo[] P4Clients = Connection.GetClientsForUser(UserName, UATLocation);
-			foreach (var Client in P4Clients)
+			foreach (P4ClientInfo Client in P4Clients)
 			{
 				if (!String.IsNullOrEmpty(Client.Host) && String.Compare(Client.Host, HostName, true) != 0)
 				{
@@ -393,7 +410,7 @@ namespace AutomationTool
 			else
 			{
 				// We may have empty host clients here, so pick the first non-empty one if possible
-				foreach (var Client in MatchingClients)
+				foreach (P4ClientInfo Client in MatchingClients)
 				{
 					if (!String.IsNullOrEmpty(Client.Host) && String.Compare(Client.Host, HostName, true) == 0)
 					{
diff --git a/Engine/Source/Programs/AutomationTool/AutomationUtils/P4Utils.cs b/Engine/Source/Programs/AutomationTool/AutomationUtils/P4Utils.cs
--- a/Engine/Source/Programs/AutomationTool/AutomationUtils/P4Utils.cs
+++ b/Engine/Source/Programs/AutomationTool/AutomationUtils/P4Utils.cs
@@ -1112,7 +1112,7 @@ namespace AutomationTool
 					{
 						if(!GlobalCommandLine.P4)
 						{
-							LogWarning("Command {0} requires P4 functionality.", Command.Name);
+							LogInformation("Command {0} requires P4 functionality.", Command.Name);
 						}
 						bRequireP4 = true;
 
@@ -2180,6 +2180,18 @@ namespace AutomationTool
 			LogP4("revert " + String.Format("-c {0} //...", CL), SpewIsVerbose: SpewIsVerbose);
 		}
 
+		/// <summary>
+		/// Submits the specified changelist.
+		/// </summary>
+		/// <param name="CL">Changelist to submit.</param>
+		/// <param name="Force">If true, the submit will be forced even if resolve is needed.</param>
+		/// <param name="RevertIfFail">If true, if the submit fails, revert the CL.</param>
+		public void Submit(int CL, bool Force = false, bool RevertIfFail = false)
+		{
+			int SubmittedCL;
+			Submit(CL, out SubmittedCL, Force, RevertIfFail);
+		}
+
 		/// <summary>
 		/// Submits the specified changelist.
 		/// </summary>
@@ -3402,31 +3414,36 @@ namespace AutomationTool
 					throw new AutomationException("Failed to retrieve p4 client info for user {0}. Unable to set up local environment", UserName);
 				}
 				
-				bool bAddClient = true;
-				// Filter the client out if the specified path is not under the client root
-				if (!String.IsNullOrEmpty(PathUnderClientRoot) && !String.IsNullOrEmpty(Info.RootPath))
+				if (IsValidClientForFile(Info, PathUnderClientRoot))
 				{
-					var ClientRootPathWithSlash = Info.RootPath;
-					if (!ClientRootPathWithSlash.EndsWith("\\") && !ClientRootPathWithSlash.EndsWith("/"))
-					{
-						ClientRootPathWithSlash = CommandUtils.ConvertSeparators(PathSeparator.Default, ClientRootPathWithSlash + "/");
-					}
-					bAddClient = PathUnderClientRoot.StartsWith(ClientRootPathWithSlash, StringComparison.CurrentCultureIgnoreCase);
+					ClientList.Add(Info);
 				}
+			}
+			return ClientList.ToArray();
+		}
 
-				if (bAddClient)
+		public bool IsValidClientForFile(P4ClientInfo Info, string PathUnderClientRoot)
+		{
+			// Filter the client out if the specified path is not under the client root
+			bool bAddClient = true;
+			if (!String.IsNullOrEmpty(PathUnderClientRoot) && !String.IsNullOrEmpty(Info.RootPath))
+			{
+				var ClientRootPathWithSlash = Info.RootPath;
+				if (!ClientRootPathWithSlash.EndsWith("\\") && !ClientRootPathWithSlash.EndsWith("/"))
 				{
-					ClientList.Add(Info);
+					ClientRootPathWithSlash = CommandUtils.ConvertSeparators(PathSeparator.Default, ClientRootPathWithSlash + "/");
 				}
+				bAddClient = PathUnderClientRoot.StartsWith(ClientRootPathWithSlash, StringComparison.CurrentCultureIgnoreCase);
 			}
-			return ClientList.ToArray();
+			return bAddClient;
 		}
-        /// <summary>
-        /// Deletes a client.
-        /// </summary>
-        /// <param name="Name">Client name.</param>
-        /// <param name="Force">Forces the operation (-f)</param>
-        public void DeleteClient(string Name, bool Force = false, bool AllowSpew = true)
+
+		/// <summary>
+		/// Deletes a client.
+		/// </summary>
+		/// <param name="Name">Client name.</param>
+		/// <param name="Force">Forces the operation (-f)</param>
+		public void DeleteClient(string Name, bool Force = false, bool AllowSpew = true)
         {
             LogP4(String.Format("client -d {0} {1}", (Force ? "-f" : ""), Name), WithClient: false, AllowSpew: AllowSpew);
         }
