Adding scripts to detect and fix inconsistent case of files in Perforce.

To check for inconsistent case:

    CheckPerforceCase -Path=//UE4/Main/Engine/...

To fix files with inconsistent case:

    FixPerforceCase -Source=//UE4/Main/Engine/Build/ios/... -Target=//UE4/Main/Engine/Build/IOS/...

diff --git a/Engine/Source/Programs/AutomationTool/Scripts/AutomationScripts.Automation.csproj b/Engine/Source/Programs/AutomationTool/Scripts/AutomationScripts.Automation.csproj
--- a/Engine/Source/Programs/AutomationTool/Scripts/AutomationScripts.Automation.csproj
+++ b/Engine/Source/Programs/AutomationTool/Scripts/AutomationScripts.Automation.csproj
@@ -78,8 +78,10 @@
     <Compile Include="BenchmarkBuild\BenchmarkRunEditorTask.cs" />
     <Compile Include="BenchmarkBuild\BenchmarkSingleCompileTask.cs" />
     <Compile Include="BenchmarkBuild\BenchmarkTaskBase.cs" />
+    <Compile Include="CheckPerforceCase.cs" />
     <Compile Include="FinalizeInstalledBuild.cs" />
     <Compile Include="BuildHlslcc.Automation.cs" />
+    <Compile Include="FixPerforceCase.cs" />
     <Compile Include="ParseMsvcTimingInfo.cs" />
     <Compile Include="StashTarget.cs" />
     <Compile Include="BlameKeyword.Automation.cs" />
diff --git a/Engine/Source/Programs/AutomationTool/Scripts/CheckPerforceCase.cs b/Engine/Source/Programs/AutomationTool/Scripts/CheckPerforceCase.cs
new file mode 100644
--- /dev/null
+++ b/Engine/Source/Programs/AutomationTool/Scripts/CheckPerforceCase.cs
@@ -0,0 +1,131 @@
+using AutomationTool;
+using System;
+using System.Collections.Generic;
+using System.Linq;
+using System.Text;
+using System.Threading.Tasks;
+
+namespace BuildScripts.Automation
+{
+	[RequireP4]
+	[DoesNotNeedP4CL]
+	[Help("Checks that the casing of files within a path on a case-insensitive Perforce server is correct.")]
+	[Help("Path", "The path to query")]
+	class CheckPerforceCase : BuildCommand
+	{
+		class TreeNode
+		{
+			public string Path;
+			public List<string> Files = new List<string>();
+			public Dictionary<string, TreeNode> ChildNodes = new Dictionary<string, TreeNode>(StringComparer.Ordinal);
+			public int Count;
+
+			public TreeNode(string Path)
+			{
+				this.Path = Path;
+			}
+
+			public void Add(string DepotPath, int Idx)
+			{
+				int NextIdx = DepotPath.IndexOf('/', Idx);
+				if(NextIdx == -1)
+				{
+					Files.Add(DepotPath);
+				}
+				else
+				{
+					string Fragment = DepotPath.Substring(Idx, NextIdx - Idx);
+
+					TreeNode ChildNode;
+					if (!ChildNodes.TryGetValue(Fragment, out ChildNode))
+					{
+						ChildNode = new TreeNode(DepotPath.Substring(0, NextIdx));
+						ChildNodes.Add(Fragment, ChildNode);
+					}
+					ChildNode.Add(DepotPath, NextIdx + 1);
+				}
+				Count++;
+			}
+
+			public IEnumerable<string> EnumerateFiles()
+			{
+				foreach (string File in Files)
+				{
+					yield return File;
+				}
+				foreach (TreeNode ChildNode in ChildNodes.Values)
+				{
+					foreach (string ChildFile in ChildNode.EnumerateFiles())
+					{
+						yield return ChildFile;
+					}
+				}
+			}
+
+			public int PrintIssues()
+			{
+				int NumIssues = 0;
+				foreach (IGrouping<string, KeyValuePair<string, TreeNode>> Group in ChildNodes.GroupBy(x => x.Key, StringComparer.OrdinalIgnoreCase))
+				{
+					if (Group.Count() > 1)
+					{
+						StringBuilder Message = new StringBuilder();
+						Message.AppendFormat("Inconsistent casing for {0}:", Group.First().Value.Path);
+						foreach (KeyValuePair<string, TreeNode> Pair in Group)
+						{
+							Message.AppendFormat("\n  Could be '{0}':", Pair.Key, Pair.Value.Count, (Pair.Value.Count == 1)? "file" : "files");
+
+							int NumFiles = 0;
+							foreach (string File in Pair.Value.EnumerateFiles())
+							{
+								Message.AppendFormat("\n    {0}", File);
+								if (++NumFiles >= 10)
+								{
+									break;
+								}
+							}
+
+							Message.AppendFormat("\n    ({0} {1})", Pair.Value.Count, (Pair.Value.Count == 1) ? "file" : "files");
+						}
+						LogError(Message.ToString());
+						NumIssues++;
+					}
+
+					foreach (KeyValuePair<string, TreeNode> Pair in Group)
+					{
+						NumIssues += Pair.Value.PrintIssues();
+					}
+				}
+				return NumIssues;
+			}
+		}
+
+		public override void ExecuteBuild()
+		{
+			// Parse the path
+			string Filter = ParseRequiredStringParam("Filter");
+			if(!Filter.StartsWith("//"))
+			{
+				throw new AutomationException("Filter is not a depot path");
+			}
+
+			// Find all the matching files
+			LogInformation("Finding files matching {0}", Filter);
+			List<string> Files = P4.Files(String.Format("-e {0}", Filter));
+			LogInformation("Found {0} files", Files.Count);
+
+			// Build a tree of all the conflicting paths and print them out
+			TreeNode RootNode = new TreeNode("//");
+			foreach (string File in Files)
+			{
+				if (!File.StartsWith("//", StringComparison.Ordinal))
+				{
+					LogWarning("File '{0}' does not start with '//'", File);
+					continue;
+				}
+				RootNode.Add(File, 2);
+			}
+			RootNode.PrintIssues();
+		}
+	}
+}
diff --git a/Engine/Source/Programs/AutomationTool/Scripts/FixPerforceCase.cs b/Engine/Source/Programs/AutomationTool/Scripts/FixPerforceCase.cs
new file mode 100644
--- /dev/null
+++ b/Engine/Source/Programs/AutomationTool/Scripts/FixPerforceCase.cs
@@ -0,0 +1,125 @@
+using System;
+using System.Collections.Generic;
+using System.Linq;
+using System.Text;
+using System.Threading.Tasks;
+
+namespace AutomationTool
+{
+	[RequireP4]
+	[DoesNotNeedP4CL]
+	[Help("Fixes the case of files on a case-insensitive Perforce server by removing and re-adding them.")]
+	[Help("Source", "Pattern for source files to match. Should be a full depot path. May end with a wildcard.")]
+	[Help("Target", "Pattern for target files. Should be identical to source, except for case.")]
+	class FixPerforceCase : BuildCommand
+	{
+		/// <summary>
+		/// Main entry point for the command
+		/// </summary>
+		public override void ExecuteBuild()
+		{
+			string SourceSpec = ParseRequiredStringParam("Source");
+			string TargetSpec = ParseRequiredStringParam("Target");
+
+			// Make sure the patterns are a valid syntax
+			if (!String.Equals(SourceSpec, TargetSpec, StringComparison.OrdinalIgnoreCase))
+			{
+				throw new AutomationException("Source and target patterns do not match");
+			}
+			if (!SourceSpec.StartsWith("//"))
+			{
+				throw new AutomationException("Source and target patterns must be full depot paths");
+			}
+
+			// Pick out the source and target prefixes
+			string SourcePrefix;
+			string TargetPrefix;
+			if (SourceSpec.EndsWith("*"))
+			{
+				SourcePrefix = SourceSpec.Substring(0, SourceSpec.Length - 1);
+				TargetPrefix = TargetSpec.Substring(0, TargetSpec.Length - 1);
+			}
+			else if (SourceSpec.EndsWith("..."))
+			{
+				SourcePrefix = SourceSpec.Substring(0, SourceSpec.Length - 3);
+				TargetPrefix = TargetSpec.Substring(0, TargetSpec.Length - 3);
+			}
+			else
+			{
+				SourcePrefix = SourceSpec;
+				TargetPrefix = TargetSpec;
+			}
+
+			// Make sure there aren't any other wildcards in the pattern
+			if (SourcePrefix.Contains("?") || SourcePrefix.Contains("*") || SourcePrefix.Contains("..."))
+			{
+				throw new AutomationException("Wildcards are only permitted at the end of filespecs");
+			}
+
+			// Find all the source files
+			List<string> SourceFiles = P4.Files(String.Format("-e {0}", SourceSpec));
+			if (SourceFiles.Count == 0)
+			{
+				throw new AutomationException("No files found matching {0}", SourceSpec);
+			}
+			SourceFiles.RemoveAll(x => x.StartsWith(TargetPrefix, StringComparison.Ordinal));
+
+			// Error if we didn't find anything
+			if (SourceFiles.Count == 0)
+			{
+				throw new AutomationException("No files found matching spec");
+			}
+
+			// Find all the target files
+			List<string> TargetFiles = new List<string>(SourceFiles.Count);
+			foreach (string SourceFile in SourceFiles)
+			{
+				if (SourceFile.StartsWith(SourcePrefix))
+				{
+					TargetFiles.Add(TargetPrefix + SourceFile.Substring(SourcePrefix.Length));
+				}
+				else
+				{
+					throw new AutomationException("Source file '{0}' does not start with '{1}'", SourceFile, SourcePrefix);
+				}
+			}
+
+			// Print what we're going to do
+			LogInformation("Ready to rename {0} files:", SourceFiles.Count);
+			for (int Idx = 0; Idx < SourceFiles.Count; Idx++)
+			{
+				LogInformation("{0,3}: {1}", Idx, SourceFiles[Idx]);
+				LogInformation("{0,3}  {1}", "", TargetFiles[Idx]);
+			}
+
+			// If we're not going through with it, print the renames
+			if (!AllowSubmit)
+			{
+				LogWarning("Skipping due to no -Submit option");
+				return;
+			}
+
+			// Force sync all the old files 
+			foreach (string OldFile in SourceFiles)
+			{
+				P4.LogP4(String.Format("sync -f {0}", OldFile));
+			}
+
+			// Delete all the old files 
+			int DeleteChangeNumber = P4.CreateChange(Description: String.Format("Fixing case of {0} (1/2)", SourceSpec));
+			foreach (string OldFile in SourceFiles)
+			{
+				P4.LogP4(String.Format("delete -k -c {0} {1}", DeleteChangeNumber, OldFile));
+			}
+			P4.Submit(DeleteChangeNumber);
+
+			// Re-add all the files in the new location
+			int AddChangeNumber = P4.CreateChange(Description: String.Format("Fixing case of {0} (2/2)", SourceSpec));
+			foreach (string NewFile in TargetFiles)
+			{
+				P4.LogP4(String.Format("add -c {0} {1}", AddChangeNumber, NewFile));
+			}
+			P4.Submit(AddChangeNumber);
+		}
+	}
+}
