Reduce arguments to pass to FixPerforceCase command.

#rnx

diff --git a/Engine/Source/Programs/AutomationTool/Scripts/FixPerforceCase.cs b/Engine/Source/Programs/AutomationTool/Scripts/FixPerforceCase.cs
--- a/Engine/Source/Programs/AutomationTool/Scripts/FixPerforceCase.cs
+++ b/Engine/Source/Programs/AutomationTool/Scripts/FixPerforceCase.cs
@@ -13,56 +13,49 @@ namespace AutomationTool
 	[Help("Target", "Pattern for target files. Should be identical to source, except for case.")]
 	class FixPerforceCase : BuildCommand
 	{
+		const string BoilerplateText = "\n\n#rb none\n#rnx";
+
 		/// <summary>
 		/// Main entry point for the command
 		/// </summary>
 		public override void ExecuteBuild()
 		{
-			string SourceSpec = ParseRequiredStringParam("Source");
-			string TargetSpec = ParseRequiredStringParam("Target");
+			string FileSpec = ParseRequiredStringParam("Files");
 
 			// Make sure the patterns are a valid syntax
-			if (!String.Equals(SourceSpec, TargetSpec, StringComparison.OrdinalIgnoreCase))
-			{
-				throw new AutomationException("Source and target patterns do not match");
-			}
-			if (!SourceSpec.StartsWith("//"))
+			if (!FileSpec.StartsWith("//"))
 			{
-				throw new AutomationException("Source and target patterns must be full depot paths");
+				throw new AutomationException("Files must be specified as full depot paths");
 			}
 
 			// Pick out the source and target prefixes
-			string SourcePrefix;
-			string TargetPrefix;
-			if (SourceSpec.EndsWith("*"))
+			string Prefix;
+			if (FileSpec.EndsWith("*"))
 			{
-				SourcePrefix = SourceSpec.Substring(0, SourceSpec.Length - 1);
-				TargetPrefix = TargetSpec.Substring(0, TargetSpec.Length - 1);
+				Prefix = FileSpec.Substring(0, FileSpec.Length - 1);
 			}
-			else if (SourceSpec.EndsWith("..."))
+			else if (FileSpec.EndsWith("..."))
 			{
-				SourcePrefix = SourceSpec.Substring(0, SourceSpec.Length - 3);
-				TargetPrefix = TargetSpec.Substring(0, TargetSpec.Length - 3);
+				Prefix = FileSpec.Substring(0, FileSpec.Length - 3);
 			}
 			else
 			{
-				SourcePrefix = SourceSpec;
-				TargetPrefix = TargetSpec;
+				Prefix = FileSpec;
 			}
 
 			// Make sure there aren't any other wildcards in the pattern
-			if (SourcePrefix.Contains("?") || SourcePrefix.Contains("*") || SourcePrefix.Contains("..."))
+			if (Prefix.Contains("?") || Prefix.Contains("*") || Prefix.Contains("..."))
 			{
 				throw new AutomationException("Wildcards are only permitted at the end of filespecs");
 			}
 
 			// Find all the source files
-			List<string> SourceFiles = P4.Files(String.Format("-e {0}", SourceSpec));
+			List<string> SourceFiles = P4.Files(String.Format("-e {0}", FileSpec));
 			if (SourceFiles.Count == 0)
 			{
-				throw new AutomationException("No files found matching {0}", SourceSpec);
+				throw new AutomationException("No files found matching {0}", FileSpec);
 			}
-			SourceFiles.RemoveAll(x => x.StartsWith(TargetPrefix, StringComparison.Ordinal));
+			SourceFiles.RemoveAll(x => x.StartsWith(Prefix, StringComparison.Ordinal));
 
 			// Error if we didn't find anything
 			if (SourceFiles.Count == 0)
@@ -74,13 +67,13 @@ namespace AutomationTool
 			List<string> TargetFiles = new List<string>(SourceFiles.Count);
 			foreach (string SourceFile in SourceFiles)
 			{
-				if (SourceFile.StartsWith(SourcePrefix))
+				if (SourceFile.StartsWith(Prefix, StringComparison.OrdinalIgnoreCase))
 				{
-					TargetFiles.Add(TargetPrefix + SourceFile.Substring(SourcePrefix.Length));
+					TargetFiles.Add(Prefix + SourceFile.Substring(Prefix.Length));
 				}
 				else
 				{
-					throw new AutomationException("Source file '{0}' does not start with '{1}'", SourceFile, SourcePrefix);
+					throw new AutomationException("Source file '{0}' does not start with '{1}'", SourceFile, Prefix);
 				}
 			}
 
@@ -106,7 +99,7 @@ namespace AutomationTool
 			}
 
 			// Delete all the old files 
-			int DeleteChangeNumber = P4.CreateChange(Description: String.Format("Fixing case of {0} (1/2)", SourceSpec));
+			int DeleteChangeNumber = P4.CreateChange(Description: String.Format("Fixing case of {0} (1/2){1}", FileSpec, BoilerplateText));
 			foreach (string OldFile in SourceFiles)
 			{
 				P4.LogP4(String.Format("delete -k -c {0} {1}", DeleteChangeNumber, OldFile));
@@ -114,7 +107,7 @@ namespace AutomationTool
 			P4.Submit(DeleteChangeNumber);
 
 			// Re-add all the files in the new location
-			int AddChangeNumber = P4.CreateChange(Description: String.Format("Fixing case of {0} (2/2)", SourceSpec));
+			int AddChangeNumber = P4.CreateChange(Description: String.Format("Fixing case of {0} (2/2){1}", FileSpec, BoilerplateText));
 			foreach (string NewFile in TargetFiles)
 			{
 				P4.LogP4(String.Format("add -c {0} {1}", AddChangeNumber, NewFile));
