Add missing copyright notices.

#rnx

diff --git a/Engine/Source/Programs/AutomationTool/Scripts/CheckPerforceCase.cs b/Engine/Source/Programs/AutomationTool/Scripts/CheckPerforceCase.cs
--- a/Engine/Source/Programs/AutomationTool/Scripts/CheckPerforceCase.cs
+++ b/Engine/Source/Programs/AutomationTool/Scripts/CheckPerforceCase.cs
@@ -1,3 +1,5 @@
+// Copyright Epic Games, Inc. All Rights Reserved.
+
 using AutomationTool;
 using System;
 using System.Collections.Generic;
diff --git a/Engine/Source/Programs/AutomationTool/Scripts/FixPerforceCase.cs b/Engine/Source/Programs/AutomationTool/Scripts/FixPerforceCase.cs
--- a/Engine/Source/Programs/AutomationTool/Scripts/FixPerforceCase.cs
+++ b/Engine/Source/Programs/AutomationTool/Scripts/FixPerforceCase.cs
@@ -1,3 +1,5 @@
+// Copyright Epic Games, Inc. All Rights Reserved.
+
 using System;
 using System.Collections.Generic;
 using System.Linq;
