Don't warn in Gauntlet if we try to delete a file that doesn't exist when prepping to copy directory.

diff --git a/Engine/Source/Programs/AutomationTool/Gauntlet/Framework/Gauntlet.Utils.cs b/Engine/Source/Programs/AutomationTool/Gauntlet/Framework/Gauntlet.Utils.cs
--- a/Engine/Source/Programs/AutomationTool/Gauntlet/Framework/Gauntlet.Utils.cs
+++ b/Engine/Source/Programs/AutomationTool/Gauntlet/Framework/Gauntlet.Utils.cs
@@ -989,6 +989,11 @@ namespace Gauntlet
 					{
 						FileInfo DestInfo = new FileInfo(Path.Combine(DestDir.FullName, RelativePath));
 
+						if (!DestInfo.Exists)
+						{
+							continue;
+						}
+
 						if (Options.Verbose)
 						{
 							Log.Info("Deleting extra file {0}", DestInfo.FullName);
