tvOS fix for missing Shader libraries and missing .udebugsymbols

diff --git a/Engine/Source/Programs/AutomationTool/TVOS/TVOSPlatform.Automation.cs b/Engine/Source/Programs/AutomationTool/TVOS/TVOSPlatform.Automation.cs
--- a/Engine/Source/Programs/AutomationTool/TVOS/TVOSPlatform.Automation.cs
+++ b/Engine/Source/Programs/AutomationTool/TVOS/TVOSPlatform.Automation.cs
@@ -115,6 +115,24 @@ public class TVOSPlatform : IOSPlatform
                     DeployGeneratePList(SC.RawProjectPath, TargetConfiguration, (SC.IsCodeBasedProject ? SC.ProjectRoot : SC.EngineRoot), !SC.IsCodeBasedProject, (SC.IsCodeBasedProject ? SC.ShortProjectName : "UE4Game"), Params.Client, SC.ShortProjectName, SC.EngineRoot, DirectoryReference.Combine((SC.IsCodeBasedProject ? SC.ProjectRoot : SC.EngineRoot), "Binaries", "TVOS", "Payload", (SC.IsCodeBasedProject ? SC.ShortProjectName : "UE4Game") + ".app"), SC.StageExecutables[0], out bSupportsPortrait, out bSupportsLandscape, out bSkipIcons);
                 }
 
+
+            // copy the udebugsymbols if they exist
+            {
+                ConfigHierarchy PlatformGameConfig;
+                bool bIncludeSymbols = false;
+                if (Params.EngineConfigs.TryGetValue(SC.StageTargetPlatform.PlatformType, out PlatformGameConfig))
+                {
+                    PlatformGameConfig.GetBool("/Script/IOSRuntimeSettings.IOSRuntimeSettings", "bGenerateCrashReportSymbols", out bIncludeSymbols);
+                }
+                if (bIncludeSymbols)
+                {
+                    FileReference SymbolFileName = FileReference.Combine((SC.IsCodeBasedProject ? SC.ProjectRoot : SC.EngineRoot), "Binaries", "TVOS", SC.StageExecutables[0] + ".udebugsymbols");
+                    if (FileReference.Exists(SymbolFileName))
+                    {
+                        SC.StageFile(StagedFileType.NonUFS, SymbolFileName, new StagedFileReference((Params.ShortProjectName + ".udebugsymbols").ToLowerInvariant()));
+                    }
+                }
+            }
                 SC.StageFile(StagedFileType.SystemNonUFS, TargetPListFile, new StagedFileReference("Info.plist"));
             }
         }
@@ -124,5 +142,31 @@ public class TVOSPlatform : IOSPlatform
             StageMovieFiles(DirectoryReference.Combine(SC.EngineRoot, "Content", "Movies"), SC);
             StageMovieFiles(DirectoryReference.Combine(SC.ProjectRoot, "Content", "Movies"), SC);
         }
+
+        {
+            // Stage any *.metallib files as NonUFS.
+            // Get the final output directory for cooked data
+            DirectoryReference CookOutputDir;
+            if (!String.IsNullOrEmpty(Params.CookOutputDir))
+            {
+                CookOutputDir = DirectoryReference.Combine(new DirectoryReference(Params.CookOutputDir), SC.CookPlatform);
+            }
+            else if (Params.CookInEditor)
+            {
+                CookOutputDir = DirectoryReference.Combine(SC.ProjectRoot, "Saved", "EditorCooked", SC.CookPlatform);
+            }
+            else
+            {
+                CookOutputDir = DirectoryReference.Combine(SC.ProjectRoot, "Saved", "Cooked", SC.CookPlatform);
+            }
+            if (DirectoryReference.Exists(CookOutputDir))
+            {
+                List<FileReference> CookedFiles = DirectoryReference.EnumerateFiles(CookOutputDir, "*.metallib", SearchOption.AllDirectories).ToList();
+                foreach (FileReference CookedFile in CookedFiles)
+                {
+                    SC.StageFile(StagedFileType.NonUFS, CookedFile, new StagedFileReference(CookedFile.MakeRelativeTo(CookOutputDir)));
+                }
+            }
+        }
     }
 }
