Add option for up to 2 overflow obb files

#ue4
#android

diff --git a/Engine/Build/Android/Java/src/com/epicgames/ue4/GameActivity.java.template b/Engine/Build/Android/Java/src/com/epicgames/ue4/GameActivity.java.template
--- a/Engine/Build/Android/Java/src/com/epicgames/ue4/GameActivity.java.template
+++ b/Engine/Build/Android/Java/src/com/epicgames/ue4/GameActivity.java.template
@@ -6216,7 +6216,7 @@ public class GameActivity extends $${gameActivitySuperClass}$$ implements Surfac
 	public native boolean nativeIsShippingBuild();
 	public native void nativeSetAndroidStartupState(boolean bDebuggerAttached);
 	public native void nativeSetGlobalActivity(boolean bUseExternalFilesDir, boolean bPublicLogFiles, String internalFilePath, String externalFilePath, boolean bOBBInAPK, String APKPath);
-	public native void nativeSetObbFilePaths(String OBBMainFilePath, String OBBPatchFilePath);
+	public native void nativeSetObbFilePaths(String OBBMainFilePath, String OBBPatchFilePath, String OBBOverflow1FilePath, String OBBOverflow2FilePath);
 	public native void nativeSetWindowInfo(boolean bIsPortrait, int DepthBufferPreference);
 	public native void nativeSetObbInfo(String ProjectName, String PackageName, int Version, int PatchVersion, String AppType);
 	public native void nativeSetAndroidVersionInformation( String AndroidVersion, String PhoneMake, String PhoneModel, String PhoneBuildNumber, String OSLanguage );
diff --git a/Engine/Config/BaseEngine.ini b/Engine/Config/BaseEngine.ini
--- a/Engine/Config/BaseEngine.ini
+++ b/Engine/Config/BaseEngine.ini
@@ -2101,6 +2101,9 @@ bBundleDensitySplit=True
 bFullScreen=True
 bForceLDLinker=False
 bForceSmallOBBFiles=False
+bAllowLargeOBBFiles=False
+bAllowPatchOBBFile=False
+bAllowOverflowOBBFiles=False
 
 [/Script/AndroidPlatformEditor.AndroidSDKSettings]
 SDKAPILevel=latest
diff --git a/Engine/Source/Programs/AutomationTool/Android/AndroidPlatform.Automation.cs b/Engine/Source/Programs/AutomationTool/Android/AndroidPlatform.Automation.cs
--- a/Engine/Source/Programs/AutomationTool/Android/AndroidPlatform.Automation.cs
+++ b/Engine/Source/Programs/AutomationTool/Android/AndroidPlatform.Automation.cs
@@ -187,6 +187,41 @@ public class AndroidPlatform : Platform
 		return PatchName;
 	}
 
+	private static string GetFinalOverflowName(string ApkName, DeploymentContext SC, int Index, bool bUseAppType = true)
+	{
+		// calculate the name for the .obb file
+		string PackageName = GetPackageInfo(ApkName, SC, false);
+		if (PackageName == null)
+		{
+			throw new AutomationException(ExitCode.Error_FailureGettingPackageInfo, "Failed to get package name from " + ApkName);
+		}
+
+		string PackageVersion = GetPackageInfo(ApkName, SC, true);
+		if (PackageVersion == null || PackageVersion.Length == 0)
+		{
+			throw new AutomationException(ExitCode.Error_FailureGettingPackageInfo, "Failed to get package version from " + ApkName);
+		}
+
+		if (PackageVersion.Length > 0)
+		{
+			int IntVersion = int.Parse(PackageVersion);
+			PackageVersion = IntVersion.ToString("0");
+		}
+
+		string AppType = bUseAppType ? GetMetaAppType() : "";
+		if (AppType.Length > 0)
+		{
+			AppType += ".";
+		}
+
+		string OverflowName = string.Format("overflow{0}.{1}.{2}.{3}obb", Index, PackageVersion, PackageName, AppType);
+
+		// plop the .obb right next to the executable
+		OverflowName = Path.Combine(Path.GetDirectoryName(ApkName), OverflowName);
+
+		return OverflowName;
+	}
+
 
 	public override string GetPlatformPakCommandLine(ProjectParams Params, DeploymentContext SC)
 	{
@@ -215,6 +250,13 @@ public class AndroidPlatform : Platform
 		return TargetAndroidLocation + PackageName + "/" + Path.GetFileName(PatchName);
 	}
 
+	private static string GetDeviceOverflowName(string ApkName, DeploymentContext SC, int Index)
+	{
+		string OverflowName = GetFinalOverflowName(ApkName, SC, Index, false);
+		string PackageName = GetPackageInfo(ApkName, SC, false);
+		return TargetAndroidLocation + PackageName + "/" + Path.GetFileName(OverflowName);
+	}
+
 	public static string GetStorageQueryCommand()
     {
 		if (Utils.IsRunningOnMono)
@@ -328,6 +370,14 @@ public class AndroidPlatform : Platform
 		return bAllowPatchOBBFile;
 	}
 
+	private bool AllowOverflowOBBFiles(DeploymentContext SC)
+	{
+		ConfigHierarchy Ini = ConfigCache.ReadHierarchy(ConfigHierarchyType.Engine, DirectoryReference.FromFile(SC.RawProjectPath), SC.StageTargetPlatform.PlatformType);
+		bool bAllowOverflowOBBFiles = false;
+		Ini.GetBool("/Script/AndroidRuntimeSettings.AndroidRuntimeSettings", "bAllowOverflowOBBFiles", out bAllowOverflowOBBFiles);
+		return bAllowOverflowOBBFiles;
+	}
+
 	private bool CreateOBBFile(DeploymentContext SC, string StageDirectoryPath, string OutputFilename, List<FileReference> FilesForObb)
 	{
 		LogInformation("Creating {0} from {1}", OutputFilename, SC.StageDirectory);
@@ -419,6 +469,8 @@ public class AndroidPlatform : Platform
 
 		string LocalObbName = SC.StageDirectory.FullName+".obb";
 		string LocalPatchName = SC.StageDirectory.FullName + ".patch.obb";
+		string LocalOverflow1Name = SC.StageDirectory.FullName + ".overflow1.obb";
+		string LocalOverflow2Name = SC.StageDirectory.FullName + ".overflow2.obb";
 
 		FileFilter ObbFileFilter = new FileFilter(FileFilterType.Include);
 		ConfigHierarchy EngineIni = ConfigCache.ReadHierarchy(ConfigHierarchyType.Engine, DirectoryReference.FromFile(Params.RawProjectPath), UnrealTargetPlatform.Android);
@@ -472,16 +524,34 @@ public class AndroidPlatform : Platform
 				File.Delete(LocalPatchName);
 			}
 
+			// Always delete the target overflow1 OBB file if it exists
+			if (File.Exists(LocalOverflow1Name))
+			{
+				File.Delete(LocalOverflow1Name);
+			}
+
+			// Always delete the target overflow2 OBB file if it exists
+			if (File.Exists(LocalOverflow2Name))
+			{
+				File.Delete(LocalOverflow2Name);
+			}
+
 			List<FileReference> FilesToObb = FilesForObb;
 			List<FileReference> FilesToPatch = new List<FileReference>();
+			List<FileReference> FilesToOverflow1 = new List<FileReference>();
+			List<FileReference> FilesToOverflow2 = new List<FileReference>();
 
 			if (AllowPatchOBBFile(SC))
 			{
+				bool bAllowOverflowOBBs = AllowOverflowOBBFiles(SC);
+
 				FilesToObb = new List<FileReference>();
 
 				// Collect the filesize and place into Obb or Patch list
-				Int64 MainObbSize = 22 + 10 + 4096;		// EOCD wit comment (store version) + padding
-				Int64 PatchObbSize = 22 + 10 + 4096;	// EOCD wit comment (store version) + padding
+				Int64 MainObbSize = 22 + 10 + 4096;			// EOCD wit comment (store version) + padding
+				Int64 PatchObbSize = 22 + 10 + 4096;		// EOCD wit comment (store version) + padding
+				Int64 Overflow1ObbSize = 22 + 10 + 4096;	// EOCD wit comment (store version) + padding
+				Int64 Overflow2ObbSize = 22 + 10 + 4096;	// EOCD wit comment (store version) + padding
 				foreach (FileReference FileRef in FilesForObb)
 				{
 					FileInfo LocalFileInfo = new FileInfo(FileRef.FullName);
@@ -502,6 +572,25 @@ public class AndroidPlatform : Platform
 						FilesToPatch.Add(FileRef);
 						PatchObbSize += FileRequirements;
 					}
+					else if (bAllowOverflowOBBs)
+					{
+						if (Overflow1ObbSize + FileRequirements < OBBSizeAllowed)
+						{
+							FilesToOverflow1.Add(FileRef);
+							Overflow1ObbSize += FileRequirements;
+						}
+						else if (Overflow2ObbSize + FileRequirements < OBBSizeAllowed)
+						{
+							FilesToOverflow2.Add(FileRef);
+							Overflow2ObbSize += FileRequirements;
+						}
+						else
+						{
+							// no room in either file
+							LogInformation("Failed to build OBB: " + LocalObbName);
+							throw new AutomationException(ExitCode.Error_AndroidOBBError, "Stage Failed. Could not build OBB {0}. The file may be too big to fit in an OBB ({1} limit)", LocalObbName, LimitString);
+						}
+					}
 					else
 					{
 						// no room in either file
@@ -527,6 +616,26 @@ public class AndroidPlatform : Platform
 					throw new AutomationException(ExitCode.Error_AndroidOBBError, "Stage Failed. Could not build OBB {0}. The file may be too big to fit in an OBB ({1} limit)", LocalPatchName, LimitString);
 				}
 			}
+
+			// Now create the overflow1 OBB as a ZIP archive if required.
+			if (FilesToOverflow1.Count() > 0)
+			{
+				if (!CreateOBBFile(SC, StageDirectoryPath, LocalOverflow1Name, FilesToOverflow1))
+				{
+					LogInformation("Failed to build OBB: " + LocalOverflow1Name);
+					throw new AutomationException(ExitCode.Error_AndroidOBBError, "Stage Failed. Could not build OBB {0}. The file may be too big to fit in an OBB ({1} limit)", LocalPatchName, LimitString);
+				}
+			}
+
+			// Now create the overflow2 OBB as a ZIP archive if required.
+			if (FilesToOverflow2.Count() > 0)
+			{
+				if (!CreateOBBFile(SC, StageDirectoryPath, LocalOverflow2Name, FilesToOverflow2))
+				{
+					LogInformation("Failed to build OBB: " + LocalOverflow2Name);
+					throw new AutomationException(ExitCode.Error_AndroidOBBError, "Stage Failed. Could not build OBB {0}. The file may be too big to fit in an OBB ({1} limit)", LocalPatchName, LimitString);
+				}
+			}
 		}
 
 		// make sure the OBB is <= 2GiB (or 4GiB if large OBB enabled)
@@ -601,6 +710,10 @@ public class AndroidPlatform : Platform
 			    string ObbName = "";
 				string DevicePatchName = "";
 				string PatchName = "";
+				string DeviceOverflow1Name = "";
+				string Overflow1Name = "";
+				string DeviceOverflow2Name = "";
+				string Overflow2Name = "";
 				if (!bPackageDataInsideApk)
 			    {
 				    DeviceObbName = GetDeviceObbName(ApkName, SC);
@@ -619,6 +732,26 @@ public class AndroidPlatform : Platform
 						// apply store version to OBB to make it unique for PlayStore upload
 						UpdateObbStoreVersion(PatchName);
 					}
+
+					if (File.Exists(LocalOverflow1Name))
+					{
+						DeviceOverflow1Name = GetDeviceOverflowName(ApkName, SC, 1);
+						Overflow1Name = GetFinalOverflowName(ApkName, SC, 1);
+						CopyFile(LocalOverflow1Name, Overflow1Name);
+
+						// apply store version to OBB to make it unique for PlayStore upload
+						UpdateObbStoreVersion(Overflow1Name);
+					}
+
+					if (File.Exists(LocalOverflow2Name))
+					{
+						DeviceOverflow2Name = GetDeviceOverflowName(ApkName, SC, 2);
+						Overflow2Name = GetFinalOverflowName(ApkName, SC, 2);
+						CopyFile(LocalOverflow2Name, Overflow2Name);
+
+						// apply store version to OBB to make it unique for PlayStore upload
+						UpdateObbStoreVersion(Overflow2Name);
+					}
 				}
 
 				// check for optional universal apk
@@ -656,7 +789,8 @@ public class AndroidPlatform : Platform
 					// Write install batch file(s).
 					string PackageName = GetPackageInfo(ApkName, SC, false);
 					string BatchName = GetFinalBatchName(ApkName, SC, bMakeSeparateApks ? Architecture : "", bMakeSeparateApks ? GPUArchitecture : "", false, EBatchType.Install, Target);
-					string[] BatchLines = GenerateInstallBatchFile(bPackageDataInsideApk, PackageName, ApkName, Params, ObbName, DeviceObbName, false, PatchName, DevicePatchName, false, bIsPC, Params.Distribution, TargetSDKVersion > 22);
+					string[] BatchLines = GenerateInstallBatchFile(bPackageDataInsideApk, PackageName, ApkName, Params, ObbName, DeviceObbName, false, PatchName, DevicePatchName, false, 
+						Overflow1Name, DeviceOverflow1Name, false, Overflow2Name, DeviceOverflow2Name, false, bIsPC, Params.Distribution, TargetSDKVersion > 22);
 					if (bHaveAPK)
 					{
 						// make a batch file that can be used to install the .apk and .obb files
@@ -664,7 +798,7 @@ public class AndroidPlatform : Platform
 					}
 					// make a batch file that can be used to uninstall the .apk and .obb files
 					string UninstallBatchName = GetFinalBatchName(ApkName, SC, bMakeSeparateApks ? Architecture : "", bMakeSeparateApks ? GPUArchitecture : "", false, EBatchType.Uninstall, Target);
-					BatchLines = GenerateUninstallBatchFile(bPackageDataInsideApk, PackageName, ApkName, Params, ObbName, DeviceObbName, false, PatchName, DevicePatchName, false, bIsPC);
+					BatchLines = GenerateUninstallBatchFile(bPackageDataInsideApk, PackageName, ApkName, Params, false, bIsPC);
 					if (bHaveAPK || bHaveUniversal)
 					{
 						File.WriteAllLines(UninstallBatchName, BatchLines);
@@ -675,7 +809,8 @@ public class AndroidPlatform : Platform
 					{
 						UniversalBatchName = GetFinalBatchName(UniversalApkName, SC, "", "", false, EBatchType.Install, Target);
 						// make a batch file that can be used to install the .apk
-						string[] UniversalBatchLines = GenerateInstallBatchFile(bPackageDataInsideApk, PackageName, UniversalApkName, Params, ObbName, DeviceObbName, false, PatchName, DevicePatchName, false, bIsPC, Params.Distribution, TargetSDKVersion > 22);
+						string[] UniversalBatchLines = GenerateInstallBatchFile(bPackageDataInsideApk, PackageName, UniversalApkName, Params, ObbName, DeviceObbName, false, PatchName, DevicePatchName,
+							Overflow1Name, DeviceOverflow1Name, false, Overflow2Name, DeviceOverflow2Name, false, bIsPC, Params.Distribution, TargetSDKVersion > 22);
 						File.WriteAllLines(UniversalBatchName, UniversalBatchLines);
 					}
 
@@ -738,7 +873,9 @@ public class AndroidPlatform : Platform
 		PrintRunTime();
 	}
 
-    private string[] GenerateInstallBatchFile(bool bPackageDataInsideApk, string PackageName, string ApkName, ProjectParams Params, string ObbName, string DeviceObbName, bool bNoObbInstall, string PatchName, string DevicePatchName, bool bNoPatchInstall, bool bIsPC, bool bIsDistribution, bool bRequireRuntimeStoragePermission)
+    private string[] GenerateInstallBatchFile(bool bPackageDataInsideApk, string PackageName, string ApkName, ProjectParams Params, string ObbName, string DeviceObbName, bool bNoObbInstall,
+		string PatchName, string DevicePatchName, bool bNoPatchInstall, string Overflow1Name, string DeviceOverflow1Name, bool bNoOverflow1Install, string Overflow2Name, string DeviceOverflow2Name, bool bNoOverflow2Install,
+		bool bIsPC, bool bIsDistribution, bool bRequireRuntimeStoragePermission)
     {
         string[] BatchLines = null;
         string ReadPermissionGrantCommand = "shell pm grant " + PackageName + " android.permission.READ_EXTERNAL_STORAGE";
@@ -752,6 +889,8 @@ public class AndroidPlatform : Platform
 		bool bDontMoveOBB = bPackageDataInsideApk || !bIsDistribution;
 
 		bool bHavePatch = (PatchName != "");
+		bool bHaveOverflow1 = (Overflow1Name != "");
+		bool bHaveOverflow2 = (Overflow2Name != "");
 
 		if (!bIsPC)
         {
@@ -759,6 +898,8 @@ public class AndroidPlatform : Platform
 			// Note that $STORAGE/Android/obb will be the folder that contains the obb if you download the app from playstore.
 			string OBBInstallCommand = bNoObbInstall ? "shell 'rm -r $EXTERNAL_STORAGE/" + DeviceObbName + "'" : "push " + Path.GetFileName(ObbName) + " $STORAGE/" + (bIsDistribution ? "Download/" : "") + DeviceObbName;
 			string PatchInstallCommand = bNoPatchInstall ? "shell 'rm -r $EXTERNAL_STORAGE/" + DevicePatchName + "'" : "push " + Path.GetFileName(PatchName) + " $STORAGE/" + (bIsDistribution ? "Download/" : "") + DevicePatchName;
+			string Overflow1InstallCommand = bNoOverflow1Install ? "shell 'rm -r $EXTERNAL_STORAGE/" + DeviceOverflow1Name + "'" : "push " + Path.GetFileName(Overflow1Name) + " $STORAGE/" + (bIsDistribution ? "Download/" : "") + DeviceOverflow1Name;
+			string Overflow2InstallCommand = bNoOverflow2Install ? "shell 'rm -r $EXTERNAL_STORAGE/" + DeviceOverflow2Name + "'" : "push " + Path.GetFileName(Overflow2Name) + " $STORAGE/" + (bIsDistribution ? "Download/" : "") + DeviceOverflow2Name;
 
 			LogInformation("Writing shell script for install with {0}", bPackageDataInsideApk ? "data in APK" : "separate obb");
             BatchLines = new string[] {
@@ -792,6 +933,8 @@ public class AndroidPlatform : Platform
 						bPackageDataInsideApk ? "" : "\t$ADB $DEVICE " + OBBInstallCommand,
 						bPackageDataInsideApk ? "if [ 1 ]; then" : "\tif [ $? -eq 0 ]; then",
 						!bHavePatch ? "" : (bPackageDataInsideApk ? "" : "\t$ADB $DEVICE " + PatchInstallCommand),
+						!bHaveOverflow1 ? "" : (bPackageDataInsideApk ? "" : "\t$ADB $DEVICE " + Overflow1InstallCommand),
+						!bHaveOverflow2 ? "" : (bPackageDataInsideApk ? "" : "\t$ADB $DEVICE " + Overflow2InstallCommand),
 						bDontMoveOBB ? "" : "\t\t$ADB $DEVICE shell mkdir $STORAGE/Android/" + TargetAndroidLocation + PackageName, // don't check for error since installing may create the obb directory
 						bDontMoveOBB ? "" : "\t\t$ADB $DEVICE shell cp $STORAGE/Download/" + TargetAndroidLocation + PackageName + "/* $STORAGE/Android/" + TargetAndroidLocation + PackageName,
 						bDontMoveOBB ? "" : "\t\t$ADB $DEVICE shell rm -r $STORAGE/Download/" + TargetAndroidLocation + PackageName,
@@ -814,6 +957,8 @@ public class AndroidPlatform : Platform
         {
 			string OBBInstallCommand = bNoObbInstall ? "shell rm -r %STORAGE%/" + DeviceObbName : "push " + Path.GetFileName(ObbName) + " %STORAGE%/" + (bIsDistribution ? "Download/" : "") + DeviceObbName;
 			string PatchInstallCommand = bNoPatchInstall ? "shell rm -r %STORAGE%/" + DevicePatchName : "push " + Path.GetFileName(PatchName) + " %STORAGE%/" + (bIsDistribution ? "Download/" : "") + DevicePatchName;
+			string Overflow1InstallCommand = bNoOverflow1Install ? "shell rm -r %STORAGE%/" + DeviceOverflow1Name : "push " + Path.GetFileName(Overflow1Name) + " %STORAGE%/" + (bIsDistribution ? "Download/" : "") + DeviceOverflow1Name;
+			string Overflow2InstallCommand = bNoOverflow2Install ? "shell rm -r %STORAGE%/" + DeviceOverflow2Name : "push " + Path.GetFileName(Overflow2Name) + " %STORAGE%/" + (bIsDistribution ? "Download/" : "") + DeviceOverflow2Name;
 
 			LogInformation("Writing bat for install with {0}", bPackageDataInsideApk ? "data in APK" : "separate OBB");
             BatchLines = new string[] {
@@ -843,6 +988,10 @@ public class AndroidPlatform : Platform
 						bPackageDataInsideApk ? "" : "if \"%ERRORLEVEL%\" NEQ \"0\" goto Error",
 						!bHavePatch ? "" : (bPackageDataInsideApk ? "" : "%ADB% %DEVICE% " + PatchInstallCommand),
 						!bHavePatch ? "" : (bPackageDataInsideApk ? "" : "if \"%ERRORLEVEL%\" NEQ \"0\" goto Error"),
+						!bHaveOverflow1 ? "" : (bPackageDataInsideApk ? "" : "%ADB% %DEVICE% " + Overflow1InstallCommand),
+						!bHaveOverflow1 ? "" : (bPackageDataInsideApk ? "" : "if \"%ERRORLEVEL%\" NEQ \"0\" goto Error"),
+						!bHaveOverflow2 ? "" : (bPackageDataInsideApk ? "" : "%ADB% %DEVICE% " + Overflow2InstallCommand),
+						!bHaveOverflow2 ? "" : (bPackageDataInsideApk ? "" : "if \"%ERRORLEVEL%\" NEQ \"0\" goto Error"),
 						bDontMoveOBB ? "" : "%ADB% %DEVICE% shell mkdir %STORAGE%/Android/" + TargetAndroidLocation + PackageName, // don't check for error since installing may create the obb directory
 						bDontMoveOBB ? "" : "%ADB% %DEVICE% shell cp %STORAGE%/Download/" + TargetAndroidLocation + PackageName + "/* %STORAGE%/Android/" + TargetAndroidLocation + PackageName,
 						bDontMoveOBB ? "" : "if \"%ERRORLEVEL%\" NEQ \"0\" goto Error",
@@ -868,7 +1017,7 @@ public class AndroidPlatform : Platform
         return BatchLines;
     }
 
-	private string[] GenerateUninstallBatchFile(bool bPackageDataInsideApk, string PackageName, string ApkName, ProjectParams Params, string ObbName, string DeviceObbName, bool bNoObbInstall, string PatchName, string DevicePatchName, bool bNoPatchInstall, bool bIsPC)
+	private string[] GenerateUninstallBatchFile(bool bPackageDataInsideApk, string PackageName, string ApkName, ProjectParams Params, bool bIsPC)
 	{
 		string[] BatchLines = null;
 
@@ -993,6 +1142,8 @@ public class AndroidPlatform : Platform
 				bool bHaveAPK = FileExists(ApkName);
 				string ObbName = GetFinalObbName(ApkName, SC);
 				string PatchName = GetFinalPatchName(ApkName, SC);
+				string Overflow1Name = GetFinalPatchName(ApkName, SC, 1);
+				string Overflow2Name = GetFinalPatchName(ApkName, SC, 2);
 				bool bBuildWithHiddenSymbolVisibility = BuildWithHiddenSymbolVisibility(SC);
 				bool bSaveSymbols = GetSaveSymbols(SC);
 				//string NoOBBBatchName = GetFinalBatchName(ApkName, Params, bMakeSeparateApks ? Architecture : "", bMakeSeparateApks ? GPUArchitecture : "", true, false);
@@ -1097,6 +1248,14 @@ public class AndroidPlatform : Platform
 						{
 							SC.ArchiveFiles(Path.GetDirectoryName(PatchName), Path.GetFileName(PatchName));
 						}
+						if (FileExists(Overflow1Name))
+						{
+							SC.ArchiveFiles(Path.GetDirectoryName(Overflow1Name), Path.GetFileName(Overflow1Name));
+						}
+						if (FileExists(Overflow2Name))
+						{
+							SC.ArchiveFiles(Path.GetDirectoryName(Overflow2Name), Path.GetFileName(Overflow2Name));
+						}
 					}
 				}
 
@@ -1402,6 +1561,8 @@ public class AndroidPlatform : Platform
             String StorageLocation = Result.Output.Trim(); // "/mnt/sdcard";
             string DeviceObbName = StorageLocation + "/" + GetDeviceObbName(ApkName, SC);
 			string DevicePatchName = StorageLocation + "/" + GetDevicePatchName(ApkName, SC);
+			string DeviceOverflow1Name = StorageLocation + "/" + GetDevicePatchName(ApkName, SC, 1);
+			string DeviceOverflow2Name = StorageLocation + "/" + GetDevicePatchName(ApkName, SC, 2);
 			string RemoteDir = StorageLocation + "/UE4Game/" + Params.ShortProjectName;
 
             // determine if APK out of date
@@ -1778,6 +1939,8 @@ public class AndroidPlatform : Platform
                 // delete the .obb file, since it will cause nothing we just deployed to be used
                 RunAdbCommand(Params, DeviceName, "shell rm " + DeviceObbName);
 				RunAdbCommand(Params, DeviceName, "shell rm " + DevicePatchName);
+				RunAdbCommand(Params, DeviceName, "shell rm " + DeviceOverflow1Name);
+				RunAdbCommand(Params, DeviceName, "shell rm " + DeviceOverflow2Name);
 			}
 			else if (SC.Archive)
             {
@@ -1802,6 +1965,28 @@ public class AndroidPlatform : Platform
 					string Commandline = string.Format("{0} \"{1}\" \"{2}\"", BaseCommandline, PatchPath, DevicePatchName);
 					RunAdbCommand(Params, DeviceName, Commandline);
 				}
+
+				// deploy the overflow1 if there is one
+				string Overflow1Path = Path.Combine(SC.StageDirectory.FullName, GetFinalOverflowName(ApkName, SC, 1));
+				if (File.Exists(Overflow1Path))
+				{
+					// cache some strings
+					string BaseCommandline = "push";
+
+					string Commandline = string.Format("{0} \"{1}\" \"{2}\"", BaseCommandline, Overflow1Path, DeviceOverflow1Name);
+					RunAdbCommand(Params, DeviceName, Commandline);
+				}
+
+				// deploy the overflow2 if there is one
+				string Overflow2Path = Path.Combine(SC.StageDirectory.FullName, GetFinalOverflowName(ApkName, SC, 2));
+				if (File.Exists(Overflow2Path))
+				{
+					// cache some strings
+					string BaseCommandline = "push";
+
+					string Commandline = string.Format("{0} \"{1}\" \"{2}\"", BaseCommandline, Overflow2Path, DeviceOverflow2Name);
+					RunAdbCommand(Params, DeviceName, Commandline);
+				}
 			}
 			else
             {
diff --git a/Engine/Source/Runtime/Android/AndroidRuntimeSettings/Classes/AndroidRuntimeSettings.h b/Engine/Source/Runtime/Android/AndroidRuntimeSettings/Classes/AndroidRuntimeSettings.h
--- a/Engine/Source/Runtime/Android/AndroidRuntimeSettings/Classes/AndroidRuntimeSettings.h
+++ b/Engine/Source/Runtime/Android/AndroidRuntimeSettings/Classes/AndroidRuntimeSettings.h
@@ -245,10 +245,14 @@ public:
 	UPROPERTY(GlobalConfig, EditAnywhere, Category = "APK Packaging", Meta = (DisplayName = "Allow large OBB files."))
 	bool bAllowLargeOBBFiles;
 
-	// If checked, a patch OBB is generated for files not fitting in the main OBB (requires using multiple PAK files so split up content by chunk id). 
+	// If checked, a patch OBB is generated for files not fitting in the main OBB (requires using multiple PAK files so split up content by chunk id).
 	UPROPERTY(GlobalConfig, EditAnywhere, Category = "APK Packaging", Meta = (DisplayName = "Allow patch OBB file."))
 	bool bAllowPatchOBBFile;
 
+	// If checked, up to two additional overflow OBB files are generated for files not fitting in the patch OBB (requires using multiple PAK files so split up content by chunk id).
+	UPROPERTY(GlobalConfig, EditAnywhere, Category = "APK Packaging", Meta = (DisplayName = "Allow overflow OBB files."))
+	bool bAllowOverflowOBBFiles;
+
 	// If checked, UE4Game files will be placed in ExternalFilesDir which is removed on uninstall.
 	// You should also check this if you need to save you game progress without requesting runtime WRITE_EXTERNAL_STORAGE permission in android api 23+
 	UPROPERTY(GlobalConfig, EditAnywhere, Category = "APK Packaging", Meta = (DisplayName = "Use ExternalFilesDir for UE4Game files?"))
diff --git a/Engine/Source/Runtime/Core/Private/Android/AndroidPlatformFile.cpp b/Engine/Source/Runtime/Core/Private/Android/AndroidPlatformFile.cpp
--- a/Engine/Source/Runtime/Core/Private/Android/AndroidPlatformFile.cpp
+++ b/Engine/Source/Runtime/Core/Private/Android/AndroidPlatformFile.cpp
@@ -87,6 +87,10 @@ FString GOBBFilePathBase;
 FString GOBBMainFilePath;
 // Obb Patch filepath
 FString GOBBPatchFilePath;
+// Obb Overflow1 filepath
+FString GOBBOverflow1FilePath;
+// Obb Overflow2 filepath
+FString GOBBOverflow2FilePath;
 // Internal File Direcory Path (for application) - setup during load
 FString GInternalFilePath;
 // External File Direcory Path (for application) - setup during load
@@ -1032,29 +1036,80 @@ public:
 			}
 			else if (FileExists(*(OBBDir1 / MainOBBName), true))
 			{
-				MountOBB(*(OBBDir1 / MainOBBName));
-				FPlatformMisc::LowLevelOutputDebugStringf(TEXT("Mounted main OBB: %s"), *(OBBDir1 / MainOBBName));
+				GOBBMainFilePath = OBBDir1 / MainOBBName;
+				MountOBB(*GOBBMainFilePath);
+				FPlatformMisc::LowLevelOutputDebugStringf(TEXT("Mounted main OBB: %s"), *GOBBMainFilePath);
 			}
 			else if (FileExists(*(OBBDir2 / MainOBBName), true))
 			{
-				MountOBB(*(OBBDir2 / MainOBBName));
-				FPlatformMisc::LowLevelOutputDebugStringf(TEXT("Mounted main OBB: %s"), *(OBBDir2 / MainOBBName));
+				GOBBMainFilePath = OBBDir2 / MainOBBName;
+				MountOBB(*GOBBMainFilePath);
+				FPlatformMisc::LowLevelOutputDebugStringf(TEXT("Mounted main OBB: %s"), *GOBBMainFilePath);
 			}
 
+			bool bHavePatch = false;
 			if (!GOBBPatchFilePath.IsEmpty() && FileExists(*GOBBPatchFilePath, true))
 			{
+				bHavePatch = true;
 				MountOBB(*GOBBPatchFilePath);
 				FPlatformMisc::LowLevelOutputDebugStringf(TEXT("Mounted patch OBB: %s"), *GOBBPatchFilePath);
 			}
 			else if (FileExists(*(OBBDir1 / PatchOBBName), true))
 			{
-				MountOBB(*(OBBDir1 / PatchOBBName));
-				FPlatformMisc::LowLevelOutputDebugStringf(TEXT("Mounted patch OBB: %s"), *(OBBDir1 / PatchOBBName));
+				bHavePatch = true;
+				GOBBPatchFilePath = OBBDir1 / PatchOBBName;
+				MountOBB(*GOBBPatchFilePath);
+				FPlatformMisc::LowLevelOutputDebugStringf(TEXT("Mounted patch OBB: %s"), *GOBBPatchFilePath);
 			}
 			else if (FileExists(*(OBBDir2 / PatchOBBName), true))
 			{
-				MountOBB(*(OBBDir2 / PatchOBBName));
-				FPlatformMisc::LowLevelOutputDebugStringf(TEXT("Mounted patch OBB: %s"), *(OBBDir2 / PatchOBBName));
+				bHavePatch = true;
+				GOBBPatchFilePath = OBBDir2 / PatchOBBName;
+				MountOBB(*GOBBPatchFilePath);
+				FPlatformMisc::LowLevelOutputDebugStringf(TEXT("Mounted patch OBB: %s"), *GOBBPatchFilePath);
+			}
+
+			// Only check for overflow files if we found a patch file
+			if (bHavePatch)
+			{
+				FString Overflow1OBBName = FString::Printf(TEXT("overflow1.%d.%s.obb"), GAndroidPackageVersion, *GPackageName);
+				FString Overflow2OBBName = FString::Printf(TEXT("overflow2.%d.%s.obb"), GAndroidPackageVersion, *GPackageName);
+
+				if (!GOBBOverflow1FilePath.IsEmpty() && FileExists(*GOBBOverflow1FilePath, true))
+				{
+					MountOBB(*GOBBOverflow1FilePath);
+					FPlatformMisc::LowLevelOutputDebugStringf(TEXT("Mounted overflow1 OBB: %s"), *GOBBOverflow1FilePath);
+				}
+				else if (FileExists(*(OBBDir1 / Overflow1OBBName), true))
+				{
+					GOBBOverflow1FilePath = OBBDir1 / Overflow1OBBName;
+					MountOBB(*GOBBOverflow1FilePath);
+					FPlatformMisc::LowLevelOutputDebugStringf(TEXT("Mounted overflow1 OBB: %s"), *GOBBOverflow1FilePath);
+				}
+				else if (FileExists(*(OBBDir2 / Overflow1OBBName), true))
+				{
+					GOBBOverflow1FilePath = OBBDir2 / Overflow1OBBName;
+					MountOBB(*GOBBOverflow1FilePath);
+					FPlatformMisc::LowLevelOutputDebugStringf(TEXT("Mounted overflow1 OBB: %s"), *GOBBOverflow1FilePath);
+				}
+
+				if (!GOBBOverflow2FilePath.IsEmpty() && FileExists(*GOBBOverflow2FilePath, true))
+				{
+					MountOBB(*GOBBOverflow2FilePath);
+					FPlatformMisc::LowLevelOutputDebugStringf(TEXT("Mounted overflow2 OBB: %s"), *GOBBOverflow2FilePath);
+				}
+				else if (FileExists(*(OBBDir1 / Overflow2OBBName), true))
+				{
+					GOBBOverflow2FilePath = OBBDir1 / Overflow2OBBName;
+					MountOBB(*GOBBOverflow2FilePath);
+					FPlatformMisc::LowLevelOutputDebugStringf(TEXT("Mounted overflow2 OBB: %s"), *GOBBOverflow2FilePath);
+				}
+				else if (FileExists(*(OBBDir2 / Overflow2OBBName), true))
+				{
+					GOBBOverflow2FilePath = OBBDir2 / Overflow2OBBName;
+					MountOBB(*GOBBOverflow2FilePath);
+					FPlatformMisc::LowLevelOutputDebugStringf(TEXT("Mounted overflow2 OBB: %s"), *GOBBOverflow2FilePath);
+				}
 			}
 		}
 
diff --git a/Engine/Source/Runtime/Launch/Private/Android/AndroidJNI.cpp b/Engine/Source/Runtime/Launch/Private/Android/AndroidJNI.cpp
--- a/Engine/Source/Runtime/Launch/Private/Android/AndroidJNI.cpp
+++ b/Engine/Source/Runtime/Launch/Private/Android/AndroidJNI.cpp
@@ -45,6 +45,8 @@ extern bool GOverrideAndroidLogDir;
 extern FString GOBBFilePathBase;
 extern FString GOBBMainFilePath;
 extern FString GOBBPatchFilePath;
+extern FString GOBBOverflow1FilePath;
+extern FString GOBBOverflow2FilePath;
 extern FString GAPKFilename;
 
 FOnActivityResult FJavaWrapper::OnActivityResultDelegate;
@@ -1598,6 +1600,8 @@ JNIEXPORT jint JNI_OnLoad(JavaVM* InJavaVM, void* InReserved)
 	GOBBFilePathBase = GFilePathBase;
 	GOBBMainFilePath = TEXT("");
 	GOBBPatchFilePath = TEXT("");
+	GOBBOverflow1FilePath = TEXT("");
+	GOBBOverflow2FilePath = TEXT("");
 
 	// then release...
 	FPlatformMisc::LowLevelOutputDebugStringf(TEXT("Path found as '%s'\n"), *GFilePathBase);
@@ -1622,10 +1626,12 @@ JNIEXPORT jint JNI_OnLoad(JavaVM* InJavaVM, void* InReserved)
 //Native-defined functions
 
 //This function is declared in the Java-defined class, GameActivity.java: "public native void naativeSetObbFilePaths();"
-JNI_METHOD void Java_com_epicgames_ue4_GameActivity_nativeSetObbFilePaths(JNIEnv* jenv, jobject thiz, jstring OBBMainFilePath, jstring OBBPatchFilePath)
+JNI_METHOD void Java_com_epicgames_ue4_GameActivity_nativeSetObbFilePaths(JNIEnv* jenv, jobject thiz, jstring OBBMainFilePath, jstring OBBPatchFilePath, jstring OBBOverflow1FilePath, jstring OBBOverflow2FilePath)
 {
 	GOBBMainFilePath = FJavaHelper::FStringFromParam(jenv, OBBMainFilePath);
 	GOBBPatchFilePath = FJavaHelper::FStringFromParam(jenv, OBBPatchFilePath);
+	GOBBOverflow1FilePath = FJavaHelper::FStringFromParam(jenv, OBBOverflow1FilePath);
+	GOBBOverflow2FilePath = FJavaHelper::FStringFromParam(jenv, OBBOverflow2FilePath);
 }
 
 //This function is declared in the Java-defined class, GameActivity.java: "public native void nativeSetGlobalActivity();"
