Fix missing parameters

#ue4
#android

diff --git a/Engine/Source/Programs/AutomationTool/Android/AndroidPlatform.Automation.cs b/Engine/Source/Programs/AutomationTool/Android/AndroidPlatform.Automation.cs
--- a/Engine/Source/Programs/AutomationTool/Android/AndroidPlatform.Automation.cs
+++ b/Engine/Source/Programs/AutomationTool/Android/AndroidPlatform.Automation.cs
@@ -798,7 +798,7 @@ public class AndroidPlatform : Platform
 					}
 					// make a batch file that can be used to uninstall the .apk and .obb files
 					string UninstallBatchName = GetFinalBatchName(ApkName, SC, bMakeSeparateApks ? Architecture : "", bMakeSeparateApks ? GPUArchitecture : "", false, EBatchType.Uninstall, Target);
-					BatchLines = GenerateUninstallBatchFile(bPackageDataInsideApk, PackageName, ApkName, Params, false, bIsPC);
+					BatchLines = GenerateUninstallBatchFile(bPackageDataInsideApk, PackageName, ApkName, Params, bIsPC);
 					if (bHaveAPK || bHaveUniversal)
 					{
 						File.WriteAllLines(UninstallBatchName, BatchLines);
@@ -809,7 +809,7 @@ public class AndroidPlatform : Platform
 					{
 						UniversalBatchName = GetFinalBatchName(UniversalApkName, SC, "", "", false, EBatchType.Install, Target);
 						// make a batch file that can be used to install the .apk
-						string[] UniversalBatchLines = GenerateInstallBatchFile(bPackageDataInsideApk, PackageName, UniversalApkName, Params, ObbName, DeviceObbName, false, PatchName, DevicePatchName,
+						string[] UniversalBatchLines = GenerateInstallBatchFile(bPackageDataInsideApk, PackageName, UniversalApkName, Params, ObbName, DeviceObbName, false, PatchName, DevicePatchName, false,
 							Overflow1Name, DeviceOverflow1Name, false, Overflow2Name, DeviceOverflow2Name, false, bIsPC, Params.Distribution, TargetSDKVersion > 22);
 						File.WriteAllLines(UniversalBatchName, UniversalBatchLines);
 					}
@@ -1142,8 +1142,8 @@ public class AndroidPlatform : Platform
 				bool bHaveAPK = FileExists(ApkName);
 				string ObbName = GetFinalObbName(ApkName, SC);
 				string PatchName = GetFinalPatchName(ApkName, SC);
-				string Overflow1Name = GetFinalPatchName(ApkName, SC, 1);
-				string Overflow2Name = GetFinalPatchName(ApkName, SC, 2);
+				string Overflow1Name = GetFinalOverflowName(ApkName, SC, 1);
+				string Overflow2Name = GetFinalOverflowName(ApkName, SC, 2);
 				bool bBuildWithHiddenSymbolVisibility = BuildWithHiddenSymbolVisibility(SC);
 				bool bSaveSymbols = GetSaveSymbols(SC);
 				//string NoOBBBatchName = GetFinalBatchName(ApkName, Params, bMakeSeparateApks ? Architecture : "", bMakeSeparateApks ? GPUArchitecture : "", true, false);
@@ -1561,8 +1561,8 @@ public class AndroidPlatform : Platform
             String StorageLocation = Result.Output.Trim(); // "/mnt/sdcard";
             string DeviceObbName = StorageLocation + "/" + GetDeviceObbName(ApkName, SC);
 			string DevicePatchName = StorageLocation + "/" + GetDevicePatchName(ApkName, SC);
-			string DeviceOverflow1Name = StorageLocation + "/" + GetDevicePatchName(ApkName, SC, 1);
-			string DeviceOverflow2Name = StorageLocation + "/" + GetDevicePatchName(ApkName, SC, 2);
+			string DeviceOverflow1Name = StorageLocation + "/" + GetDeviceOverflowName(ApkName, SC, 1);
+			string DeviceOverflow2Name = StorageLocation + "/" + GetDeviceOverflowName(ApkName, SC, 2);
 			string RemoteDir = StorageLocation + "/UE4Game/" + Params.ShortProjectName;
 
             // determine if APK out of date
