Fixed incorrect target receipt path breaking iOS packaing on Mac
Fixed incorrect storyboard image copy for blueprint projects

diff --git a/Engine/Source/Programs/AutomationTool/IOS/IOSPlatform.Automation.cs b/Engine/Source/Programs/AutomationTool/IOS/IOSPlatform.Automation.cs
--- a/Engine/Source/Programs/AutomationTool/IOS/IOSPlatform.Automation.cs
+++ b/Engine/Source/Programs/AutomationTool/IOS/IOSPlatform.Automation.cs
@@ -194,7 +194,14 @@ public class IOSPlatform : Platform
 	{
 		string TargetName = Path.GetFileNameWithoutExtension(InExecutablePath).Split("-".ToCharArray())[0];
 		FileReference TargetReceiptFileName;
-		TargetReceiptFileName = TargetReceipt.GetDefaultPath(InProjectDirectory, TargetName, UnrealTargetPlatform.IOS, Config, "");
+		if (bIsUE4Game)
+		{
+			TargetReceiptFileName = TargetReceipt.GetDefaultPath(InEngineDir, "UE4Game", UnrealTargetPlatform.IOS, Config, "");
+		}
+		else
+		{
+			TargetReceiptFileName = TargetReceipt.GetDefaultPath(InProjectDirectory, TargetName, UnrealTargetPlatform.IOS, Config, "");
+		}
 		return TargetReceiptFileName;
 	}
 
@@ -358,6 +365,13 @@ public class IOSPlatform : Platform
 			}
 		}
 
+		// Stage the user defined LaunchScreen if any.
+		string BuildGraphicsDirectory = Path.GetDirectoryName(Params.RawProjectPath.FullName) + "/Build/IOS/Resources/Graphics/";
+		if (File.Exists(BuildGraphicsDirectory + "LaunchScreenIOS.png"))
+		{
+			InternalUtils.SafeCopyFile(BuildGraphicsDirectory + "LaunchScreenIOS.png", SC.StageDirectory + "/LaunchScreenIOS.png");
+		}
+
 		IOSExports.GenerateAssetCatalog(Params.RawProjectPath, new FileReference(FullExePath), new DirectoryReference(CombinePaths(Params.BaseStageDirectory, (Platform == UnrealTargetPlatform.IOS ? "IOS" : "TVOS"))), Platform);
 
 		bCreatedIPA = false;
diff --git a/Engine/Source/Programs/UnrealBuildTool/Platform/IOS/UEBuildIOS.cs b/Engine/Source/Programs/UnrealBuildTool/Platform/IOS/UEBuildIOS.cs
--- a/Engine/Source/Programs/UnrealBuildTool/Platform/IOS/UEBuildIOS.cs
+++ b/Engine/Source/Programs/UnrealBuildTool/Platform/IOS/UEBuildIOS.cs
@@ -956,20 +956,6 @@ namespace UnrealBuildTool
 			return false;
 		}
 
-		public bool HasCustomLaunchScreen(DirectoryReference ProjectDirectoryName)
-		{
-			string LaunchScreenDir = Path.Combine(ProjectDirectoryName.FullName, "Build", "IOS", "Resources", "Graphics");
-			if (Directory.Exists(LaunchScreenDir))
-			{
-				if (File.Exists(LaunchScreenDir + "\\LaunchScreenIOS.png"))
-				{
-					Log.TraceInformation("Requiring custom build because project {0} has a custom LaunchScreen", Path.GetFileName(ProjectDirectoryName.FullName));
-					return true;
-				}
-			}
-			return false;
-		}
-
 		/// <summary>
 		/// Check for the default configuration
 		/// return true if the project uses the default build config
@@ -989,7 +975,7 @@ namespace UnrealBuildTool
 			};
 
 			// check for custom icons
-			if (HasCustomIcons(ProjectDirectoryName) || HasCustomLaunchScreen(ProjectDirectoryName))
+			if (HasCustomIcons(ProjectDirectoryName))
 			{
 				return false;
 			}
