- Moved some platform-related binaries to their platformextension

diff --git a/Engine/Build/InstalledEngineFilters.xml b/Engine/Build/InstalledEngineFilters.xml
--- a/Engine/Build/InstalledEngineFilters.xml
+++ b/Engine/Build/InstalledEngineFilters.xml
@@ -539,8 +539,8 @@
 		Engine/Build/XboxOne/...
 
 		<!-- Additional Binaries -->
-		Engine/Binaries/Win64/XboxOne/XboxOnePackageNameUtil.exe
-		Engine/Binaries/Win64/XboxOne/XboxOnePDBFileUtil.exe
+		Engine/Platforms/XboxOne/Binaries/Win64/XboxOnePackageNameUtil.exe
+		Engine/Platforms/XboxOne/Binaries/Win64/XboxOnePDBFileUtil.exe
 	</Property>
 
 	<Property Name="CopyXboxOneExceptions">
diff --git a/Engine/Source/Programs/AutomationTool/Gauntlet/Unreal/RunUnreal.cs b/Engine/Source/Programs/AutomationTool/Gauntlet/Unreal/RunUnreal.cs
--- a/Engine/Source/Programs/AutomationTool/Gauntlet/Unreal/RunUnreal.cs
+++ b/Engine/Source/Programs/AutomationTool/Gauntlet/Unreal/RunUnreal.cs
@@ -290,7 +290,7 @@ namespace Gauntlet
 
 			if (UsedPlatforms.Contains(UnrealTargetPlatform.PS4))
 			{
-				String DevKitUtilPath = Path.Combine(Environment.CurrentDirectory, "Engine/Binaries/DotNET/PS4/PS4DevKitUtil.exe");
+				String DevKitUtilPath = Path.Combine(Environment.CurrentDirectory, "Engine/Platforms/PS4/Binaries/DotNET/PS4DevKitUtil.exe");
 				Gauntlet.Log.Verbose("PS4DevkitUtil executing 'removeall'");
 				IProcessResult BootResult = CommandUtils.Run(DevKitUtilPath, "removeall");
 			}
diff --git a/Engine/Source/Programs/AutomationTool/Scripts/BuildCommonTools.Automation.cs b/Engine/Source/Programs/AutomationTool/Scripts/BuildCommonTools.Automation.cs
--- a/Engine/Source/Programs/AutomationTool/Scripts/BuildCommonTools.Automation.cs
+++ b/Engine/Source/Programs/AutomationTool/Scripts/BuildCommonTools.Automation.cs
@@ -1,4 +1,4 @@
-﻿// Copyright Epic Games, Inc. All Rights Reserved.
+// Copyright Epic Games, Inc. All Rights Reserved.
 
 using System;
 using System.IO;
@@ -143,11 +143,11 @@ public class BuildCommonTools : BuildCommand
 		{
 			Agenda.AddTarget("PS4MapFileUtil", UnrealBuildTool.UnrealTargetPlatform.Win64, UnrealBuildTool.UnrealTargetConfiguration.Development);
 
-			Agenda.DotNetProjects.Add(CommandUtils.CombinePaths(CommandUtils.CmdEnv.LocalRoot, @"Engine/Source/Programs/PS4/PS4DevKitUtil/PS4DevKitUtil.csproj"));
-			ExtraBuildProducts.Add(CommandUtils.CombinePaths(CommandUtils.CmdEnv.LocalRoot, @"Engine/Binaries/DotNET/PS4/PS4DevKitUtil.exe"));
+			Agenda.DotNetProjects.Add(CommandUtils.CombinePaths(CommandUtils.CmdEnv.LocalRoot, @"Engine/Platforms/PS4/Source/Programs/PS4DevKitUtil/PS4DevKitUtil.csproj"));
+			ExtraBuildProducts.Add(CommandUtils.CombinePaths(CommandUtils.CmdEnv.LocalRoot, @"Engine/Platforms/PS4/Binaries/DotNETPS4DevKitUtil.exe"));
 
-			Agenda.DotNetProjects.Add(CommandUtils.CombinePaths(CommandUtils.CmdEnv.LocalRoot, @"Engine/Source/Programs/PS4/PS4SymbolTools/PS4SymbolTool.csproj"));
-			ExtraBuildProducts.Add(CommandUtils.CombinePaths(CommandUtils.CmdEnv.LocalRoot, @"Engine/Binaries/DotNET/PS4/PS4SymbolTool.exe"));
+			Agenda.DotNetProjects.Add(CommandUtils.CombinePaths(CommandUtils.CmdEnv.LocalRoot, @"Engine/Platforms/PS4/Source/Programs/PS4SymbolTools/PS4SymbolTool.csproj"));
+			ExtraBuildProducts.Add(CommandUtils.CombinePaths(CommandUtils.CmdEnv.LocalRoot, @"Engine/Platforms/PS4/Binaries/DotNET/PS4SymbolTool.exe"));
 		}
 
 		// Xbox One binaries
