Put Gauntlet platform files in the correct subfolder.

diff --git a/Engine/Source/Programs/AutomationTool/Gauntlet/Gauntlet.Automation.csproj b/Engine/Source/Programs/AutomationTool/Gauntlet/Gauntlet.Automation.csproj
--- a/Engine/Source/Programs/AutomationTool/Gauntlet/Gauntlet.Automation.csproj
+++ b/Engine/Source/Programs/AutomationTool/Gauntlet/Gauntlet.Automation.csproj
@@ -215,6 +215,10 @@
   files and include them in the build. They will also appear in VisualStudio -->
   <ItemGroup>
     <Compile Include="../../../../Platforms/*/Source/Programs/AutomationTool/Gauntlet/**/*.cs">
+      <!-- RecursiveDir metadata expands to [PlatformName]/Source/Programs/[etc]. The Replace() will
+      replace everything from the first slash to the end leaving just the platform name. Should the
+      Replace() fail then RecursiveDir is still an agreeable value -->
+      <Link>Platform/$([System.Text.RegularExpressions.Regex]::Replace(%(Compile.RecursiveDir), [\\/].+$, ``))/_</Link>
     </Compile>
   </ItemGroup>
   
