Strip comments from ini files when staging

josh.adams

diff --git a/Engine/Source/Programs/AutomationTool/AutomationUtils/Utils.cs b/Engine/Source/Programs/AutomationTool/AutomationUtils/Utils.cs
--- a/Engine/Source/Programs/AutomationTool/AutomationUtils/Utils.cs
+++ b/Engine/Source/Programs/AutomationTool/AutomationUtils/Utils.cs
@@ -277,7 +277,7 @@ namespace AutomationTool
 		}
 
 		// Characters that can appear at the start of
-		private static char[] IgnoredIniValuePrefixes = { '+', '-', ' ', '\t', ';' };
+		private static char[] IgnoredIniValuePrefixes = { '+', '-', ' ', '\t' };
 
 		private static void FilterIniFile(string SourceName, string TargetName, List<string> IniKeyBlacklist, List<string> InSectionBlacklist)
 		{
@@ -288,6 +288,12 @@ namespace AutomationTool
 			foreach (string OriginalLine in Lines)
 			{
 				string Line = OriginalLine.Trim();
+
+				if (Line.StartsWith(";"))
+				{
+					continue;
+				}
+
 				bool bFiltered = bFilteringSection;
 
 				// look for each filter on each line
