Lower a warning in the gif-ify process in gauntlet to an info like all of the rest of the errors. Also, cause image conversion not to delete the source item when it has the same location as the post conversion file.

diff --git a/Engine/Source/Programs/AutomationTool/Gauntlet/Framework/Gauntlet.Utils.cs b/Engine/Source/Programs/AutomationTool/Gauntlet/Framework/Gauntlet.Utils.cs
--- a/Engine/Source/Programs/AutomationTool/Gauntlet/Framework/Gauntlet.Utils.cs
+++ b/Engine/Source/Programs/AutomationTool/Gauntlet/Framework/Gauntlet.Utils.cs
@@ -1367,7 +1367,7 @@ namespace Gauntlet
 
 				if (Files.Count() == 0)
 				{
-					Log.Warning("Could not find files at {0} to Gif-ify", InDirectory);
+					Log.Info("Could not find files at {0} to Gif-ify", InDirectory);
 					return false;
 				}
 
@@ -1424,24 +1424,44 @@ namespace Gauntlet
 
 				try
 				{
+					List<FileInfo> FilesToCleanUp = new List<FileInfo>();
 					foreach (FileInfo File in Files)
 					{
 						using (MagickImage Image = new MagickImage(File.FullName))
 						{
 							string OutFile = Path.Combine(OutDirectory, File.Name);
-							OutFile = Path.ChangeExtension(OutFile, OutExtension);	
-							Image.Write(OutFile);
+							OutFile = Path.ChangeExtension(OutFile, OutExtension);
+							// If we're trying to convert something to itself in place, skip the step.
+							if (OutFile != File.FullName)
+							{
+								Image.Write(OutFile);
+								if (DeleteOriginals)
+								{
+									FilesToCleanUp.Add(File);
+								}
+							}
 						}
 					}
 
-					if (DeleteOriginals)
+					foreach (FileInfo File in FilesToCleanUp)
 					{
-						Files.ToList().ForEach(F => F.Delete());
+						File.Delete();
 					}
 				}
 				catch (System.Exception Ex)
 				{
 					Log.Warning("ConvertImages failed: {0}", Ex);
+					try
+					{
+						if (DeleteOriginals)
+						{
+							Files.ToList().ForEach(F => F.Delete());
+						}
+					}
+					catch (System.Exception e)
+					{
+						Log.Warning("Cleaning up original files failed: {0}", e);
+					}
 					return false;
 				}
 
