Fix for Android builds not being discovered in Gauntlet

- Typo where when trying to remove 32-bit Android builds in favor of 64-bit builds we were setting them back to 32-bit.

diff --git a/Engine/Source/Programs/AutomationTool/Gauntlet/Platform/Android/Gauntlet.AndroidBuildSource.cs b/Engine/Source/Programs/AutomationTool/Gauntlet/Platform/Android/Gauntlet.AndroidBuildSource.cs
--- a/Engine/Source/Programs/AutomationTool/Gauntlet/Platform/Android/Gauntlet.AndroidBuildSource.cs
+++ b/Engine/Source/Programs/AutomationTool/Gauntlet/Platform/Android/Gauntlet.AndroidBuildSource.cs
@@ -157,7 +157,7 @@ namespace Gauntlet
 			// If we have both 32 and 64-bit builds, prefer 64-bit
 			if (DiscoveredBuilds.Where(B => B.Is32Bit == false).Any())
 			{
-				DiscoveredBuilds = DiscoveredBuilds.Where(B => B.Is32Bit = false).ToList();
+				DiscoveredBuilds = DiscoveredBuilds.Where(B => !B.Is32Bit).ToList();
 			}
 
 			return DiscoveredBuilds;
