Removed overly pessimistic error message when running UAT with -iostore without -pak
#rnx

diff --git a/Engine/Source/Programs/AutomationTool/AutomationUtils/ProjectParams.cs b/Engine/Source/Programs/AutomationTool/AutomationUtils/ProjectParams.cs
--- a/Engine/Source/Programs/AutomationTool/AutomationUtils/ProjectParams.cs
+++ b/Engine/Source/Programs/AutomationTool/AutomationUtils/ProjectParams.cs
@@ -2758,11 +2758,6 @@ namespace AutomationTool
 			{
 				throw new AutomationException("DiscVersion is required for generating remaster package.");
 			}*/
-
-			if ((IoStore && !SkipIoStore) && (!Pak || SkipPak))
-			{
-				throw new AutomationException("-iostore can only be used with -pak");
-			}
 		}
 
 		protected bool bLogged = false;
