Move Build platform files to platform extensions folder
#rnx

diff --git a/Engine/Source/Programs/AutomationTool/Scripts/CopyBuildToStagingDirectory.Automation.cs b/Engine/Source/Programs/AutomationTool/Scripts/CopyBuildToStagingDirectory.Automation.cs
--- a/Engine/Source/Programs/AutomationTool/Scripts/CopyBuildToStagingDirectory.Automation.cs
+++ b/Engine/Source/Programs/AutomationTool/Scripts/CopyBuildToStagingDirectory.Automation.cs
@@ -1868,7 +1868,11 @@ public partial class Project : CommandUtils
 		List<string> Blacklist = null;
 		if (SC.StageTargetConfigurations.Count == 1)
 		{
-			FileReference PakBlacklistFilename = FileReference.Combine(SC.ProjectRoot, "Build", SC.PlatformDir, string.Format("PakBlacklist-{0}.txt", SC.StageTargetConfigurations[0].ToString()));
+			FileReference PakBlacklistFilename = FileReference.Combine(SC.ProjectRoot, "Platforms", SC.PlatformDir, "Build", string.Format("PakBlacklist-{0}.txt", SC.StageTargetConfigurations[0].ToString()));
+			if (!FileReference.Exists(PakBlacklistFilename))
+			{
+				PakBlacklistFilename = FileReference.Combine(SC.ProjectRoot, "Build", SC.PlatformDir, string.Format("PakBlacklist-{0}.txt", SC.StageTargetConfigurations[0].ToString()));
+			}
 			if (FileReference.Exists(PakBlacklistFilename))
 			{
 				LogInformation("Applying PAK blacklist file {0}. This is deprecated in favor of DefaultPakFileRules.ini", PakBlacklistFilename);
