Remove unused test code.

diff --git a/Engine/Source/Programs/AutomationTool/Scripts/Tests.Automation.cs b/Engine/Source/Programs/AutomationTool/Scripts/Tests.Automation.cs
--- a/Engine/Source/Programs/AutomationTool/Scripts/Tests.Automation.cs
+++ b/Engine/Source/Programs/AutomationTool/Scripts/Tests.Automation.cs
@@ -831,269 +831,6 @@ class TestChangeFileType : BuildCommand
 	}
 }
 
-class TestGamePerf : BuildCommand
-{
-	[DllImport("advapi32.dll", SetLastError = true)]
-	private static extern bool LogonUser(string UserName, string Domain, string Password, int LogonType, int LogonProvider, out SafeTokenHandle SafeToken);
-
-	public override void ExecuteBuild()
-	{
-		LogInformation("*********************** TestGamePerf");
-
-		string UserName = "test.automation";
-		string Password = "JJGfh9CX";
-		string Domain = "epicgames.net";
-		string FileName = string.Format("UStats_{0:MM-d-yyyy_hh-mm-ss-tt}.csv", DateTime.Now);
-		string Exe = CombinePaths(CmdEnv.LocalRoot, "Engine", "Binaries", "Win64", "UE4Editor.exe");
-
-		string ClientLogFile = "";
-		string CmdLine = "";
-		string UStatsPath = "";
-
-		string LevelParam = ParseParamValue("level", "");
-
-		switch (LevelParam)
-		{
-			case "PerfShooterGame_Santuary":
-				ClientLogFile = CombinePaths(CmdEnv.LogFolder, "ShooterGame.txt");
-				UStatsPath = CombinePaths(CmdEnv.LocalRoot, "ShooterGame", "Saved", "Profiling", "UnrealStats");
-				CmdLine = CombinePaths(CmdEnv.LocalRoot, "ShooterGame", "ShooterGame.uproject ") + CmdEnv.LocalRoot + @"/ShooterGame/Content/Maps/TestMaps/Sanctuary_PerfLoader.umap?game=ffa -game -novsync";
-				break;
-
-			case "PerfShooterGame_Highrise":
-				ClientLogFile = CombinePaths(CmdEnv.LogFolder, "ShooterGame.txt");
-				UStatsPath = CombinePaths(CmdEnv.LocalRoot, "ShooterGame", "Saved", "Profiling", "UnrealStats");
-				CmdLine = CombinePaths(CmdEnv.LocalRoot, "ShooterGame", "ShooterGame.uproject ") + CmdEnv.LocalRoot + @"/ShooterGame/Content/Maps/TestMaps/Highrise_PerfLoader.umap?game=ffa -game -novsync";
-				break;
-
-			case "PerfHoverShip":
-				ClientLogFile = CombinePaths(CmdEnv.LogFolder, "HoverShip.txt");
-				UStatsPath = CombinePaths(CmdEnv.LocalRoot, "Samples", "SampleGames", "HoverShip", "Saved", "Profiling", "UnrealStats");
-				CmdLine = CombinePaths(CmdEnv.LocalRoot, "Samples", "HoverShip", "HoverShip.uproject ") + CmdEnv.LocalRoot + @"/Samples/HoverShip/Content/Maps/HoverShip_01.umap -game -novsync";
-				break;
-
-			case "PerfInfiltrator":
-				ClientLogFile = CombinePaths(CmdEnv.LogFolder, "Infiltrator.txt");
-				UStatsPath = CombinePaths(CmdEnv.LocalRoot, "Samples", "Showcases", "Infiltrator", "Saved", "Profiling", "UnrealStats");
-				CmdLine = CombinePaths(CmdEnv.LocalRoot, "Samples", "Showcases", "Infiltrator", "Infiltrator.uproject ") + CmdEnv.LocalRoot + @"/Samples/Showcases/Infiltrator/Content/Maps/TestMaps/Visuals/VIS_ArtDemo_PerfLoader.umap -game -novsync";
-				break;
-
-			case "PerfElemental":
-				ClientLogFile = CombinePaths(CmdEnv.LogFolder, "Elemental.txt");
-				UStatsPath = CombinePaths(CmdEnv.LocalRoot, "Samples", "Showcases", "Elemental", "Saved", "Profiling", "UnrealStats");
-				CmdLine = CombinePaths(CmdEnv.LocalRoot, "Samples", "Showcases", "Elemental", "Elemental.uproject ") + CmdEnv.LocalRoot + @"/Samples/Showcases/Elemental/Content/Misc/Automated_Perf/Elemental_PerfLoader.umap -game -novsync";
-				break;
-
-			case "PerfStrategyGame":
-				ClientLogFile = CombinePaths(CmdEnv.LogFolder, "StrategyGame.txt");
-				UStatsPath = CombinePaths(CmdEnv.LocalRoot, "StrategyGame", "Saved", "Profiling", "UnrealStats");
-				CmdLine = CombinePaths(CmdEnv.LocalRoot, "StrategyGame", "StrategyGame.uproject ") + CmdEnv.LocalRoot + @"/StrategyGame/Content/Maps/TestMaps/TowerDefenseMap_PerfLoader.umap -game -novsync";
-				break;
-
-			case "PerfVehicleGame":
-				ClientLogFile = CombinePaths(CmdEnv.LogFolder, "VehicleGame.txt");
-				UStatsPath = CombinePaths(CmdEnv.LocalRoot, "VehicleGame", "Saved", "Profiling", "UnrealStats");
-				CmdLine = CombinePaths(CmdEnv.LocalRoot, "VehicleGame", "VehicleGame.uproject ") + CmdEnv.LocalRoot + @"/VehicleGame/Content/Maps/TestMaps/VehicleTrack_PerfLoader.umap -game -novsync";
-				break;
-
-			case "PerfPlatformerGame":
-				ClientLogFile = CombinePaths(CmdEnv.LogFolder, "PlatformerGame.txt");
-				UStatsPath = CombinePaths(CmdEnv.LocalRoot, "PlatformerGame", "Saved", "Profiling", "UnrealStats");
-				CmdLine = CombinePaths(CmdEnv.LocalRoot, "PlatformerGame", "PlatformerGame.uproject ") + CmdEnv.LocalRoot + @"/PlatformerGame/Content/Maps/Platformer_StreetSection.umap -game -novsync";
-				break;
-
-			case "PerfLanderGame":
-				ClientLogFile = CombinePaths(CmdEnv.LogFolder, "LanderGame.txt");
-				UStatsPath = CombinePaths(CmdEnv.LocalRoot, "LanderGame", "Saved", "Profiling", "UnrealStats");
-				CmdLine = CombinePaths(CmdEnv.LocalRoot, "LanderGame", "LanderGame.uproject ") + CmdEnv.LocalRoot + @"/LanderGame/Content/Maps/sphereworld.umap -game -novsync";
-				break;
-
-			case "PerfDemoLets_DynamicLighting":
-				ClientLogFile = CombinePaths(CmdEnv.LogFolder, "DemoLets_DynamicLighting.txt");
-				UStatsPath = CombinePaths(CmdEnv.LocalRoot, "Samples", "DemoLets", "Saved", "Profiling", "UnrealStats");
-				CmdLine = CombinePaths(CmdEnv.LocalRoot, "DemoLets", "DemoLets.uproject ") + CmdEnv.LocalRoot + @"/BasicMaterials/Content/Maps/Performance/Demolet_DynamicLighting_PerfLoader.umap -game -novsync";
-				break;
-
-			case "PerfDemoLets_BasicMaterials":
-				ClientLogFile = CombinePaths(CmdEnv.LogFolder, "DemoLets_BasicMaterials.txt");
-				UStatsPath = CombinePaths(CmdEnv.LocalRoot, "Samples", "DemoLets", "Saved", "Profiling", "UnrealStats");
-				CmdLine = CombinePaths(CmdEnv.LocalRoot, "DemoLets", "DemoLets.uproject ") + CmdEnv.LocalRoot + @"/EffectsGallery/Content/Maps/Performance/BasicMaterials_PerfLoader.umap -game -novsync";
-				break;
-
-			case "PerfDemoLets_DemoRoom":
-				ClientLogFile = CombinePaths(CmdEnv.LogFolder, "DemoLets_DemoRoom.txt");
-				UStatsPath = CombinePaths(CmdEnv.LocalRoot, "Samples", "DemoLets", "Saved", "Profiling", "UnrealStats");
-				CmdLine = CombinePaths(CmdEnv.LocalRoot, "DemoLets", "DemoLets.uproject ") + CmdEnv.LocalRoot + @"/EffectsGallery/Content/Maps/Performance/DemoRoom_Effects_hallway_PerfLoader.umap -game -novsync";
-				break;
-
-			case "PerfDemoLets_Reflections":
-				ClientLogFile = CombinePaths(CmdEnv.LogFolder, "DemoLets_Reflections.txt");
-				UStatsPath = CombinePaths(CmdEnv.LocalRoot, "Samples", "DemoLets", "Saved", "Profiling", "UnrealStats");
-				CmdLine = CombinePaths(CmdEnv.LocalRoot, "DemoLets", "DemoLets.uproject ") + CmdEnv.LocalRoot + @"/EffectsGallery/Content/Maps/Performance/DemoLet_Reflections_PerfLoader.umap -game -novsync";
-				break;
-
-			case "PerfDemoLets_PostProcessing":
-				ClientLogFile = CombinePaths(CmdEnv.LogFolder, "DemoLets_PostProcessing.txt");
-				UStatsPath = CombinePaths(CmdEnv.LocalRoot, "Samples", "DemoLets", "Saved", "Profiling", "UnrealStats");
-				CmdLine = CombinePaths(CmdEnv.LocalRoot, "DemoLets", "DemoLets.uproject ") + CmdEnv.LocalRoot + @"/EffectsGallery/Content/Maps/Performance/Demolet_PostProcessing_PerfLoader.umap -game -novsync";
-				break;
-
-			case "PerfDemoLets_Physics":
-				ClientLogFile = CombinePaths(CmdEnv.LogFolder, "DemoLets_Physics.txt");
-				UStatsPath = CombinePaths(CmdEnv.LocalRoot, "Samples", "DemoLets", "Saved", "Profiling", "UnrealStats");
-				CmdLine = CombinePaths(CmdEnv.LocalRoot, "DemoLets", "DemoLets.uproject ") + CmdEnv.LocalRoot + @"/EffectsGallery/Content/Maps/Performance/Demolet_Physics_PerfLoader.umap -game -novsync";
-				break;
-
-			case "PerfDemoLets_ShadowMaps":
-				ClientLogFile = CombinePaths(CmdEnv.LogFolder, "DemoLets_ShadowMaps.txt");
-				UStatsPath = CombinePaths(CmdEnv.LocalRoot, "Samples", "DemoLets", "Saved", "Profiling", "UnrealStats");
-				CmdLine = CombinePaths(CmdEnv.LocalRoot, "DemoLets", "DemoLets.uproject ") + CmdEnv.LocalRoot + @"/EffectsGallery/Content/Maps/Performance/Demolet_CascadingShadowMaps_PerfLoader.umap -game -novsync";
-				break;
-
-			case "PerfDemoLets_Weather":
-				ClientLogFile = CombinePaths(CmdEnv.LogFolder, "DemoLets_Weather.txt");
-				UStatsPath = CombinePaths(CmdEnv.LocalRoot, "Samples", "DemoLets", "Saved", "Profiling", "UnrealStats");
-				CmdLine = CombinePaths(CmdEnv.LocalRoot, "DemoLets", "DemoLets.uproject ") + CmdEnv.LocalRoot + @"/EffectsGallery/Content/Maps/Performance/dynamic_weather_selector_PerfLoader.umap -game -novsync";
-				break;
-
-			case "PerfRealisticRendering_RoomNight":
-				ClientLogFile = CombinePaths(CmdEnv.LogFolder, "RealisticRendering_RoomNight.txt");
-				UStatsPath = CombinePaths(CmdEnv.LocalRoot, "Samples", "ShowCases", "RealisticRendering", "Saved", "Profiling", "UnrealStats");
-				CmdLine = CombinePaths(CmdEnv.LocalRoot, "Samples", "Showcases", "RealisticRendering", "RealisticRendering.uproject ") + CmdEnv.LocalRoot + @"/RealisticRendering/Content/Maps/Performance/RoomNight_PerfLoader.umap -game -novsync";
-				break;
-
-			case "PerfRealisticRendering_NoLight":
-				ClientLogFile = CombinePaths(CmdEnv.LogFolder, "RealisticRendering_NoLight.txt");
-				UStatsPath = CombinePaths(CmdEnv.LocalRoot, "Samples", "ShowCases", "RealisticRendering", "Saved", "Profiling", "UnrealStats");
-				CmdLine = CombinePaths(CmdEnv.LocalRoot, "Samples", "Showcases", "RealisticRendering", "RealisticRendering.uproject ") + CmdEnv.LocalRoot + @"/RealisticRendering/Content/Maps/Performance/RoomNightNoLights_PerfLoader.umap -game -novsync";
-				break;
-
-			case "PerfRealisticRendering_Room":
-				ClientLogFile = CombinePaths(CmdEnv.LogFolder, "RealisticRendering_Room.txt");
-				UStatsPath = CombinePaths(CmdEnv.LocalRoot, "Samples", "ShowCases", "RealisticRendering", "Saved", "Profiling", "UnrealStats");
-				CmdLine = CombinePaths(CmdEnv.LocalRoot, "Samples", "Showcases", "RealisticRendering", "RealisticRendering.uproject ") + CmdEnv.LocalRoot + @"/RealisticRendering/Content/Maps/Performance/Room_PerfLoader.umap -game -novsync";
-				break;
-
-			case "PerfMorphTargets":
-				ClientLogFile = CombinePaths(CmdEnv.LogFolder, "MorphTargets.txt");
-				UStatsPath = CombinePaths(CmdEnv.LocalRoot, "Samples", "ShowCases", "MorphTargets", "Saved", "Profiling", "UnrealStats");
-				CmdLine = CombinePaths(CmdEnv.LocalRoot, "Samples", "Showcases", "MorphTargets", "MorphTargets.uproject ") + CmdEnv.LocalRoot + @"/MorphTargets/Content/Maps/Performance/MorphTargets_PerfLoader.umap -game -novsync";
-				break;
-
-			case "PerfMatinee":
-				ClientLogFile = CombinePaths(CmdEnv.LogFolder, "Matinee.txt");
-				UStatsPath = CombinePaths(CmdEnv.LocalRoot, "Samples", "ShowCases", "PostProcessMatinee", "Saved", "Profiling", "UnrealStats");
-				CmdLine = CombinePaths(CmdEnv.LocalRoot, "Samples", "Showcases", "PostProcessMatinee", "PostProcessMatinee.uproject ") + CmdEnv.LocalRoot + @"/PostProcessMatinee/Content/Maps/Performance/PostProcessMatinee_PerfLoader.umap -game -novsync";
-				break;
-
-			case "PerfLandScape":
-				ClientLogFile = CombinePaths(CmdEnv.LogFolder, "LandScape.txt");
-				UStatsPath = CombinePaths(CmdEnv.LocalRoot, "Samples", "ShowCases", "Landscape_WorldMachine", "Saved", "Profiling", "UnrealStats");
-				CmdLine = CombinePaths(CmdEnv.LocalRoot, "Samples", "Showcases", "Landscape_WorldMachine", "Landscape_WorldMap.uproject ") + CmdEnv.LocalRoot + @"/Landscape_WorldMachine/Content/Maps/Performance/LandscapeMap_PerfLoader.umap -game -novsync";
-				break;
-
-			case "PerfLandScape_Moutain":
-				ClientLogFile = CombinePaths(CmdEnv.LogFolder, "LandScape_Moutain.txt");
-				UStatsPath = CombinePaths(CmdEnv.LocalRoot, "Samples", "ShowCases", "Landscape", "Saved", "Profiling", "UnrealStats");
-				CmdLine = CombinePaths(CmdEnv.LocalRoot, "Samples", "Showcases", "Landscape", "Landscape.uproject ") + CmdEnv.LocalRoot + @"/Landscape/Content/Maps/Performance/MoutainRange_PerfLoader.umap -game -novsync";
-				break;
-
-			case "PerfMobile":
-				ClientLogFile = CombinePaths(CmdEnv.LogFolder, "Mobile.txt");
-				UStatsPath = CombinePaths(CmdEnv.LocalRoot, "Samples", "ShowCases", "Mobile", "Saved", "Profiling", "UnrealStats");
-				CmdLine = CombinePaths(CmdEnv.LocalRoot, "Samples", "Showcases", "Mobile", "Mobile.uproject ") + CmdEnv.LocalRoot + @"/Mobile/Content/Maps/Performance/Testmap_PerfLoader.umap -game -novsync";
-				break;
-
-			case "PerfSampleGames":
-				ClientLogFile = CmdEnv.LogFolder;
-				break;
-
-			default:
-				break;
-		}
-
-		try
-		{
-			SafeTokenHandle SafeToken;
-
-			int LogonInteractive = 2;
-			int ProviderDefualt = 0;
-
-			/*bool bLogonSuccess = */
-			LogonUser(UserName, Domain, Password, LogonInteractive, ProviderDefualt, out SafeToken);
-
-			using (SafeToken)
-			{
-				using (WindowsIdentity NewUser = new WindowsIdentity(SafeToken.DangerousGetHandle()))
-				{
-					using (WindowsImpersonationContext TestAccount = NewUser.Impersonate())
-					{
-						if (LevelParam == "PerfSampleGames")
-						{
-
-						}
-						else
-						{
-							if (!File.Exists(ClientLogFile))
-							{
-								LogInformation("Creating log file");
-
-								File.Create(ClientLogFile).Close();
-
-								LogInformation(ClientLogFile);
-
-								LogInformation("Log file created");
-							}
-
-							RunAndLog(Exe, CmdLine, ClientLogFile);
-						}
-
-						DirectoryInfo[] UStatsDirectories = new DirectoryInfo(UStatsPath).GetDirectories();
-						DirectoryInfo CurrentUStatsDirectory = UStatsDirectories[0];
-
-						for (int i = 0; i < UStatsDirectories.Length; i++)
-						{
-							if (UStatsDirectories[i].LastWriteTime > CurrentUStatsDirectory.LastWriteTime)
-							{
-								CurrentUStatsDirectory = UStatsDirectories[i];
-							}
-						}
-
-						FileInfo[] UStatsFilePaths = new DirectoryInfo(CurrentUStatsDirectory.FullName).GetFiles();
-						FileInfo CurrentUStatsFilePath = UStatsFilePaths[0];
-
-						for (int i = 0; i < UStatsFilePaths.Length; i++)
-						{
-							if (UStatsFilePaths[i].LastWriteTime > CurrentUStatsFilePath.LastWriteTime)
-							{
-								CurrentUStatsFilePath = UStatsFilePaths[i];
-							}
-						}
-						try
-						{
-							string FrontEndExe = CombinePaths(CmdEnv.LocalRoot, "Engine", "Binaries", "Win64", "UnrealFrontend.exe");
-
-							CmdLine = " -run=convert -infile=" + CurrentUStatsFilePath.FullName + @" -outfile=D:\" + FileName + " -statlist=GameThread+RenderThread+StatsThread+STAT_FrameTime+STAT_PhysicalAllocSize+STAT_VirtualAllocSize";
-
-							RunAndLog(FrontEndExe, CmdLine, ClientLogFile);
-						}
-
-						catch (Exception Err)
-						{
-							LogInformation(Err.Message);
-						}
-					}
-				}
-			}
-		}
-
-		catch (Exception Err)
-		{
-			LogInformation(Err.Message);
-		}
-	}
-}
-
 public sealed class SafeTokenHandle : SafeHandleZeroOrMinusOneIsInvalid
 {
 	private SafeTokenHandle()
