- Updated Turnkey, including missing file

diff --git a/Engine/Source/Programs/AutomationTool/Turnkey/Commands/InstallSdk.cs b/Engine/Source/Programs/AutomationTool/Turnkey/Commands/InstallSdk.cs
--- a/Engine/Source/Programs/AutomationTool/Turnkey/Commands/InstallSdk.cs
+++ b/Engine/Source/Programs/AutomationTool/Turnkey/Commands/InstallSdk.cs
@@ -24,7 +24,6 @@ namespace Turnkey.Commands
 
 		protected override void Execute(string[] CommandOptions)
 		{
-			string PlatformString = TurnkeyUtils.ParseParamValue("Platform", null, CommandOptions);
 			string DeviceName = TurnkeyUtils.ParseParamValue("Device", null, CommandOptions);
 
 			bool bBestAvailable = TurnkeyUtils.ParseParam("BestAvailable", CommandOptions);
@@ -50,39 +49,10 @@ namespace Turnkey.Commands
 				}
 			}
 
-			List<UnrealTargetPlatform> PlatformsLeftToInstall = new List<UnrealTargetPlatform>();
-			// prompt user for a platform
-			if (PlatformString == null)
+			// get the platforms to install either from user or from commandline
+			List<UnrealTargetPlatform> PlatformsLeftToInstall = TurnkeyUtils.GetPlatformsFromCommandLineOrUser(CommandOptions, PlatformsWithSdks);
+			if (PlatformsLeftToInstall == null)
 			{
-				List<string> PlatformOptions = PlatformsWithSdks.ConvertAll(x => x.ToString());
-				PlatformOptions.Add("All Possible");
-				int PlatformChoice = TurnkeyUtils.ReadInputInt("Choose a platform:", PlatformOptions, true);
-
-				if (PlatformChoice == 0)
-				{
-					return;
-				}
-				// All platforms means to install every platform with an installer
-				if (PlatformChoice == PlatformOptions.Count - 1)
-				{
-					PlatformsLeftToInstall = PlatformsWithSdks;
-				}
-				else
-				{
-					PlatformsLeftToInstall.Add(PlatformsWithSdks[PlatformChoice - 1]);
-				}
-			}
-			else if (PlatformString == "All")
-			{
-				PlatformsLeftToInstall = PlatformsWithSdks;
-			}
-			else if (UnrealTargetPlatform.IsValidName(PlatformString))
-			{
-				PlatformsLeftToInstall.Add(UnrealTargetPlatform.Parse(PlatformString));
-			}
-			else
-			{
-				TurnkeyUtils.Log("Unknown platform {0}", PlatformString);
 				return;
 			}
 
@@ -145,7 +115,7 @@ namespace Turnkey.Commands
 									BestMiscSdk = Sdk;
 								}
 							}
-							else if (Sdk.Type == SdkInfo.SdkType.Full)
+							else if (Sdk.Type == SdkInfo.SdkType.Full || Sdk.Type == SdkInfo.SdkType.RunOnly || Sdk.Type == SdkInfo.SdkType.BuildOnly)
 							{
 								if (BestNormalSdk == null || string.Compare(Sdk.Version, BestNormalSdk.Version, true) > 0)
 								{
@@ -268,6 +238,7 @@ namespace Turnkey.Commands
 
 					// standard variables
 					TurnkeyUtils.SetVariable("Platform", Platform.ToString());
+					TurnkeyUtils.SetVariable("Version", Sdk.Version);
 
 					Sdk.Install(Platform);
 				}
diff --git a/Engine/Source/Programs/AutomationTool/Turnkey/Commands/VerifySdk.cs b/Engine/Source/Programs/AutomationTool/Turnkey/Commands/VerifySdk.cs
new file mode 100644
--- /dev/null
+++ b/Engine/Source/Programs/AutomationTool/Turnkey/Commands/VerifySdk.cs
@@ -0,0 +1,44 @@
+using System;
+using System.Collections.Generic;
+using System.Linq;
+using UnrealBuildTool;
+
+namespace Turnkey.Commands
+{
+	class VerifySdk : TurnkeyCommand
+	{
+		protected override void Execute(string[] CommandOptions)
+		{
+			string DeviceName = TurnkeyUtils.ParseParamValue("Device", null, CommandOptions);
+
+			// 		bool bBestAvailable = TurnkeyUtils.ParseParam("BestAvailable", CommandOptions);
+			// 		bool bUpdateOnly = TurnkeyUtils.ParseParam("UpdateOnly", CommandOptions);
+
+			List<UnrealTargetPlatform> ChosenPlatforms = TurnkeyUtils.GetPlatformsFromCommandLineOrUser(CommandOptions, UnrealTargetPlatform.GetValidPlatforms().ToList());
+
+			if (ChosenPlatforms == null)
+			{
+				return;
+			}
+
+			TurnkeyUtils.Log("Installed Sdk validity:");
+			foreach (UnrealTargetPlatform Platform in ChosenPlatforms)
+			{
+				// get the platform object
+				AutomationTool.Platform PlatformObject = AutomationTool.Platform.Platforms[new AutomationTool.TargetPlatformDescriptor(Platform)];
+
+				// check its installed Sdk is valid
+				if (!TurnkeyUtils.IsValueValid(PlatformObject.GetInstalledSdk(), PlatformObject.GetAllowedSdks()))
+				{
+					TurnkeyUtils.Log("{0}: INVALID [Installed: '{1}', Required: '{2}']", Platform, PlatformObject.GetInstalledSdk(), PlatformObject.GetAllowedSdks());
+					Environment.ExitCode = (int)AutomationTool.ExitCode.Error_SDKNotFound;
+				}
+				else
+				{
+					TurnkeyUtils.Log("{0}: Valid!", Platform);
+//					TurnkeyUtils.Log("{0}: Valid [Installed: '{1}', Required: '{2}']", Platform, PlatformObject.GetInstalledSdk(), PlatformObject.GetAllowedSdks());
+				}
+			}
+		}
+	}
+}
diff --git a/Engine/Source/Programs/AutomationTool/Turnkey/CopyAndRun.cs b/Engine/Source/Programs/AutomationTool/Turnkey/CopyAndRun.cs
--- a/Engine/Source/Programs/AutomationTool/Turnkey/CopyAndRun.cs
+++ b/Engine/Source/Programs/AutomationTool/Turnkey/CopyAndRun.cs
@@ -92,7 +92,7 @@ namespace Turnkey
 				}
 			}
 
-			if ((CommandPath != null && CommandPath.Contains("$(CopyOutputPath)")) || (CommandLine != null && CommandLine.Contains("$(CopyOutputPath)")))
+			if (OutputPath == null && ((CommandPath != null && CommandPath.Contains("$(CopyOutputPath)")) || (CommandLine != null && CommandLine.Contains("$(CopyOutputPath)"))))
 			{
 				throw new AutomationException("CopyAndRun required an valid $(CopyOutputPath) from the Copy operation, but it failed");
 			}
diff --git a/Engine/Source/Programs/AutomationTool/Turnkey/Turnkey.Automation.cs b/Engine/Source/Programs/AutomationTool/Turnkey/Turnkey.Automation.cs
--- a/Engine/Source/Programs/AutomationTool/Turnkey/Turnkey.Automation.cs
+++ b/Engine/Source/Programs/AutomationTool/Turnkey/Turnkey.Automation.cs
@@ -56,6 +56,13 @@ namespace Turnkey
 			{
 				TurnkeyUtils.CleanupPaths();
 			}
+
+			if (Environment.ExitCode != 0)
+			{
+				// @todo turnkey - would be nice to return a failure, without an exception in UAT, which looks violent
+				AutomationTool.ExitCode ExitCode = (AutomationTool.ExitCode)Environment.ExitCode;
+				throw new AutomationException(ExitCode, ExitCode.ToString());
+			}
 		}
 	}
 }
diff --git a/Engine/Source/Programs/AutomationTool/Turnkey/TurnkeyCommand.cs b/Engine/Source/Programs/AutomationTool/Turnkey/TurnkeyCommand.cs
--- a/Engine/Source/Programs/AutomationTool/Turnkey/TurnkeyCommand.cs
+++ b/Engine/Source/Programs/AutomationTool/Turnkey/TurnkeyCommand.cs
@@ -96,19 +96,17 @@ namespace Turnkey
 
 			try
 			{
+				TurnkeyUtils.Log("");
+
 				// run the command!
 				Command.Execute(CommandOptions);
+
+				TurnkeyUtils.Log("");
 			}
 			catch (System.Exception)
 			{
 				throw;
 			}
-			finally
-			{
-// 				// cleanup any paths we need to, now that the command is complete
-// 				TurnkeyUtils.Log("Cleaning up paths....");
-// 				TurnkeyUtils.CleanupPaths();
-			}
 
 			// true if we should keep going
 			return true;
diff --git a/Engine/Source/Programs/AutomationTool/Turnkey/TurnkeyManifest.cs b/Engine/Source/Programs/AutomationTool/Turnkey/TurnkeyManifest.cs
--- a/Engine/Source/Programs/AutomationTool/Turnkey/TurnkeyManifest.cs
+++ b/Engine/Source/Programs/AutomationTool/Turnkey/TurnkeyManifest.cs
@@ -108,7 +108,7 @@ namespace Turnkey
 						continue;
 					}
 
-					string ThisManifestDir = Path.GetDirectoryName(Source);
+					string ThisManifestDir = Path.GetDirectoryName(Source).Replace("\\", "/");
 
 					// if a directory is returned, look for the standardized manifest name, as we have not much else we can do with a directory
 					if (LocalManifestPath.EndsWith("/") || LocalManifestPath.EndsWith("\\"))
diff --git a/Engine/Source/Programs/AutomationTool/Turnkey/TurnkeyUtils.cs b/Engine/Source/Programs/AutomationTool/Turnkey/TurnkeyUtils.cs
--- a/Engine/Source/Programs/AutomationTool/Turnkey/TurnkeyUtils.cs
+++ b/Engine/Source/Programs/AutomationTool/Turnkey/TurnkeyUtils.cs
@@ -7,6 +7,7 @@ using System.Text.RegularExpressions;
 using System.Globalization;
 using Tools.DotNETCommon;
 using AutomationTool;
+using UnrealBuildTool;
 
 namespace Turnkey
 {
@@ -26,14 +27,14 @@ namespace Turnkey
 
 		public static string SetVariable(string Key, string Value)
 		{
-			string Previous = null;
+			string Previous;
 			TurnkeyVariables.TryGetValue(Key, out Previous);
 			TurnkeyVariables[Key] = Value;
 			return Previous;
 		}
 		public static string GetVariableValue(string Key)
 		{
-			string Value = null;
+			string Value;
 			TurnkeyVariables.TryGetValue(Key, out Value);
 			return Value;
 		}
@@ -79,6 +80,59 @@ namespace Turnkey
 			return Value == null ? Default : Value;
 		}
 
+		public static List<UnrealTargetPlatform> GetPlatformsFromCommandLineOrUser(string[] CommandOptions, List<UnrealTargetPlatform> PossiblePlatforms)
+		{
+			string PlatformString = TurnkeyUtils.ParseParamValue("Platform", null, CommandOptions);
+
+			List<UnrealTargetPlatform> Platforms = new List<UnrealTargetPlatform>();
+			// prompt user for a platform
+			if (PlatformString == null)
+			{
+				List<string> PlatformOptions = PossiblePlatforms.ConvertAll(x => x.ToString());
+				PlatformOptions.Add("All of the Above");
+				int PlatformChoice = TurnkeyUtils.ReadInputInt("Choose a platform:", PlatformOptions, true);
+
+				if (PlatformChoice == 0)
+				{
+					return null;
+				}
+				// All platforms means to install every platform with an installer
+				if (PlatformChoice == PlatformOptions.Count)
+				{
+					Platforms = PossiblePlatforms;
+				}
+				else
+				{
+					Platforms.Add(PossiblePlatforms[PlatformChoice - 1]);
+				}
+			}
+			else if (PlatformString == "All")
+			{
+				Platforms = PossiblePlatforms;
+			}
+			else
+			{
+				string[] Tokens = PlatformString.Split("+".ToCharArray(), StringSplitOptions.RemoveEmptyEntries);
+
+				foreach (string Token in Tokens)
+				{
+					UnrealTargetPlatform Platform;
+					if (!UnrealTargetPlatform.TryParse(Token, out Platform))
+					{
+						TurnkeyUtils.Log("Platform {0} is unknown", Token);
+						return null;
+					}
+
+					// if the platform isn't in the possible list, then don't add it
+					if (PossiblePlatforms.Contains(Platform))
+					{
+						Platforms.Add(Platform);
+					}
+				}
+			}
+
+			return Platforms;
+		}
 
 
 		private static IDictionary[] SavedEnvVars = new System.Collections.IDictionary[2];
