- Cleaned up manifest locations, finally checking in some manifests
  - Epic's manifest uses google drive, with the "secrets" file pointed to via a variable in the TurnkeyManifest
- Improved caching saving/loading
- Have switchig BuildOnly Sdks working well

diff --git a/Engine/Build/Turnkey/TurnkeyManifest.xml b/Engine/Build/Turnkey/TurnkeyManifest.xml
new file mode 100644
--- /dev/null
+++ b/Engine/Build/Turnkey/TurnkeyManifest.xml
@@ -0,0 +1,9 @@
+<?xml version="1.0" encoding="utf-8" ?>
+<TurnkeyManifest>
+	<AdditionalManifests>
+		<Manifest>$(ThisManifestDir)/StudioTurnkeyManifest.xml</Manifest>
+		<Manifest>file:$(EngineDir)/Restricted/NoRedist/Build/Turnkey/TurnkeyManifest.xml</Manifest>
+		<Manifest>file:$(UE_SDKS_ROOT)/TurnkeyManifest.xml</Manifest>
+		<Manifest>$(UE_STUDIO_TURNKEY_LOCATION)</Manifest>
+	</AdditionalManifests>
+</TurnkeyManifest>
diff --git a/Engine/Source/Programs/AutomationTool/Turnkey/Commands/InstallSdk.cs b/Engine/Source/Programs/AutomationTool/Turnkey/Commands/InstallSdk.cs
--- a/Engine/Source/Programs/AutomationTool/Turnkey/Commands/InstallSdk.cs
+++ b/Engine/Source/Programs/AutomationTool/Turnkey/Commands/InstallSdk.cs
@@ -17,7 +17,6 @@ namespace Turnkey.Commands
 			{
 				{ "Auto Install All Needed Sdks", new string[] { "-platform=All", "-NeededOnly", "-BestAvailable" } },
 				{ "Auto Update All Sdks", new string[] { "-platform=All", "-UpdateOnly", "-BestAvailable" } },
-				{ "Download QuickSwitch Sdk", new string[] { "-AllowInvalid", "-DownloadOnly" } },
 			};
 		}
 
@@ -25,6 +24,7 @@ namespace Turnkey.Commands
 		protected override void Execute(string[] CommandOptions)
 		{
 			string DeviceName = TurnkeyUtils.ParseParamValue("Device", null, CommandOptions);
+			string SdkTypeString = TurnkeyUtils.ParseParamValue("SdkType", null, CommandOptions);
 
 			bool bBestAvailable = TurnkeyUtils.ParseParam("BestAvailable", CommandOptions);
 			bool bUpdateOnly = TurnkeyUtils.ParseParam("UpdateOnly", CommandOptions);
@@ -33,6 +33,17 @@ namespace Turnkey.Commands
 			bool bValidOnly = bBestAvailable || !TurnkeyUtils.ParseParam("AllowInvalid", CommandOptions);
 			bool bNeededOnly = bBestAvailable || TurnkeyUtils.ParseParam("NeededOnly", CommandOptions);
 
+
+			SdkInfo.SdkType DesiredType = SdkInfo.SdkType.Misc;
+			if (SdkTypeString != null)
+			{
+				if (!Enum.TryParse(SdkTypeString, out DesiredType))
+				{
+					TurnkeyUtils.Log("Invalid SdkType given with -SdkType={0}", SdkTypeString);
+					return;
+				}
+			}
+
 			// we need all sdks we can find
 			List<SdkInfo> AllSdks = TurnkeyManifest.GetDiscoveredSdks();
 
@@ -95,38 +106,37 @@ namespace Turnkey.Commands
 					Sdks = Sdks.FindAll(x => x.IsNeeded(Platform, DeviceName));
 				}
 
-				if (bBestAvailable)
+				// filter on type
+				if (SdkTypeString != null && Sdks.Count > 0)
 				{
-					// loop through Sdks that are left, and install the best one available
+					List<SdkInfo> TypedSdks = Sdks.FindAll(x => x.Type == DesiredType);
+					if (TypedSdks.Count == 0)
+					{
+						TurnkeyUtils.Log("No valid Sdks found of type {0}, using Full", DesiredType);
+						TypedSdks = Sdks.FindAll(x => x.Type == SdkInfo.SdkType.Full);
+					}
+					Sdks = TypedSdks;
+				}
 
-					SdkInfo BestNormalSdk = null;
-					SdkInfo BestMiscSdk = null;
-					SdkInfo BestFlash = null;
+				if (bBestAvailable && Sdks.Count > 0)
+				{
+					// loop through Sdks that are left, and install the best one available
+					Dictionary<SdkInfo.SdkType, SdkInfo> BestByType = new Dictionary<SdkInfo.SdkType, SdkInfo>();
 
 					foreach (SdkInfo Sdk in Sdks)
 					{
 						// always install custom Sdks since there are no versions to check, if it's valid, let it decide what to do
 						if (Sdk.CustomSdkId == null)
 						{
-							if (Sdk.Type == SdkInfo.SdkType.Misc)
+							if (!BestByType.ContainsKey(Sdk.Type))
 							{
-								if (BestMiscSdk == null || string.Compare(Sdk.Version, BestMiscSdk.Version, true) > 0)
-								{
-									BestMiscSdk = Sdk;
-								}
-							}
-							else if (Sdk.Type == SdkInfo.SdkType.Full || Sdk.Type == SdkInfo.SdkType.RunOnly || Sdk.Type == SdkInfo.SdkType.BuildOnly)
-							{
-								if (BestNormalSdk == null || string.Compare(Sdk.Version, BestNormalSdk.Version, true) > 0)
-								{
-									BestNormalSdk = Sdk;
-								}
+								BestByType[Sdk.Type] = Sdk;
 							}
-							else if (Sdk.Type == SdkInfo.SdkType.Flash)
+							else
 							{
-								if (BestFlash == null || string.Compare(Sdk.Version, BestFlash.Version, true) > 0)
+								if (string.Compare(Sdk.Version, BestByType[Sdk.Type].Version, true) > 0)
 								{
-									BestFlash = Sdk;
+									BestByType[Sdk.Type] = Sdk;
 								}
 							}
 						}
@@ -134,18 +144,29 @@ namespace Turnkey.Commands
 
 					// first, keep only custom Sdks
 					Sdks = Sdks.FindAll(x => x.CustomSdkId != null);
-					// then add the best matches above
-					if (BestNormalSdk != null)
-					{
-						Sdks.Add(BestNormalSdk);
-					}
-					if (BestMiscSdk != null)
-					{
-						Sdks.Add(BestNormalSdk);
-					}
-					if (BestFlash != null)
+
+// 					// if the user chose a type, then pick that one
+// 					if (SdkTypeString != null)
+// 					{
+// 						if ((DesiredType == SdkInfo.SdkType.BuildOnly || DesiredType == SdkInfo.SdkType.RunOnly) &&
+// 							!BestByType.ContainsKey(DesiredType))
+// 						{
+// 							TurnkeyUtils.Log("No valid Sdks found of type {0}, trying for a Full",  DesiredType);
+// 							DesiredType = SdkInfo.SdkType.Flash;
+// 						}
+// 
+// 						if (!BestByType.ContainsKey(DesiredType))
+// 						{
+// 							TurnkeyUtils.Log("No valid Sdks found of type {0}, giving up.", DesiredType);
+// 							return;
+// 						}
+// 
+// 						SdkInfo Sdk = BestByType[DesiredType];
+// 					}
+// 					// otherwise, add all matches to the list
+// 					else
 					{
-						Sdks.Add(BestNormalSdk);
+						Sdks.AddRange(BestByType.Values);
 					}
 				}
 				// if we are not doing best available Sdks, then let used choose one
diff --git a/Engine/Source/Programs/AutomationTool/Turnkey/CopyAndRun.cs b/Engine/Source/Programs/AutomationTool/Turnkey/CopyAndRun.cs
--- a/Engine/Source/Programs/AutomationTool/Turnkey/CopyAndRun.cs
+++ b/Engine/Source/Programs/AutomationTool/Turnkey/CopyAndRun.cs
@@ -5,8 +5,7 @@ using UnrealBuildTool;
 using AutomationTool;
 using Tools.DotNETCommon;
 using System.Diagnostics;
-using System.Security.Permissions;
-using System.Security.Principal;
+using System.ComponentModel;
 
 namespace Turnkey
 {
@@ -62,6 +61,7 @@ namespace Turnkey
 			CommandLine = TurnkeyUtils.ExpandVariables(CommandLine, bUseOnlyTurnkeyVariables: true);
 		}
 
+
 		public bool Execute(CopyExecuteSpecialMode SpecialMode=CopyExecuteSpecialMode.None, string ModeHint=null)
 		{
 			if (ShouldExecute() == false)
@@ -105,25 +105,68 @@ namespace Turnkey
 				FixedCommandPath = FixedCommandPath.Replace(Path.AltDirectorySeparatorChar, Path.DirectorySeparatorChar);
 
 				string PreviousCWD = Environment.CurrentDirectory;
+				// if a directory was included in the command path, then run from there
+				string CommandWorkingDir = Path.GetDirectoryName(FixedCommandPath);
+
+				// if we run UseShellExecute with true, StartInfo.WorkingDirectory won't set the directory as expected
+				if (!string.IsNullOrEmpty(CommandWorkingDir))
+				{
+					Environment.CurrentDirectory = CommandWorkingDir;
+				}
+
 
 				TurnkeyUtils.StartTrackingExternalEnvVarChanges();
 
+				// run installer as administrator, as some need it
+				Process InstallProcess = new Process();
+				InstallProcess.StartInfo.UseShellExecute = false;
+				InstallProcess.StartInfo.FileName = FixedCommandPath;
+				InstallProcess.StartInfo.Arguments = FixedCommandLine;
+				InstallProcess.StartInfo.WindowStyle = ProcessWindowStyle.Hidden;
+
+  				InstallProcess.OutputDataReceived += (Sender, Args) => { if (Args != null && Args.Data != null) TurnkeyUtils.Log(Args.Data.TrimEnd()); };
+  				InstallProcess.ErrorDataReceived += (Sender, Args) => { if (Args != null && Args.Data != null) TurnkeyUtils.Log("Error: {0}", Args.Data.TrimEnd()); };
+
+				//installers may require administrator access to succeed. so run as an admmin.
+
 				// run in a loop in case of a failure
 				bool bDone = false;
+				int ExitCode;
+
 				while (!bDone)
 				{
-					// if a directory was included in the command path, then run from there
-					string CommandWorkingDir = Path.GetDirectoryName(FixedCommandPath);
-					if (!string.IsNullOrEmpty(CommandWorkingDir))
+					try
 					{
-						Environment.CurrentDirectory = CommandWorkingDir;
+						InstallProcess.Start();
+						InstallProcess.WaitForExit();
+						ExitCode = InstallProcess.ExitCode;
+					}
+					catch (Exception Ex)
+					{
+						// native error in a Win32Exception, of 740, means the process needs elevation, so we need to runas. However,
+						// this will not allow capturing stdout, so run with window as Normal
+						if (InstallProcess.StartInfo.UseShellExecute == false && Ex is Win32Exception && ((Win32Exception)Ex).NativeErrorCode == 740)
+						{
+							InstallProcess.StartInfo.UseShellExecute = true;
+							InstallProcess.StartInfo.Verb = "runas";
+							InstallProcess.StartInfo.WindowStyle = ProcessWindowStyle.Normal;
+
+							TurnkeyUtils.Log("The installer {0} needed to run with elevated permissions, trying with Admin privileges (output may be hidden)", FixedCommandPath);
+
+							// try again
+							continue;
+						}
+
+						TurnkeyUtils.Log("Error: {0} caused an exception: {1}", FixedCommandPath, Ex.Message);
+						ExitCode = -1;
 					}
 
-					ProcessStartInfo ProcInfo = new ProcessStartInfo(FixedCommandPath, FixedCommandLine);
-					int ExitCode = Utils.RunLocalProcessAndLogOutput(ProcInfo);
 					if (ExitCode != 0)
 					{
-						TurnkeyUtils.Log("Command {0} {1} failed [Exit code {2}]", FixedCommandPath, FixedCommandLine, ExitCode);
+						TurnkeyUtils.Log("");
+						TurnkeyUtils.Log("Command {0} {1} failed [Exit code {2}, working dir = {3}]", FixedCommandPath, FixedCommandLine, ExitCode, CommandWorkingDir);
+						TurnkeyUtils.Log("");
+
 						string Response = TurnkeyUtils.ReadInput("Do you want to attempt again? [y/N]", "N");
 						if (string.Compare(Response, "Y", true) != 0)
 						{
@@ -133,7 +176,7 @@ namespace Turnkey
 					else
 					{
 						bDone = true;
-					}
+					} 
 				}
 
 				TurnkeyUtils.EndTrackingExternalEnvVarChanges();
diff --git a/Engine/Source/Programs/AutomationTool/Turnkey/CopyProvider.cs b/Engine/Source/Programs/AutomationTool/Turnkey/CopyProvider.cs
--- a/Engine/Source/Programs/AutomationTool/Turnkey/CopyProvider.cs
+++ b/Engine/Source/Programs/AutomationTool/Turnkey/CopyProvider.cs
@@ -58,13 +58,20 @@ namespace Turnkey
 			}
 		}
 
-		private static void ParseOperation(string Operation, out CopyProvider Provider, out string ProviderParam)
+		private static bool ParseOperation(string Operation, out CopyProvider Provider, out string ProviderParam, bool bCanFail)
 		{
 			Operation = TurnkeyUtils.ExpandVariables(Operation);
 
+			Provider = null;
+			ProviderParam = null;
+
 			int ColonLocation = Operation.IndexOf(':');
 			if (ColonLocation < 0)
 			{
+				if (bCanFail)
+				{
+					return false;
+				}
 				throw new AutomationException("Malformed copy operation: {0}", Operation);
 			}
 			// get the token before the :
@@ -72,10 +79,15 @@ namespace Turnkey
 
 			if (!CachedProviders.TryGetValue(Token, out Provider))
 			{
+				if (bCanFail)
+				{
+					return false;
+				}
 				throw new AutomationException("Unable to find a CopyProvider for copy type {0}", Token);
 			}
 
 			ProviderParam = Operation.Substring(ColonLocation + 1);
+			return true;
 		}
 
 		/// <summary>
@@ -87,7 +99,8 @@ namespace Turnkey
 		{
 			CopyProvider Provider;
 			string ProviderParam;
-			ParseOperation(CopyOperation, out Provider, out ProviderParam);
+
+			ParseOperation(CopyOperation, out Provider, out ProviderParam, false);
 
 			// execute what comes after the colon
 			string OutputPath = Provider.Execute(ProviderParam, SpecialMode, SpecialModeHint);
@@ -95,14 +108,16 @@ namespace Turnkey
 			// always unzip the file into a temp directory (or downloadpath if it's a permanent download), and make that directory be the outputpath
 			if (OutputPath != null)
 			{
+				OutputPath = OutputPath.Replace(Path.AltDirectorySeparatorChar, Path.DirectorySeparatorChar);
 				string Ext = Path.GetExtension(OutputPath).ToLower();
 				if (Ext == ".zip" || Ext == ".7z")
 				{
 					// check if this file already has been unzipped - if the provider output to a new temp location but the same file, 
-					// we can't trust it, so we re-decompress
+					// we can't trust it, so we re-decompress, or if the source file's date changed
+					string FileVersion = File.GetLastWriteTimeUtc(OutputPath).ToString();
 					string Tag = OutputPath + "_Decompressed";
 
-					string CachedLocation = LocalCache.GetCachedPathByTag(Tag);
+					string CachedLocation = LocalCache.GetCachedPathByTag(Tag, FileVersion);
 
 					// if it was there, use it directly
 					if (CachedLocation != null)
@@ -111,8 +126,17 @@ namespace Turnkey
 					}
 					else
 					{
-						// make a random temp directory, or use the download directory
-						string TempDir = SpecialMode == CopyExecuteSpecialMode.UsePermanentStorage ? Path.GetDirectoryName(OutputPath) : LocalCache.CreateTempDirectory();
+						// make a random temp directory, or use the download directory for permanent downloads
+						string DecompressLocation;
+						if (SpecialMode == CopyExecuteSpecialMode.UsePermanentStorage)
+						{
+							DecompressLocation = Path.Combine(Path.GetDirectoryName(OutputPath), Path.GetFileNameWithoutExtension(OutputPath) + "_Uncompressed");
+						}
+						else
+						{
+							DecompressLocation = LocalCache.CreateTempDirectory();
+						}
+
 						bool bFailed = false;
 						if (Ext == ".zip")
 						{
@@ -120,23 +144,24 @@ namespace Turnkey
 							{
 								using (ZipFile ZipFile = ZipFile.Read(OutputPath))
 								{
-									TurnkeyUtils.Log("Unzipping {0} to {1}...", OutputPath, TempDir);
+									TurnkeyUtils.Log("Unzipping {0} to {1}...", OutputPath, DecompressLocation);
 
 									// extract the zip to it
-									ZipFile.ExtractAll(TempDir);
+									ZipFile.ExtractAll(DecompressLocation);
 								}
 							}
-							catch (Exception)
+							catch (Exception Ex)
 							{
+								TurnkeyUtils.Log("Unzip failed: {0}", Ex);
 								bFailed = true;
 							}
 						}
 						else if (Ext == ".7z")
 						{
-							TurnkeyUtils.Log("7zip decompressing {0} to {1}...", OutputPath, TempDir);
+							TurnkeyUtils.Log("7zip decompressing {0} to {1}...", OutputPath, DecompressLocation);
 
 							int ExitCode = UnrealBuildTool.Utils.RunLocalProcessAndLogOutput(TurnkeyUtils.ExpandVariables("$(EngineDir)/Restricted/NotForLicensees/Extras/ThirdPartyNotUE/7-Zip/7z.exe"),
-								string.Format("x -o{0} {1}", TempDir, OutputPath));
+								string.Format("x -o{0} {1}", DecompressLocation, OutputPath));
 							if (ExitCode != 0)
 							{
 								TurnkeyUtils.Log("Failed to uncompress a .7z file {0}", OutputPath);
@@ -147,8 +172,13 @@ namespace Turnkey
 						// the temp dir is now the outputpatb to return to later installation steps
 						if (!bFailed)
 						{
-							OutputPath = TempDir;
-							LocalCache.CacheLocationByTag(Tag, OutputPath);
+							OutputPath = DecompressLocation;
+							LocalCache.CacheLocationByTag(Tag, OutputPath, FileVersion);
+						}
+						else
+						{
+							// return null on a failure
+							OutputPath = null;
 						}
 
 						// @todo turnkey: let the CopyProvider delete the file (p4 needs to perform deletes of synced files)
@@ -163,7 +193,12 @@ namespace Turnkey
 		{
 			CopyProvider Provider;
 			string ProviderParam;
-			ParseOperation(CopyOperation, out Provider, out ProviderParam);
+
+			// we allow this to fail (like an unknown variable, etc)
+			if (!ParseOperation(CopyOperation, out Provider, out ProviderParam, true))
+			{
+				return null;
+			}
 
 			return Provider.Enumerate(ProviderParam);
 		}
diff --git a/Engine/Source/Programs/AutomationTool/Turnkey/LocalCache.cs b/Engine/Source/Programs/AutomationTool/Turnkey/LocalCache.cs
--- a/Engine/Source/Programs/AutomationTool/Turnkey/LocalCache.cs
+++ b/Engine/Source/Programs/AutomationTool/Turnkey/LocalCache.cs
@@ -6,16 +6,64 @@ using System.Text;
 using System.Threading.Tasks;
 using System.IO;
 using System.Runtime.Serialization.Formatters.Binary;
+using System.Reflection;
+using Tools.DotNETCommon;
+
 
 namespace Turnkey
 {
-	static class LocalCache
+	[Serializable]
+	public class CacheData
 	{
-		// these temp files 
-		public static string TempCacheLocation = Path.Combine(Path.GetTempPath(), "Turnkey", "TempFiles");
+		public string LocalPath;
+		public string Version;
+		public bool bIsFile;
+		public Dictionary<string, long> FileToTimeCache = new Dictionary<string, long>();
+
+		public CacheData()
+		{
+
+		}
+
+		public CacheData(string LocalPath, string Version)
+		{
+			this.LocalPath = LocalPath;
+			this.Version = Version;
+
+			// check if it's a file, and if it is, add itself to the date cache
+			if (File.Exists(LocalPath))
+			{
+				bIsFile = true;
+				FileToTimeCache.Add("", File.GetLastWriteTimeUtc(LocalPath).ToFileTimeUtc());
+			}
+			else
+			{
+				bIsFile = false;
+				string PathToRemove = LocalPath + Path.DirectorySeparatorChar;
+				// gather the files under the localpath and remember dates, for later cleaning
+				foreach (string Filename in Directory.EnumerateFiles(LocalPath, "*", SearchOption.AllDirectories))
+				{
+					FileToTimeCache.Add(Filename.Replace(PathToRemove, "").Replace("\\", "/"), File.GetLastWriteTimeUtc(Filename).ToFileTimeUtc());
+				}
+			}
+		}
+
+	}
+
+	[Serializable]
+	public class SavedCache
+	{
+		public Dictionary<string, CacheData> Cache = new Dictionary<string, CacheData>();
+		public SavedCache()
+		{
+
+		}
+	}
+
 
-		public static string DownloadCacheLocation = Path.Combine(Path.GetTempPath(), "Turnkey", "DownloadCache");
 
+	static class LocalCache
+	{
 		public static string CreateTempDirectory()
 		{
 			string TempDir = Path.Combine(TempCacheLocation, Path.GetRandomFileName());
@@ -28,98 +76,185 @@ namespace Turnkey
 		}
 
 
+		public static void CacheLocationByTag(string Tag, string CachePath, string VersionMatch = "")
+		{
+			// Xml serialization doesn't work well with \ in strings in dictionaries
+			Tag = Tag.Replace("\\", "/");
+			if (CachePath.StartsWith(TempCacheLocation))
+			{
+				TagToTempCache.Cache[Tag] = new CacheData(CachePath, VersionMatch);
+			}
+			else
+			{
+				TagToDownloadCache.Cache[Tag] = new CacheData(CachePath, VersionMatch);
+				SerializeObject(TagToDownloadCache, TagCacheFile);
+			}
+		}
 
-		static void SerializeObject(ISerializable Object, string Filename)
+		public static string GetCachedPathByTag(string Tag, string VersionMatch = "")
 		{
-			using (FileStream Stream = new FileStream(Filename, FileMode.Create))
+			// Xml serialization doesn't work well with \ in strings in dictionaries
+			Tag = Tag.Replace("\\", "/");
+
+			// @todo turnkey - this assumes a tag can't be in both
+
+			// look to see if the tag is known
+			CacheData Data;
+			if (TagToDownloadCache.Cache.TryGetValue(Tag, out Data) || TagToTempCache.Cache.TryGetValue(Tag, out Data))
 			{
-				new BinaryFormatter().Serialize(Stream, Object);
+				// last minute double-check before using it, that it exists
+				bool bExists = (Data.bIsFile && File.Exists(Data.LocalPath)) || (!Data.bIsFile && Directory.Exists(Data.LocalPath));
+				if (!bExists)
+				{
+					TagToDownloadCache.Cache.Remove(Tag);
+					TagToTempCache.Cache.Remove(Tag);
+					return null;
+				}
+
+				// if the version doesn't match, return it, but let the called know it's out of date
+				if (Data.Version != VersionMatch)
+				{
+					CleanCacheByTag(Tag);
+					return null;
+				}
+
+				// if the cersion matches, then we can use it!
+				return Data.LocalPath;
 			}
+
+			return null;
+		}
+
+		public static string CreateDownloadCacheDirectory()
+		{
+			return Path.Combine(DownloadCacheLocation, Path.GetRandomFileName());
+		}
+
+		public static string GetInstallCacheDirectory()
+		{
+			// this will prompt user if needed
+			return TurnkeySettings.GetUserSetting("User_QuickSwitchSdkLocation");
+		}
+
+
+
+
+
+		// these temp files 
+		public static string TempCacheLocation = Path.Combine(Path.GetTempPath(), "Turnkey", "TempFiles");
+
+		public static string DownloadCacheLocation = Path.Combine(Path.GetTempPath(), "Turnkey", "DownloadCache");
+
+
+
+		static void SerializeObject(SavedCache Object, string Filename)
+		{
+			string Str = Json.Serialize(Object);
+			File.WriteAllText(Filename, Str);
+
+// 			using (FileStream Stream = new FileStream(Filename, FileMode.Create))
+// 			{
+// 				new BinaryFormatter().Serialize(Stream, Object);
+// 			}
 		}
 
-		static ISerializable DeserializeObject(string Filename)
+		static SavedCache DeserializeObject(string Filename)
 		{
 			try
 			{
-				if (File.Exists(Filename))
+				string Str = File.ReadAllText(Filename);
+				if (!string.IsNullOrEmpty(Str))
 				{
-					using (FileStream Stream = new FileStream(Filename, FileMode.Open))
-					{
-						return (ISerializable)new BinaryFormatter().Deserialize(Stream);
-					}
+					return Json.Deserialize<SavedCache>(Str);
 				}
+
+// 				if (File.Exists(Filename))
+// 				{
+// 					using (FileStream Stream = new FileStream(Filename, FileMode.Open))
+// 					{
+// 						return (SavedCache)new BinaryFormatter().Deserialize(Stream);
+// 					}
+// 				}
 			}
-			catch(Exception)
+			catch(Exception Ex)
 			{
+				TurnkeyUtils.Log("Exception: {0}", Ex.ToString());
 			}
 
 			// for any error, just return null
 			return null;
 		}
 
-
 		// @todo turnkey - expire/cache dates/something to be able to check - likely the caller of this would need to do it, like GoogleDrive
 		// checking versions of each file, while enumerating, but only download if newer
-		static Dictionary<string, KeyValuePair<string, string>> TagToDownloadCache;
-		static string TagCacheFile = Path.Combine(DownloadCacheLocation, "TagCache.bin");
+		static SavedCache TagToDownloadCache;
+		static string TagCacheFile = Path.Combine(DownloadCacheLocation, "TagCache.xml");
+
+		// the Temp cache is for during-run temp files, we don't serialize this cache out
+		static SavedCache TagToTempCache = new SavedCache();
 
 		static LocalCache()
 		{
-			TagToDownloadCache = DeserializeObject(TagCacheFile) as Dictionary<string, KeyValuePair<string, string>>;
+			TurnkeyUtils.Log("Loading cache bin {0}", TagCacheFile);
+			TagToDownloadCache = DeserializeObject(TagCacheFile);
+			TurnkeyUtils.Log("Loadied: {0}", TagToDownloadCache);
 
 			if (TagToDownloadCache == null)
 			{
-				TagToDownloadCache = new Dictionary<string, KeyValuePair<string, string>>();
+				TagToDownloadCache = new SavedCache();
 			}
-		}
-
-		public static void CacheLocationByTag(string Tag, string CachePath, string VersionMatch="")
-		{
-			TagToDownloadCache[Tag] = new KeyValuePair<string, string>(CachePath, VersionMatch);
-			SerializeObject(TagToDownloadCache, TagCacheFile);
-		}
-
-		public static string GetCachedPathByTag(string Tag, string VersionMatch, out bool bFoundOldVersion)
-		{
-			// return it if it was there
-			KeyValuePair<string, string> CachedPathWithVersion;
-			if (TagToDownloadCache.TryGetValue(Tag, out CachedPathWithVersion))
+			
+			if (TagToDownloadCache.Cache.Count > 0)
 			{
-				// make sure any referenced cache locations still exist
-				if (!File.Exists(CachedPathWithVersion.Key) && !Directory.Exists(CachedPathWithVersion.Key))
+				TurnkeyUtils.Log("Cleaning old download cache...");
+
+				List<string> EntriesToDelete = new List<string>();
+				foreach (var Pair in TagToDownloadCache.Cache)
 				{
-					bFoundOldVersion = false;
-					TagToDownloadCache.Remove(Tag);
-					return null;
+					CacheData Data = Pair.Value;
+					foreach (var FileWithDate in Data.FileToTimeCache)
+					{
+						string Filename = Path.Combine(Data.LocalPath, FileWithDate.Key);
+						long FileTime = File.GetLastWriteTimeUtc(Filename).ToFileTimeUtc();
+						long SavedTime = FileWithDate.Value;
+						if (!File.Exists(Filename) || FileTime != SavedTime)
+						{
+							TurnkeyUtils.Log("Cleaning old download cache...");
+							// remove this guy from the cache (delayed after the iterator)
+							EntriesToDelete.Add(Pair.Key);
+
+							break;
+						}
+					}
 				}
 
-				// if the version doesn't match, return it, but let the called know it's out of date
-				bFoundOldVersion = CachedPathWithVersion.Value != VersionMatch;
-				return CachedPathWithVersion.Key;
+				// clean any we deleted above (even if delete failed, we remove from the cache so we will download again)
+				EntriesToDelete.ForEach(x => CleanCacheByTag(x));
 			}
 
-			bFoundOldVersion = false;
-			return null;
+			// @todo turnkey: now delete directories that aren't represented in the cache, or are too old
 		}
 
-		// version with no version checking
-		public static string GetCachedPathByTag(string Tag)
+		public static void CleanCacheByTag(string Tag)
 		{
-			return GetCachedPathByTag(Tag, "", out _);
-		}
-
-		public static string CreateDownloadCacheDirectory()
-		{
-			return Path.Combine(DownloadCacheLocation, Path.GetRandomFileName());
-		}
-
+			// Xml serialization doesn't work well with \ in strings in dictionaries
+			Tag = Tag.Replace("\\", "/");
 
+			TurnkeyUtils.Log("Cleaning tag {0}...", Tag);
 
+			// if any files were bad, then we just delete the whole pile
+			CacheData Data = TagToDownloadCache.Cache[Tag];
+			TagToDownloadCache.Cache.Remove(Tag);
 
-		public static string GetInstallCacheDirectory()
-		{
-			// this will prompt user if needed
-			return TurnkeySettings.GetUserSetting("User_QuickSwitchSdkLocation");
+			if (Data.bIsFile)
+			{
+				AutomationTool.InternalUtils.SafeDeleteFile(Data.LocalPath);
+			}
+			else
+			{
+				AutomationTool.InternalUtils.SafeDeleteDirectory(Data.LocalPath);
+			}
+			SerializeObject(TagToDownloadCache, TagCacheFile);
 		}
 	}
 }
diff --git a/Engine/Source/Programs/AutomationTool/Turnkey/Providers/GoogleDriveCopyProvider.cs b/Engine/Source/Programs/AutomationTool/Turnkey/Providers/GoogleDriveCopyProvider.cs
--- a/Engine/Source/Programs/AutomationTool/Turnkey/Providers/GoogleDriveCopyProvider.cs
+++ b/Engine/Source/Programs/AutomationTool/Turnkey/Providers/GoogleDriveCopyProvider.cs
@@ -12,16 +12,16 @@ namespace Turnkey
 		public override string ProviderToken { get { return "googledrive"; } }
 
 
-
 		private static DriveServiceHelper ServiceHelper = null;
 		private Dictionary<string, Google.Apis.Drive.v3.Data.File> PathToFileCache = new Dictionary<string, Google.Apis.Drive.v3.Data.File>();
 
 
-
-
 		public override string Execute(string Operation, CopyExecuteSpecialMode SpecialMode, string SpecialModeHint)
 		{
-			Prepare();
+			if (!Prepare())
+			{
+				return null;
+			}
 
 			// operations that end in / must have an * appended, so the wildcard matching works properlty
 			if (Operation.EndsWith("/"))
@@ -77,19 +77,14 @@ namespace Turnkey
 				Dictionary<Google.Apis.Drive.v3.Data.File, string> NewIdsToCopyToFilename = new Dictionary<Google.Apis.Drive.v3.Data.File, string>();
 				foreach (var Pair in IdsToCopyToFilename)
 				{
-					bool bFoundOldVersion;
 					// check the cache status and version
-					string ExistingPath = LocalCache.GetCachedPathByTag(Pair.Key.Id, Pair.Key.Version.ToString(), out bFoundOldVersion);
+					// @todo turnkey: we don't have per-version files under a directory tag cache entry, so we tag each file, although it will be duplicated with the op cache tag above
+					string ExistingPath = LocalCache.GetCachedPathByTag(Pair.Key.Id, Pair.Key.Version.ToString());
 
-					if (ExistingPath == null || bFoundOldVersion)
+					if (ExistingPath == null)
 					{
 						// still need to download this one
 						NewIdsToCopyToFilename.Add(Pair.Key, Pair.Value);
-
-						if (File.Exists(ExistingPath))
-						{
-							File.Delete(ExistingPath);
-						}
 					}
 					// if we are removing our one and only file from the list to be deleted, we need to change our OutputPath to be that single file
 					// since we won't be going through the lower loop
@@ -110,6 +105,7 @@ namespace Turnkey
 				if (SpecialMode == CopyExecuteSpecialMode.UsePermanentStorage && CachedOperationLocation == null)
 				{
 					DownloadDirectory = Path.Combine(LocalCache.GetInstallCacheDirectory(), SpecialModeHint);
+					AutomationTool.InternalUtils.SafeDeleteDirectory(DownloadDirectory);
 				}
 
 				// if we didn't have a cached directory at all, then make a new one
@@ -150,13 +146,16 @@ namespace Turnkey
 
 		public override string[] Enumerate(string Operation)
 		{
-			// if we have no wildcards, there's no need to waste time touching p4, just return the spec
+			// if we have no wildcards, there's no need to waste time touching google drive, just return the spec
 			if (!Operation.Contains("*"))
 			{
 				return new string[] { ProviderToken + ":" + Operation };
 			}
 
-			Prepare();
+			if (!Prepare())
+			{
+				return null;
+			}
 
 			Dictionary<string, Google.Apis.Drive.v3.Data.File> Output = new Dictionary<string, Google.Apis.Drive.v3.Data.File>();
 
@@ -207,19 +206,27 @@ namespace Turnkey
 			return Result;
 		}
 
-		private void Prepare()
+		private bool Prepare()
 		{
 			if (ServiceHelper == null)
 			{
-				string SecretsFile = CopyProvider.ExecuteCopy(@"file:D:\SDK\credentials.json");
+				if (!TurnkeyUtils.HasVariable("GoogleDrive_Credentials"))
+				{
+					TurnkeyUtils.Log("ERROR: Unable to use GoogleDrive without GoogleDrive_Credentials being set first!");
+					return false;
+				}
+
+				string SecretsFile = CopyProvider.ExecuteCopy("$(GoogleDrive_Credentials)");
 				if (SecretsFile == null)
 				{
-					TurnkeyUtils.Log("Unable to get a secrets file");
-					return;
+					TurnkeyUtils.Log("ERROR: Unable to get GoogleDrive secrets file from {0}", TurnkeyUtils.GetVariableValue("GoogleDrive_Credentials"));
+					return false;
 				}
 
 				ServiceHelper = new DriveServiceHelper("Quickstart", SecretsFile, Path.GetDirectoryName(SecretsFile));
 			}
+
+			return true;
 		}		
 		
 // 		private string GetHighestCachedDirectoryIdForPath(string DrivePath, out string RemainingPathAfterId)
diff --git a/Engine/Source/Programs/AutomationTool/Turnkey/SdkInfo.cs b/Engine/Source/Programs/AutomationTool/Turnkey/SdkInfo.cs
--- a/Engine/Source/Programs/AutomationTool/Turnkey/SdkInfo.cs
+++ b/Engine/Source/Programs/AutomationTool/Turnkey/SdkInfo.cs
@@ -19,7 +19,6 @@ namespace Turnkey
 			Full,
 			Flash,
 			Misc,
-			QuickSwitch,
 		};
 
 
@@ -46,44 +45,44 @@ namespace Turnkey
 		[XmlIgnore]
 		private Dictionary<UnrealTargetPlatform, AutomationTool.Platform> AutomationPlatforms;
 
-		private SdkInfo CloneForQuickSwitch(string LocalDirectoryForCopy)
-		{
-			SdkInfo Clone = new SdkInfo();
-			Clone.PlatformString = PlatformString;
-			Clone.Version = Version;
-			Clone.DisplayName = DisplayName;
-			Clone.AllowedFlashDeviceTypes = AllowedFlashDeviceTypes;
-			Clone.CustomSdkId = CustomSdkId;
-			Clone.CustomSdkParams = CustomSdkParams;
-			Clone.PlatformString = PlatformString;
-
-			Clone.Type = SdkType.QuickSwitch;
-
-			if (Installers != null)
-			{
-				List<CopyAndRun> NewInstallers = new List<CopyAndRun>();
-				foreach (CopyAndRun Installer in Installers)
-				{
-					// if we want to execute the installer, then copy it over
-					if (Installer.ShouldExecute())
-					{
-						CopyAndRun NewInstaller = new CopyAndRun(Installer);
-
-						// the downloaded location will be the location to install from later. if we blanked this out, the CopyOutputPath variable wouldn't get set appropriately
-						NewInstaller.Copy = "file:" + LocalDirectoryForCopy;
-						NewInstallers.Add(NewInstaller);
-					}
-				}
-				Clone.Installers = NewInstallers.ToArray();
-			}
-
-			if (CustomSdkInputFiles != null)
-			{
-				throw new AutomationException("CustomSdkInputFiles are not supported yet for QuickSwitch Sdks");
-			}
-
-			return Clone;
-		}
+// 		private SdkInfo CloneForQuickSwitch(string LocalDirectoryForCopy)
+// 		{
+// 			SdkInfo Clone = new SdkInfo();
+// 			Clone.PlatformString = PlatformString;
+// 			Clone.Version = Version;
+// 			Clone.DisplayName = DisplayName;
+// 			Clone.AllowedFlashDeviceTypes = AllowedFlashDeviceTypes;
+// 			Clone.CustomSdkId = CustomSdkId;
+// 			Clone.CustomSdkParams = CustomSdkParams;
+// 			Clone.PlatformString = PlatformString;
+// 
+// 			Clone.Type = SdkType.QuickSwitch;
+// 
+// 			if (Installers != null)
+// 			{
+// 				List<CopyAndRun> NewInstallers = new List<CopyAndRun>();
+// 				foreach (CopyAndRun Installer in Installers)
+// 				{
+// 					// if we want to execute the installer, then copy it over
+// 					if (Installer.ShouldExecute())
+// 					{
+// 						CopyAndRun NewInstaller = new CopyAndRun(Installer);
+// 
+// 						// the downloaded location will be the location to install from later. if we blanked this out, the CopyOutputPath variable wouldn't get set appropriately
+// 						NewInstaller.Copy = "file:" + LocalDirectoryForCopy;
+// 						NewInstallers.Add(NewInstaller);
+// 					}
+// 				}
+// 				Clone.Installers = NewInstallers.ToArray();
+// 			}
+// 
+// 			if (CustomSdkInputFiles != null)
+// 			{
+// 				throw new AutomationException("CustomSdkInputFiles are not supported yet for QuickSwitch Sdks");
+// 			}
+// 
+// 			return Clone;
+// 		}
 
 		#endregion
 
@@ -250,33 +249,33 @@ namespace Turnkey
 					string SubDir = string.Format("{0}/{1}", Platform.ToString(), Version);
 					Install.Execute(CopyExecuteSpecialMode.UsePermanentStorage, SubDir);
 
-					// now that we are copied, write out a manifest to the root of the output location for later quick switching,
-					// and not continue to install anything
-					string ManifestPath = TurnkeyUtils.GetVariableValue("CopyOutputPath");
-
-					// if it's a file, then get it's directory
-					if (File.Exists(ManifestPath))
-					{
-						ManifestPath = Path.GetDirectoryName(ManifestPath);
-					}
-					Directory.CreateDirectory(ManifestPath);
-					ManifestPath = Path.Combine(ManifestPath, "TurnkeyQuickSwitch.xml");
-
-					// if the manifest already existed as downloaded, then there's no need to make one, we assume it was correct
-					if (!File.Exists(ManifestPath))
-					{
-						TurnkeyManifest QuickSwitchManifest = new TurnkeyManifest();
-
-						// copy this Sdk to a new one
-						SdkInfo NewSdk = CloneForQuickSwitch(Path.GetDirectoryName(ManifestPath));
-						QuickSwitchManifest.SdkInfos = new SdkInfo[] { NewSdk };
-
-						// register the Sdk with th runtime so that we can install it without quitting
-						TurnkeyManifest.AddCreatedSdk(NewSdk);
-
-						// save out a single switch manifest
-						QuickSwitchManifest.Write(ManifestPath);
-					}
+// 					// now that we are copied, write out a manifest to the root of the output location for later quick switching,
+// 					// and not continue to install anything
+// 					string ManifestPath = TurnkeyUtils.GetVariableValue("CopyOutputPath");
+// 
+// 					// if it's a file, then get it's directory
+// 					if (File.Exists(ManifestPath))
+// 					{
+// 						ManifestPath = Path.GetDirectoryName(ManifestPath);
+// 					}
+// 					Directory.CreateDirectory(ManifestPath);
+// 					ManifestPath = Path.Combine(ManifestPath, "TurnkeyQuickSwitch.xml");
+// 
+// 					// if the manifest already existed as downloaded, then there's no need to make one, we assume it was correct
+// 					if (!File.Exists(ManifestPath))
+// 					{
+// 						TurnkeyManifest QuickSwitchManifest = new TurnkeyManifest();
+// 
+// 						// copy this Sdk to a new one
+// 						SdkInfo NewSdk = CloneForQuickSwitch(Path.GetDirectoryName(ManifestPath));
+// 						QuickSwitchManifest.SdkInfos = new SdkInfo[] { NewSdk };
+// 
+// 						// register the Sdk with th runtime so that we can install it without quitting
+// 						TurnkeyManifest.AddCreatedSdk(NewSdk);
+// 
+// 						// save out a single switch manifest
+// 						QuickSwitchManifest.Write(ManifestPath);
+//					}
 				}
 				else
 				{
diff --git a/Engine/Source/Programs/AutomationTool/Turnkey/Turnkey.Automation.cs b/Engine/Source/Programs/AutomationTool/Turnkey/Turnkey.Automation.cs
--- a/Engine/Source/Programs/AutomationTool/Turnkey/Turnkey.Automation.cs
+++ b/Engine/Source/Programs/AutomationTool/Turnkey/Turnkey.Automation.cs
@@ -48,9 +48,9 @@ namespace Turnkey
 					}
 				}
 			}
-			catch (Exception)
+			catch (Exception Ex)
 			{
-
+				TurnkeyUtils.Log("Turnkey exception: {0}", Ex.ToString());
 			}
 			finally
 			{
diff --git a/Engine/Source/Programs/AutomationTool/Turnkey/TurnkeyManifest.cs b/Engine/Source/Programs/AutomationTool/Turnkey/TurnkeyManifest.cs
--- a/Engine/Source/Programs/AutomationTool/Turnkey/TurnkeyManifest.cs
+++ b/Engine/Source/Programs/AutomationTool/Turnkey/TurnkeyManifest.cs
@@ -55,23 +55,8 @@ namespace Turnkey
 			{
 				DiscoveredSdks = new List<SdkInfo>();
 
-				// @todo turnkey: look in the Sdk Download directory!~
-				if (TurnkeySettings.HasSetUserSetting("User_QuickSwitchSdkLocation"))
-				{
-					DiscoveredSdks.AddRange(GetSdkInfosFromProvider("file:$(User_QuickSwitchSdkLocation)/*/*/TurnkeyQuickSwitch.xml"));
-				}
-
-
-				// some known standard locations to the chain with
-				DiscoveredSdks.AddRange(GetSdkInfosFromProvider("file:$(EngineDir)/Build/TurnkeyManifest.xml"));
-				if (Environment.GetEnvironmentVariable("UE_SDKS_ROOT") != null)
-				{
-					DiscoveredSdks.AddRange(GetSdkInfosFromProvider("file:$(UE_SDKS_ROOT)/TurnkeyManifest.xml"));
-				}
-				if (Environment.GetEnvironmentVariable("UE_STUDIO_TURNKEY_LOCATION") != null)
-				{
-					DiscoveredSdks.AddRange(GetSdkInfosFromProvider("$(UE_STUDIO_TURNKEY_LOCATION)"));
-				}
+				// known location to branch from, this will include a few other locations
+				DiscoveredSdks.AddRange(GetSdkInfosFromProvider("file:$(EngineDir)/Build/Turnkey/TurnkeyManifest.xml"));
 			}
 		}
 
