- Early platform Turnkey platform support (some using made up numbers, some using real sdk values)

diff --git a/Engine/Source/Programs/AutomationTool/IOS/IOSPlatform.Automation.cs b/Engine/Source/Programs/AutomationTool/IOS/IOSPlatform.Automation.cs
--- a/Engine/Source/Programs/AutomationTool/IOS/IOSPlatform.Automation.cs
+++ b/Engine/Source/Programs/AutomationTool/IOS/IOSPlatform.Automation.cs
@@ -169,6 +169,22 @@ public class IOSPlatform : Platform
 		SDKName = (TargetPlatform == UnrealTargetPlatform.TVOS) ? "appletvos" : "iphoneos";
 	}
 
+
+	public override bool IsCustomVersionNeeded(string CustomVersionId, string CustomVersionParams)
+	{
+		return true;
+	}
+	public override bool CustomVersionUpdate(string CustomVersionId, string UpdateParams)
+	{
+		Console.WriteLine("UPDATING CERT FROM " + UpdateParams);
+
+		return true;
+	}
+
+
+
+
+
 	public override string GetPlatformPakCommandLine(ProjectParams Params, DeploymentContext SC)
 	{
 		string PakParams = "";
diff --git a/Engine/Source/Programs/UnrealBuildTool/Configuration/UEBuildPlatformSDK.cs b/Engine/Source/Programs/UnrealBuildTool/Configuration/UEBuildPlatformSDK.cs
--- a/Engine/Source/Programs/UnrealBuildTool/Configuration/UEBuildPlatformSDK.cs
+++ b/Engine/Source/Programs/UnrealBuildTool/Configuration/UEBuildPlatformSDK.cs
@@ -51,7 +51,7 @@ namespace UnrealBuildTool
 		/// </summary>
 		protected const string SDKEnvironmentVarsFile = "OutputEnvVars.txt";
 
-		protected const string SDKRootEnvVar = "UE_SDKS_ROOT";
+		protected const string SDKRootEnvVar = "UE_SDKS_ROOT_LEGACY";
 
 		protected const string AutoSetupEnvVar = "AutoSDKSetup";
 
