- Turnkey embracing AutoSdk system
- Added "expansion" which can find valid Sdks (to use with AutoSdk, etc)
- Added very early Build discovery to use with Gauntlet

diff --git a/Engine/Build/Turnkey/TurnkeyManifest.xml b/Engine/Build/Turnkey/TurnkeyManifest.xml
--- a/Engine/Build/Turnkey/TurnkeyManifest.xml
+++ b/Engine/Build/Turnkey/TurnkeyManifest.xml
@@ -3,7 +3,6 @@
 	<AdditionalManifests>
 		<Manifest>$(ThisManifestDir)/StudioTurnkeyManifest.xml</Manifest>
 		<Manifest>file:$(EngineDir)/Restricted/NoRedist/Build/Turnkey/TurnkeyManifest.xml</Manifest>
-		<Manifest>file:$(UE_SDKS_ROOT)/TurnkeyManifest.xml</Manifest>
 		<Manifest>$(UE_STUDIO_TURNKEY_LOCATION)</Manifest>
 	</AdditionalManifests>
 </TurnkeyManifest>
diff --git a/Engine/Source/Programs/AutomationTool/AutomationUtils/Platform.cs b/Engine/Source/Programs/AutomationTool/AutomationUtils/Platform.cs
--- a/Engine/Source/Programs/AutomationTool/AutomationUtils/Platform.cs
+++ b/Engine/Source/Programs/AutomationTool/AutomationUtils/Platform.cs
@@ -11,6 +11,15 @@ using Tools.DotNETCommon;
 
 namespace AutomationTool
 {
+	public class DeviceInfo
+	{
+		public string Name;
+		public string Id;
+		public string SoftwareVersion;
+		public string Type;
+		public bool bIsDefault = false;
+	}
+
 	/// <summary>
 	/// Holds information for targeting specific platform (platform type + cook flavor)
 	/// </summary>
@@ -141,6 +150,39 @@ namespace AutomationTool
 			return new TargetPlatformDescriptor(TargetPlatformType, "");
 		}
 
+		public virtual string GetInstalledSdk()
+		{
+			return "";
+		}
+
+		public virtual string GetAllowedSdks()
+		{
+			return "";
+		}
+		public virtual string GetAutoSdkPlatformName()
+		{
+			return IniPlatformType.ToString();
+		}
+
+		public virtual DeviceInfo[] GetDevices()
+		{
+			return null;
+		}
+
+		public virtual string GetAllowedSoftwareVersions()
+		{
+			return "";
+		}
+		public virtual bool IsCustomVersionNeeded(string CustomVersionId, string CustomVersionParams)
+		{
+			return true;
+		}
+		public virtual bool CustomVersionUpdate(string CustomVersionId, string UpdateParams)
+		{
+			return true;
+		}
+
+
 		/// <summary>
 		/// Package files for the current platform.
 		/// </summary>
diff --git a/Engine/Source/Programs/AutomationTool/Turnkey/Commands/VerifySdk.cs b/Engine/Source/Programs/AutomationTool/Turnkey/Commands/VerifySdk.cs
--- a/Engine/Source/Programs/AutomationTool/Turnkey/Commands/VerifySdk.cs
+++ b/Engine/Source/Programs/AutomationTool/Turnkey/Commands/VerifySdk.cs
@@ -1,5 +1,7 @@
+using AutomationTool;
 using System;
 using System.Collections.Generic;
+using System.IO;
 using System.Linq;
 using UnrealBuildTool;
 
@@ -10,9 +12,8 @@ namespace Turnkey.Commands
 		protected override void Execute(string[] CommandOptions)
 		{
 			string DeviceName = TurnkeyUtils.ParseParamValue("Device", null, CommandOptions);
-
-			// 		bool bBestAvailable = TurnkeyUtils.ParseParam("BestAvailable", CommandOptions);
-			// 		bool bUpdateOnly = TurnkeyUtils.ParseParam("UpdateOnly", CommandOptions);
+			bool bUpdateIfNeeded = TurnkeyUtils.ParseParam("UpdateIfNeeded", CommandOptions);
+			bool bUnattended = TurnkeyUtils.ParseParam("Unattended", CommandOptions);
 
 			List<UnrealTargetPlatform> ChosenPlatforms = TurnkeyUtils.GetPlatformsFromCommandLineOrUser(CommandOptions, UnrealTargetPlatform.GetValidPlatforms().ToList());
 
@@ -22,21 +23,173 @@ namespace Turnkey.Commands
 			}
 
 			TurnkeyUtils.Log("Installed Sdk validity:");
+			TurnkeyUtils.ExitCode = AutomationTool.ExitCode.Success;
+
+			// check all the platforms
 			foreach (UnrealTargetPlatform Platform in ChosenPlatforms)
 			{
 				// get the platform object
-				AutomationTool.Platform PlatformObject = AutomationTool.Platform.Platforms[new AutomationTool.TargetPlatformDescriptor(Platform)];
+				AutomationTool.Platform AutomationPlatform = AutomationTool.Platform.Platforms[new AutomationTool.TargetPlatformDescriptor(Platform)];
 
-				// check its installed Sdk is valid
-				if (!TurnkeyUtils.IsValueValid(PlatformObject.GetInstalledSdk(), PlatformObject.GetAllowedSdks()))
+				SdkInfo.LocalAvailability LocalState = SdkInfo.GetLocalAvailability(AutomationPlatform);
+
+				if ((LocalState & (SdkInfo.LocalAvailability.AutoSdk_PlatformExists | SdkInfo.LocalAvailability.InstalledSdk_ValidVersionExists)) == 0)
 				{
-					TurnkeyUtils.Log("{0}: INVALID [Installed: '{1}', Required: '{2}']", Platform, PlatformObject.GetInstalledSdk(), PlatformObject.GetAllowedSdks());
-					Environment.ExitCode = (int)AutomationTool.ExitCode.Error_SDKNotFound;
+					TurnkeyUtils.Log("{0}: No AutoSdk or Installed Sdk found matching {1}", Platform, AutomationPlatform.GetAllowedSdks());
+					TurnkeyUtils.ExitCode = AutomationTool.ExitCode.Error_SDKNotFound;
 				}
 				else
 				{
-					TurnkeyUtils.Log("{0}: Valid!", Platform);
-//					TurnkeyUtils.Log("{0}: Valid [Installed: '{1}', Required: '{2}']", Platform, PlatformObject.GetInstalledSdk(), PlatformObject.GetAllowedSdks());
+					TurnkeyUtils.Log("{0}: Valid! [{1}]", Platform, (LocalState & (SdkInfo.LocalAvailability.AutoSdk_PlatformExists | SdkInfo.LocalAvailability.InstalledSdk_ValidVersionExists)).ToString());
+					//					TurnkeyUtils.Log("{0}: Valid [Installed: '{1}', Required: '{2}']", Platform, PlatformObject.GetInstalledSdk(), PlatformObject.GetAllowedSdks());
+				}
+
+				if (bUpdateIfNeeded && (TurnkeyUtils.ExitCode != AutomationTool.ExitCode.Success))
+				{
+					TurnkeyUtils.ExitCode = AutomationTool.ExitCode.Success;
+
+					bool bAttemptAutoSdkSetup = false;
+					bool bSetupEnvVarAfterInstall = false;
+					if (LocalState.HasFlag(SdkInfo.LocalAvailability.AutoSdk_VariableExists))
+					{
+						bAttemptAutoSdkSetup = true;
+					}
+					else
+					{
+						if (!bUnattended)
+						{
+							// @todo turnkey: help set up UE_SDKS_ROOT if it's not there
+							string Response = TurnkeyUtils.ReadInput("AutoSdks are not setup, but your studio has support. Would you like to set it up now? [Y/n]", "Y");
+							if (string.Compare(Response, "Y", true) == 0)
+							{
+								bAttemptAutoSdkSetup = true;
+								bSetupEnvVarAfterInstall = true;
+							}
+						}
+					}
+
+					if (bAttemptAutoSdkSetup)
+					{
+						TurnkeyUtils.Log("{0}: AutoSdk is setup on this compter, will look for available AutoSdk to download", Platform);
+
+						List<SdkInfo> MatchingAutoSdk = SdkInfo.FindMatchingSdks(AutomationPlatform, SdkInfo.SdkType.AutoSdk, bSelectSingleBest: true);
+
+						if (MatchingAutoSdk.Count == 0)
+						{
+							// no matching AutoSdk found - will fall through to look for full install
+							if (bSetupEnvVarAfterInstall)
+							{
+								TurnkeyUtils.Log("{0}: Unable to find a matching AutoSdk, skipping AutoSdk setup", Platform);
+							}
+						}
+						else
+						{
+							// make sure this is unset so that we can know if it worked or not after install
+							TurnkeyUtils.ClearVariable("CopyOutputPath");
+
+							// now download it (AutoSdks don't "install") on download
+							// @todo turnkey: handle errors, handle p4 going to wrong location, handle one Sdk for multiple platforms
+							MatchingAutoSdk[0].Install(Platform);
+
+							if (bSetupEnvVarAfterInstall)
+							{
+								// @todo turnkey - have studio settings 
+
+								// this is where we synced the Sdk to
+								string InstalledRoot = TurnkeyUtils.GetVariableValue("CopyOutputPath");
+
+								// failed to install, nothing we can do
+								if (string.IsNullOrEmpty(InstalledRoot))
+								{
+									TurnkeyUtils.ExitCode = AutomationTool.ExitCode.Error_SDKNotFound;
+									continue;
+								}
+
+								// walk up to one above Host* directory
+								DirectoryInfo AutoSdkSearch;
+								if (Directory.Exists(InstalledRoot))
+								{
+									AutoSdkSearch = new DirectoryInfo(InstalledRoot);
+								}
+								else
+								{
+									AutoSdkSearch = new FileInfo(InstalledRoot).Directory;
+								}
+								while (AutoSdkSearch.Name != "Host" + HostPlatform.Current.HostEditorPlatform.ToString())
+								{
+									AutoSdkSearch = AutoSdkSearch.Parent;
+								}
+
+								// now go one up to the parent of Host
+								AutoSdkSearch = AutoSdkSearch.Parent;
+
+								string Response = TurnkeyUtils.ReadInput("Enter directory for root of AutoSdks. Use detected value, or enter another:", AutoSdkSearch.FullName);
+								if (string.IsNullOrEmpty(Response))
+								{
+									continue;
+								}
+
+								// set the env var, globally
+								TurnkeyUtils.StartTrackingExternalEnvVarChanges();
+								Environment.SetEnvironmentVariable("UE_SDKS_ROOT", Response);
+								Environment.SetEnvironmentVariable("UE_SDKS_ROOT", Response, EnvironmentVariableTarget.User);
+								TurnkeyUtils.EndTrackingExternalEnvVarChanges();
+							}
+
+							// @todo turnkey - re-run detection again
+
+							continue;
+						}
+					}
+
+					if (!bUnattended)
+					{
+						string Response = TurnkeyUtils.ReadInput("Your Sdk installation is not up to date. Would you like to install a valid Sdk? [Y/n]", "Y");
+						if (string.Compare(Response, "Y", true) != 0)
+						{
+							continue;
+						}
+					}
+
+					// at this point, AutoSdk isn't viable, so look for a full install
+					List<SdkInfo> MatchingInstallableSdks = SdkInfo.FindMatchingSdks(AutomationPlatform, SdkInfo.SdkType.BuildOnly, bSelectSingleBest: bUnattended);
+					if (MatchingInstallableSdks.Count == 0)
+					{
+						MatchingInstallableSdks = SdkInfo.FindMatchingSdks(AutomationPlatform, SdkInfo.SdkType.Full, bSelectSingleBest: bUnattended);
+					}
+
+					if (MatchingInstallableSdks.Count == 0)
+					{
+						TurnkeyUtils.Log("ERROR: {0}: Unable top find any Sdks that could be installed");
+						TurnkeyUtils.ExitCode = AutomationTool.ExitCode.Error_SDKNotFound;
+					}
+
+					SdkInfo SdkToInstall = MatchingInstallableSdks[0];
+
+					// if there are multiple, ask the user
+					if (MatchingInstallableSdks.Count > 1)
+					{
+						int Choice = TurnkeyUtils.ReadInputInt("Multiple Sdks found that could be installed. Please select one:", MatchingInstallableSdks.Select(x => x.DisplayName).ToList(), true);
+						if (Choice == 0)
+						{
+							TurnkeyUtils.ExitCode = AutomationTool.ExitCode.Error_SDKNotFound;
+							continue;
+						}
+
+						SdkToInstall = MatchingInstallableSdks[Choice - 1];
+					}
+
+					// finall install the best or chosen Sdk
+					SdkToInstall.Install(Platform);
+
+// 					if (bInstallSdk)
+// 					{
+// 						// reset exit code
+// 						TurnkeyUtils.ExitCode = AutomationTool.ExitCode.Success;
+// 
+// 						InstallSdk InstallCommand = new InstallSdk();
+// 						InstallCommand.InternalExecute(new string[] { "-platform=" + Platform, "-BestAvailable" });
+// 					}
 				}
 			}
 		}
diff --git a/Engine/Source/Programs/AutomationTool/Turnkey/CopyProvider.cs b/Engine/Source/Programs/AutomationTool/Turnkey/CopyProvider.cs
--- a/Engine/Source/Programs/AutomationTool/Turnkey/CopyProvider.cs
+++ b/Engine/Source/Programs/AutomationTool/Turnkey/CopyProvider.cs
@@ -37,8 +37,9 @@ namespace Turnkey
 		/// we can't test locally what type it is
 		/// </summary>
 		/// <param name="Operation">Description for the operation, including a ProviderToken</param>
+		/// <param name="Expansions">A list of what *'s expanded to, one entry for each result</param>
 		/// <returns></returns>
-		public abstract string[] Enumerate(string Operation);
+		public abstract string[] Enumerate(string Operation, List<List<string>> Expansions);
 
 
 
@@ -95,7 +96,7 @@ namespace Turnkey
 		/// </summary>
 		/// <param name="CopyOperation"></param>
 		/// <returns>Output path, which could then be used as $(OutputPath) in later operations</returns>
-		public static string ExecuteCopy(string CopyOperation, CopyExecuteSpecialMode SpecialMode=CopyExecuteSpecialMode.None, string SpecialModeHint=null)
+		public static string ExecuteCopy(string CopyOperation, CopyExecuteSpecialMode SpecialMode = CopyExecuteSpecialMode.None, string SpecialModeHint = null)
 		{
 			CopyProvider Provider;
 			string ProviderParam;
@@ -189,7 +190,7 @@ namespace Turnkey
 			return OutputPath;
 		}
 
-		public static string[] ExecuteEnumerate(string CopyOperation)
+		public static string[] ExecuteEnumerate(string CopyOperation, List<List<string>> Expansions=null)
 		{
 			CopyProvider Provider;
 			string ProviderParam;
@@ -200,7 +201,7 @@ namespace Turnkey
 				return null;
 			}
 
-			return Provider.Enumerate(ProviderParam);
+			return Provider.Enumerate(ProviderParam, Expansions);
 		}
 	}
 }
diff --git a/Engine/Source/Programs/AutomationTool/Turnkey/Providers/GoogleDriveCopyProvider.cs b/Engine/Source/Programs/AutomationTool/Turnkey/Providers/GoogleDriveCopyProvider.cs
--- a/Engine/Source/Programs/AutomationTool/Turnkey/Providers/GoogleDriveCopyProvider.cs
+++ b/Engine/Source/Programs/AutomationTool/Turnkey/Providers/GoogleDriveCopyProvider.cs
@@ -13,7 +13,7 @@ namespace Turnkey
 
 
 		private static DriveServiceHelper ServiceHelper = null;
-		private Dictionary<string, Google.Apis.Drive.v3.Data.File> PathToFileCache = new Dictionary<string, Google.Apis.Drive.v3.Data.File>();
+		private Dictionary<string, Google.Apis.Drive.v3.Data.File> PathToFileCache;// = new Dictionary<string, Google.Apis.Drive.v3.Data.File>();
 
 
 		public override string Execute(string Operation, CopyExecuteSpecialMode SpecialMode, string SpecialModeHint)
@@ -41,12 +41,12 @@ namespace Turnkey
 				int PrevSlash = Operation.LastIndexOf('/', StarLocation);
 				PathBeforeFirstStar = Operation.Substring(0, PrevSlash + 1);
 
-				Dictionary<string, Google.Apis.Drive.v3.Data.File> Output = new Dictionary<string, Google.Apis.Drive.v3.Data.File>();
-				ExpandWildcards("", Operation, Output);
+				Dictionary<string, Tuple<Google.Apis.Drive.v3.Data.File, List<string>>> Output = new Dictionary<string, Tuple<Google.Apis.Drive.v3.Data.File, List<string>>>();
+				ExpandWildcards("", Operation, Output, null);
 
 				foreach (var Pair in Output)
 				{
-					IdsToCopyToFilename.Add(Pair.Value, Pair.Key.Replace(PathBeforeFirstStar, ""));
+					IdsToCopyToFilename.Add(Pair.Value.Item1, Pair.Key.Replace(PathBeforeFirstStar, ""));
 				}
 			}
 			else
@@ -144,7 +144,7 @@ namespace Turnkey
 		}
 
 
-		public override string[] Enumerate(string Operation)
+		public override string[] Enumerate(string Operation, List<List<string>> Expansions)
 		{
 			// if we have no wildcards, there's no need to waste time touching google drive, just return the spec
 			if (!Operation.Contains("*"))
@@ -157,35 +157,40 @@ namespace Turnkey
 				return null;
 			}
 
-			Dictionary<string, Google.Apis.Drive.v3.Data.File> Output = new Dictionary<string, Google.Apis.Drive.v3.Data.File>();
+			Dictionary<string, Tuple<Google.Apis.Drive.v3.Data.File, List<string>>> Output = new Dictionary<string, Tuple<Google.Apis.Drive.v3.Data.File, List<string>>>();
 
-			ExpandWildcards("", Operation, Output);
+			List<string> ExpansionSet = new List<string>();
+			ExpandWildcards("", Operation, Output, ExpansionSet);
 
 			// convert them into a nicer-for-Turnkey "'ParentId'/Filename" format
 			List<string> FixedPaths = new List<string>();
 			foreach (var Pair in Output)
 			{
-				if (Pair.Value.Parents.Count != 1)
+				if (Pair.Value.Item1.Parents.Count != 1)
 				{
 					TurnkeyUtils.Log("File {0} did not have 1 parent as expected. GoogleDriveCopyProvider.ExpandWildcards will probably have to change to remember the parent for each object as it goes down");
 					continue;
 				}
 
-				string FixedPath = string.Format("'{0}'/{1}", Pair.Value.Parents[0], Pair.Value.Name);
+				string FixedPath = string.Format("'{0}'/{1}", Pair.Value.Item1.Parents[0], Pair.Value.Item1.Name);
 
 				// cache this new path for faster Execute, etc
 				if (!PathToFileCache.ContainsKey(FixedPath))
 				{
-					PathToFileCache.Add(FixedPath, Pair.Value);
+					PathToFileCache.Add(FixedPath, Pair.Value.Item1);
 				}
 
 				// if it's a directory, then append a slash
-				if (Pair.Value.MimeType == MIMETypes.GoogleDriveFolder.ToMimeString())
+				if (Pair.Value.Item1.MimeType == MIMETypes.GoogleDriveFolder.ToMimeString())
 				{
 					FixedPath += "/";
 				}
 
 				FixedPaths.Add(ProviderToken + ":" + FixedPath);
+				if (Expansions != null)
+				{
+					Expansions.Add(Pair.Value.Item2);
+				}
 			}
 
 			return FixedPaths.ToArray();
@@ -224,6 +229,8 @@ namespace Turnkey
 				}
 
 				ServiceHelper = new DriveServiceHelper("Quickstart", SecretsFile, Path.GetDirectoryName(SecretsFile));
+
+				PathToFileCache = new Dictionary<string, Google.Apis.Drive.v3.Data.File>();
 			}
 
 			return true;
@@ -393,7 +400,7 @@ namespace Turnkey
 			}
 		}
 
-		private void ExpandWildcards(string Prefix, string PathString, Dictionary<string, Google.Apis.Drive.v3.Data.File> Output)
+		private void ExpandWildcards(string Prefix, string PathString, Dictionary<string, Tuple<Google.Apis.Drive.v3.Data.File, List<string>>> Output, List<string> ExpansionSet)
 		{
 			char Slash = '/';
 
@@ -407,7 +414,7 @@ namespace Turnkey
 				Google.Apis.Drive.v3.Data.File FinalFile = GetFileForPath(Prefix + PathString);
 				if (FinalFile != null)
 				{
-					Output.Add(Prefix + PathString, FinalFile);
+					Output.Add(Prefix + PathString, new Tuple<Google.Apis.Drive.v3.Data.File, List<string>>(FinalFile, ExpansionSet));
 				}
 				return;
 			}
@@ -441,8 +448,9 @@ namespace Turnkey
 				{
 					Match = Wildcard.Substring(1);
 					// we can only do contains, not endswith, but it will reduce the results
-					Query += string.Format(" and name contains '{0}'", Match);
+//					Query += string.Format(" and name contains '{0}'", Match);
 				}
+
 				List<Google.Apis.Drive.v3.Data.File> Results = ServiceHelper.SearchFiles(Query);
 
 				foreach (Google.Apis.Drive.v3.Data.File Result in Results)
@@ -453,6 +461,14 @@ namespace Turnkey
 						continue;
 					}
 
+					List<string> NewExpansionSet = null;
+					if (ExpansionSet != null)
+					{
+						NewExpansionSet = new List<string>(ExpansionSet);
+						// the match is the part of the filename that was not the *, so removing that will give us what we wanted to match
+						NewExpansionSet.Add(Result.Name.Replace(Match, ""));
+					}
+
 					// cache the file (it may have already been cached tho from a previous test)
 					string PathToFile = FullPathBeforeWildcard + "/" + Result.Name;
 					if (!PathToFileCache.ContainsKey(PathToFile))
@@ -465,7 +481,7 @@ namespace Turnkey
 						// always add directories, but only add files if the path didn't end in / (/Sdks/*/ would only want to return directories)
 						if (Result.MimeType == MIMETypes.GoogleDriveFolder.ToMimeString() || NextSlash == -1)
 						{
-							Output.Add(PathToFile, Result);
+							Output.Add(PathToFile, new Tuple<Google.Apis.Drive.v3.Data.File, List<string>>(Result, NewExpansionSet));
 						}
 					}
 					else
@@ -473,7 +489,7 @@ namespace Turnkey
 						// recurse with parts after the wildcard component
 						// @todo turnkey
 						string NewPrefix = (FullPathBeforeWildcard != "" ? (FullPathBeforeWildcard + Slash) : "") + Result.Name + Slash;
-						ExpandWildcards(NewPrefix, PathString.Substring(NextSlash + 1), Output);
+						ExpandWildcards(NewPrefix, PathString.Substring(NextSlash + 1), Output, NewExpansionSet);
 					}
 				}
 			}
diff --git a/Engine/Source/Programs/AutomationTool/Turnkey/Providers/NullCopyProvider.cs b/Engine/Source/Programs/AutomationTool/Turnkey/Providers/NullCopyProvider.cs
--- a/Engine/Source/Programs/AutomationTool/Turnkey/Providers/NullCopyProvider.cs
+++ b/Engine/Source/Programs/AutomationTool/Turnkey/Providers/NullCopyProvider.cs
@@ -1,5 +1,6 @@
 using System.IO;
 using System.Collections.Generic;
+using System.Linq;
 
 namespace Turnkey
 {
@@ -12,20 +13,26 @@ namespace Turnkey
 			Operation = Operation.Replace(Path.AltDirectorySeparatorChar, Path.DirectorySeparatorChar);
 
 			// this provider can use the file directly, so just return the input after expanding variables, if it exists
-			string FinalPath = TurnkeyUtils.ExpandVariables(Operation);
+			string OutputPath = TurnkeyUtils.ExpandVariables(Operation);
 
-			if (File.Exists(FinalPath) == false && Directory.Exists(FinalPath) == false)
+			int WildcardLocation = OutputPath.IndexOf('*');
+			if (WildcardLocation >= 0)
 			{
-				TurnkeyUtils.Log("Reqeusted local path {0} does not exist!", FinalPath);
-				return null;
+				// chop down to the last / before the wildcard
+				int LastSlashLocation = OutputPath.Substring(0, WildcardLocation).LastIndexOf(Path.DirectorySeparatorChar);
+				OutputPath = OutputPath.Substring(0, LastSlashLocation);
 			}
 
-			// @todo turnkey : if wildcard, do like p4 does
+			if (File.Exists(OutputPath) == false && Directory.Exists(OutputPath) == false)
+			{
+				TurnkeyUtils.Log("Reqeusted local path {0} does not exist!", OutputPath);
+				return null;
+			}
 			
-			return FinalPath ;
+			return OutputPath;
 		}
 
-		private void ExpandWildcards(string Prefix, string PathString, List<string> Output)
+		private void ExpandWildcards(string Prefix, string PathString, Dictionary<string, List<string>> Output, List<string> ExpansionSet)
 		{
 			char Slash = Path.DirectorySeparatorChar;
 
@@ -39,11 +46,11 @@ namespace Turnkey
 				if (Directory.Exists(PathString))
 				{
 					// make sure we end with a single slash
-					Output.Add(ProviderToken + ":" + PathString.TrimEnd("/\\".ToCharArray()) + Slash);
+					Output.Add(ProviderToken + ":" + PathString.TrimEnd("/\\".ToCharArray()) + Slash, ExpansionSet);
 				}
 				else if (File.Exists(PathString))
 				{
-					Output.Add(ProviderToken + ":" + PathString);
+					Output.Add(ProviderToken + ":" + PathString, ExpansionSet);
 				}
 				return;
 			}
@@ -62,18 +69,33 @@ namespace Turnkey
 
 			if (Directory.Exists(FullPathBeforeWildcard))
 			{
+				// get the path component that has a wildcard
 				string Wildcard = (NextSlash == -1) ? PathString.Substring(PrevSlash + 1) : PathString.Substring(PrevSlash + 1, (NextSlash - PrevSlash) - 1);
+
+				// track what's before and after the * to return what it expanded to
+				int StarLoc = Wildcard.IndexOf('*');
+				int PrefixLen = StarLoc;
+				int SuffixLen = Wildcard.Length - (StarLoc + 1);
 				foreach (string Dirname in Directory.EnumerateDirectories(FullPathBeforeWildcard, Wildcard))
 				{
+					List<string> NewExpansionSet = null;
+					if (ExpansionSet != null)
+					{
+						NewExpansionSet = new List<string>(ExpansionSet);
+						string PathComponent = Path.GetFileName(Dirname);
+						// the match is the part of the filename that was not the *, so removing that will give us what we wanted to match
+						NewExpansionSet.Add(PathComponent.Substring(PrefixLen, PathComponent.Length - (PrefixLen + SuffixLen)));
+					}
+
 					if (bIsLastComponent)
 					{
 						// indicate directories with a slash at the end
-						Output.Add(ProviderToken + ":" + Dirname + Slash);
+						Output.Add(ProviderToken + ":" + Dirname + Slash, NewExpansionSet);
 					}
 					// recurse
 					else
 					{
-						ExpandWildcards(Dirname + Slash, PathString.Substring(NextSlash + 1), Output);
+						ExpandWildcards(Dirname + Slash, PathString.Substring(NextSlash + 1), Output, NewExpansionSet);
 					}
 				}
 
@@ -82,22 +104,35 @@ namespace Turnkey
 				{
 					foreach (string Filename in Directory.EnumerateFiles(FullPathBeforeWildcard, Wildcard))
 					{
-						Output.Add(ProviderToken + ":" + Filename);
+						List<string> NewExpansionSet = null;
+						if (ExpansionSet != null)
+						{
+							NewExpansionSet = new List<string>(ExpansionSet);
+							string PathComponent = Path.GetFileName(Filename);
+							// the match is the part of the filename that was not the *, so removing that will give us what we wanted to match
+							NewExpansionSet.Add(PathComponent.Substring(PrefixLen, PathComponent.Length - (PrefixLen + SuffixLen)));
+						}
+
+						Output.Add(ProviderToken + ":" + Filename, NewExpansionSet);
 					}
 				}
 			}
 		}
 
-		public override string[] Enumerate(string Operation)
+		public override string[] Enumerate(string Operation, List<List<string>> Expansions)
 		{
-			List<string> Output = new List<string>();
+			Dictionary<string, List<string>> Output = new Dictionary<string, List<string>>();
 
 			// we want consistent slashes
 			Operation = Operation.Replace(Path.AltDirectorySeparatorChar, Path.DirectorySeparatorChar);
 
-			ExpandWildcards("", Operation, Output);
+			ExpandWildcards("", Operation, Output, Expansions == null ? null : new List<string>());
 
-			return Output.ToArray();
+			if (Expansions != null)
+			{
+				Expansions.AddRange(Output.Values);
+			}
+			return Output.Keys.ToArray();
 		}
 	}
 }
diff --git a/Engine/Source/Programs/AutomationTool/Turnkey/Providers/PerforceCopyProvider.cs b/Engine/Source/Programs/AutomationTool/Turnkey/Providers/PerforceCopyProvider.cs
--- a/Engine/Source/Programs/AutomationTool/Turnkey/Providers/PerforceCopyProvider.cs
+++ b/Engine/Source/Programs/AutomationTool/Turnkey/Providers/PerforceCopyProvider.cs
@@ -1,11 +1,10 @@
 using System;
 using System.Collections.Generic;
 using System.Linq;
-using System.Text;
-using System.Threading.Tasks;
 using Tools.DotNETCommon;
 using AutomationTool;
 using System.Text.RegularExpressions;
+using System.IO;
 
 namespace Turnkey
 {
@@ -23,19 +22,13 @@ namespace Turnkey
 
 
 			// calculate the output path
-			string OutputPath = Operation.Replace(ConnectedStream, PerforceClient.RootPath);
-			// find the first wildcard and that's the lowest down outputpath we can use
-			int DotsLocation = OutputPath.IndexOf("...");
-			int StarLocation = OutputPath.IndexOf("*");
-			// hjandle various combos of -1 and non -1
-			int WildcardLocation = (DotsLocation >= 0 && StarLocation >= 0) ? Math.Min(DotsLocation, StarLocation) : Math.Max(DotsLocation, StarLocation);
-			if (WildcardLocation != -1)
-			{
-				// chop down to the last / before the wildcard
-				int LastSlashLocation = OutputPath.Substring(0, WildcardLocation).LastIndexOf("/");
-				OutputPath = OutputPath.Substring(0, LastSlashLocation);
-			}
-			OutputPath.Replace('/', System.IO.Path.DirectorySeparatorChar);
+			// @todo turnkey: use p4 where, because a depot mapping could be strange and make this invlalid
+			string OutputPath = Operation.Replace(ConnectedRoot, PerforceClient.RootPath);
+
+			// turn //a/b/c/foo*/... to //a/b/c
+			OutputPath = GetDirectoryBeforeWildcard(OutputPath);
+
+			OutputPath = OutputPath.Replace(Path.AltDirectorySeparatorChar, Path.DirectorySeparatorChar);
 
 			TurnkeyUtils.Log("Syncing '{0}' to '{1}'", Operation, OutputPath);
 
@@ -45,7 +38,7 @@ namespace Turnkey
 			return OutputPath;
 		}
 
-		public override string[] Enumerate(string Operation)
+		public override string[] Enumerate(string Operation, List<List<string>> Expansions)
 		{
 			// if we have no wildcards, there's no need to waste time touching p4, just return the spec
 			if (!Operation.Contains("*") && !Operation.Contains("..."))
@@ -62,16 +55,53 @@ namespace Turnkey
 			// now get the file list that matches
 			List<string> Results = PerforceConnection.Files(Operation);
 
+			if (Expansions != null)
+			{
+				// figure out what the *'s turned into
+				Regex Expression = new Regex(Operation.Replace("/", "\\/").Replace("*", "([^\\/]*)").Replace("...", ".*"));
+
+				// run the regex on each result
+				foreach (string Result in Results)
+				{
+					List<string> ResultExpansions = new List<string>();
+					Expansions.Add(ResultExpansions);
+
+					Match ResultMatch = Expression.Match(Result);
+					for (int Index = 1; Index < ResultMatch.Groups.Count; Index++)
+					{
+						ResultExpansions.Add(ResultMatch.Groups[Index].Value);
+					}
+				}
+			}
+
 			return Results.Select(x => ProviderToken + ":" + x).ToArray();
 		}
 
 
 
 
+		private string GetDirectoryBeforeWildcard(string Input)
+		{
+			// find the first wildcard and that's the lowest down outputpath we can use
+			int DotsLocation = Input.IndexOf("...");
+			int StarLocation = Input.IndexOf("*");
+			// hjandle various combos of -1 and non -1
+			int WildcardLocation = (DotsLocation >= 0 && StarLocation >= 0) ? Math.Min(DotsLocation, StarLocation) : Math.Max(DotsLocation, StarLocation);
+			if (WildcardLocation != -1)
+			{
+				// chop down to the last / before the wildcard
+				int LastSlashLocation = Input.Substring(0, WildcardLocation).Replace("\\", "/").LastIndexOf("/");
+				Input = Input.Substring(0, LastSlashLocation);
+			}
+
+			return Input;
+		}
+
+
 		static private P4Connection PerforceConnection = null;
 		static private P4ClientInfo PerforceClient = null;
 		// only non-null if we are connected
-		static private string ConnectedStream = null;
+		static private string ConnectedRoot = null;
 
 		private P4ClientInfo DetectClientForStream(string Stream, string Username, string Hostname)
 		{
@@ -92,6 +122,74 @@ namespace Turnkey
 			return null;
 		}
 
+		private P4ClientInfo CanClientHandleOperation(string ClientName, string Operation, string Hostname)
+		{
+			// make sure the Operation can be supported by the client (Tokens[1] is the clientspec name)
+			IProcessResult P4Result = PerforceConnection.P4(string.Format("-c {0} where {1}", ClientName, Operation), AllowSpew: false, WithClient: false);
+
+			if (P4Result.ExitCode == 0)
+			{
+				// sadly, this doesn't set an ExitCode, so we have to look at output
+				if (!P4Result.Output.Trim().EndsWith("not in client view."))
+				{
+					P4ClientInfo ClientInfo = PerforceConnection.GetClientInfo(ClientName, true);
+					// make sure it's usable on this computer
+					if (string.IsNullOrEmpty(ClientInfo.Host) || string.Compare(ClientInfo.Host, Hostname) == 0)
+					{
+						return ClientInfo;
+					}
+				}
+			}
+
+			return null;
+		}
+
+		private P4ClientInfo DetectClientForDepot(string Operation, string Username, string Hostname)
+		{
+			// use the bits up to a wildcard
+			Operation = GetDirectoryBeforeWildcard(Operation);
+
+			P4ClientInfo Client;
+			if (TurnkeySettings.HasSetUserSetting("User_LastPerforceClient"))
+			{
+				Client = CanClientHandleOperation(TurnkeySettings.GetUserSetting("User_LastPerforceClient"), Operation, Hostname);
+				if (Client != null)
+				{
+					return Client;
+				}
+			}
+
+			// Get all clients for this user
+			string P4Command = String.Format("clients -u {0}", Username);
+
+			var P4Result = PerforceConnection.P4(string.Format("clients -u {0}", Username), AllowSpew: false, WithClient: false);
+			if (P4Result.ExitCode != 0)
+			{
+				return null;
+			}
+
+			// Parse output.
+			var Lines = P4Result.Output.Split(new string[] { Environment.NewLine }, StringSplitOptions.RemoveEmptyEntries);
+			foreach (string Line in Lines)
+			{
+				var Tokens = Line.Split(new char[] { ' ' }, StringSplitOptions.RemoveEmptyEntries);
+				if (Tokens[0] == "Client")
+				{
+					Client = CanClientHandleOperation(Tokens[1], Operation, Hostname);
+					if (Client != null)
+					{
+						// remember this for next time
+						TurnkeySettings.SetUserSetting("User_LastPerforceClient", Tokens[1]);
+
+						return Client;
+					}
+				}
+			}
+
+			// if no clients at all, return null
+			return null;
+		}
+
 		private bool PrepareForOperation(string Operation)
 		{
 			Match StreamMatch = new Regex(@"(\/\/\w*\/\w*)\/.*").Match(Operation);
@@ -100,19 +198,32 @@ namespace Turnkey
 				throw new AutomationException("Unable to find stream spec in perforce operation {0}", Operation);
 			}
 
-			string Stream = StreamMatch.Groups[1].ToString();
+			string SpecRoot = StreamMatch.Groups[1].ToString();
+			string DepotName = SpecRoot.Substring(0, SpecRoot.LastIndexOf('/'));
 			string Hostname = System.Net.Dns.GetHostName();
 
-			TurnkeyUtils.Log("Stream: {0}", Stream);
-
-			if (ConnectedStream != Stream || PerforceConnection == null)
+			// @todo turnkey - for depot types, we can't totally safely check DepotName, we need to make sure the recent Client can handle the Operation - which we could with tricky View parsing
+			if (ConnectedRoot != SpecRoot || PerforceConnection == null)
 			{
+				TurnkeyUtils.Log("Finding clientspec usable with {0}...", Operation);
+
 				PerforceConnection = new P4Connection(null, null);
 				string Username = P4Environment.DetectUserName(PerforceConnection);
-				PerforceClient = DetectClientForStream(Stream, Username, Hostname);
-				if (PerforceClient == null)
+
+				// make sure it's a stream
+				var P4Result = PerforceConnection.P4("stream -o " + SpecRoot, AllowSpew: false);
+				bool bIsStream = P4Result.ExitCode == 0;
+
+				// hunt down a client that can be used
+				PerforceClient = bIsStream ? DetectClientForStream(SpecRoot, Username, Hostname) : DetectClientForDepot(Operation, Username, Hostname);
+
+				if (PerforceClient != null)
+				{
+					TurnkeyUtils.Log("Using client {0}", PerforceClient.Name);
+				}
+				else
 				{
-					TurnkeyUtils.Log("Unable to find a clientspec for the perforce stream {0}", Stream);
+					TurnkeyUtils.Log("Unable to find a clientspec for the perforce operation {0}, looking for a depot client", Operation);
 
 					string Response = TurnkeyUtils.ReadInput("Would you like to create one? [y/N]", "N");
 					if (string.Compare(Response, "Y", true) != 0)
@@ -127,7 +238,7 @@ namespace Turnkey
 					string ClientName = TurnkeyUtils.ReadInput("Enter clientspec name:", string.Format("{0}_sdks", Username));
 
 					// get local pathname
-					string LocalPath = TurnkeyUtils.ReadInput("Enter local path:", @"C:\Sdks");
+					string LocalPath = TurnkeyUtils.ReadInput(string.Format("Enter where to map {0} on your computer:", bIsStream ? SpecRoot : DepotName), @"D:\Sdks");
 
 					// create a client from the input settings
 					P4ClientInfo NewClient = new P4ClientInfo();
@@ -135,7 +246,18 @@ namespace Turnkey
 					NewClient.Owner = Username;
 					NewClient.Host = Hostname;
 					NewClient.RootPath = LocalPath;
-					NewClient.Stream = Stream;
+					// @todo turnkey handle depot type :|
+					if (bIsStream)
+					{
+						NewClient.Stream = SpecRoot;
+					}
+					else
+					{
+						// set up mapping
+						string DepotPath = SpecRoot + "/...";
+						string ClientPath = DepotPath.Replace("//" + DepotName + "/", "//" + ClientName + "/");
+						NewClient.View.Add(new KeyValuePair<string, string>(DepotPath, ClientPath));
+					}
 
 					PerforceClient = PerforceConnection.CreateClient(NewClient);
 				}
@@ -151,7 +273,7 @@ namespace Turnkey
 
 				// @todo turnkey: how to check that for errors?
 
-				ConnectedStream = Stream;
+				ConnectedRoot = bIsStream ? SpecRoot : DepotName;
 			}
 
 			return true;
diff --git a/Engine/Source/Programs/AutomationTool/Turnkey/SdkInfo.cs b/Engine/Source/Programs/AutomationTool/Turnkey/SdkInfo.cs
--- a/Engine/Source/Programs/AutomationTool/Turnkey/SdkInfo.cs
+++ b/Engine/Source/Programs/AutomationTool/Turnkey/SdkInfo.cs
@@ -15,6 +15,7 @@ namespace Turnkey
 		public enum SdkType
 		{
 			BuildOnly,
+			AutoSdk,
 			RunOnly,
 			Full,
 			Flash,
@@ -34,6 +35,7 @@ namespace Turnkey
 		public string CustomSdkId = null;
 		public string CustomSdkParams = null;
 		public CopyAndRun[] CustomSdkInputFiles = null;
+		public CopyAndRun[] Expansion = null;
 
 		[XmlIgnore]
 		// if there are custom files, we copy them once, then cache the location for $(CustomSdkInputFiles) on a later execution
@@ -45,44 +47,48 @@ namespace Turnkey
 		[XmlIgnore]
 		private Dictionary<UnrealTargetPlatform, AutomationTool.Platform> AutomationPlatforms;
 
-// 		private SdkInfo CloneForQuickSwitch(string LocalDirectoryForCopy)
-// 		{
-// 			SdkInfo Clone = new SdkInfo();
-// 			Clone.PlatformString = PlatformString;
-// 			Clone.Version = Version;
-// 			Clone.DisplayName = DisplayName;
-// 			Clone.AllowedFlashDeviceTypes = AllowedFlashDeviceTypes;
-// 			Clone.CustomSdkId = CustomSdkId;
-// 			Clone.CustomSdkParams = CustomSdkParams;
-// 			Clone.PlatformString = PlatformString;
-// 
-// 			Clone.Type = SdkType.QuickSwitch;
-// 
-// 			if (Installers != null)
-// 			{
-// 				List<CopyAndRun> NewInstallers = new List<CopyAndRun>();
-// 				foreach (CopyAndRun Installer in Installers)
-// 				{
-// 					// if we want to execute the installer, then copy it over
-// 					if (Installer.ShouldExecute())
-// 					{
-// 						CopyAndRun NewInstaller = new CopyAndRun(Installer);
-// 
-// 						// the downloaded location will be the location to install from later. if we blanked this out, the CopyOutputPath variable wouldn't get set appropriately
-// 						NewInstaller.Copy = "file:" + LocalDirectoryForCopy;
-// 						NewInstallers.Add(NewInstaller);
-// 					}
-// 				}
-// 				Clone.Installers = NewInstallers.ToArray();
-// 			}
-// 
-// 			if (CustomSdkInputFiles != null)
-// 			{
-// 				throw new AutomationException("CustomSdkInputFiles are not supported yet for QuickSwitch Sdks");
-// 			}
-// 
-// 			return Clone;
-// 		}
+		public SdkInfo CloneForExpansion()
+		{
+			SdkInfo Clone = new SdkInfo();
+			Clone.PlatformString = PlatformString;
+			Clone.Version = Version;
+			Clone.DisplayName = DisplayName;
+			Clone.AllowedFlashDeviceTypes = AllowedFlashDeviceTypes;
+			Clone.CustomSdkId = CustomSdkId;
+			Clone.CustomSdkParams = CustomSdkParams;
+			Clone.Type = Type;
+
+			if (Installers != null)
+			{
+				List<CopyAndRun> NewInstallers = new List<CopyAndRun>();
+				foreach (CopyAndRun Installer in Installers)
+				{
+					// if we want to execute the installer, then copy it over
+					if (Installer.ShouldExecute())
+					{
+						NewInstallers.Add(new CopyAndRun(Installer));
+					}
+				}
+				Clone.Installers = NewInstallers.ToArray();
+			}
+			if (CustomSdkInputFiles != null)
+			{
+				List<CopyAndRun> NewInstallers = new List<CopyAndRun>();
+				foreach (CopyAndRun Installer in CustomSdkInputFiles)
+				{
+					// if we want to execute the installer, then copy it over
+					if (Installer.ShouldExecute())
+					{
+						NewInstallers.Add(new CopyAndRun(Installer));
+					}
+				}
+				Clone.CustomSdkInputFiles = NewInstallers.ToArray();
+			}
+
+			Clone.PostDeserialize();
+
+			return Clone;
+		}
 
 		#endregion
 
@@ -93,6 +99,96 @@ namespace Turnkey
 		};
 
 
+		[Flags]
+		public enum LocalAvailability
+		{
+			None = 0,
+			AutoSdk_VariableExists = 1,
+			AutoSdk_PlatformExists = 2,
+			// InstalledSdk_BuildOnlyWasInstalled = 4,
+			InstalledSdk_ValidVersionExists = 8,
+			InstalledSdk_InvalidVersionExists = 16,
+		}
+
+		static public LocalAvailability GetLocalAvailability(AutomationTool.Platform AutomationPlatform)
+		{
+			LocalAvailability Result = LocalAvailability.None;
+
+			// if we have the variable at all, 
+			string AutoSdkVar = Environment.GetEnvironmentVariable("UE_SDKS_ROOT");
+			if (AutoSdkVar != null)
+			{
+				Result |= LocalAvailability.AutoSdk_VariableExists;
+
+				// get platform subdirectory
+				string AutoSubdir = string.Format("Host{0}/{1}", HostPlatform.Current.HostEditorPlatform.ToString(), AutomationPlatform.GetAutoSdkPlatformName());
+				DirectoryInfo PlatformDir = new DirectoryInfo(Path.Combine(AutoSdkVar, AutoSubdir));
+				if (PlatformDir.Exists)
+				{
+					foreach (DirectoryInfo Dir in PlatformDir.EnumerateDirectories())
+					{
+						// check for valid version number
+						if (TurnkeyUtils.IsValueValid(Dir.Name, AutomationPlatform.GetAllowedSdks()))
+						{
+							// make sure it actually has bits in it
+							if (File.Exists(Path.Combine(Dir.FullName, "setup.bat")))
+							{
+								Result |= LocalAvailability.AutoSdk_PlatformExists;
+							}
+							break;
+						}
+					}
+				}
+			}
+
+			// if anything is installed, this will return a value
+			string InstalledVersion = AutomationPlatform.GetInstalledSdk();
+			if (!string.IsNullOrEmpty(InstalledVersion))
+			{
+				if (TurnkeyUtils.IsValueValid(InstalledVersion, AutomationPlatform.GetAllowedSdks()))
+				{
+					Result |= LocalAvailability.InstalledSdk_ValidVersionExists;
+				}
+				else
+				{
+					Result |= LocalAvailability.InstalledSdk_InvalidVersionExists;
+				}
+			}
+
+
+			return Result;
+		}
+
+		static public List<SdkInfo> FindMatchingSdks(AutomationTool.Platform Platform, SdkType Type, bool bSelectSingleBest)
+		{
+
+			List<SdkInfo> Sdks = TurnkeyManifest.GetDiscoveredSdks().FindAll(x => 
+				x.SupportsPlatform(Platform.IniPlatformType) && 
+				x.Type == Type && 
+				TurnkeyUtils.IsValueValid(x.Version, Platform.GetAllowedSdks())
+				);
+
+			// handle easy cases
+			if (Sdks.Count <= 1 || !bSelectSingleBest)
+			{
+				return Sdks;
+			}
+
+			// find best one
+			SdkInfo Best = null;
+			foreach (SdkInfo Sdk in Sdks)
+			{
+				// @todo turnkey: let platform decide what's best
+				if (Best == null || string.Compare(Sdk.Version, Best.Version, true) > 0)
+				{
+					Best = Sdk;
+				}
+			}
+
+			return new List<SdkInfo> { Best };
+		}
+
+
 
 		public bool SupportsPlatform(UnrealTargetPlatform Platform)
 		{
@@ -230,6 +326,24 @@ namespace Turnkey
 			// custom sdk installation is enabled if ChosenSdk.CustomVersionId is !null
 			if (CustomSdkId != null)
 			{
+				// copy files down, which are needed to check if whatever is installed is up to date (only do it once)
+				if (CustomVersionLocalFiles == null && CustomSdkInputFiles != null)
+				{
+					foreach (CopyAndRun CustomCopy in CustomSdkInputFiles)
+					{
+						if (CustomCopy.ShouldExecute())
+						{
+							if (CustomVersionLocalFiles != null)
+							{
+								throw new AutomationTool.AutomationException("CustomSdkInputFiles specified multiple locations to be copied for this platform, which is not supported (only one value of $(CustomVersionLocalFiles) allowed)");
+							}
+
+							CustomCopy.Execute();
+							CustomVersionLocalFiles = TurnkeyUtils.GetVariableValue("CopyOutputPath");
+						}
+					}
+				}
+
 				// in case the custom sdk modifies global env vars, make sure we capture them
 				TurnkeyUtils.StartTrackingExternalEnvVarChanges();
 
@@ -241,49 +355,52 @@ namespace Turnkey
 			}
 
 			// now run any installers
-			foreach (CopyAndRun Install in Installers)
+			if (Installers != null)
 			{
-				if (Type == SdkType.BuildOnly)
-				{
-					// download subdir is based on platform 
-					string SubDir = string.Format("{0}/{1}", Platform.ToString(), Version);
-					Install.Execute(CopyExecuteSpecialMode.UsePermanentStorage, SubDir);
-
-// 					// now that we are copied, write out a manifest to the root of the output location for later quick switching,
-// 					// and not continue to install anything
-// 					string ManifestPath = TurnkeyUtils.GetVariableValue("CopyOutputPath");
-// 
-// 					// if it's a file, then get it's directory
-// 					if (File.Exists(ManifestPath))
-// 					{
-// 						ManifestPath = Path.GetDirectoryName(ManifestPath);
-// 					}
-// 					Directory.CreateDirectory(ManifestPath);
-// 					ManifestPath = Path.Combine(ManifestPath, "TurnkeyQuickSwitch.xml");
-// 
-// 					// if the manifest already existed as downloaded, then there's no need to make one, we assume it was correct
-// 					if (!File.Exists(ManifestPath))
-// 					{
-// 						TurnkeyManifest QuickSwitchManifest = new TurnkeyManifest();
-// 
-// 						// copy this Sdk to a new one
-// 						SdkInfo NewSdk = CloneForQuickSwitch(Path.GetDirectoryName(ManifestPath));
-// 						QuickSwitchManifest.SdkInfos = new SdkInfo[] { NewSdk };
-// 
-// 						// register the Sdk with th runtime so that we can install it without quitting
-// 						TurnkeyManifest.AddCreatedSdk(NewSdk);
-// 
-// 						// save out a single switch manifest
-// 						QuickSwitchManifest.Write(ManifestPath);
-//					}
-				}
-				else
+				foreach (CopyAndRun Install in Installers)
 				{
-					Install.Execute();
-				}
+					if (Type == SdkType.BuildOnly)
+					{
+						// download subdir is based on platform 
+						string SubDir = string.Format("{0}/{1}", Platform.ToString(), Version);
+						Install.Execute(CopyExecuteSpecialMode.UsePermanentStorage, SubDir);
+
+						// 					// now that we are copied, write out a manifest to the root of the output location for later quick switching,
+						// 					// and not continue to install anything
+						// 					string ManifestPath = TurnkeyUtils.GetVariableValue("CopyOutputPath");
+						// 
+						// 					// if it's a file, then get it's directory
+						// 					if (File.Exists(ManifestPath))
+						// 					{
+						// 						ManifestPath = Path.GetDirectoryName(ManifestPath);
+						// 					}
+						// 					Directory.CreateDirectory(ManifestPath);
+						// 					ManifestPath = Path.Combine(ManifestPath, "TurnkeyQuickSwitch.xml");
+						// 
+						// 					// if the manifest already existed as downloaded, then there's no need to make one, we assume it was correct
+						// 					if (!File.Exists(ManifestPath))
+						// 					{
+						// 						TurnkeyManifest QuickSwitchManifest = new TurnkeyManifest();
+						// 
+						// 						// copy this Sdk to a new one
+						// 						SdkInfo NewSdk = CloneForQuickSwitch(Path.GetDirectoryName(ManifestPath));
+						// 						QuickSwitchManifest.SdkInfos = new SdkInfo[] { NewSdk };
+						// 
+						// 						// register the Sdk with th runtime so that we can install it without quitting
+						// 						TurnkeyManifest.AddCreatedSdk(NewSdk);
+						// 
+						// 						// save out a single switch manifest
+						// 						QuickSwitchManifest.Write(ManifestPath);
+						//					}
+					}
+					else
+					{
+						Install.Execute();
+					}
 
-				// grab where the Install copied the files
+					// grab where the Install copied the files
 
+				}
 			}
 		}
 
@@ -340,6 +457,7 @@ namespace Turnkey
 				{
 					foreach (CopyAndRun CustomInput in CustomSdkInputFiles)
 					{
+						bContainsVar = bContainsVar || (CustomInput.Copy != null && CustomInput.Copy.Contains(Format));
 						bContainsVar = bContainsVar || (CustomInput.CommandPath != null && CustomInput.CommandPath.Contains(Format));
 						bContainsVar = bContainsVar || (CustomInput.CommandLine != null && CustomInput.CommandLine.Contains(Format));
 					}
@@ -349,6 +467,7 @@ namespace Turnkey
 			{
 				foreach (CopyAndRun Installer in Installers)
 				{
+					bContainsVar = bContainsVar || (Installer.Copy != null && Installer.Copy.Contains(Format));
 					bContainsVar = bContainsVar || (Installer.CommandPath != null && Installer.CommandPath.Contains(Format));
 					bContainsVar = bContainsVar || (Installer.CommandLine != null && Installer.CommandLine.Contains(Format));
 				}
@@ -365,9 +484,10 @@ namespace Turnkey
 			// validate
 			if (Version == null && CustomSdkId == null)
 			{
-				throw new AutomationTool.AutomationException("SdkInfo {0} needs to have either Version or CustomSdkId specified");
+				throw new AutomationTool.AutomationException("SdkInfo {0} needs to have either Version or CustomSdkId specified", DisplayName);
 			}
 
+			PlatformString = TurnkeyUtils.ExpandVariables(PlatformString, true);
 			string[] PlatformStrings = PlatformString.Split(",".ToCharArray());
 			Platforms = new List<UnrealTargetPlatform>();
 			AutomationPlatforms = new Dictionary<UnrealTargetPlatform, Platform>();
@@ -382,6 +502,12 @@ namespace Turnkey
 				}
 			}
 
+			Version = TurnkeyUtils.ExpandVariables(Version, true);
+			DisplayName = TurnkeyUtils.ExpandVariables(DisplayName, true);
+			AllowedFlashDeviceTypes = TurnkeyUtils.ExpandVariables(AllowedFlashDeviceTypes, true);
+			CustomSdkId = TurnkeyUtils.ExpandVariables(CustomSdkId, true);
+			CustomSdkParams = TurnkeyUtils.ExpandVariables(CustomSdkParams, true);
+
 			if (CustomSdkInputFiles != null)
 			{
 				Array.ForEach(CustomSdkInputFiles, x => x.PostDeserialize());
@@ -391,6 +517,14 @@ namespace Turnkey
 			{
 				Array.ForEach(Installers, x => x.PostDeserialize());
 			}
+
+			if (Expansion != null)
+			{
+				Array.ForEach(Expansion, x => x.PostDeserialize());
+
+				// remove installers that don't match this host platform
+				Expansion = Array.FindAll(Expansion, x => x.ShouldExecute());
+			}
 		}
 
 
diff --git a/Engine/Source/Programs/AutomationTool/Turnkey/Turnkey.Automation.cs b/Engine/Source/Programs/AutomationTool/Turnkey/Turnkey.Automation.cs
--- a/Engine/Source/Programs/AutomationTool/Turnkey/Turnkey.Automation.cs
+++ b/Engine/Source/Programs/AutomationTool/Turnkey/Turnkey.Automation.cs
@@ -11,9 +11,14 @@ namespace Turnkey
 {
 	class Turnkey : BuildCommand
 	{
-		public override void ExecuteBuild()
+		public override ExitCode Execute()
 		{
-			TurnkeyUtils.Initialize(new ConsoleIOProvider(), this);
+			return Turnkey.Execute(new ConsoleIOProvider(), this);
+		}
+
+		public static AutomationTool.ExitCode Execute(IOProvider IOProvider, BuildCommand CommandUtilHelper)
+		{
+			TurnkeyUtils.Initialize(IOProvider, CommandUtilHelper);
 			TurnkeySettings.Initialize();
 
 			// cache some settings for other classes
@@ -28,7 +33,7 @@ namespace Turnkey
 				TurnkeyUtils.Log("");
 				TurnkeyUtils.Log("If you installed an SDK, you should NOT \"Terminate batch job\"!");
 				TurnkeyUtils.Log("");
-				Environment.ExitCode = 0;
+				TurnkeyUtils.ExitCode = ExitCode.Success;
 			};
 
 			try
@@ -57,12 +62,7 @@ namespace Turnkey
 				TurnkeyUtils.CleanupPaths();
 			}
 
-			if (Environment.ExitCode != 0)
-			{
-				// @todo turnkey - would be nice to return a failure, without an exception in UAT, which looks violent
-				AutomationTool.ExitCode ExitCode = (AutomationTool.ExitCode)Environment.ExitCode;
-				throw new AutomationException(ExitCode, ExitCode.ToString());
-			}
+			return TurnkeyUtils.ExitCode;
 		}
 	}
 }
diff --git a/Engine/Source/Programs/AutomationTool/Turnkey/Turnkey.Automation.csproj b/Engine/Source/Programs/AutomationTool/Turnkey/Turnkey.Automation.csproj
--- a/Engine/Source/Programs/AutomationTool/Turnkey/Turnkey.Automation.csproj
+++ b/Engine/Source/Programs/AutomationTool/Turnkey/Turnkey.Automation.csproj
@@ -39,6 +39,8 @@
     <Compile Include="..\..\DotNETCommon\MetaData.cs">
       <Link>Properties\MetaData.cs</Link>
     </Compile>
+    <Compile Include="BuildInfo.cs" />
+    <Compile Include="Commands\InstallBuild.cs" />
     <Compile Include="Commands\InstallSdk.cs" />
     <Compile Include="Commands\ListPlatforms.cs" />
     <Compile Include="Commands\ListSdks.cs" />
@@ -61,10 +63,6 @@
     <Compile Include="TurnkeyUtils.cs" />
   </ItemGroup>
   <ItemGroup>
-    <ProjectReference Include="..\..\..\..\Restricted\NotForLicensees\Source\Programs\AutomationTool\GauntletExtras\GauntletExtras.Automation.csproj">
-      <Project>{ad48925a-1c54-46fd-8a4f-45d87115e18d}</Project>
-      <Name>GauntletExtras.Automation</Name>
-    </ProjectReference>
     <ProjectReference Include="..\..\DotNETCommon\DotNETUtilities\DotNETUtilities.csproj">
       <Project>{5d7d66e8-8c76-4af9-b3ec-2ef03421d730}</Project>
       <Name>DotNETUtilities</Name>
@@ -80,6 +78,10 @@
       <Name>AutomationUtils.Automation</Name>
       <Private>False</Private>
     </ProjectReference>
+    <ProjectReference Include="..\Gauntlet\Gauntlet.Automation.csproj">
+      <Project>{767b4f85-ab56-4b00-a033-04c7600acc3d}</Project>
+      <Name>Gauntlet.Automation</Name>
+    </ProjectReference>
   </ItemGroup>
   <ItemGroup>
     <Reference Include="Google.Apis, Version=1.33.1.0, Culture=neutral, PublicKeyToken=4b01fa6e34db77ab, processorArchitecture=MSIL">
diff --git a/Engine/Source/Programs/AutomationTool/Turnkey/TurnkeyCommand.cs b/Engine/Source/Programs/AutomationTool/Turnkey/TurnkeyCommand.cs
--- a/Engine/Source/Programs/AutomationTool/Turnkey/TurnkeyCommand.cs
+++ b/Engine/Source/Programs/AutomationTool/Turnkey/TurnkeyCommand.cs
@@ -35,6 +35,10 @@ namespace Turnkey
 		{
 			return null;
 		}
+		internal void InternalExecute(string[] CommandOptions)
+		{
+			Execute(CommandOptions);
+		}
 
 
 		// cached commands
diff --git a/Engine/Source/Programs/AutomationTool/Turnkey/TurnkeyManifest.cs b/Engine/Source/Programs/AutomationTool/Turnkey/TurnkeyManifest.cs
--- a/Engine/Source/Programs/AutomationTool/Turnkey/TurnkeyManifest.cs
+++ b/Engine/Source/Programs/AutomationTool/Turnkey/TurnkeyManifest.cs
@@ -3,6 +3,8 @@ using System.Collections.Generic;
 using System.Text;
 using System.Xml.Serialization;
 using System.IO;
+using System.Linq;
+using System.Text.RegularExpressions;
 using Tools.DotNETCommon;
 using UnrealBuildTool;
 using AutomationTool;
@@ -14,6 +16,7 @@ namespace Turnkey
 		static private string StandardManifestName = "TurnkeyManifest.xml";
 
 		public SdkInfo[] SdkInfos = null;
+		public BuildSource[] BuildSources = null;
 
 		[XmlArrayItem(ElementName = "Manifest")]
 		public string[] AdditionalManifests = null;
@@ -22,19 +25,80 @@ namespace Turnkey
 
 		internal void PostDeserialize()
 		{
-			if (SdkInfos != null)
-			{
-				Array.ForEach(SdkInfos, x => x.PostDeserialize());
-			}
-
 			// load any settings set in the xml
 			if (SavedSettings != null)
 			{
 				Array.ForEach(SavedSettings, x => TurnkeyUtils.SetVariable(x.Variable, x.Value));
 			}
+
+			if (BuildSources != null)
+			{
+				Array.ForEach(BuildSources, x => x.PostDeserialize());
+			}
+
+			if (SdkInfos != null)
+			{
+				Array.ForEach(SdkInfos, x => x.PostDeserialize());
+
+				SdkInfo[] ExpandedSdks = Array.FindAll(SdkInfos, x => x.Expansion != null && x.Expansion.Length > 0);
+
+				if (ExpandedSdks.Length > 0)
+				{
+					// we will only have one active Expansion
+					List<SdkInfo> NewSdks = SdkInfos.Where(x => !ExpandedSdks.Contains(x)).ToList();
+
+					// now expand
+					foreach (SdkInfo Sdk in ExpandedSdks)
+					{
+						// we only allow one expansion active
+						if (Sdk.Expansion.Length > 1)
+						{
+							throw new AutomationTool.AutomationException("SdkInfo {0} had multiple expansions active on this platform. THat is not allowed", Sdk.DisplayName);
+						}
+
+						// now enumerate and get the values
+						List<List<string>> Expansions = new List<List<string>>();
+						string[] ExpandedInstallerResults = CopyProvider.ExecuteEnumerate(Sdk.Expansion[0].Copy, Expansions);
+
+						if (Expansions.Count != ExpandedInstallerResults.Length)
+						{
+							throw new AutomationException("Bad expansions output from CopyProvider");
+						}
+
+						// @todo turnkey: this will be uysed in Builds also, make it a function with a lambda
+						// make a new SdkInfo for each expansion
+						int MaxIndex = 0;
+						for (int ResultIndex=0; ResultIndex < ExpandedInstallerResults.Length; ResultIndex++)
+						{
+							// set temp variables potentially used somewhere in the Sdks
+							TurnkeyUtils.SetVariable("Expansion", ExpandedInstallerResults[ResultIndex]);
+							for (int ExpansionIndex = 0; ExpansionIndex < Expansions[ResultIndex].Count; ExpansionIndex++)
+							{
+								TurnkeyUtils.SetVariable(ExpansionIndex.ToString(), Expansions[ResultIndex][ExpansionIndex]);
+							}
+
+							// remember how many we beed to unset
+							MaxIndex = Math.Max(MaxIndex, Expansions[ResultIndex].Count);
+
+							// make a new Sdk for each result in the expansion
+							NewSdks.Add(Sdk.CloneForExpansion());
+						}
+
+						// clear temp variables
+						TurnkeyUtils.ClearVariable("Expansion");
+						for (int Index = 0; Index <= MaxIndex; Index++)
+						{
+							TurnkeyUtils.ClearVariable(Index.ToString());
+						}
+					}
+
+					SdkInfos = NewSdks.ToArray();
+				}
+			}
 		}
 
 		static List<SdkInfo> DiscoveredSdks = null;
+		static List<BuildSource> DiscoveredBuildSources = null;
 
 		public static List<SdkInfo> GetDiscoveredSdks()
 		{
@@ -43,6 +107,15 @@ namespace Turnkey
 
 			return DiscoveredSdks;
 		}
+
+		public static List<BuildSource> GetDiscoveredBuildSources()
+		{
+			// hunt down manifests if needed
+			DiscoverManifests();
+
+			return DiscoveredBuildSources;
+		}
+
 		public static void AddCreatedSdk(SdkInfo Sdk)
 		{
 			DiscoveredSdks.Add(Sdk);
@@ -54,25 +127,23 @@ namespace Turnkey
 			if (DiscoveredSdks == null)
 			{
 				DiscoveredSdks = new List<SdkInfo>();
+				DiscoveredBuildSources = new List<BuildSource>();
 
 				// known location to branch from, this will include a few other locations
-				DiscoveredSdks.AddRange(GetSdkInfosFromProvider("file:$(EngineDir)/Build/Turnkey/TurnkeyManifest.xml"));
-			}
-		}
+				string RootOperation = "file:$(EngineDir)/Build/Turnkey/TurnkeyManifest.xml";
 
-		public static List<SdkInfo> GetSdkInfosFromProvider(string EnumerationOperation)
-		{
-			// gather the 
-			List<SdkInfo> Sdks = new List<SdkInfo>();
-			LoadManifestsFromProvider(EnumerationOperation).ForEach(x =>
-			{
-				if (x.SdkInfos != null)
+				LoadManifestsFromProvider(RootOperation).ForEach(x =>
 				{
-					Sdks.AddRange(x.SdkInfos);
-				}
-			});
-
-			return Sdks;
+					if (x.SdkInfos != null)
+					{
+						DiscoveredSdks.AddRange(x.SdkInfos);
+					}
+					if (x.BuildSources != null)
+					{
+						DiscoveredBuildSources.AddRange(x.BuildSources);
+					}
+				});
+			}
 		}
 
 		public static List<TurnkeyManifest> LoadManifestsFromProvider(string EnumerationOperation)
diff --git a/Engine/Source/Programs/AutomationTool/Turnkey/TurnkeySettings.cs b/Engine/Source/Programs/AutomationTool/Turnkey/TurnkeySettings.cs
--- a/Engine/Source/Programs/AutomationTool/Turnkey/TurnkeySettings.cs
+++ b/Engine/Source/Programs/AutomationTool/Turnkey/TurnkeySettings.cs
@@ -35,6 +35,7 @@ namespace Turnkey
 		public static UserSetting[] AllUserSettings = 
 		{
 			new UserSetting("User_QuickSwitchSdkLocation", null, "Location for downloaded Sdks to be stored (currently just for Google Drive, perforce will sync per clientspec)"),
+			new UserSetting("User_LastPerforceClient", null, "Perforce depot clientspec to look at first for finding matching clientspecs. Will be set automatically when one is found, but this can be used to force a particular client"),
 		};
 
 		// basically same as Turnkey variables, but this only contains ones that were loaded so we can write them back out `
diff --git a/Engine/Source/Programs/AutomationTool/Turnkey/TurnkeyUtils.cs b/Engine/Source/Programs/AutomationTool/Turnkey/TurnkeyUtils.cs
--- a/Engine/Source/Programs/AutomationTool/Turnkey/TurnkeyUtils.cs
+++ b/Engine/Source/Programs/AutomationTool/Turnkey/TurnkeyUtils.cs
@@ -16,11 +16,16 @@ namespace Turnkey
 		// object that commands, etc can use to access UAT functionality
 		static public BuildCommand CommandUtilHelper;
 
+		// replacement for Environment.ExitCode
+		static public AutomationTool.ExitCode ExitCode = ExitCode.Success;
+
 		static private IOProvider IOProvider;
 		static public void Initialize(IOProvider InIOProvider, BuildCommand InCommandUtilHelper)
 		{
 			IOProvider = InIOProvider;
 			CommandUtilHelper = InCommandUtilHelper;
+
+			SetVariable("HOST_PLATFORM_NAME", HostPlatform.Current.HostEditorPlatform.ToString());
 		}
 
 		static Dictionary<string, string> TurnkeyVariables = new Dictionary<string, string>();
@@ -83,11 +88,18 @@ namespace Turnkey
 		public static List<UnrealTargetPlatform> GetPlatformsFromCommandLineOrUser(string[] CommandOptions, List<UnrealTargetPlatform> PossiblePlatforms)
 		{
 			string PlatformString = TurnkeyUtils.ParseParamValue("Platform", null, CommandOptions);
+			bool bUnattended = TurnkeyUtils.ParseParam("Unattended", CommandOptions);
 
 			List<UnrealTargetPlatform> Platforms = new List<UnrealTargetPlatform>();
 			// prompt user for a platform
 			if (PlatformString == null)
 			{
+				if (bUnattended)
+				{
+					// can't ask
+					return null;
+				}
+
 				List<string> PlatformOptions = PossiblePlatforms.ConvertAll(x => x.ToString());
 				PlatformOptions.Add("All of the Above");
 				int PlatformChoice = TurnkeyUtils.ReadInputInt("Choose a platform:", PlatformOptions, true);
@@ -134,6 +146,45 @@ namespace Turnkey
 			return Platforms;
 		}
 
+		public static DeviceInfo GetDeviceFromCommandLineOrUser(string[] CommandOptions, UnrealTargetPlatform Platform)
+		{
+			string DeviceName = TurnkeyUtils.ParseParamValue("Device", null, CommandOptions);
+
+			AutomationTool.Platform AutomationPlatform = AutomationTool.Platform.Platforms[new TargetPlatformDescriptor(Platform)];
+			
+			if (DeviceName == null)
+			{
+				List<string> Options = new List<string>();
+				// @todo turnkey: filter devices that have valid flash versions
+				//				DeviceInfo[] PossibleDevices = Array.FindAll(AutomationPlatform.GetDevices(), x => TurnkeyUtils.IsValueValid(x.SoftwareVersion, AutomationPlatform.GetAllowedSoftwareVersions()));
+				DeviceInfo[] PossibleDevices = AutomationPlatform.GetDevices();
+				if (PossibleDevices == null)
+				{
+					return null;
+				}
+
+				foreach (DeviceInfo Device in PossibleDevices)
+				{
+					Options.Add(string.Format("[{0} {1}] {2}", Platform, Device.Type, Device.Name));
+				}
+
+				// get the choice
+				int Choice = TurnkeyUtils.ReadInputInt("Select the number of a device to flash:", Options, true);
+
+				if (Choice == 0)
+				{
+					return null;
+				}
+
+				// get the name of the device chosen
+				DeviceName = PossibleDevices[Choice - 1].Name;
+			}
+
+			// get device info of the chosen or supplied device
+			DeviceInfo InstallDevice = Array.Find(AutomationPlatform.GetDevices(), x => string.Compare(x.Name, DeviceName, true) == 0);
+			return InstallDevice;
+		}
+
 
 		private static IDictionary[] SavedEnvVars = new System.Collections.IDictionary[2];
 		private static Dictionary<string, string> EnvVarsToSaveToBatchFile = new Dictionary<string, string>();
