- Missing files

diff --git a/Engine/Source/Programs/AutomationTool/Turnkey/BuildInfo.cs b/Engine/Source/Programs/AutomationTool/Turnkey/BuildInfo.cs
new file mode 100644
--- /dev/null
+++ b/Engine/Source/Programs/AutomationTool/Turnkey/BuildInfo.cs
@@ -0,0 +1,26 @@
+using System;
+using System.Collections.Generic;
+using System.Linq;
+using System.Text;
+using System.Threading.Tasks;
+
+namespace Turnkey
+{
+	public class BuildSource
+	{
+ 		public string Project;
+		public CopyAndRun[] Sources;
+		public string BuildEnumerationSuffix;
+		// 		string[] Platforms;
+		// 		string Project;
+
+
+		internal void PostDeserialize()
+		{
+			if (Sources != null)
+			{
+				Array.ForEach(Sources, x => x.PostDeserialize());
+			}
+		}
+	}
+}
diff --git a/Engine/Source/Programs/AutomationTool/Turnkey/Commands/InstallBuild.cs b/Engine/Source/Programs/AutomationTool/Turnkey/Commands/InstallBuild.cs
new file mode 100644
--- /dev/null
+++ b/Engine/Source/Programs/AutomationTool/Turnkey/Commands/InstallBuild.cs
@@ -0,0 +1,164 @@
+using System;
+using System.Collections.Generic;
+using System.Linq;
+using Gauntlet;
+using UnrealBuildTool;
+using AutomationTool;
+
+namespace Turnkey.Commands
+{
+	class InstallBuild : TurnkeyCommand
+	{
+		protected override void Execute(string[] CommandOptions)
+		{
+		
+			// get all projects
+			List<BuildSource> AllSources = TurnkeyManifest.GetDiscoveredBuildSources();
+
+			// get distinct project names
+			List<string> AllProjects = AllSources.Select(x => x.Project).Distinct().ToList();
+
+			int Choice = TurnkeyUtils.ReadInputInt("Select a project:", AllProjects, true);
+			if (Choice == 0)
+			{
+				return;
+			}
+
+			string Project = AllProjects[Choice - 1];
+			
+			// go over sources for this project
+			List<BuildSource> ProjectSources = AllSources.Where(x => x.Project == Project).ToList();
+
+
+			// go over all possible sources, assume there is an wildcard with an expansion in it
+			// @todo turnkey: if there is no wildcard, then we could just use the final part of pathname in the results
+			Dictionary<string, Tuple<string, string>> BuildSourceToLocation = new Dictionary<string, Tuple<string, string>>();
+			foreach (BuildSource Source in ProjectSources)
+			{
+				foreach (CopyAndRun Copy in Source.Sources)
+				{
+					List<List<string>> Expansions = new List<List<string>>();
+					string[] EnumeratedSources = CopyProvider.ExecuteEnumerate(Copy.Copy, Expansions);
+					for (int Index = 0; Index < EnumeratedSources.Length; Index++)
+					{
+						// get a mapping of description (expanded bit) to location
+						BuildSourceToLocation.Add(Expansions[Index][0], new Tuple<string,string>(EnumeratedSources[Index], Source.BuildEnumerationSuffix));
+					}
+				}
+			}
+
+			if (BuildSourceToLocation.Count == 0)
+			{
+				TurnkeyUtils.Log("No builds found!");
+				return;
+			}
+
+			List<string> BuildKeys = BuildSourceToLocation.Keys.ToList();
+			Choice = TurnkeyUtils.ReadInputInt("Choose a build version", BuildKeys, true);
+			if (Choice == 0)
+			{
+				return;
+			}
+
+			// get the chosen path
+			string BuildLocation = BuildSourceToLocation[BuildKeys[Choice - 1]].Item1;
+			string EnumerationSuffix = BuildSourceToLocation[BuildKeys[Choice - 1]].Item2;
+
+			// now find platforms under it
+			List<List<string>> PlatformExpansion = new List<List<string>>();
+			string[] PlatformSources = CopyProvider.ExecuteEnumerate(BuildLocation + EnumerationSuffix, PlatformExpansion);
+
+			if (PlatformSources.Length == 0)
+			{
+				TurnkeyUtils.Log("No platform builds found in {0}!", BuildLocation);
+				return;
+			}
+
+			List<string> CollapsedPlatformNames = PlatformExpansion.Select(x => x[0]).ToList();
+			Choice = TurnkeyUtils.ReadInputInt("Choose a platform", CollapsedPlatformNames, true);
+			if (Choice == 0)
+			{
+				return;
+			}
+
+			string LocalPath = CopyProvider.ExecuteCopy(PlatformSources[Choice - 1]);
+
+			// @todo turnkey this could be improved i'm sure
+			string TargetPlatformString = CollapsedPlatformNames[Choice - 1];
+			UnrealTargetPlatform Platform;
+			if (!UnrealTargetPlatform.TryParse(TargetPlatformString, out Platform))
+			{
+				Action<string> FindAndStrip = x => { int Loc = TargetPlatformString.IndexOf(x); if (Loc > 0) { TargetPlatformString = TargetPlatformString.Substring(0, Loc); } };
+
+				FindAndStrip("NoEditor");
+				FindAndStrip("Client");
+				FindAndStrip("Server");
+
+				// Windows hack
+				if (TargetPlatformString == "Windows")
+				{
+					TargetPlatformString = "Win64";
+				}
+
+				if (!UnrealTargetPlatform.TryParse(TargetPlatformString, out Platform))
+				{
+					TurnkeyUtils.Log("Unable to figure out the UnrealTargetPlatform from {0}!", CollapsedPlatformNames[Choice - 1]);
+					return;
+				}
+			}
+
+			// find all build sources that can be created a folder path
+			IEnumerable<IFolderBuildSource> BuildSources = Gauntlet.Utils.InterfaceHelpers.FindImplementations<IFolderBuildSource>();
+			List<IBuild> Builds = BuildSources.Where(S => S.CanSupportPlatform(Platform)).SelectMany(S => S.GetBuildsAtPath(Project, LocalPath)).ToList();
+			if (Builds.Count() == 0)
+			{
+				throw new AutomationException("No builds for {0} found at {1}", Platform, LocalPath);
+			}
+
+			// finally choose a build
+			List<string> Options = Builds.Select(x => string.Format("{0} {1} {2} {3}", x.Platform, x.Configuration, x.Flags, x.GetType())).ToList();
+			Choice = TurnkeyUtils.ReadInputInt("Choose a build", Options, true);
+			if (Choice == 0)
+			{
+				return;
+			}
+
+			IBuild Build = Builds[Choice - 1];
+
+			// install the build and run it
+			UnrealAppConfig Config = new UnrealAppConfig();
+
+			if (Build.Flags.HasFlag(BuildFlags.CanReplaceCommandLine))
+			{
+				Config.CommandLine = TurnkeyUtils.ReadInput("Enter replacement commandline. Leave blank for original, or space to wipe it out", "");
+			}
+			Config.Build = Build;
+			Config.ProjectName = Project;
+
+			DeviceInfo TurnkeyDevice = TurnkeyUtils.GetDeviceFromCommandLineOrUser(CommandOptions, Platform);
+
+			ITargetDevice GauntletDevice;
+			if (TurnkeyDevice == null)
+			{
+				IEnumerable<IDefaultDeviceSource> DeviceSources = Gauntlet.Utils.InterfaceHelpers.FindImplementations<IDefaultDeviceSource>();
+				GauntletDevice = DeviceSources.Where(S => S.CanSupportPlatform(Platform)).SelectMany(S => S.GetDefaultDevices()).FirstOrDefault();
+			}
+			else
+			{
+				IDeviceFactory Factory = Gauntlet.Utils.InterfaceHelpers.FindImplementations<IDeviceFactory>()
+					.Where(F => F.CanSupportPlatform(Platform))
+					.FirstOrDefault();
+				GauntletDevice = Factory.CreateDevice(TurnkeyDevice.Name, null);
+			}
+
+			if (GauntletDevice == null)
+			{
+				TurnkeyUtils.Log("Could not find a device to install on!");
+				return;
+			}
+
+			IAppInstall Install = GauntletDevice.InstallApplication(Config);
+			GauntletDevice.Run(Install);
+		}
+	}
+}
