- Fixed Turnkey copyrights
- Latest changes

diff --git a/Engine/Source/Programs/AutomationTool/Turnkey/BuildInfo.cs b/Engine/Source/Programs/AutomationTool/Turnkey/BuildInfo.cs
--- a/Engine/Source/Programs/AutomationTool/Turnkey/BuildInfo.cs
+++ b/Engine/Source/Programs/AutomationTool/Turnkey/BuildInfo.cs
@@ -1,3 +1,5 @@
+// Copyright Epic Games, Inc. All Rights Reserved.
+
 using System;
 using System.Collections.Generic;
 using System.Linq;
diff --git a/Engine/Source/Programs/AutomationTool/Turnkey/Commands/InstallBuild.cs b/Engine/Source/Programs/AutomationTool/Turnkey/Commands/InstallBuild.cs
--- a/Engine/Source/Programs/AutomationTool/Turnkey/Commands/InstallBuild.cs
+++ b/Engine/Source/Programs/AutomationTool/Turnkey/Commands/InstallBuild.cs
@@ -1,3 +1,5 @@
+// Copyright Epic Games, Inc. All Rights Reserved.
+
 using System;
 using System.Collections.Generic;
 using System.Linq;
diff --git a/Engine/Source/Programs/AutomationTool/Turnkey/Commands/InstallSdk.cs b/Engine/Source/Programs/AutomationTool/Turnkey/Commands/InstallSdk.cs
--- a/Engine/Source/Programs/AutomationTool/Turnkey/Commands/InstallSdk.cs
+++ b/Engine/Source/Programs/AutomationTool/Turnkey/Commands/InstallSdk.cs
@@ -1,3 +1,5 @@
+// Copyright Epic Games, Inc. All Rights Reserved.
+
 using System;
 using System.Collections.Generic;
 using System.Text;
diff --git a/Engine/Source/Programs/AutomationTool/Turnkey/Commands/ListPlatforms.cs b/Engine/Source/Programs/AutomationTool/Turnkey/Commands/ListPlatforms.cs
--- a/Engine/Source/Programs/AutomationTool/Turnkey/Commands/ListPlatforms.cs
+++ b/Engine/Source/Programs/AutomationTool/Turnkey/Commands/ListPlatforms.cs
@@ -1,3 +1,5 @@
+// Copyright Epic Games, Inc. All Rights Reserved.
+
 using System;
 using System.Collections.Generic;
 using UnrealBuildTool;
diff --git a/Engine/Source/Programs/AutomationTool/Turnkey/Commands/ListSdks.cs b/Engine/Source/Programs/AutomationTool/Turnkey/Commands/ListSdks.cs
--- a/Engine/Source/Programs/AutomationTool/Turnkey/Commands/ListSdks.cs
+++ b/Engine/Source/Programs/AutomationTool/Turnkey/Commands/ListSdks.cs
@@ -1,3 +1,5 @@
+// Copyright Epic Games, Inc. All Rights Reserved.
+
 using System;
 using System.Collections.Generic;
 using UnrealBuildTool;
diff --git a/Engine/Source/Programs/AutomationTool/Turnkey/Commands/ManageSettings.cs b/Engine/Source/Programs/AutomationTool/Turnkey/Commands/ManageSettings.cs
--- a/Engine/Source/Programs/AutomationTool/Turnkey/Commands/ManageSettings.cs
+++ b/Engine/Source/Programs/AutomationTool/Turnkey/Commands/ManageSettings.cs
@@ -1,3 +1,5 @@
+// Copyright Epic Games, Inc. All Rights Reserved.
+
 using System;
 using System.Collections.Generic;
 using System.Linq;
diff --git a/Engine/Source/Programs/AutomationTool/Turnkey/Commands/Test.cs b/Engine/Source/Programs/AutomationTool/Turnkey/Commands/Test.cs
--- a/Engine/Source/Programs/AutomationTool/Turnkey/Commands/Test.cs
+++ b/Engine/Source/Programs/AutomationTool/Turnkey/Commands/Test.cs
@@ -1,3 +1,5 @@
+// Copyright Epic Games, Inc. All Rights Reserved.
+
 using System;
 using System.Collections.Generic;
 using System.Linq;
diff --git a/Engine/Source/Programs/AutomationTool/Turnkey/Commands/VerifySdk.cs b/Engine/Source/Programs/AutomationTool/Turnkey/Commands/VerifySdk.cs
--- a/Engine/Source/Programs/AutomationTool/Turnkey/Commands/VerifySdk.cs
+++ b/Engine/Source/Programs/AutomationTool/Turnkey/Commands/VerifySdk.cs
@@ -1,3 +1,5 @@
+// Copyright Epic Games, Inc. All Rights Reserved.
+
 using AutomationTool;
 using System;
 using System.Collections.Generic;
@@ -33,14 +35,14 @@ namespace Turnkey.Commands
 
 				SdkInfo.LocalAvailability LocalState = SdkInfo.GetLocalAvailability(AutomationPlatform);
 
-				if ((LocalState & (SdkInfo.LocalAvailability.AutoSdk_PlatformExists | SdkInfo.LocalAvailability.InstalledSdk_ValidVersionExists)) == 0)
+				if ((LocalState & (SdkInfo.LocalAvailability.AutoSdk_ValidVersionExists | SdkInfo.LocalAvailability.InstalledSdk_ValidVersionExists)) == 0)
 				{
-					TurnkeyUtils.Log("{0}: No AutoSdk or Installed Sdk found matching {1}", Platform, AutomationPlatform.GetAllowedSdks());
+					TurnkeyUtils.Report("[{0}]: Invalid [No AutoSdk or Installed Sdk found matching {1}]", Platform, AutomationPlatform.GetAllowedSdks());
 					TurnkeyUtils.ExitCode = AutomationTool.ExitCode.Error_SDKNotFound;
 				}
 				else
 				{
-					TurnkeyUtils.Log("{0}: Valid! [{1}]", Platform, (LocalState & (SdkInfo.LocalAvailability.AutoSdk_PlatformExists | SdkInfo.LocalAvailability.InstalledSdk_ValidVersionExists)).ToString());
+					TurnkeyUtils.Report("[{0}]: Valid [{1}]", Platform, (LocalState & (SdkInfo.LocalAvailability.AutoSdk_ValidVersionExists | SdkInfo.LocalAvailability.InstalledSdk_ValidVersionExists)).ToString());
 					//					TurnkeyUtils.Log("{0}: Valid [Installed: '{1}', Required: '{2}']", Platform, PlatformObject.GetInstalledSdk(), PlatformObject.GetAllowedSdks());
 				}
 
diff --git a/Engine/Source/Programs/AutomationTool/Turnkey/CopyAndRun.cs b/Engine/Source/Programs/AutomationTool/Turnkey/CopyAndRun.cs
--- a/Engine/Source/Programs/AutomationTool/Turnkey/CopyAndRun.cs
+++ b/Engine/Source/Programs/AutomationTool/Turnkey/CopyAndRun.cs
@@ -1,3 +1,5 @@
+// Copyright Epic Games, Inc. All Rights Reserved.
+
 using System;
 using System.Xml.Serialization;
 using System.IO;
diff --git a/Engine/Source/Programs/AutomationTool/Turnkey/CopyProvider.cs b/Engine/Source/Programs/AutomationTool/Turnkey/CopyProvider.cs
--- a/Engine/Source/Programs/AutomationTool/Turnkey/CopyProvider.cs
+++ b/Engine/Source/Programs/AutomationTool/Turnkey/CopyProvider.cs
@@ -1,3 +1,5 @@
+// Copyright Epic Games, Inc. All Rights Reserved.
+
 using System;
 using System.Collections.Generic;
 using System.Linq;
diff --git a/Engine/Source/Programs/AutomationTool/Turnkey/IOProvider.cs b/Engine/Source/Programs/AutomationTool/Turnkey/IOProvider.cs
--- a/Engine/Source/Programs/AutomationTool/Turnkey/IOProvider.cs
+++ b/Engine/Source/Programs/AutomationTool/Turnkey/IOProvider.cs
@@ -1,3 +1,5 @@
+// Copyright Epic Games, Inc. All Rights Reserved.
+
 using System;
 using System.Collections.Generic;
 using System.Linq;
@@ -9,7 +11,14 @@ namespace Turnkey
 	abstract class IOProvider
 	{
 		public abstract void Log(string Message, bool bAppendNewLine);
+		
+		public virtual void Report(string Message, bool bAppendNewLine)
+		{
+			Log(Message, bAppendNewLine);
+		}
+
 		public abstract string ReadInput(string Prompt, string DefaultValue, bool bAppendNewLine);
+		
 		public abstract int ReadInputInt(string Prompt, List<string> Options, bool bIsCancellable, int DefaultValue, bool bAppendNewLine);
 	}
 }
diff --git a/Engine/Source/Programs/AutomationTool/Turnkey/LocalCache.cs b/Engine/Source/Programs/AutomationTool/Turnkey/LocalCache.cs
--- a/Engine/Source/Programs/AutomationTool/Turnkey/LocalCache.cs
+++ b/Engine/Source/Programs/AutomationTool/Turnkey/LocalCache.cs
@@ -1,3 +1,5 @@
+// Copyright Epic Games, Inc. All Rights Reserved.
+
 using System;
 using System.Collections.Generic;
 using System.Runtime.Serialization;
diff --git a/Engine/Source/Programs/AutomationTool/Turnkey/Providers/ConsoleIOProvider.cs b/Engine/Source/Programs/AutomationTool/Turnkey/Providers/ConsoleIOProvider.cs
--- a/Engine/Source/Programs/AutomationTool/Turnkey/Providers/ConsoleIOProvider.cs
+++ b/Engine/Source/Programs/AutomationTool/Turnkey/Providers/ConsoleIOProvider.cs
@@ -1,3 +1,5 @@
+// Copyright Epic Games, Inc. All Rights Reserved.
+
 using System;
 using System.Collections.Generic;
 using System.Text;
diff --git a/Engine/Source/Programs/AutomationTool/Turnkey/Providers/GoogleDriveCopyProvider.cs b/Engine/Source/Programs/AutomationTool/Turnkey/Providers/GoogleDriveCopyProvider.cs
--- a/Engine/Source/Programs/AutomationTool/Turnkey/Providers/GoogleDriveCopyProvider.cs
+++ b/Engine/Source/Programs/AutomationTool/Turnkey/Providers/GoogleDriveCopyProvider.cs
@@ -1,3 +1,5 @@
+// Copyright Epic Games, Inc. All Rights Reserved.
+
 using System;
 using System.Collections.Generic;
 using System.Linq;
diff --git a/Engine/Source/Programs/AutomationTool/Turnkey/Providers/HybridIOProvider.cs b/Engine/Source/Programs/AutomationTool/Turnkey/Providers/HybridIOProvider.cs
new file mode 100644
--- /dev/null
+++ b/Engine/Source/Programs/AutomationTool/Turnkey/Providers/HybridIOProvider.cs
@@ -0,0 +1,184 @@
+// Copyright Epic Games, Inc. All Rights Reserved.
+
+using System;
+using System.Collections.Generic;
+using System.Linq;
+using System.Text;
+using System.Threading.Tasks;
+using System.Windows.Forms;
+using System.Drawing;
+using System.Drawing.Configuration;
+using Tools.DotNETCommon;
+
+namespace Turnkey
+{
+	class HybridIOProvider : ConsoleIOProvider
+	{
+		// @todo turnkey: if this is really bad on Mono, then we can set it in the manifest, but that means changing AutomationTool's and ATLaunchers' manifests
+		[System.Runtime.InteropServices.DllImport("user32.dll")]
+		private static extern bool SetProcessDPIAware();
+
+		static HybridIOProvider()
+		{
+			// make the form look good on modern displays!
+			if (Environment.OSVersion.Version.Major >= 6)
+			{
+				SetProcessDPIAware();
+			}
+
+			Application.EnableVisualStyles();
+			Application.SetCompatibleTextRenderingDefault(false);
+		}
+
+
+		private  string ShowDialog(string Prompt, string Default, bool bIsList)
+		{
+			// Create a new instance of the form.
+			Form Form1 = new Form();
+
+			System.Windows.Forms.TextBox Label;
+			System.Windows.Forms.Button CancelBtn;
+			System.Windows.Forms.Button OkBtn;
+			System.Windows.Forms.TextBox TextBox;
+
+			TextBox = new System.Windows.Forms.TextBox();
+			Label = new System.Windows.Forms.TextBox();
+			CancelBtn = new System.Windows.Forms.Button();
+			OkBtn = new System.Windows.Forms.Button();
+			Form1.SuspendLayout();
+			// 
+			// Label
+			// 
+			Label.Anchor = ((System.Windows.Forms.AnchorStyles)(((System.Windows.Forms.AnchorStyles.Top | System.Windows.Forms.AnchorStyles.Left)
+			| System.Windows.Forms.AnchorStyles.Right | System.Windows.Forms.AnchorStyles.Bottom)));
+			Label.BackColor = System.Drawing.SystemColors.Control;
+			Label.BorderStyle = System.Windows.Forms.BorderStyle.None;
+			Label.Location = new System.Drawing.Point(13, 13);
+			Label.Multiline = true;
+			Label.Name = "Label";
+			Label.ScrollBars = System.Windows.Forms.ScrollBars.Vertical;
+			Label.ShortcutsEnabled = false;
+			Label.Size = new System.Drawing.Size(775, 199);
+			Label.TabIndex = 0;
+			Label.TabStop = false;
+			Label.Text = Prompt;
+
+			// 
+			// CancelBtn
+			// 
+			CancelBtn.Anchor = ((System.Windows.Forms.AnchorStyles)((System.Windows.Forms.AnchorStyles.Right | System.Windows.Forms.AnchorStyles.Bottom)));
+			CancelBtn.DialogResult = System.Windows.Forms.DialogResult.Cancel;
+			CancelBtn.Location = new System.Drawing.Point(679, 262);
+			CancelBtn.Margin = new System.Windows.Forms.Padding(4);
+			CancelBtn.Name = "CancelBtn";
+			CancelBtn.Size = new System.Drawing.Size(109, 34);
+			CancelBtn.TabIndex = 6;
+			CancelBtn.Text = "Cancel";
+			CancelBtn.UseVisualStyleBackColor = true;
+			// 
+			// TextBox
+			// 
+			TextBox.Anchor = ((System.Windows.Forms.AnchorStyles)(((System.Windows.Forms.AnchorStyles.Left | System.Windows.Forms.AnchorStyles.Bottom)
+			| System.Windows.Forms.AnchorStyles.Right)));
+			TextBox.Location = new System.Drawing.Point(13, 219);
+			TextBox.Margin = new System.Windows.Forms.Padding(4);
+			TextBox.Name = "TextBox";
+			TextBox.Size = new System.Drawing.Size(772, 22);
+			TextBox.TabIndex = 4;
+			TextBox.Text = Default;
+			// 
+			// OkBtn
+			// 
+			OkBtn.Anchor = ((System.Windows.Forms.AnchorStyles)((System.Windows.Forms.AnchorStyles.Right | System.Windows.Forms.AnchorStyles.Bottom)));
+			OkBtn.DialogResult = System.Windows.Forms.DialogResult.OK;
+			OkBtn.Location = new System.Drawing.Point(562, 262);
+			OkBtn.Margin = new System.Windows.Forms.Padding(4);
+			OkBtn.Name = "OkBtn";
+			OkBtn.Size = new System.Drawing.Size(109, 34);
+			OkBtn.TabIndex = 5;
+			OkBtn.Text = "Ok";
+			OkBtn.UseVisualStyleBackColor = true;
+			// 
+			// Form1
+			// 
+			Form1.AcceptButton = OkBtn;
+			Form1.StartPosition = FormStartPosition.CenterScreen;
+			Form1.AutoScaleDimensions = new System.Drawing.SizeF(8F, 16F);
+			Form1.AutoScaleMode = System.Windows.Forms.AutoScaleMode.Font;
+			Form1.CancelButton = CancelBtn;
+			Form1.ClientSize = new System.Drawing.Size(800, 315);
+			Form1.Controls.Add(CancelBtn);
+			Form1.Controls.Add(OkBtn);
+			Form1.Controls.Add(TextBox);
+			Form1.Controls.Add(Label);
+			Form1.Name = "Form1";
+			Form1.Text = "Make selection";
+			Form1.ResumeLayout(false);
+			Form1.PerformLayout();
+
+			if (!bIsList)
+			{
+				Form1.Height -= 150;
+				Form1.PerformLayout();
+			}
+
+			// Display the form as a modal dialog box.
+			Form1.ShowDialog();
+
+			// Determine if the OK button was clicked on the dialog box.
+			if (Form1.DialogResult == DialogResult.OK)
+			{
+				return TextBox.Text;
+			}
+
+			return null;
+		}
+
+		public override string ReadInput(string Prompt, string Default, bool bAppendNewLine)
+		{
+			return ShowDialog(Prompt, Default, false);
+		}
+
+		public override int ReadInputInt(string Prompt, List<string> Options, bool bIsCancellable, int DefaultValue, bool bAppendNewLine)
+		{
+			StringBuilder FullPromptBuilder = new StringBuilder();
+
+			// start with given prompt
+			FullPromptBuilder.Append(Prompt);
+			if (bAppendNewLine)
+			{
+				FullPromptBuilder.AppendLine("");
+			}
+
+			// now add the options given
+			int Index = 1;
+			foreach (string Option in Options)
+			{
+				FullPromptBuilder.AppendLine(" ({0}) {1}", Index++, Option);
+			}
+
+			string FullPrompt = FullPromptBuilder.ToString();
+
+			// go until good choice
+			while (true)
+			{
+				string ChoiceString = ShowDialog(FullPrompt, DefaultValue >= 0 ? DefaultValue.ToString() : "", true);
+
+				if (ChoiceString == null)
+				{
+					return bIsCancellable ? 0 : -1;
+				}
+
+				int Choice;
+				if (Int32.TryParse(ChoiceString, out Choice) == false || Choice < 0 || Choice >= Options.Count + (bIsCancellable ? 1 : 0))
+				{
+					TurnkeyUtils.Log("Invalid choice");
+				}
+				else
+				{
+					return Choice;
+				}
+			}
+		}
+	}
+}
diff --git a/Engine/Source/Programs/AutomationTool/Turnkey/Providers/NullCopyProvider.cs b/Engine/Source/Programs/AutomationTool/Turnkey/Providers/NullCopyProvider.cs
--- a/Engine/Source/Programs/AutomationTool/Turnkey/Providers/NullCopyProvider.cs
+++ b/Engine/Source/Programs/AutomationTool/Turnkey/Providers/NullCopyProvider.cs
@@ -1,3 +1,5 @@
+// Copyright Epic Games, Inc. All Rights Reserved.
+
 using System.IO;
 using System.Collections.Generic;
 using System.Linq;
diff --git a/Engine/Source/Programs/AutomationTool/Turnkey/Providers/PerforceCopyProvider.cs b/Engine/Source/Programs/AutomationTool/Turnkey/Providers/PerforceCopyProvider.cs
--- a/Engine/Source/Programs/AutomationTool/Turnkey/Providers/PerforceCopyProvider.cs
+++ b/Engine/Source/Programs/AutomationTool/Turnkey/Providers/PerforceCopyProvider.cs
@@ -1,3 +1,5 @@
+// Copyright Epic Games, Inc. All Rights Reserved.
+
 using System;
 using System.Collections.Generic;
 using System.Linq;
@@ -20,7 +22,6 @@ namespace Turnkey
 				return null;
 			}
 
-
 			// calculate the output path
 			// @todo turnkey: use p4 where, because a depot mapping could be strange and make this invlalid
 			string OutputPath = Operation.Replace(ConnectedRoot, PerforceClient.RootPath);
@@ -162,6 +163,7 @@ namespace Turnkey
 			// Get all clients for this user
 			string P4Command = String.Format("clients -u {0}", Username);
 
+			// @todo turnkey: sort the results on third token (date)
 			var P4Result = PerforceConnection.P4(string.Format("clients -u {0}", Username), AllowSpew: false, WithClient: false);
 			if (P4Result.ExitCode != 0)
 			{
@@ -238,7 +240,15 @@ namespace Turnkey
 					string ClientName = TurnkeyUtils.ReadInput("Enter clientspec name:", string.Format("{0}_sdks", Username));
 
 					// get local pathname
-					string LocalPath = TurnkeyUtils.ReadInput(string.Format("Enter where to map {0} on your computer:", bIsStream ? SpecRoot : DepotName), @"D:\Sdks");
+					// if we had UE_SDKS_ROOT already set, we use it as a default
+					string AutoSdksDir = Environment.GetEnvironmentVariable("UE_SDKS_ROOT");
+					string DefaultDir = @"D:\Sdks";
+					if (!string.IsNullOrEmpty(AutoSdksDir))
+					{
+						// @todo turnkey - for a deeper directory structure, we should match directory names to the perforce Operation, to figure it out properly
+						DefaultDir = Path.GetDirectoryName(AutoSdksDir);
+					}
+					string LocalPath = TurnkeyUtils.ReadInput(string.Format("Enter where to map {0} on your computer:", bIsStream ? SpecRoot : DepotName), DefaultDir);
 
 					// create a client from the input settings
 					P4ClientInfo NewClient = new P4ClientInfo();
@@ -246,17 +256,15 @@ namespace Turnkey
 					NewClient.Owner = Username;
 					NewClient.Host = Hostname;
 					NewClient.RootPath = LocalPath;
-					// @todo turnkey handle depot type :|
+					NewClient.Options |= P4ClientOption.RmDir;
 					if (bIsStream)
 					{
 						NewClient.Stream = SpecRoot;
 					}
 					else
 					{
-						// set up mapping
-						string DepotPath = SpecRoot + "/...";
-						string ClientPath = DepotPath.Replace("//" + DepotName + "/", "//" + ClientName + "/");
-						NewClient.View.Add(new KeyValuePair<string, string>(DepotPath, ClientPath));
+						// set up mapping (CreateClient prepends the Key with //<clientname>)
+						NewClient.View.Add(new KeyValuePair<string, string>(DepotName + "/...", "/..."));
 					}
 
 					PerforceClient = PerforceConnection.CreateClient(NewClient);
diff --git a/Engine/Source/Programs/AutomationTool/Turnkey/SdkInfo.cs b/Engine/Source/Programs/AutomationTool/Turnkey/SdkInfo.cs
--- a/Engine/Source/Programs/AutomationTool/Turnkey/SdkInfo.cs
+++ b/Engine/Source/Programs/AutomationTool/Turnkey/SdkInfo.cs
@@ -1,3 +1,5 @@
+// Copyright Epic Games, Inc. All Rights Reserved.
+
 using System;
 using System.Collections.Generic;
 using System.Text;
@@ -104,10 +106,11 @@ namespace Turnkey
 		{
 			None = 0,
 			AutoSdk_VariableExists = 1,
-			AutoSdk_PlatformExists = 2,
-			// InstalledSdk_BuildOnlyWasInstalled = 4,
-			InstalledSdk_ValidVersionExists = 8,
-			InstalledSdk_InvalidVersionExists = 16,
+			AutoSdk_ValidVersionExists = 2,
+			AutoSdk_InvalidVersionExists = 4,
+			// InstalledSdk_BuildOnlyWasInstalled = 8,
+			InstalledSdk_ValidVersionExists = 16,
+			InstalledSdk_InvalidVersionExists = 32,
 		}
 
 		static public LocalAvailability GetLocalAvailability(AutomationTool.Platform AutomationPlatform)
@@ -125,6 +128,7 @@ namespace Turnkey
 				DirectoryInfo PlatformDir = new DirectoryInfo(Path.Combine(AutoSdkVar, AutoSubdir));
 				if (PlatformDir.Exists)
 				{
+					bool bValidVersionFound = false;
 					foreach (DirectoryInfo Dir in PlatformDir.EnumerateDirectories())
 					{
 						// check for valid version number
@@ -133,11 +137,14 @@ namespace Turnkey
 							// make sure it actually has bits in it
 							if (File.Exists(Path.Combine(Dir.FullName, "setup.bat")))
 							{
-								Result |= LocalAvailability.AutoSdk_PlatformExists;
+								bValidVersionFound = true;
+								break;
 							}
-							break;
 						}
 					}
+
+					// if we had a platform directory, but no valid versions, note that
+					Result |= bValidVersionFound ? LocalAvailability.AutoSdk_ValidVersionExists : LocalAvailability.AutoSdk_InvalidVersionExists;
 				}
 			}
 
diff --git a/Engine/Source/Programs/AutomationTool/Turnkey/Turnkey.Automation.cs b/Engine/Source/Programs/AutomationTool/Turnkey/Turnkey.Automation.cs
--- a/Engine/Source/Programs/AutomationTool/Turnkey/Turnkey.Automation.cs
+++ b/Engine/Source/Programs/AutomationTool/Turnkey/Turnkey.Automation.cs
@@ -1,4 +1,3 @@
-
 // Copyright Epic Games, Inc. All Rights Reserved.
 
 using System;
@@ -13,7 +12,25 @@ namespace Turnkey
 	{
 		public override ExitCode Execute()
 		{
-			return Turnkey.Execute(new ConsoleIOProvider(), this);
+			IOProvider IOProvider;
+			if (ParseParam("EditorIO"))
+			{
+				IOProvider = new HybridIOProvider();
+			}
+			else
+			{
+				string ReportFilename = ParseParamValue("ReportFilename");
+				if (!string.IsNullOrEmpty(ReportFilename))
+				{
+					IOProvider = new ReportIOProvider(ReportFilename);
+				}
+				else
+				{
+					IOProvider = new ConsoleIOProvider();
+				}
+			}
+
+			return Turnkey.Execute(IOProvider, this);
 		}
 
 		public static AutomationTool.ExitCode Execute(IOProvider IOProvider, BuildCommand CommandUtilHelper)
diff --git a/Engine/Source/Programs/AutomationTool/Turnkey/TurnkeyCommand.cs b/Engine/Source/Programs/AutomationTool/Turnkey/TurnkeyCommand.cs
--- a/Engine/Source/Programs/AutomationTool/Turnkey/TurnkeyCommand.cs
+++ b/Engine/Source/Programs/AutomationTool/Turnkey/TurnkeyCommand.cs
@@ -1,3 +1,5 @@
+// Copyright Epic Games, Inc. All Rights Reserved.
+
 using System;
 using System.Collections.Generic;
 using System.Reflection;
diff --git a/Engine/Source/Programs/AutomationTool/Turnkey/TurnkeyManifest.cs b/Engine/Source/Programs/AutomationTool/Turnkey/TurnkeyManifest.cs
--- a/Engine/Source/Programs/AutomationTool/Turnkey/TurnkeyManifest.cs
+++ b/Engine/Source/Programs/AutomationTool/Turnkey/TurnkeyManifest.cs
@@ -1,3 +1,5 @@
+// Copyright Epic Games, Inc. All Rights Reserved.
+
 using System;
 using System.Collections.Generic;
 using System.Text;
@@ -60,6 +62,12 @@ namespace Turnkey
 						List<List<string>> Expansions = new List<List<string>>();
 						string[] ExpandedInstallerResults = CopyProvider.ExecuteEnumerate(Sdk.Expansion[0].Copy, Expansions);
 
+						// expansion may not work
+						if (ExpandedInstallerResults == null)
+						{
+							continue;
+						}
+
 						if (Expansions.Count != ExpandedInstallerResults.Length)
 						{
 							throw new AutomationException("Bad expansions output from CopyProvider");
diff --git a/Engine/Source/Programs/AutomationTool/Turnkey/TurnkeySettings.cs b/Engine/Source/Programs/AutomationTool/Turnkey/TurnkeySettings.cs
--- a/Engine/Source/Programs/AutomationTool/Turnkey/TurnkeySettings.cs
+++ b/Engine/Source/Programs/AutomationTool/Turnkey/TurnkeySettings.cs
@@ -1,3 +1,5 @@
+// Copyright Epic Games, Inc. All Rights Reserved.
+
 using System;
 using System.IO;
 using System.Collections.Generic;
@@ -104,7 +106,7 @@ namespace Turnkey
 			string PreviousValue = null;
 
 			SetUserSettings.TryGetValue(VariableName, out PreviousValue);
-			SetUserSettings.Add(VariableName, Value);
+			SetUserSettings[VariableName] = Value;
 			Save();
 
 			return PreviousValue;
diff --git a/Engine/Source/Programs/AutomationTool/Turnkey/TurnkeyUtils.cs b/Engine/Source/Programs/AutomationTool/Turnkey/TurnkeyUtils.cs
--- a/Engine/Source/Programs/AutomationTool/Turnkey/TurnkeyUtils.cs
+++ b/Engine/Source/Programs/AutomationTool/Turnkey/TurnkeyUtils.cs
@@ -1,3 +1,5 @@
+// Copyright Epic Games, Inc. All Rights Reserved.
+
 using System;
 using System.Collections;
 using System.Collections.Generic;
@@ -19,7 +21,6 @@ namespace Turnkey
 		// replacement for Environment.ExitCode
 		static public AutomationTool.ExitCode ExitCode = ExitCode.Success;
 
-		static private IOProvider IOProvider;
 		static public void Initialize(IOProvider InIOProvider, BuildCommand InCommandUtilHelper)
 		{
 			IOProvider = InIOProvider;
@@ -28,6 +29,8 @@ namespace Turnkey
 			SetVariable("HOST_PLATFORM_NAME", HostPlatform.Current.HostEditorPlatform.ToString());
 		}
 
+		#region Turnkey Variables
+
 		static Dictionary<string, string> TurnkeyVariables = new Dictionary<string, string>();
 
 		public static string SetVariable(string Key, string Value)
@@ -65,7 +68,9 @@ namespace Turnkey
 			return UnrealBuildTool.Utils.ExpandVariables(ExpandedUserVariables, TurnkeyVariables, bUseOnlyTurnkeyVariables);
 		}
 
+		#endregion
 
+		#region Commandline Handling
 
 		public static bool ParseParam(string Param, string[] ExtraOptions)
 		{
@@ -132,7 +137,7 @@ namespace Turnkey
 					if (!UnrealTargetPlatform.TryParse(Token, out Platform))
 					{
 						TurnkeyUtils.Log("Platform {0} is unknown", Token);
-						return null;
+						continue;
 					}
 
 					// if the platform isn't in the possible list, then don't add it
@@ -184,7 +189,9 @@ namespace Turnkey
 			DeviceInfo InstallDevice = Array.Find(AutomationPlatform.GetDevices(), x => string.Compare(x.Name, DeviceName, true) == 0);
 			return InstallDevice;
 		}
+		#endregion
 
+		#region Env vars
 
 		private static IDictionary[] SavedEnvVars = new System.Collections.IDictionary[2];
 		private static Dictionary<string, string> EnvVarsToSaveToBatchFile = new Dictionary<string, string>();
@@ -245,6 +252,10 @@ namespace Turnkey
 			}
 		}
 
+		#endregion
+
+		#region Regex Matching
+
 		static bool TryConvertToUint64(string InValue, out UInt64 OutValue)
 		{ 
 			if (InValue.StartsWith("0x"))
@@ -311,6 +322,11 @@ namespace Turnkey
 			return string.Compare(Value, AllowedValues, true) == 0;
 		}
 
+		#endregion
+
+		#region IO
+		static private IOProvider IOProvider;
+
 		public static void Log(string Message)
 		{
 			IOProvider.Log(Message, bAppendNewLine: true);
@@ -325,6 +341,20 @@ namespace Turnkey
 			IOProvider.Log(string.Format(Message, Params), bAppendNewLine: true);
 		}
 
+		public static void Report(string Message)
+		{
+			IOProvider.Report(Message, bAppendNewLine: true);
+		}
+		public static void Report(ref StringBuilder Message)
+		{
+			IOProvider.Report(Message.ToString(), bAppendNewLine: false);
+			Message.Clear();
+		}
+		public static void Report(string Message, params object[] Params)
+		{
+			IOProvider.Report(string.Format(Message, Params), bAppendNewLine: true);
+		}
+
 		public static string ReadInput(string Prompt, string Default = "")
 		{
 			return IOProvider.ReadInput(Prompt, Default, bAppendNewLine: true);
@@ -339,7 +369,9 @@ namespace Turnkey
 		{
 			return IOProvider.ReadInputInt(Prompt, Options, bIsCancellable, DefaultValue, bAppendNewLine: true);
 		}
+		#endregion
 
+		#region Temp files [move to LocalCache]
 		static List<string> PathsToCleanup = new List<string>();
 		public static void AddPathToCleanup(string Path)
 		{
@@ -365,5 +397,6 @@ namespace Turnkey
 
 			PathsToCleanup.Clear();
 		}
+		#endregion
 	}
 }
