- Missing file in checkin

diff --git a/Engine/Source/Programs/AutomationTool/Turnkey/Providers/ReportIOProvider.cs b/Engine/Source/Programs/AutomationTool/Turnkey/Providers/ReportIOProvider.cs
new file mode 100644
--- /dev/null
+++ b/Engine/Source/Programs/AutomationTool/Turnkey/Providers/ReportIOProvider.cs
@@ -0,0 +1,43 @@
+// Copyright Epic Games, Inc. All Rights Reserved.
+
+using System;
+using System.Collections.Generic;
+using System.Linq;
+using System.Text;
+using System.Threading.Tasks;
+
+namespace Turnkey
+{
+	class ReportIOProvider : ConsoleIOProvider
+	{
+		private System.IO.StreamWriter Out;
+
+		public ReportIOProvider(string ReportFilename)
+		{
+			Out = new System.IO.StreamWriter(ReportFilename);
+		}
+
+		public override void Report(string Message, bool bAppendNewLine)
+		{
+			if (bAppendNewLine)
+			{
+				Out.WriteLine(Message);
+			}
+			else
+			{
+				Out.Write(Message);
+			}
+			Out.Flush();
+		}
+
+		public override string ReadInput(string Prompt, string Default, bool bAppendNewLine)
+		{
+			throw new AutomationTool.AutomationException("ReportIOProvider is unable to ask for input. The Turnkey commandline likely didn't specify enough information.");
+		}
+
+		public override int ReadInputInt(string Prompt, List<string> Options, bool bIsCancellable, int DefaultValue, bool bAppendNewLine)
+		{
+			throw new AutomationTool.AutomationException("ReportIOProvider is unable to ask for input. The Turnkey commandline likely didn't specify enough information.");
+		}
+	}
+}
diff --git a/Engine/Source/Programs/AutomationTool/Turnkey/Turnkey.Automation.csproj b/Engine/Source/Programs/AutomationTool/Turnkey/Turnkey.Automation.csproj
--- a/Engine/Source/Programs/AutomationTool/Turnkey/Turnkey.Automation.csproj
+++ b/Engine/Source/Programs/AutomationTool/Turnkey/Turnkey.Automation.csproj
@@ -51,6 +51,8 @@
     <Compile Include="IOProvider.cs" />
     <Compile Include="LocalCache.cs" />
     <Compile Include="Providers\ConsoleIOProvider.cs" />
+    <Compile Include="Providers\HybridIOProvider.cs" />
+    <Compile Include="Providers\ReportIOProvider.cs" />
     <Compile Include="SdkInfo.cs" />
     <Compile Include="TurnkeyCommand.cs" />
     <Compile Include="CopyProvider.cs" />
@@ -97,6 +99,8 @@
       <Private>False</Private>
     </Reference>
     <Reference Include="System" />
+    <Reference Include="System.Drawing" />
+    <Reference Include="System.Windows.Forms" />
     <Reference Include="System.Xml" />
     <Reference Include="System.Xml.Serialization" />
   </ItemGroup>
