Adding support in CopyUAT to write out a copied file list for debugging purposes.

#rnx

diff --git a/Engine/Source/Programs/AutomationTool/Scripts/CopyUAT.Automation.cs b/Engine/Source/Programs/AutomationTool/Scripts/CopyUAT.Automation.cs
--- a/Engine/Source/Programs/AutomationTool/Scripts/CopyUAT.Automation.cs
+++ b/Engine/Source/Programs/AutomationTool/Scripts/CopyUAT.Automation.cs
@@ -1,4 +1,4 @@
-﻿// Copyright Epic Games, Inc. All Rights Reserved.
+// Copyright Epic Games, Inc. All Rights Reserved.
 
 using System;
 using System.Collections.Generic;
@@ -62,5 +62,6 @@ public class CopyUAT : BuildCommand
 		}
 
 		LogInformation("Copied {0} files to {1}", SourceFiles.Count, TargetDir);
+		File.WriteAllLines(Path.Combine(TargetDirParam, "CopiedFiles.txt"), SourceFiles.Select(F => F.FullName));
 	}
 }
