- Turnkey script changes for editor integration

diff --git a/Engine/Source/Programs/AutomationTool/Turnkey/Commands/InstallSdk.cs b/Engine/Source/Programs/AutomationTool/Turnkey/Commands/InstallSdk.cs
--- a/Engine/Source/Programs/AutomationTool/Turnkey/Commands/InstallSdk.cs
+++ b/Engine/Source/Programs/AutomationTool/Turnkey/Commands/InstallSdk.cs
@@ -30,7 +30,8 @@ namespace Turnkey.Commands
 
 			bool bBestAvailable = TurnkeyUtils.ParseParam("BestAvailable", CommandOptions);
 			bool bUpdateOnly = TurnkeyUtils.ParseParam("UpdateOnly", CommandOptions);
-			
+			bool bAllowAutoSdk = TurnkeyUtils.ParseParam("AllowAutoSdk", CommandOptions);
+
 			// best available installation requires valid and needed Sdks
 			bool bValidOnly = bBestAvailable || !TurnkeyUtils.ParseParam("AllowInvalid", CommandOptions);
 			bool bNeededOnly = bBestAvailable || TurnkeyUtils.ParseParam("NeededOnly", CommandOptions);
@@ -79,8 +80,25 @@ namespace Turnkey.Commands
 				// remove the platform (we may remove more later if an Sdk supports multiple platforms)
 				PlatformsLeftToInstall.RemoveAt(0);
 
+
 				// cache the automation platform object
-				AutomationTool.Platform AutomationPlatform = AutomationTool.Platform.Platforms[new TargetPlatformDescriptor(Platform)];
+				AutomationTool.Platform AutomationPlatform = AutomationTool.Platform.GetPlatform(Platform);
+
+// 				if (bAllowAutoSdk)
+// 				{
+// 					// first, attempt AutoSdk
+// 					SdkInfo.LocalAvailability LocalState = SdkInfo.GetLocalAvailability(AutomationPlatform);
+// 					bool bWasAutoSdkSetup;
+// 					
+// 					SdkInfo.ConditionalSetupAutoSdk(Platform, ref LocalState, out bWasAutoSdkSetup, bUnattended: false);
+// 
+// 					// if we got an AutoSdk for this platform, then we don't need to continue
+// 					if (bWasAutoSdkSetup)
+// 					{
+// 						continue;
+// 					}
+// 				}
+
 
 				// filter the Sdks if a platform was given
 				List<SdkInfo> Sdks = AllSdks.FindAll(x => x.SupportsPlatform(Platform));
@@ -90,6 +108,10 @@ namespace Turnkey.Commands
 				Sdks = Sdks.FindAll(x => x.Type != SdkInfo.SdkType.Flash || (AutomationPlatform.GetDevices() != null && AutomationPlatform.GetDevices().Length > 0));
 				bool bStrippedDevices = Sdks.Count != SdkCount;
 
+				// we don't need to do AutoSdks, we already attempted one above if we wanted to
+				Sdks = Sdks.FindAll(x => x.Type != SdkInfo.SdkType.AutoSdk);
+
+
 				// strip out Sdks where there is no Sdk installed yet
 				if (bUpdateOnly)
 				{
@@ -220,6 +242,7 @@ namespace Turnkey.Commands
 					}
 					SdksAlreadyInstalled.Add(Sdk);
 
+					DeviceInfo InstallDevice = null;
 					// set variables
 					if (Sdk.Type == SdkInfo.SdkType.Flash)
 					{
@@ -246,24 +269,16 @@ namespace Turnkey.Commands
 						}
 
 						// get device info of the chosen or supplied device
-						DeviceInfo InstallDevice = Array.Find(AutomationPlatform.GetDevices(), x => string.Compare(x.Name, InstallDeviceName, true) == 0);
+						InstallDevice = Array.Find(AutomationPlatform.GetDevices(), x => string.Compare(x.Name, InstallDeviceName, true) == 0);
 
 						if (InstallDevice == null)
 						{
 							TurnkeyUtils.Log("Unable to find {0} device {1}", Platform, InstallDeviceName);
 							return;
 						}
-
-						TurnkeyUtils.SetVariable("DeviceName", InstallDevice.Name);
-						TurnkeyUtils.SetVariable("DeviceId", InstallDevice.Id);
-						TurnkeyUtils.SetVariable("DeviceType", InstallDevice.Type);
 					}
 
-					// standard variables
-					TurnkeyUtils.SetVariable("Platform", Platform.ToString());
-					TurnkeyUtils.SetVariable("Version", Sdk.Version);
-
-					Sdk.Install(Platform);
+					Sdk.Install(Platform, InstallDevice);
 				}
 			}
 		}
diff --git a/Engine/Source/Programs/AutomationTool/Turnkey/Commands/ListSdks.cs b/Engine/Source/Programs/AutomationTool/Turnkey/Commands/ListSdks.cs
--- a/Engine/Source/Programs/AutomationTool/Turnkey/Commands/ListSdks.cs
+++ b/Engine/Source/Programs/AutomationTool/Turnkey/Commands/ListSdks.cs
@@ -4,7 +4,7 @@ using System;
 using System.Collections.Generic;
 using UnrealBuildTool;
 using AutomationTool;
-
+using System.Linq;
 
 namespace Turnkey.Commands
 {
@@ -14,7 +14,32 @@ namespace Turnkey.Commands
 		{
 			TurnkeyUtils.Log("");
 			TurnkeyUtils.Log("Available Installers:");
-			foreach (SdkInfo Sdk in TurnkeyManifest.GetDiscoveredSdks())
+
+			string TypeString = TurnkeyUtils.ParseParamValue("Type", null, CommandOptions);
+			string PlatformString = TurnkeyUtils.ParseParamValue("Platform", null, CommandOptions);
+
+			List<SdkInfo> Sdks = TurnkeyManifest.GetDiscoveredSdks();
+
+			if (TypeString != null)
+			{
+				SdkInfo.SdkType Type;
+				if (Enum.TryParse(TypeString, out Type))
+				{
+					Sdks = Sdks.FindAll(x => x.Type == Type);
+				}
+			}
+
+			if (PlatformString != null)
+			{
+				UnrealTargetPlatform Platform;
+				if (UnrealTargetPlatform.TryParse(PlatformString, out Platform))
+				{
+					Sdks = Sdks.FindAll(x => x.SupportsPlatform(Platform));
+				}
+			}
+
+
+			foreach (SdkInfo Sdk in Sdks)
 			{
 				TurnkeyUtils.Log(Sdk.ToString(2));
 			}
diff --git a/Engine/Source/Programs/AutomationTool/Turnkey/Commands/VerifySdk.cs b/Engine/Source/Programs/AutomationTool/Turnkey/Commands/VerifySdk.cs
--- a/Engine/Source/Programs/AutomationTool/Turnkey/Commands/VerifySdk.cs
+++ b/Engine/Source/Programs/AutomationTool/Turnkey/Commands/VerifySdk.cs
@@ -1,6 +1,7 @@
 // Copyright Epic Games, Inc. All Rights Reserved.
 
 using AutomationTool;
+using Gauntlet;
 using System;
 using System.Collections.Generic;
 using System.IO;
@@ -13,185 +14,213 @@ namespace Turnkey.Commands
 	{
 		protected override void Execute(string[] CommandOptions)
 		{
-			string DeviceName = TurnkeyUtils.ParseParamValue("Device", null, CommandOptions);
 			bool bUpdateIfNeeded = TurnkeyUtils.ParseParam("UpdateIfNeeded", CommandOptions);
 			bool bUnattended = TurnkeyUtils.ParseParam("Unattended", CommandOptions);
+			bool bPreferFullSdk = TurnkeyUtils.ParseParam("PreferFull", CommandOptions);
+			bool bForceInstall = TurnkeyUtils.ParseParam("ForceInstall", CommandOptions);
 
-			List<UnrealTargetPlatform> ChosenPlatforms = TurnkeyUtils.GetPlatformsFromCommandLineOrUser(CommandOptions, UnrealTargetPlatform.GetValidPlatforms().ToList());
+			// track each platform to check, and 
+			Dictionary<UnrealTargetPlatform, List<string>> PlatformsAndDevices = null;
 
-			if (ChosenPlatforms == null)
+			// look at any devices on the commandline, and see if they have platforms or not
+			string DeviceList = TurnkeyUtils.ParseParamValue("Device", null, CommandOptions);
+			List<string> SplitDeviceList = null;
+			if (DeviceList != null)
 			{
+				SplitDeviceList = DeviceList.Split("+".ToCharArray()).ToList();
+
+				// look if they have platform@ tags
+				bool bAnyHavePlatform = SplitDeviceList.Any(x => x.Contains("@"));
+				if (bAnyHavePlatform)
+				{
+					if (!SplitDeviceList.All(x => x.Contains("@")))
+					{
+						throw new AutomationException("If any device in -device has a platform indicator ('Platform@Device'), they must all have a platform indicator");
+					}
+
+					// now split it up for devices for each platform
+					foreach (string DeviceToken in SplitDeviceList)
+					{
+						string[] Tokens = DeviceToken.Split("@".ToCharArray(), StringSplitOptions.RemoveEmptyEntries);
+						if (Tokens.Length != 2)
+						{
+							throw new AutomationException("{0} did not have the Platform@Device format", DeviceToken);
+						}
+						UnrealTargetPlatform Platform;
+						if (!UnrealTargetPlatform.TryParse(Tokens[0], out Platform))
+						{
+							TurnkeyUtils.Log("Platform indicator {0} is an invalid platform, skipping", Tokens[0]);
+							continue;
+						}
+
+						string DeviceName = Tokens[1];
+
+						// track it
+						if (PlatformsAndDevices == null)
+						{
+							PlatformsAndDevices = new Dictionary<UnrealTargetPlatform, List<string>>();
+						}
+						if (!PlatformsAndDevices.ContainsKey(Platform))
+						{
+							PlatformsAndDevices[Platform] = new List<string>();
+						}
+						PlatformsAndDevices[Platform].Add(DeviceName);
+					}
+					SplitDeviceList = null;
+				}
+			}
+
+			// if we didn't get some platforms already from -device list, then get or ask the user for platforms
+			if (PlatformsAndDevices == null)
+			{
+				PlatformsAndDevices = new Dictionary<UnrealTargetPlatform, List<string>>();
+				List<UnrealTargetPlatform> ChosenPlatforms = TurnkeyUtils.GetPlatformsFromCommandLineOrUser(CommandOptions, UnrealTargetPlatform.GetValidPlatforms().ToList());
+				ChosenPlatforms.ForEach(x => PlatformsAndDevices.Add(x, null));
+
+				if (ChosenPlatforms.Count > 1 && SplitDeviceList != null && !(SplitDeviceList.Count == 1 && SplitDeviceList[0].CompareTo("All") == 0))
+				{
+					throw new AutomationException("When passing -Device to VerifySdk without platform specifiers ('Platform:Device'), a single platform must be specified (unless -Device=All is used)");
+				}
+			}
+
+			if (PlatformsAndDevices.Count == 0)
+			{
+				TurnkeyUtils.Log("Platform(s) needed for VerifySdk command. Ending command.");
 				return;
 			}
 
+
 			TurnkeyUtils.Log("Installed Sdk validity:");
 			TurnkeyUtils.ExitCode = AutomationTool.ExitCode.Success;
 
 			// check all the platforms
-			foreach (UnrealTargetPlatform Platform in ChosenPlatforms)
+			foreach (var Pair in PlatformsAndDevices)
 			{
+				UnrealTargetPlatform Platform = Pair.Key;
+
 				// get the platform object
-				AutomationTool.Platform AutomationPlatform = AutomationTool.Platform.Platforms[new AutomationTool.TargetPlatformDescriptor(Platform)];
+				AutomationTool.Platform AutomationPlatform = AutomationTool.Platform.GetPlatform(Platform);
 
 				SdkInfo.LocalAvailability LocalState = SdkInfo.GetLocalAvailability(AutomationPlatform);
 
 				if ((LocalState & (SdkInfo.LocalAvailability.AutoSdk_ValidVersionExists | SdkInfo.LocalAvailability.InstalledSdk_ValidVersionExists)) == 0)
 				{
-					TurnkeyUtils.Report("[{0}]: Invalid [No AutoSdk or Installed Sdk found matching {1}]", Platform, AutomationPlatform.GetAllowedSdks());
+					TurnkeyUtils.Report("{0}: Invalid: [No AutoSdk or Installed Sdk found matching {1} - {2} - '{3}']", Platform, AutomationPlatform.GetAllowedSdks(), LocalState.ToString(), Environment.GetEnvironmentVariable("UE_SDKS_ROOT"));
 					TurnkeyUtils.ExitCode = AutomationTool.ExitCode.Error_SDKNotFound;
 				}
 				else
 				{
-					TurnkeyUtils.Report("[{0}]: Valid [{1}]", Platform, (LocalState & (SdkInfo.LocalAvailability.AutoSdk_ValidVersionExists | SdkInfo.LocalAvailability.InstalledSdk_ValidVersionExists)).ToString());
+					TurnkeyUtils.Report("{0}: Valid: [{1}]", Platform, (LocalState & (SdkInfo.LocalAvailability.AutoSdk_ValidVersionExists | SdkInfo.LocalAvailability.InstalledSdk_ValidVersionExists)).ToString());
 					//					TurnkeyUtils.Log("{0}: Valid [Installed: '{1}', Required: '{2}']", Platform, PlatformObject.GetInstalledSdk(), PlatformObject.GetAllowedSdks());
 				}
 
-				if (bUpdateIfNeeded && (TurnkeyUtils.ExitCode != AutomationTool.ExitCode.Success))
+				// track if we installed an AutoSdk since we may need to setup autosdks before checking for devices
+				bool bAutoSdkSetupSucceeded = false;
+
+				// install if out of date, or if forcing it
+				if (bForceInstall || (bUpdateIfNeeded && TurnkeyUtils.ExitCode != AutomationTool.ExitCode.Success))
 				{
 					TurnkeyUtils.ExitCode = AutomationTool.ExitCode.Success;
 
-					bool bAttemptAutoSdkSetup = false;
-					bool bSetupEnvVarAfterInstall = false;
-					if (LocalState.HasFlag(SdkInfo.LocalAvailability.AutoSdk_VariableExists))
-					{
-						bAttemptAutoSdkSetup = true;
-					}
-					else
-					{
-						if (!bUnattended)
-						{
-							// @todo turnkey: help set up UE_SDKS_ROOT if it's not there
-							string Response = TurnkeyUtils.ReadInput("AutoSdks are not setup, but your studio has support. Would you like to set it up now? [Y/n]", "Y");
-							if (string.Compare(Response, "Y", true) == 0)
-							{
-								bAttemptAutoSdkSetup = true;
-								bSetupEnvVarAfterInstall = true;
-							}
-						}
-					}
-
-					if (bAttemptAutoSdkSetup)
+					// if autosdk worked, we don't need to install full sdk
+					if (!bAutoSdkSetupSucceeded)
 					{
-						TurnkeyUtils.Log("{0}: AutoSdk is setup on this compter, will look for available AutoSdk to download", Platform);
-
-						List<SdkInfo> MatchingAutoSdk = SdkInfo.FindMatchingSdks(AutomationPlatform, SdkInfo.SdkType.AutoSdk, bSelectSingleBest: true);
-
-						if (MatchingAutoSdk.Count == 0)
+// 						if (!bUnattended)
+// 						{
+// 							string Response = TurnkeyUtils.ReadInput("Your Sdk installation is not up to date. Would you like to install a valid Sdk? [Y/n]", "Y");
+// 							if (string.Compare(Response, "Y", true) != 0)
+// 							{
+// 								continue;
+// 							}
+// 						}
+
+						SdkInfo BestSdk = null;
+						// find the best Sdk, prioritizing as request
+						if (bPreferFullSdk)
 						{
-							// no matching AutoSdk found - will fall through to look for full install
-							if (bSetupEnvVarAfterInstall)
-							{
-								TurnkeyUtils.Log("{0}: Unable to find a matching AutoSdk, skipping AutoSdk setup", Platform);
-							}
+							BestSdk = SdkInfo.FindMatchingSdk(AutomationPlatform, new SdkInfo.SdkType[] { SdkInfo.SdkType.Full, SdkInfo.SdkType.BuildOnly, SdkInfo.SdkType.AutoSdk }, bSelectBest: bUnattended);
 						}
 						else
 						{
-							// make sure this is unset so that we can know if it worked or not after install
-							TurnkeyUtils.ClearVariable("CopyOutputPath");
-
-							// now download it (AutoSdks don't "install") on download
-							// @todo turnkey: handle errors, handle p4 going to wrong location, handle one Sdk for multiple platforms
-							MatchingAutoSdk[0].Install(Platform);
-
-							if (bSetupEnvVarAfterInstall)
-							{
-								// @todo turnkey - have studio settings 
-
-								// this is where we synced the Sdk to
-								string InstalledRoot = TurnkeyUtils.GetVariableValue("CopyOutputPath");
-
-								// failed to install, nothing we can do
-								if (string.IsNullOrEmpty(InstalledRoot))
-								{
-									TurnkeyUtils.ExitCode = AutomationTool.ExitCode.Error_SDKNotFound;
-									continue;
-								}
-
-								// walk up to one above Host* directory
-								DirectoryInfo AutoSdkSearch;
-								if (Directory.Exists(InstalledRoot))
-								{
-									AutoSdkSearch = new DirectoryInfo(InstalledRoot);
-								}
-								else
-								{
-									AutoSdkSearch = new FileInfo(InstalledRoot).Directory;
-								}
-								while (AutoSdkSearch.Name != "Host" + HostPlatform.Current.HostEditorPlatform.ToString())
-								{
-									AutoSdkSearch = AutoSdkSearch.Parent;
-								}
+							BestSdk = SdkInfo.FindMatchingSdk(AutomationPlatform, new SdkInfo.SdkType[] { SdkInfo.SdkType.AutoSdk, SdkInfo.SdkType.BuildOnly, SdkInfo.SdkType.Full }, bSelectBest: bUnattended);
+						}
 
-								// now go one up to the parent of Host
-								AutoSdkSearch = AutoSdkSearch.Parent;
+						if (BestSdk == null)
+						{
+							TurnkeyUtils.Log("ERROR: {0}: Unable to find any Sdks that could be installed");
+							TurnkeyUtils.ExitCode = AutomationTool.ExitCode.Error_SDKNotFound;
+							continue;
+						}
 
-								string Response = TurnkeyUtils.ReadInput("Enter directory for root of AutoSdks. Use detected value, or enter another:", AutoSdkSearch.FullName);
-								if (string.IsNullOrEmpty(Response))
-								{
-									continue;
-								}
+						TurnkeyUtils.Log("Will install {0}", BestSdk.DisplayName);
 
-								// set the env var, globally
-								TurnkeyUtils.StartTrackingExternalEnvVarChanges();
-								Environment.SetEnvironmentVariable("UE_SDKS_ROOT", Response);
-								Environment.SetEnvironmentVariable("UE_SDKS_ROOT", Response, EnvironmentVariableTarget.User);
-								TurnkeyUtils.EndTrackingExternalEnvVarChanges();
-							}
+//						bool bWasSdkInstalled = BestSdk.Install(Platform, null, bUnattended);
 
-							// @todo turnkey - re-run detection again
+						// update LocalState
+						LocalState = SdkInfo.GetLocalAvailability(AutomationPlatform);
 
-							continue;
-						}
+						// @todo turnkey: validate!
 					}
+				}
 
-					if (!bUnattended)
+				// use the per-platform device list, unless it's not specifed, then use the global device list (which is set when not using platform specifiers)
+				List<string> DeviceNames = Pair.Value != null ? Pair.Value : SplitDeviceList;
+				if (DeviceNames != null && DeviceNames.Count > 0)
+				{
+					// if we just setup AutoSdk, then run the setup so that we have a hope of finding devices
+					if (bAutoSdkSetupSucceeded)
 					{
-						string Response = TurnkeyUtils.ReadInput("Your Sdk installation is not up to date. Would you like to install a valid Sdk? [Y/n]", "Y");
-						if (string.Compare(Response, "Y", true) != 0)
-						{
-							continue;
-						}
+						// @todo turnkey: Run the autosdk thing
+						//							AutomationPlatform.
 					}
 
-					// at this point, AutoSdk isn't viable, so look for a full install
-					List<SdkInfo> MatchingInstallableSdks = SdkInfo.FindMatchingSdks(AutomationPlatform, SdkInfo.SdkType.BuildOnly, bSelectSingleBest: bUnattended);
-					if (MatchingInstallableSdks.Count == 0)
+					DeviceInfo[] Devices = AutomationPlatform.GetDevices();
+					if (Devices == null)
 					{
-						MatchingInstallableSdks = SdkInfo.FindMatchingSdks(AutomationPlatform, SdkInfo.SdkType.Full, bSelectSingleBest: bUnattended);
+						TurnkeyUtils.Log("Platform {0} didn't have any devices, ignoring any devices specified", Platform);
+						continue;
 					}
 
-					if (MatchingInstallableSdks.Count == 0)
+					TurnkeyUtils.Log("Installed Device validity:");
+
+					// a single device named all means all devices
+					if (!(DeviceNames.Count == 1 && DeviceNames[0].CompareTo("All") == 0))
 					{
-						TurnkeyUtils.Log("ERROR: {0}: Unable top find any Sdks that could be installed");
-						TurnkeyUtils.ExitCode = AutomationTool.ExitCode.Error_SDKNotFound;
+						Devices = Devices.Where(x => DeviceNames.Contains(x.Name, StringComparer.OrdinalIgnoreCase)).ToArray();
 					}
 
-					SdkInfo SdkToInstall = MatchingInstallableSdks[0];
-
-					// if there are multiple, ask the user
-					if (MatchingInstallableSdks.Count > 1)
+					// now check software verison of each device
+					foreach (DeviceInfo Device in Devices)
 					{
-						int Choice = TurnkeyUtils.ReadInputInt("Multiple Sdks found that could be installed. Please select one:", MatchingInstallableSdks.Select(x => x.DisplayName).ToList(), true);
-						if (Choice == 0)
+						bool bIsSoftwareValid = TurnkeyUtils.IsValueValid(Device.SoftwareVersion, AutomationPlatform.GetAllowedSoftwareVersions());
+
+
+						if (bIsSoftwareValid)
 						{
-							TurnkeyUtils.ExitCode = AutomationTool.ExitCode.Error_SDKNotFound;
-							continue;
+							TurnkeyUtils.Report("{0}@{1}: Valid: [{2}]", Platform, Device.Name, Device.SoftwareVersion);
 						}
+						else
+						{
+							TurnkeyUtils.Report("{0}@{1}: Invalid: [Has {2}, needs {3}]", Platform, Device.Name, Device.SoftwareVersion, AutomationPlatform.GetAllowedSoftwareVersions());
+							TurnkeyUtils.ExitCode = AutomationTool.ExitCode.Error_SDKNotFound;
 
-						SdkToInstall = MatchingInstallableSdks[Choice - 1];
-					}
+							if (bUpdateIfNeeded)
+							{
+								SdkInfo MatchingInstallableSdk = SdkInfo.FindMatchingSdk(AutomationPlatform, new SdkInfo.SdkType[] { SdkInfo.SdkType.Flash }, bSelectBest: bUnattended);
 
-					// finall install the best or chosen Sdk
-					SdkToInstall.Install(Platform);
-
-// 					if (bInstallSdk)
-// 					{
-// 						// reset exit code
-// 						TurnkeyUtils.ExitCode = AutomationTool.ExitCode.Success;
-// 
-// 						InstallSdk InstallCommand = new InstallSdk();
-// 						InstallCommand.InternalExecute(new string[] { "-platform=" + Platform, "-BestAvailable" });
-// 					}
+								if (MatchingInstallableSdk == null)
+								{
+									TurnkeyUtils.Log("ERROR: {0}: Unable top find any Sdks that could be installed on {1}", Platform, Device.Name);
+									TurnkeyUtils.ExitCode = AutomationTool.ExitCode.Error_SDKNotFound;
+								}
+								else
+								{
+									MatchingInstallableSdk.Install(Platform, Device, bUnattended);
+									TurnkeyUtils.ExitCode = AutomationTool.ExitCode.Success;
+								}
+							}
+						}
+					}
 				}
 			}
 		}
diff --git a/Engine/Source/Programs/AutomationTool/Turnkey/Providers/GoogleDriveCopyProvider.cs b/Engine/Source/Programs/AutomationTool/Turnkey/Providers/GoogleDriveCopyProvider.cs
--- a/Engine/Source/Programs/AutomationTool/Turnkey/Providers/GoogleDriveCopyProvider.cs
+++ b/Engine/Source/Programs/AutomationTool/Turnkey/Providers/GoogleDriveCopyProvider.cs
@@ -58,6 +58,10 @@ namespace Turnkey
 				{
 					IdsToCopyToFilename.Add(File, File.Name);
 				}
+				else
+				{
+					TurnkeyUtils.Log("ERROR: Unable to find GoogleDrive file {0}", Operation);
+				}
 			}
 
 
diff --git a/Engine/Source/Programs/AutomationTool/Turnkey/Providers/HybridIOProvider.cs b/Engine/Source/Programs/AutomationTool/Turnkey/Providers/HybridIOProvider.cs
--- a/Engine/Source/Programs/AutomationTool/Turnkey/Providers/HybridIOProvider.cs
+++ b/Engine/Source/Programs/AutomationTool/Turnkey/Providers/HybridIOProvider.cs
@@ -125,13 +125,17 @@ namespace Turnkey
 			// Display the form as a modal dialog box.
 			Form1.ShowDialog();
 
+			string Result = null;
+
 			// Determine if the OK button was clicked on the dialog box.
 			if (Form1.DialogResult == DialogResult.OK)
 			{
-				return TextBox.Text;
+				Result = TextBox.Text;
 			}
 
-			return null;
+			Form1.Dispose();
+
+			return Result;
 		}
 
 		public override string ReadInput(string Prompt, string Default, bool bAppendNewLine)
diff --git a/Engine/Source/Programs/AutomationTool/Turnkey/Providers/PerforceCopyProvider.cs b/Engine/Source/Programs/AutomationTool/Turnkey/Providers/PerforceCopyProvider.cs
--- a/Engine/Source/Programs/AutomationTool/Turnkey/Providers/PerforceCopyProvider.cs
+++ b/Engine/Source/Programs/AutomationTool/Turnkey/Providers/PerforceCopyProvider.cs
@@ -58,8 +58,9 @@ namespace Turnkey
 
 			if (Expansions != null)
 			{
-				// figure out what the *'s turned into
-				Regex Expression = new Regex(Operation.Replace("/", "\\/").Replace("*", "([^\\/]*)").Replace("...", ".*"));
+				// figure out what the *'s turned into (note we use case-insensitive because p4 has an annoying habit of being case-preserving of directories
+				// on a per-file basis (//depot/UE4/Foo/* might return //depot/UE4/Foo/A and //Depot/ue4/FOO/B)
+				Regex Expression = new Regex(Operation.Replace("/", "\\/").Replace("*", "([^\\/]*)").Replace("...", ".*"), RegexOptions.IgnoreCase);
 
 				// run the regex on each result
 				foreach (string Result in Results)
diff --git a/Engine/Source/Programs/AutomationTool/Turnkey/SdkInfo.cs b/Engine/Source/Programs/AutomationTool/Turnkey/SdkInfo.cs
--- a/Engine/Source/Programs/AutomationTool/Turnkey/SdkInfo.cs
+++ b/Engine/Source/Programs/AutomationTool/Turnkey/SdkInfo.cs
@@ -8,6 +8,8 @@ using System.IO;
 using Tools.DotNETCommon;
 using UnrealBuildTool;
 using AutomationTool;
+using System.Linq;
+using System.Windows.Forms;
 
 namespace Turnkey
 {
@@ -166,33 +168,65 @@ namespace Turnkey
 			return Result;
 		}
 
-		static public List<SdkInfo> FindMatchingSdks(AutomationTool.Platform Platform, SdkType Type, bool bSelectSingleBest)
+		static public SdkInfo FindMatchingSdk(AutomationTool.Platform Platform, SdkType[] TypePriority, bool bSelectBest)
 		{
+			foreach (SdkType Type in TypePriority)
+			{
+				List<SdkInfo> Sdks = TurnkeyManifest.GetDiscoveredSdks().FindAll(x =>
+					x.SupportsPlatform(Platform.IniPlatformType) &&
+					x.Type == Type &&
+					TurnkeyUtils.IsValueValid(x.Version, Type == SdkType.Flash ? Platform.GetAllowedSoftwareVersions() : Platform.GetAllowedSdks())
+					); ;
 
-			List<SdkInfo> Sdks = TurnkeyManifest.GetDiscoveredSdks().FindAll(x => 
-				x.SupportsPlatform(Platform.IniPlatformType) && 
-				x.Type == Type && 
-				TurnkeyUtils.IsValueValid(x.Version, Platform.GetAllowedSdks())
-				);
+				Sdks = TurnkeyManifest.GetDiscoveredSdks();
+				Sdks = Sdks.FindAll(x => x.SupportsPlatform(Platform.IniPlatformType));
+				Sdks = Sdks.FindAll(x => x.Type == Type);
+				Sdks = Sdks.FindAll(x => TurnkeyUtils.IsValueValid(x.Version, Type == SdkType.Flash ? Platform.GetAllowedSoftwareVersions() : Platform.GetAllowedSdks()));
 
-			// handle easy cases
-			if (Sdks.Count <= 1 || !bSelectSingleBest)
-			{
-				return Sdks;
-			}
+				// if none were found try next type
+				if (Sdks.Count == 0)
+				{
+					continue;
+				}
 
-			// find best one
-			SdkInfo Best = null;
-			foreach (SdkInfo Sdk in Sdks)
-			{
-				// @todo turnkey: let platform decide what's best
-				if (Best == null || string.Compare(Sdk.Version, Best.Version, true) > 0)
+				// if one was founf return it!
+				if (Sdks.Count == 1)
 				{
-					Best = Sdk;
+					return Sdks[0];
 				}
+
+				// if more than 1, select one automatically if requested
+				if (bSelectBest)
+				{
+					// find best one
+					SdkInfo Best = null;
+					foreach (SdkInfo Sdk in Sdks)
+					{
+						// @todo turnkey: let platform decide what's best
+						if (Best == null || string.Compare(Sdk.Version, Best.Version, true) > 0)
+						{
+							Best = Sdk;
+						}
+					}
+
+					return Best;
+				}
+
+				// if unable to pock one automatically, ask the user
+				int Choice = TurnkeyUtils.ReadInputInt("Multiple Sdks found that could be installed. Please select one:", Sdks.Select(x => x.DisplayName).ToList(), true);
+
+				// we take canceling to mean to try the next type
+				if (Choice == 0)
+				{
+					continue;
+				}
+
+				return Sdks[Choice - 1];
 			}
 
-			return new List<SdkInfo> { Best };
+
+			// if we never found anything, fail
+			return null;
 		}
 
 
@@ -322,12 +356,32 @@ namespace Turnkey
 			}
 		}
 
-		public void Install(UnrealTargetPlatform Platform)
+		public bool Install(UnrealTargetPlatform Platform, DeviceInfo Device = null, bool bUnattended = false)
 		{
+			if (Type == SdkType.AutoSdk)
+			{
+				// AutoSdk has some extra setup needed
+				if (SdkInfo.ConditionalSetupAutoSdk(Platform, bUnattended))
+				{
+					return true;
+				}
+			}
+
 			if (!CheckForExtendedVariables(true))
 			{
 				// user canceled a choice
-				return;
+				return false;
+			}
+
+			// standard variables
+			TurnkeyUtils.SetVariable("Platform", Platform.ToString());
+			TurnkeyUtils.SetVariable("Version", Version);
+
+			if (Device != null)
+			{
+				TurnkeyUtils.SetVariable("DeviceName", Device.Name);
+				TurnkeyUtils.SetVariable("DeviceId", Device.Id);
+				TurnkeyUtils.SetVariable("DeviceType", Device.Type);
 			}
 
 			// custom sdk installation is enabled if ChosenSdk.CustomVersionId is !null
@@ -362,53 +416,132 @@ namespace Turnkey
 			}
 
 			// now run any installers
-			if (Installers != null)
+			if (Installers == null)
+			{
+				// if there were no installers, then we succeeded
+				return true;
+			}
+
+			bool bSucceeded = true;
+			foreach (CopyAndRun Install in Installers)
+			{
+				// build only types we keep in a nice directory structure in the PermanentStorage (if the copy provider allows for choosing destination)
+				if (Type == SdkType.BuildOnly)
+				{
+					// download subdir is based on platform 
+					string SubDir = string.Format("{0}/{1}", Platform.ToString(), Version);
+					bSucceeded = bSucceeded && Install.Execute(CopyExecuteSpecialMode.UsePermanentStorage, SubDir);
+				}
+				else
+				{
+					// !@todo turnkey : verify it worked
+					bSucceeded = bSucceeded && Install.Execute();
+				}
+			}
+
+			return bSucceeded;
+		}
+
+		public static bool ConditionalSetupAutoSdk(UnrealTargetPlatform Platform, bool bUnattended)
+		{
+			bool bAttemptAutoSdkSetup = false;
+			bool bSetupEnvVarAfterInstall = false;
+			if (Environment.GetEnvironmentVariable("UE_SDKS_ROOT") != null)
 			{
-				foreach (CopyAndRun Install in Installers)
+				bAttemptAutoSdkSetup = true;
+			}
+			else
+			{
+				if (!bUnattended)
 				{
-					if (Type == SdkType.BuildOnly)
+					// @todo turnkey - have studio settings 
+					string Response = TurnkeyUtils.ReadInput("AutoSdks are not setup, but your studio has support. Would you like to set it up now? [Y/n]", "Y");
+					if (string.Compare(Response, "Y", true) == 0)
 					{
-						// download subdir is based on platform 
-						string SubDir = string.Format("{0}/{1}", Platform.ToString(), Version);
-						Install.Execute(CopyExecuteSpecialMode.UsePermanentStorage, SubDir);
-
-						// 					// now that we are copied, write out a manifest to the root of the output location for later quick switching,
-						// 					// and not continue to install anything
-						// 					string ManifestPath = TurnkeyUtils.GetVariableValue("CopyOutputPath");
-						// 
-						// 					// if it's a file, then get it's directory
-						// 					if (File.Exists(ManifestPath))
-						// 					{
-						// 						ManifestPath = Path.GetDirectoryName(ManifestPath);
-						// 					}
-						// 					Directory.CreateDirectory(ManifestPath);
-						// 					ManifestPath = Path.Combine(ManifestPath, "TurnkeyQuickSwitch.xml");
-						// 
-						// 					// if the manifest already existed as downloaded, then there's no need to make one, we assume it was correct
-						// 					if (!File.Exists(ManifestPath))
-						// 					{
-						// 						TurnkeyManifest QuickSwitchManifest = new TurnkeyManifest();
-						// 
-						// 						// copy this Sdk to a new one
-						// 						SdkInfo NewSdk = CloneForQuickSwitch(Path.GetDirectoryName(ManifestPath));
-						// 						QuickSwitchManifest.SdkInfos = new SdkInfo[] { NewSdk };
-						// 
-						// 						// register the Sdk with th runtime so that we can install it without quitting
-						// 						TurnkeyManifest.AddCreatedSdk(NewSdk);
-						// 
-						// 						// save out a single switch manifest
-						// 						QuickSwitchManifest.Write(ManifestPath);
-						//					}
+						bAttemptAutoSdkSetup = true;
+						bSetupEnvVarAfterInstall = true;
 					}
-					else
+				}
+			}
+
+			if (bAttemptAutoSdkSetup)
+			{
+				AutomationTool.Platform AutomationPlatform = AutomationTool.Platform.GetPlatform(Platform);
+
+				TurnkeyUtils.Log("{0}: AutoSdk is setup on this computer, will look for available AutoSdk to download", Platform);
+
+				SdkInfo MatchingAutoSdk = SdkInfo.FindMatchingSdk(AutomationPlatform, new SdkType[] { SdkInfo.SdkType.AutoSdk } , bSelectBest: true);
+
+				if (MatchingAutoSdk == null)
+				{
+					// no matching AutoSdk found - will fall through to look for full install
+					if (bSetupEnvVarAfterInstall)
 					{
-						Install.Execute();
+						TurnkeyUtils.Log("{0}: Unable to find a matching AutoSdk, skipping AutoSdk setup", Platform);
 					}
+				}
+				else
+				{
+					// make sure this is unset so that we can know if it worked or not after install
+					TurnkeyUtils.ClearVariable("CopyOutputPath");
+
+					// now download it (AutoSdks don't "install") on download
+					// @todo turnkey: handle errors, handle p4 going to wrong location, handle one Sdk for multiple platforms
+					MatchingAutoSdk.Install(Platform, null, bUnattended);
+
+					if (bSetupEnvVarAfterInstall)
+					{
 
-					// grab where the Install copied the files
+						// this is where we synced the Sdk to
+						string InstalledRoot = TurnkeyUtils.GetVariableValue("CopyOutputPath");
 
+						// failed to install, nothing we can do
+						if (string.IsNullOrEmpty(InstalledRoot))
+						{
+							TurnkeyUtils.ExitCode = AutomationTool.ExitCode.Error_SDKNotFound;
+							return false;
+						}
+
+						// walk up to one above Host* directory
+						DirectoryInfo AutoSdkSearch;
+						if (Directory.Exists(InstalledRoot))
+						{
+							AutoSdkSearch = new DirectoryInfo(InstalledRoot);
+						}
+						else
+						{
+							AutoSdkSearch = new FileInfo(InstalledRoot).Directory;
+						}
+						while (AutoSdkSearch.Name != "Host" + HostPlatform.Current.HostEditorPlatform.ToString())
+						{
+							AutoSdkSearch = AutoSdkSearch.Parent;
+						}
+
+						// now go one up to the parent of Host
+						AutoSdkSearch = AutoSdkSearch.Parent;
+
+						string AutoSdkDir = AutoSdkSearch.FullName;
+						if (!bUnattended)
+						{
+							string Response = TurnkeyUtils.ReadInput("Enter directory for root of AutoSdks. Use detected value, or enter another:", AutoSdkSearch.FullName);
+							if (string.IsNullOrEmpty(Response))
+							{
+								return false;
+							}
+						}
+
+						// set the env var, globally
+						TurnkeyUtils.StartTrackingExternalEnvVarChanges();
+						Environment.SetEnvironmentVariable("UE_SDKS_ROOT", AutoSdkDir);
+						Environment.SetEnvironmentVariable("UE_SDKS_ROOT", AutoSdkDir, EnvironmentVariableTarget.User);
+						TurnkeyUtils.EndTrackingExternalEnvVarChanges();
+					}
+
+					return true;
 				}
 			}
+
+			return false;
 		}
 
 
@@ -417,6 +550,7 @@ namespace Turnkey
 
 
 
+
 		private bool CheckForExtendedVariables(bool bIncludeCustomSdk)
 		{
 			foreach (string Var in ExtendedVariables)
