- Moved UAT project commandline parsing up to BuildCommand to be usable by scripts other than BuildCookRun

diff --git a/Engine/Source/Programs/AutomationTool/AutomationUtils/BuildCommand.cs b/Engine/Source/Programs/AutomationTool/AutomationUtils/BuildCommand.cs
--- a/Engine/Source/Programs/AutomationTool/AutomationUtils/BuildCommand.cs
+++ b/Engine/Source/Programs/AutomationTool/AutomationUtils/BuildCommand.cs
@@ -226,6 +226,65 @@ namespace AutomationTool
 			}
 		}
 
+		public FileReference ParseProjectParam()
+		{
+			FileReference ProjectFullPath = null;
+
+			var bForeign = ParseParam("foreign");
+			var bForeignCode = ParseParam("foreigncode");
+			if (bForeign)
+			{
+				var DestSample = ParseParamValue("DestSample", "CopiedHoverShip");
+				var Dest = ParseParamValue("ForeignDest", CombinePaths(@"C:\testue4\foreign\", DestSample + "_ _Dir"));
+				ProjectFullPath = new FileReference(CombinePaths(Dest, DestSample + ".uproject"));
+			}
+			else if (bForeignCode)
+			{
+				var DestSample = ParseParamValue("DestSample", "PlatformerGame");
+				var Dest = ParseParamValue("ForeignDest", CombinePaths(@"C:\testue4\foreign\", DestSample + "_ _Dir"));
+				ProjectFullPath = new FileReference(CombinePaths(Dest, DestSample + ".uproject"));
+			}
+			else
+			{
+				var OriginalProjectName = ParseParamValue("project", "");
+
+				if (string.IsNullOrEmpty(OriginalProjectName))
+				{
+					return null;
+				}
+
+				var ProjectName = OriginalProjectName;
+				ProjectName = ProjectName.Trim(new char[] { '\"' });
+				if (ProjectName.IndexOfAny(new char[] { '\\', '/' }) < 0)
+				{
+					ProjectName = CombinePaths(CmdEnv.LocalRoot, ProjectName, ProjectName + ".uproject");
+				}
+				else if (!FileExists_NoExceptions(ProjectName))
+				{
+					ProjectName = CombinePaths(CmdEnv.LocalRoot, ProjectName);
+				}
+				if (FileExists_NoExceptions(ProjectName))
+				{
+					ProjectFullPath = new FileReference(ProjectName);
+				}
+				else
+				{
+					var Branch = new BranchInfo();
+					var GameProj = Branch.FindGame(OriginalProjectName);
+					if (GameProj != null)
+					{
+						ProjectFullPath = GameProj.FilePath;
+					}
+					if (ProjectFullPath == null || !FileExists_NoExceptions(ProjectFullPath.FullName))
+					{
+						throw new AutomationException("Could not find a project file {0}.", ProjectName);
+					}
+				}
+			}
+
+			return ProjectFullPath;
+		}
+
 		/// <summary>
 		/// Checks that all of the required params are present, throws an exception if not
 		/// </summary>
diff --git a/Engine/Source/Programs/AutomationTool/Scripts/BuildCookRun.Automation.cs b/Engine/Source/Programs/AutomationTool/Scripts/BuildCookRun.Automation.cs
--- a/Engine/Source/Programs/AutomationTool/Scripts/BuildCookRun.Automation.cs
+++ b/Engine/Source/Programs/AutomationTool/Scripts/BuildCookRun.Automation.cs
@@ -281,58 +281,14 @@ public class BuildCookRun : BuildCommand
 		{
 			if (ProjectFullPath == null)
 			{
-				var bForeign = ParseParam("foreign");
-				var bForeignCode = ParseParam("foreigncode");
-				if (bForeign)
-				{
-					var DestSample = ParseParamValue("DestSample", "CopiedHoverShip");
-					var Dest = ParseParamValue("ForeignDest", CombinePaths(@"C:\testue4\foreign\", DestSample + "_ _Dir"));
-					ProjectFullPath = new FileReference(CombinePaths(Dest, DestSample + ".uproject"));
-				}
-				else if (bForeignCode)
-				{
-					var DestSample = ParseParamValue("DestSample", "PlatformerGame");
-					var Dest = ParseParamValue("ForeignDest", CombinePaths(@"C:\testue4\foreign\", DestSample + "_ _Dir"));
-					ProjectFullPath = new FileReference(CombinePaths(Dest, DestSample + ".uproject"));
-				}
-				else
+				ProjectFullPath = ParseProjectParam();
+
+				if (ProjectFullPath == null)
 				{
-					var OriginalProjectName = ParseParamValue("project", "");
-
-                    if (string.IsNullOrEmpty(OriginalProjectName))
-                    {
-                        throw new AutomationException("No project file specified. Use -project=<project>.");
-                    }
-
-					var ProjectName = OriginalProjectName;
-					ProjectName = ProjectName.Trim(new char[] { '\"' });
-					if (ProjectName.IndexOfAny(new char[] { '\\', '/' }) < 0)
-					{
-						ProjectName = CombinePaths(CmdEnv.LocalRoot, ProjectName, ProjectName + ".uproject");
-					}
-					else if (!FileExists_NoExceptions(ProjectName))
-					{
-						ProjectName = CombinePaths(CmdEnv.LocalRoot, ProjectName);
-					}
-					if(FileExists_NoExceptions(ProjectName))
-					{
-						ProjectFullPath = new FileReference(ProjectName);
-					}
-					else
-					{
-						var Branch = new BranchInfo();
-						var GameProj = Branch.FindGame(OriginalProjectName);
-						if (GameProj != null)
-						{
-							ProjectFullPath = GameProj.FilePath;
-						}
-						if (ProjectFullPath == null || !FileExists_NoExceptions(ProjectFullPath.FullName))
-						{
-							throw new AutomationException("Could not find a project file {0}.", ProjectName);
-						}
-					}
+					throw new AutomationException("No project file specified. Use -project=<project>.");
 				}
 			}
+
 			return ProjectFullPath;
 		}
 	}
