- Turnkey updates for running on Mac,
- Turnkey support for IOS (installing Xcode, installing .ipsw files, installing certs, updating and downloading provisions from Dev Center, etc) (initial version of a lot, but bare minimum working!)

diff --git a/Engine/Config/BaseEngine.ini b/Engine/Config/BaseEngine.ini
--- a/Engine/Config/BaseEngine.ini
+++ b/Engine/Config/BaseEngine.ini
@@ -1618,7 +1618,6 @@ PathToRAC=RuntimeAssetCache
 ; Each graph should start with 'Root' node. Names of all the other nodes are not predefined.
 ; Supported node types are: KeyLength, AsyncPut, Hierarchical, Boot, Filesystem, ReadPak, WritePak, Verify
 ; The order nodes are define in is not relevant
-; Filesystem backends can be disabled by setting the corresponding environment variable to 'None'. E.g UE-SharedDataCachePath=None
 
 [DerivedDataBackendGraph]
 MinimumDaysToKeepFile=7
@@ -1626,9 +1625,9 @@ Root=(Type=KeyLength, Length=120, Inner=AsyncPut)
 AsyncPut=(Type=AsyncPut, Inner=Hierarchy)
 Hierarchy=(Type=Hierarchical, Inner=Boot, Inner=Pak, Inner=EnginePak, Inner=Local, Inner=Shared)
 Boot=(Type=Boot, Filename="%GAMEDIR%DerivedDataCache/Boot.ddc", MaxCacheSize=512)
-Local=(Type=FileSystem, ReadOnly=false, Clean=false, Flush=false, PurgeTransient=true, DeleteUnused=true, UnusedFileAge=34, FoldersToClean=-1, PromptIfMissing=false, Path=%ENGINEDIR%DerivedDataCache, EnvPathOverride=UE-LocalDataCachePath, EditorOverrideSetting=LocalDerivedDataCache)
-Shared=(Type=FileSystem, ReadOnly=false, Clean=false, Flush=false, DeleteUnused=true, UnusedFileAge=10, FoldersToClean=10, MaxFileChecksPerSec=1, ConsiderSlowAt=70, PromptIfMissing=false, Path=?EpicDDC, EnvPathOverride=UE-SharedDataCachePath, EditorOverrideSetting=SharedDerivedDataCache, CommandLineOverride=SharedDataCachePath)
-AltShared=(Type=FileSystem, ReadOnly=true, Clean=false, Flush=false, DeleteUnused=true, UnusedFileAge=23, FoldersToClean=10, MaxFileChecksPerSec=1, ConsiderSlowAt=70, PromptIfMissing=false, Path=?EpicDDC2, EnvPathOverride=UE-SharedDataCachePath2)
+Local=(Type=FileSystem, ReadOnly=false, Clean=false, Flush=false, PurgeTransient=true, DeleteUnused=true, UnusedFileAge=34, FoldersToClean=-1, Path=%ENGINEDIR%DerivedDataCache, EnvPathOverride=UE-LocalDataCachePath, EditorOverrideSetting=LocalDerivedDataCache)
+Shared=(Type=FileSystem, ReadOnly=false, Clean=false, Flush=false, DeleteUnused=true, UnusedFileAge=10, FoldersToClean=10, MaxFileChecksPerSec=1, ConsiderSlowAt=70, Path=?EpicDDC, EnvPathOverride=UE-SharedDataCachePath, EditorOverrideSetting=SharedDerivedDataCache, CommandLineOverride=SharedDataCachePath)
+AltShared=(Type=FileSystem, ReadOnly=true, Clean=false, Flush=false, DeleteUnused=true, UnusedFileAge=23, FoldersToClean=10, MaxFileChecksPerSec=1, ConsiderSlowAt=70, Path=?EpicDDC2, EnvPathOverride=UE-SharedDataCachePath2)
 Pak=(Type=ReadPak, Filename="%GAMEDIR%DerivedDataCache/DDC.ddp")
 EnginePak=(Type=ReadPak, Filename=%ENGINEDIR%DerivedDataCache/DDC.ddp)
 
@@ -2021,7 +2020,9 @@ bDisableHTTPS=false
 bUseRSync=True
 BundleDisplayName=[PROJECT_NAME]
 BundleName=[PROJECT_NAME]
-BundleIdentifier=com.YourCompany.[PROJECT_NAME]
+BundleIdentifier=com.YourCXompany.[PROJECT_NAME]
+DevCodeSigningIdentity=iPhone Developer:;Apple Developer:
+DistroCodeSigningIdentity=iPhone Distribution:;Apple Distribution:
 VersionInfo=1.0
 FrameRateLock=PUFRL_30
 bEnableDynamicMaxFPS=False
diff --git a/Engine/Source/Programs/AutomationTool/AutomationUtils/Platform.cs b/Engine/Source/Programs/AutomationTool/AutomationUtils/Platform.cs
--- a/Engine/Source/Programs/AutomationTool/AutomationUtils/Platform.cs
+++ b/Engine/Source/Programs/AutomationTool/AutomationUtils/Platform.cs
@@ -11,6 +11,11 @@ using Tools.DotNETCommon;
 
 namespace AutomationTool
 {
+	public interface FileRetriever
+	{
+		string RetrieveByTags(string[] RequiredTags, string[] PreferredTags, Dictionary<string,string> ExtraVariables=null);
+	}
+
 	public class DeviceInfo
 	{
 		public string Name;
@@ -177,11 +182,20 @@ namespace AutomationTool
 		{
 			return true;
 		}
-		public virtual bool CustomVersionUpdate(string CustomVersionId, string UpdateParams)
+		public virtual bool CustomVersionUpdate(string CustomVersionId, string UpdateParams, FileRetriever Retriever)
+		{
+			return true;
+		}
+
+		public virtual bool UpdateHostPrerequisites(BuildCommand Command, FileRetriever Retriever)
 		{
 			return true;
 		}
 
+		public virtual bool UpdateDevicePrerequisites(DeviceInfo Device, BuildCommand Command, FileRetriever Retriever)
+		{
+			return true;
+		}
 
 		/// <summary>
 		/// Package files for the current platform.
diff --git a/Engine/Source/Programs/AutomationTool/IOS/IOSPlatform.Automation.cs b/Engine/Source/Programs/AutomationTool/IOS/IOSPlatform.Automation.cs
--- a/Engine/Source/Programs/AutomationTool/IOS/IOSPlatform.Automation.cs
+++ b/Engine/Source/Programs/AutomationTool/IOS/IOSPlatform.Automation.cs
@@ -169,19 +169,240 @@ public class IOSPlatform : Platform
 		SDKName = (TargetPlatform == UnrealTargetPlatform.TVOS) ? "appletvos" : "iphoneos";
 	}
 
+	public override string GetAllowedSdks()
+	{
+		return "11.5";
+	}
+	public override string GetInstalledSdk()
+	{
+		// @todo turnkey: get installed Xcode version
+		string Output = UnrealBuildTool.Utils.RunLocalProcessAndReturnStdOut("sh", "-c 'xcodebuild -version'");
+		Match Result = Regex.Match(Output, @"Xcode (\S*)");
+		return Result.Success ? Result.Groups[1].Value : "";
+	}
+	public override string GetAllowedSoftwareVersions()
+	{
+		return @"regex:^13\.5.*";
+	}
 
 	public override bool IsCustomVersionNeeded(string CustomVersionId, string CustomVersionParams)
 	{
 		return true;
 	}
-	public override bool CustomVersionUpdate(string CustomVersionId, string UpdateParams)
+	public override bool CustomVersionUpdate(string CustomVersionId, string UpdateParams, FileRetriever Retriever)
 	{
-		Console.WriteLine("UPDATING CERT FROM " + UpdateParams);
+		if (CustomVersionId == "QueryIOSDeviceType")
+		{
+			Console.WriteLine("Getting DeviceType device with " + UpdateParams);
+
+			// cfgtool needs ECID, not UDID, so find it
+			string Configurator = Path.Combine(GetConfiguratorLocation().Replace(" ", "\\ "), "Contents/MacOS/cfgutil");
+
+			string Device = CommandUtils.ParseParamValue(UpdateParams.Split(' '), "-device=");
+			Console.WriteLine("DeviceId: {0}", Device);
+
+			string Params = string.Format(" -c '{0} list | grep {1}'", Configurator, Device);
+			string Output = UnrealBuildTool.Utils.RunLocalProcessAndReturnStdOut("sh", Params);
+
+			Match Result = Regex.Match(Output, @"Type: (\S*).*ECID: (\S*)");
+			if (!Result.Success)
+			{
+				Console.WriteLine("Unable to find the given deviceid: {0}", Device);
+			}
+
+			// set variables for Installers following this
+			Environment.SetEnvironmentVariable("DeviceModel", Result.Groups[1].Value, EnvironmentVariableTarget.Process);
+			Environment.SetEnvironmentVariable("ECID", Result.Groups[2].Value, EnvironmentVariableTarget.Process);
+			Environment.SetEnvironmentVariable("FlashApplication", Configurator, EnvironmentVariableTarget.Process);
+
+			//// get the flash for this guy
+			//Dictionary<string, string> Variables = new Dictionary<string, string>();
+			//Variables["DeviceType"] = Result.Groups[1].Value;
+			//string ECID = Result.Groups[2].Value;
+
+			//string FlashLocation = Retriever.RetrieveByTags(new string[] { "IOS_IPSW_Location" }, null, Variables);
+
+			//Params = string.Format("-c {0} --ecid {1} update -i {2}", Configurator, ECID, FlashLocation);
+
+			//Console.WriteLine("Update cmdline: {0}", Params);
+		}
 
 		return true;
 	}
 
+	private bool InstallCert(FileRetriever Retriever)
+	{
+		// get the cert password from Studio settings
+		string CertPassword = "Epic123!";
+
+		string CertLoc = Retriever.RetrieveByTags(new string[] { "IOS_Cert_Generic_Development" }, null);
+
+		if (CertLoc != null)
+		{
+			Console.WriteLine("Will install cert from: '{0}'", CertLoc);
+
+			string CommandLine = string.Format("-e 'require \"cert\"; " +
+				"FastlaneCore::KeychainImporter.import_file(\"{0}\", FastlaneCore::Helper.keychain_path(\"login\"), certificate_password:\"{1}\")'",
+				CertLoc, CertPassword);
+			UnrealBuildTool.Utils.RunLocalProcessAndLogOutput("ruby", CommandLine);
+
+			// @todo turnkey verify output
+
+			return true;
+		}
+		else
+		{
+			Console.WriteLine("Unable to find a tagged source for IOS_Cert_Generic_Development");
+
+			return false;
+		}
+
+	}
+
+	private class VerifyIOSSettings
+	{
+		public string CodeSigningIdentity = null;
+		public string BundleId = null;
+		public string Account = null;
+		public string Password = null;
+
+		public string RubyScript = Path.Combine(CommandUtils.EngineDirectory.FullName, "Build/Turnkey/VerifyIOS.ru");
+
+		public VerifyIOSSettings(BuildCommand Command)
+		{
+			FileReference ProjectPath = Command.ParseProjectParam();
+			ConfigHierarchy EngineConfig = ConfigCache.ReadHierarchy(ConfigHierarchyType.Engine, ProjectPath == null ? null : ProjectPath.Directory, UnrealTargetPlatform.IOS);
+
+			// first look for settings on the commandline:
+			CodeSigningIdentity = Command.ParseParamValue("certificate");
+			BundleId = Command.ParseParamValue("bundleid");
+			Account = Command.ParseParamValue("devcenterusername");
+			Password = Command.ParseParamValue("devcenterpassword");
+
+			// fall back to ini for anything else
+			if (string.IsNullOrEmpty(CodeSigningIdentity)) EngineConfig.GetString("/Script/IOSRuntimeSettings.IOSRuntimeSettings", "DevCodeSigningIdentity", out CodeSigningIdentity);
+			if (string.IsNullOrEmpty(BundleId)) EngineConfig.GetString("/Script/IOSRuntimeSettings.IOSRuntimeSettings", "BundleIdentifier", out BundleId);
+			if (string.IsNullOrEmpty(Account)) EngineConfig.GetString("TurnkeySettings", "IOS_DevCenterUsername", out Account);
+			if (string.IsNullOrEmpty(Password)) EngineConfig.GetString("TurnkeySettings", "IOS_DevCenterPassword", out Password);
+
+			string Foo;
+			bool Bar = EngineConfig.GetString("/Script/IOSRuntimeSettings.IOSRuntimeSettings", "DevCodeSigningIdentity", out Foo);
+
+			// some are required
+			if (string.IsNullOrEmpty(CodeSigningIdentity) || string.IsNullOrEmpty(BundleId))
+			{
+				throw new AutomationException("Turnkey IOS verification requires code signing identity (have ='{0}', ex: iPhone Developer) and a bundle id (have '{1}', ex: com.company.foo)", CodeSigningIdentity, BundleId);
+			}
+		}
+
+		public string MakeRubyCommandline(bool bCertOnly, string DeviceName = null)
+		{
+			if (bCertOnly)
+			{
+				return string.Format("--certonly --identity {0}", CodeSigningIdentity);
+			}
+
 
+			string Params = string.Format("--identity {0} --bundleid {1}", CodeSigningIdentity, BundleId);
+
+			if (!string.IsNullOrEmpty(Account))
+			{
+				Params += string.Format(" --login {0}", Account);
+			}
+			if (!string.IsNullOrEmpty(Password))
+			{
+				Params += string.Format(" --password {0}", Password);
+			}
+
+
+			if (!string.IsNullOrEmpty(DeviceName))
+			{
+				Params += string.Format(" --device {0}", DeviceName);
+			}
+
+			return Params;
+		}
+	}
+
+	string GetConfiguratorLocation()
+	{
+		string FindCommand = "-c 'mdfind \"kMDItemKind == Application\" | grep \"Apple Configurator 2.app\"'";
+		return UnrealBuildTool.Utils.RunLocalProcessAndReturnStdOut("sh", FindCommand);
+	}
+
+	public override bool UpdateHostPrerequisites(BuildCommand Command, FileRetriever Retriever)
+	{
+		// make sure the Configurator is installed
+		string ConfiguratorLocation = GetConfiguratorLocation();
+
+		if (ConfiguratorLocation == "")
+		{
+			Console.WriteLine("Apple Configurator 2 is required for some automation to work. You should install it from the App Store. Launching...");
+			
+			// we need to install Configurator 2, and we will block until it's done
+			UnrealBuildTool.Utils.RunLocalProcessAndReturnStdOut("open", "macappstore://apps.apple.com/us/app/apple-configurator-2/id1037126344?mt=12");
+
+			while ((ConfiguratorLocation = GetConfiguratorLocation()) == "")
+			{
+				Thread.Sleep(1000);
+			}
+		}
+		
+		return UpdateDevicePrerequisites(null, Command, Retriever);
+	}
+
+	public override bool UpdateDevicePrerequisites(DeviceInfo Device, BuildCommand Command, FileRetriever Retriever)
+	{
+		VerifyIOSSettings Settings = new VerifyIOSSettings(Command);
+
+		// look if we have a cert that matches it
+		int ExitCode;
+		string StdOut = UnrealBuildTool.Utils.RunLocalProcessAndReturnStdOut(Settings.RubyScript, Settings.MakeRubyCommandline(true), out ExitCode, true);
+
+		if (ExitCode != 0)
+		{
+			if (!InstallCert(Retriever))
+			{
+				return false;
+			}
+		}
+
+		// @todo turnkey - better to use the device's udid if it's set properly in DeviceInfo
+		string DeviceName = Device == null ? null : Device.Name;
+
+		// now look for a provision that can be used with a (maybe newly) instally cert
+		StdOut = UnrealBuildTool.Utils.RunLocalProcessAndReturnStdOut(Settings.RubyScript, Settings.MakeRubyCommandline(false, DeviceName), out ExitCode, true);
+
+		Console.WriteLine("VerifyIOS exited with code {0}", ExitCode);
+
+		return ExitCode == 0;
+	}
+
+
+	public override DeviceInfo[] GetDevices()
+	{
+		string DeviceType = Platform == UnrealTargetPlatform.TVOS ? "tvOS" : "iOS";
+
+		// print out each connected device's udid, os version, and name (which might have spaces)
+		string Params = "-e 'require \"fastlane\"; FastlaneCore::DeviceManager.connected_devices(\"" + DeviceType + "\").each { |x| puts \"#{x.udid} #{x.os_type} #{x.os_version} #{x.name}\" }'";
+		string StdOut = UnrealBuildTool.Utils.RunLocalProcessAndReturnStdOut("ruby", Params);
+
+		List<DeviceInfo> Devices = new List<DeviceInfo>();
+		foreach (string Line in StdOut.Split("\n".ToCharArray(), StringSplitOptions.RemoveEmptyEntries))
+		{
+			string[] Tokens = Line.Split(" ".ToCharArray());
+
+			DeviceInfo Device = new DeviceInfo();
+			Device.Id = Tokens[0];
+			Device.Type = Tokens[1];
+			Device.SoftwareVersion = Tokens[2];
+			Device.Name = string.Join(" ", Tokens.Skip(3));
+
+			Devices.Add(Device);
+		}
+
+		return Devices.ToArray();
+	}
 
 
 
diff --git a/Engine/Source/Programs/AutomationTool/Turnkey/Commands/ManageSettings.cs b/Engine/Source/Programs/AutomationTool/Turnkey/Commands/ManageSettings.cs
--- a/Engine/Source/Programs/AutomationTool/Turnkey/Commands/ManageSettings.cs
+++ b/Engine/Source/Programs/AutomationTool/Turnkey/Commands/ManageSettings.cs
@@ -76,7 +76,8 @@ namespace Turnkey.Commands
 				string NewValue = TurnkeyUtils.ReadInput(string.Format("Enter value for {0} [Currently '{1}']", Settings[Choice - 1].VariableName, TurnkeyUtils.GetVariableValue(Settings[Choice - 1].VariableName)));
 				if (NewValue != "")
 				{
-					TurnkeyUtils.SetVariable(Settings[Choice - 1].VariableName, NewValue);
+					TurnkeyUtils.Log("Setting Settings[{0}].{1} = {2}", Choice - 1, Settings[Choice - 1].VariableName, NewValue);
+					TurnkeySettings.SetUserSetting(Settings[Choice - 1].VariableName, NewValue);
 					TurnkeySettings.Save();
 				}
 			}
diff --git a/Engine/Source/Programs/AutomationTool/Turnkey/Commands/VerifySdk.cs b/Engine/Source/Programs/AutomationTool/Turnkey/Commands/VerifySdk.cs
--- a/Engine/Source/Programs/AutomationTool/Turnkey/Commands/VerifySdk.cs
+++ b/Engine/Source/Programs/AutomationTool/Turnkey/Commands/VerifySdk.cs
@@ -93,6 +93,8 @@ namespace Turnkey.Commands
 			TurnkeyUtils.Log("Installed Sdk validity:");
 			TurnkeyUtils.ExitCode = AutomationTool.ExitCode.Success;
 
+			CopyProviderRetriever Retriever = new CopyProviderRetriever();
+
 			// check all the platforms
 			foreach (var Pair in PlatformsAndDevices)
 			{
@@ -101,11 +103,13 @@ namespace Turnkey.Commands
 				// get the platform object
 				AutomationTool.Platform AutomationPlatform = AutomationTool.Platform.GetPlatform(Platform);
 
+				AutomationPlatform.UpdateHostPrerequisites(TurnkeyUtils.CommandUtilHelper, Retriever);
+
 				SdkInfo.LocalAvailability LocalState = SdkInfo.GetLocalAvailability(AutomationPlatform);
 
 				if ((LocalState & (SdkInfo.LocalAvailability.AutoSdk_ValidVersionExists | SdkInfo.LocalAvailability.InstalledSdk_ValidVersionExists)) == 0)
 				{
-					TurnkeyUtils.Report("{0}: Invalid: [No AutoSdk or Installed Sdk found matching {1} - {2} - '{3}']", Platform, AutomationPlatform.GetAllowedSdks(), LocalState.ToString(), Environment.GetEnvironmentVariable("UE_SDKS_ROOT"));
+					TurnkeyUtils.Report("{0}: Invalid: [No AutoSdk or Installed Sdk found matching {1} - {2}]", Platform, AutomationPlatform.GetAllowedSdks(), LocalState.ToString());
 					TurnkeyUtils.ExitCode = AutomationTool.ExitCode.Error_SDKNotFound;
 				}
 				else
@@ -154,7 +158,7 @@ namespace Turnkey.Commands
 
 						TurnkeyUtils.Log("Will install {0}", BestSdk.DisplayName);
 
-//						bool bWasSdkInstalled = BestSdk.Install(Platform, null, bUnattended);
+						bool bWasSdkInstalled = BestSdk.Install(Platform, null, bUnattended);
 
 						// update LocalState
 						LocalState = SdkInfo.GetLocalAvailability(AutomationPlatform);
@@ -192,6 +196,8 @@ namespace Turnkey.Commands
 					// now check software verison of each device
 					foreach (DeviceInfo Device in Devices)
 					{
+						AutomationPlatform.UpdateDevicePrerequisites(Device, TurnkeyUtils.CommandUtilHelper, Retriever);
+
 						bool bIsSoftwareValid = TurnkeyUtils.IsValueValid(Device.SoftwareVersion, AutomationPlatform.GetAllowedSoftwareVersions());
 
 
diff --git a/Engine/Source/Programs/AutomationTool/Turnkey/CopyProvider.cs b/Engine/Source/Programs/AutomationTool/Turnkey/CopyProvider.cs
--- a/Engine/Source/Programs/AutomationTool/Turnkey/CopyProvider.cs
+++ b/Engine/Source/Programs/AutomationTool/Turnkey/CopyProvider.cs
@@ -206,4 +206,35 @@ namespace Turnkey
 			return Provider.Enumerate(ProviderParam, Expansions);
 		}
 	}
+
+	class CopyProviderRetriever : AutomationTool.FileRetriever
+	{
+		public string RetrieveByTags(string[] RequiredTags, string[] PreferredTags, Dictionary<string, string> ExtraVariables = null)
+		{
+			// @todo turnkey: unset these
+			if (ExtraVariables != null)
+			{
+				foreach (var Pair in ExtraVariables)
+				{
+					TurnkeyUtils.SetVariable(Pair.Key, Pair.Value);
+				}
+			}
+
+			List<SdkInfo> Sdks = TurnkeyManifest.GetDiscoveredSdks();
+
+			Sdks = Sdks.FindAll(x => x.Type == SdkInfo.SdkType.Misc && x.CustomSdkId == RequiredTags[0]);
+
+			SdkInfo Sdk = Sdks.First();
+
+			foreach (CopyAndRun Copy in Sdk.CustomSdkInputFiles)
+			{
+				if (Copy.ShouldExecute())
+				{
+					return CopyProvider.ExecuteCopy(Copy.Copy);
+				}
+			}
+
+			return null;
+		}
+	}
 }
diff --git a/Engine/Source/Programs/AutomationTool/Turnkey/LocalCache.cs b/Engine/Source/Programs/AutomationTool/Turnkey/LocalCache.cs
--- a/Engine/Source/Programs/AutomationTool/Turnkey/LocalCache.cs
+++ b/Engine/Source/Programs/AutomationTool/Turnkey/LocalCache.cs
@@ -162,6 +162,11 @@ namespace Turnkey
 
 		static SavedCache DeserializeObject(string Filename)
 		{
+			if (!File.Exists(Filename))
+			{
+				return null;
+			}
+
 			try
 			{
 				string Str = File.ReadAllText(Filename);
diff --git a/Engine/Source/Programs/AutomationTool/Turnkey/SdkInfo.cs b/Engine/Source/Programs/AutomationTool/Turnkey/SdkInfo.cs
--- a/Engine/Source/Programs/AutomationTool/Turnkey/SdkInfo.cs
+++ b/Engine/Source/Programs/AutomationTool/Turnkey/SdkInfo.cs
@@ -172,16 +172,10 @@ namespace Turnkey
 		{
 			foreach (SdkType Type in TypePriority)
 			{
-				List<SdkInfo> Sdks = TurnkeyManifest.GetDiscoveredSdks().FindAll(x =>
-					x.SupportsPlatform(Platform.IniPlatformType) &&
-					x.Type == Type &&
-					TurnkeyUtils.IsValueValid(x.Version, Type == SdkType.Flash ? Platform.GetAllowedSoftwareVersions() : Platform.GetAllowedSdks())
-					); ;
-
-				Sdks = TurnkeyManifest.GetDiscoveredSdks();
+				List<SdkInfo> Sdks = TurnkeyManifest.GetDiscoveredSdks();
 				Sdks = Sdks.FindAll(x => x.SupportsPlatform(Platform.IniPlatformType));
 				Sdks = Sdks.FindAll(x => x.Type == Type);
-				Sdks = Sdks.FindAll(x => TurnkeyUtils.IsValueValid(x.Version, Type == SdkType.Flash ? Platform.GetAllowedSoftwareVersions() : Platform.GetAllowedSdks()));
+				Sdks = Sdks.FindAll(x => (x.Version == null && x.CustomSdkId != null) || TurnkeyUtils.IsValueValid(x.Version, Type == SdkType.Flash ? Platform.GetAllowedSoftwareVersions() : Platform.GetAllowedSdks()));
 
 				// if none were found try next type
 				if (Sdks.Count == 0)
@@ -202,7 +196,7 @@ namespace Turnkey
 					SdkInfo Best = null;
 					foreach (SdkInfo Sdk in Sdks)
 					{
-						// @todo turnkey: let platform decide what's best
+						// @todo turnkey: let platform decide what's best, and handle no-version custom sdks
 						if (Best == null || string.Compare(Sdk.Version, Best.Version, true) > 0)
 						{
 							Best = Sdk;
@@ -410,7 +404,7 @@ namespace Turnkey
 
 				// re-set the path to local files
 				TurnkeyUtils.SetVariable("CustomVersionLocalFiles", CustomVersionLocalFiles);
-				AutomationPlatforms[Platform].CustomVersionUpdate(CustomSdkId, TurnkeyUtils.ExpandVariables(CustomSdkParams));
+				AutomationPlatforms[Platform].CustomVersionUpdate(CustomSdkId, TurnkeyUtils.ExpandVariables(CustomSdkParams), new CopyProviderRetriever());
 
 				TurnkeyUtils.EndTrackingExternalEnvVarChanges();
 			}
@@ -629,22 +623,32 @@ namespace Turnkey
 			}
 
 			PlatformString = TurnkeyUtils.ExpandVariables(PlatformString, true);
-			string[] PlatformStrings = PlatformString.Split(",".ToCharArray());
 			Platforms = new List<UnrealTargetPlatform>();
 			AutomationPlatforms = new Dictionary<UnrealTargetPlatform, Platform>();
-			// parse into runtime usable values
-			foreach (string Plat in PlatformStrings)
+
+			if (PlatformString != null)
 			{
-				if (UnrealTargetPlatform.IsValidName(Plat))
+				string[] PlatformStrings = PlatformString.Split(",".ToCharArray());
+				// parse into runtime usable values
+				foreach (string Plat in PlatformStrings)
 				{
-					UnrealTargetPlatform TargetPlat = UnrealTargetPlatform.Parse(Plat);
-					Platforms.Add(TargetPlat);
-					AutomationPlatforms.Add(TargetPlat, AutomationTool.Platform.Platforms[new TargetPlatformDescriptor(TargetPlat)]);
+					if (UnrealTargetPlatform.IsValidName(Plat))
+					{
+						UnrealTargetPlatform TargetPlat = UnrealTargetPlatform.Parse(Plat);
+						Platforms.Add(TargetPlat);
+						AutomationPlatforms.Add(TargetPlat, AutomationTool.Platform.Platforms[new TargetPlatformDescriptor(TargetPlat)]);
+					}
 				}
 			}
 
 			Version = TurnkeyUtils.ExpandVariables(Version, true);
 			DisplayName = TurnkeyUtils.ExpandVariables(DisplayName, true);
+
+			if (string.IsNullOrEmpty(DisplayName))
+			{
+				DisplayName = string.Format("{0} {1}", PlatformString, Version);
+			}
+
 			AllowedFlashDeviceTypes = TurnkeyUtils.ExpandVariables(AllowedFlashDeviceTypes, true);
 			CustomSdkId = TurnkeyUtils.ExpandVariables(CustomSdkId, true);
 			CustomSdkParams = TurnkeyUtils.ExpandVariables(CustomSdkParams, true);
@@ -682,8 +686,8 @@ namespace Turnkey
 		{
 			StringBuilder Builder = new StringBuilder();
 			Builder.AppendLine("{1}Name: {0}", DisplayName, Indent(BaseIndent));
-			Builder.AppendLine("{1}Version: {0}", Version, Indent(BaseIndent + 2));
-			Builder.AppendLine("{1}Platform: {0}", string.Join(",", Platforms), Indent(BaseIndent + 2));
+			Builder.AppendLine("{1}Version: {0}", Version == null ? "<Any>" : Version, Indent(BaseIndent + 2));
+			Builder.AppendLine("{1}Platform: {0}", Platforms == null ? "" : string.Join(",", Platforms), Indent(BaseIndent + 2));
 			Builder.AppendLine("{1}Type: {0}", Type, Indent(BaseIndent + 2));
 			if (Type == SdkType.Flash)
 			{
@@ -694,20 +698,26 @@ namespace Turnkey
 				Builder.AppendLine("{0}CustomSdk:", Indent(BaseIndent + 2));
 				Builder.AppendLine("{1}CusttomSdkId: {0}", CustomSdkId, Indent(BaseIndent + 4));
 				Builder.AppendLine("{1}CustomSdkParams: {0}", CustomSdkParams, Indent(BaseIndent + 4));
-				Builder.AppendLine("{0}Input File Sources:", Indent(BaseIndent + 4));
-				foreach (CopyAndRun CustomInputFiles in CustomSdkInputFiles)
+				if (CustomSdkInputFiles != null)
 				{
-					Builder.AppendLine("{1}HostPlatform: {0}", CustomInputFiles.Platform, Indent(BaseIndent + 6));
-					Builder.AppendLine("{1}FileSource: {0}", CustomInputFiles.Copy, Indent(BaseIndent + 8));
+					Builder.AppendLine("{0}Input File Sources:", Indent(BaseIndent + 4));
+					foreach (CopyAndRun CustomInputFiles in CustomSdkInputFiles)
+					{
+						Builder.AppendLine("{1}HostPlatform: {0}", CustomInputFiles.Platform, Indent(BaseIndent + 6));
+						Builder.AppendLine("{1}FileSource: {0}", CustomInputFiles.Copy, Indent(BaseIndent + 8));
+					}
 				}
 			}
-			Builder.AppendLine("{0}Installers:", Indent(BaseIndent + 2));
-			foreach (CopyAndRun Installer in Installers)
+			if (Installers != null)
 			{
-				Builder.AppendLine("{1}HostPlatform: {0}", Installer.Platform, Indent(BaseIndent + 4));
-				Builder.AppendLine("{1}CopyOperation: {0}", Installer.Copy, Indent(BaseIndent + 6));
-				Builder.AppendLine("{1}InstallerPath: {0}", Installer.CommandPath, Indent(BaseIndent + 6));
-				Builder.AppendLine("{1}InstallerCommandLine: {0}", Installer.CommandLine, Indent(BaseIndent + 6));
+				Builder.AppendLine("{0}Installers:", Indent(BaseIndent + 2));
+				foreach (CopyAndRun Installer in Installers)
+				{
+					Builder.AppendLine("{1}HostPlatform: {0}", Installer.Platform, Indent(BaseIndent + 4));
+					Builder.AppendLine("{1}CopyOperation: {0}", Installer.Copy, Indent(BaseIndent + 6));
+					Builder.AppendLine("{1}InstallerPath: {0}", Installer.CommandPath, Indent(BaseIndent + 6));
+					Builder.AppendLine("{1}InstallerCommandLine: {0}", Installer.CommandLine, Indent(BaseIndent + 6));
+				}
 			}
 			return Builder.ToString().TrimEnd();
 
diff --git a/Engine/Source/Programs/AutomationTool/Turnkey/TurnkeyUtils.cs b/Engine/Source/Programs/AutomationTool/Turnkey/TurnkeyUtils.cs
--- a/Engine/Source/Programs/AutomationTool/Turnkey/TurnkeyUtils.cs
+++ b/Engine/Source/Programs/AutomationTool/Turnkey/TurnkeyUtils.cs
@@ -43,8 +43,11 @@ namespace Turnkey
 		public static string GetVariableValue(string Key)
 		{
 			string Value;
-			TurnkeyVariables.TryGetValue(Key, out Value);
-			return Value;
+			if (TurnkeyVariables.TryGetValue(Key, out Value) || TurnkeySettings.GetSetUserSettings().TryGetValue(Key, out Value))
+			{
+				return Value;
+			}
+			return null;
 		}
 
 		public static void ClearVariable(string Key)
