- Missed files in last checkin

diff --git a/Engine/Source/Programs/AutomationTool/AutomationUtils/Platform.cs b/Engine/Source/Programs/AutomationTool/AutomationUtils/Platform.cs
--- a/Engine/Source/Programs/AutomationTool/AutomationUtils/Platform.cs
+++ b/Engine/Source/Programs/AutomationTool/AutomationUtils/Platform.cs
@@ -155,6 +155,9 @@ namespace AutomationTool
 			return new TargetPlatformDescriptor(TargetPlatformType, "");
 		}
 
+
+		#region Turnkey
+
 		public virtual string GetInstalledSdk()
 		{
 			return "";
@@ -196,6 +199,13 @@ namespace AutomationTool
 		{
 			return true;
 		}
+		public virtual bool TryConvertVersionToInt(string StringValue, out UInt64 IntValue)
+		{
+			IntValue = 0;
+			return false;
+		}
+
+		#endregion
 
 		/// <summary>
 		/// Package files for the current platform.
diff --git a/Engine/Source/Programs/AutomationTool/Turnkey/Commands/InstallSdk.cs b/Engine/Source/Programs/AutomationTool/Turnkey/Commands/InstallSdk.cs
--- a/Engine/Source/Programs/AutomationTool/Turnkey/Commands/InstallSdk.cs
+++ b/Engine/Source/Programs/AutomationTool/Turnkey/Commands/InstallSdk.cs
@@ -250,7 +250,7 @@ namespace Turnkey.Commands
 						if (InstallDeviceName == null)
 						{
 							List<string> Options = new List<string>();
-							DeviceInfo[] PossibleDevices = Array.FindAll(AutomationPlatform.GetDevices(), x => TurnkeyUtils.IsValueValid(x.Type, Sdk.AllowedFlashDeviceTypes));
+							DeviceInfo[] PossibleDevices = Array.FindAll(AutomationPlatform.GetDevices(), x => TurnkeyUtils.IsValueValid(x.Type, Sdk.AllowedFlashDeviceTypes, AutomationPlatform));
 							foreach (DeviceInfo Device in PossibleDevices)
 							{
 								Options.Add(string.Format("[{0} {1}] {2}", Platform, Device.Type, Device.Name));
diff --git a/Engine/Source/Programs/AutomationTool/Turnkey/Commands/ListPlatforms.cs b/Engine/Source/Programs/AutomationTool/Turnkey/Commands/ListPlatforms.cs
--- a/Engine/Source/Programs/AutomationTool/Turnkey/Commands/ListPlatforms.cs
+++ b/Engine/Source/Programs/AutomationTool/Turnkey/Commands/ListPlatforms.cs
@@ -21,7 +21,7 @@ namespace Turnkey.Commands
 				string InstalledSdk = Platform.GetInstalledSdk();
 				string AllowedSdks = Platform.GetAllowedSdks();
 				string AllowedSoftware = Platform.GetAllowedSoftwareVersions();
-				bool bIsSdkValid = TurnkeyUtils.IsValueValid(InstalledSdk, AllowedSdks);
+				bool bIsSdkValid = TurnkeyUtils.IsValueValid(InstalledSdk, AllowedSdks, Platform);
 				TurnkeyUtils.Log("  Platform: {0}", TargetPlatform.ToString());
 				TurnkeyUtils.Log("  Installed Sdk: {0}", InstalledSdk);
 				TurnkeyUtils.Log("  Allowed Sdk(s): {0}", AllowedSdks);
@@ -55,7 +55,7 @@ namespace Turnkey.Commands
 				{
 					foreach (DeviceInfo Device in Devices)
 					{
-						bool bIsSoftwareValid = TurnkeyUtils.IsValueValid(Device.SoftwareVersion, AllowedSoftware);
+						bool bIsSoftwareValid = TurnkeyUtils.IsValueValid(Device.SoftwareVersion, AllowedSoftware, Platform);
 
 						TurnkeyUtils.Log("    Name: {0}{1}", Device.Name, Device.bIsDefault ? "*" : "");
 						TurnkeyUtils.Log("      Id: {0}", Device.Id);
diff --git a/Engine/Source/Programs/AutomationTool/Turnkey/TurnkeyUtils.cs b/Engine/Source/Programs/AutomationTool/Turnkey/TurnkeyUtils.cs
--- a/Engine/Source/Programs/AutomationTool/Turnkey/TurnkeyUtils.cs
+++ b/Engine/Source/Programs/AutomationTool/Turnkey/TurnkeyUtils.cs
@@ -164,7 +164,7 @@ namespace Turnkey
 			{
 				List<string> Options = new List<string>();
 				// @todo turnkey: filter devices that have valid flash versions
-				//				DeviceInfo[] PossibleDevices = Array.FindAll(AutomationPlatform.GetDevices(), x => TurnkeyUtils.IsValueValid(x.SoftwareVersion, AutomationPlatform.GetAllowedSoftwareVersions()));
+				//				DeviceInfo[] PossibleDevices = Array.FindAll(AutomationPlatform.GetDevices(), x => TurnkeyUtils.IsValueValid(x.SoftwareVersion, AutomationPlatform.GetAllowedSoftwareVersions(), AutomationPlatforms[Platform]));
 				DeviceInfo[] PossibleDevices = AutomationPlatform.GetDevices();
 				if (PossibleDevices == null)
 				{
@@ -269,7 +269,7 @@ namespace Turnkey
 			return UInt64.TryParse(InValue, out OutValue);
 		}
 
-		public static bool IsValueValid(string Value, string AllowedValues)
+		public static bool IsValueValid(string Value, string AllowedValues, AutomationTool.Platform Platform)
 		{
 			if (string.IsNullOrEmpty(Value) || string.IsNullOrEmpty(AllowedValues))
 			{
@@ -295,11 +295,11 @@ namespace Turnkey
 					return false;
 				}
 
-				// convert inputs to uints
-				UInt64 ValueInt;
-				if (!TryConvertToUint64(Value, out ValueInt))
+				// convert inputs to uint (unless already converted above by platform)
+				UInt64 ValueInt = 0;
+				if (!TryConvertToUint64(Value, out ValueInt) && !Platform.TryConvertVersionToInt(Value, out ValueInt))
 				{
-					TurnkeyUtils.Log("Warning: range: input value [{0}] was not an unsigned integer", Value);
+					TurnkeyUtils.Log("Warning: range: input value [{0}] was not an unsigned integer, and platform couldn't convert it", Value);
 					return false;
 				}
 
