- Fixed some C# code to be compliant with older .NET

diff --git a/Engine/Source/Programs/AutomationTool/Turnkey/Commands/VerifySdk.cs b/Engine/Source/Programs/AutomationTool/Turnkey/Commands/VerifySdk.cs
--- a/Engine/Source/Programs/AutomationTool/Turnkey/Commands/VerifySdk.cs
+++ b/Engine/Source/Programs/AutomationTool/Turnkey/Commands/VerifySdk.cs
@@ -216,7 +216,7 @@ namespace Turnkey.Commands
 
 							if (bUpdateIfNeeded)
 							{
-								SdkInfo MatchingInstallableSdk = SdkInfo.FindMatchingSdk(AutomationPlatform, new SdkInfo.SdkType[] { SdkInfo.SdkType.Flash }, bSelectBest: bUnattended, Device.Type);
+								SdkInfo MatchingInstallableSdk = SdkInfo.FindMatchingSdk(AutomationPlatform, new SdkInfo.SdkType[] { SdkInfo.SdkType.Flash }, bSelectBest: bUnattended, DeviceType: Device.Type);
 
 								if (MatchingInstallableSdk == null)
 								{
