Merging Horde changes for BuildGraph.

#rnx

diff --git a/Engine/Source/Programs/AutomationTool/AutomationUtils/P4Utils.cs b/Engine/Source/Programs/AutomationTool/AutomationUtils/P4Utils.cs
--- a/Engine/Source/Programs/AutomationTool/AutomationUtils/P4Utils.cs
+++ b/Engine/Source/Programs/AutomationTool/AutomationUtils/P4Utils.cs
@@ -541,7 +541,7 @@ namespace AutomationTool
 		/// <summary>
 		/// Size of the file, or -1 if not specified
 		/// </summary>
-		public readonly int FileSize;
+		public readonly long FileSize;
 
 		/// <summary>
 		/// Digest of the file, or null if not specified
@@ -572,7 +572,7 @@ namespace AutomationTool
 		/// <param name="Digest">Digest of the file, or null if not specified</param>
 		/// <param name="Description">Description of the changelist</param>
 		/// <param name="Integrations">Integrations performed to the file</param>
-		public P4RevisionRecord(int RevisionNumber, int ChangeNumber, P4Action Action, string Type, DateTime DateTime, string UserName, string ClientName, int FileSize, string Digest, string Description, P4IntegrationRecord[] Integrations)
+		public P4RevisionRecord(int RevisionNumber, int ChangeNumber, P4Action Action, string Type, DateTime DateTime, string UserName, string ClientName, long FileSize, string Digest, string Description, P4IntegrationRecord[] Integrations)
 		{
 			this.RevisionNumber = RevisionNumber;
 			this.ChangeNumber = ChangeNumber;
@@ -3380,12 +3380,20 @@ namespace AutomationTool
 		/// </summary>
 		/// <param name="UserName"></param>
 		/// <returns>List of clients owned by the user.</returns>
-		public P4ClientInfo[] GetClientsForUser(string UserName, string PathUnderClientRoot = null)
+		public P4ClientInfo[] GetClientsForUser(string UserName, string PathUnderClientRoot = null, string AllowedStream = null)
 		{
 			var ClientList = new List<P4ClientInfo>();
 
 			// Get all clients for this user
-            var P4Result = P4(String.Format("clients -u {0}", UserName), AllowSpew: false, WithClient: false);
+			string P4Command = String.Format("clients -u {0}", UserName);
+
+			// filter by Stream if desired
+			if (AllowedStream != null)
+			{
+				P4Command += " -S " + AllowedStream;
+			}
+
+			var P4Result = P4(P4Command, AllowSpew: false, WithClient: false);
 			if (P4Result.ExitCode != 0)
 			{
 				throw new AutomationException("p4 clients -u {0} failed.", UserName);
@@ -3459,8 +3467,11 @@ namespace AutomationTool
             SpecInput += "Owner: " + ClientSpec.Owner + Environment.NewLine;
             SpecInput += "Host: " + ClientSpec.Host + Environment.NewLine;
             SpecInput += "Root: " + ClientSpec.RootPath + Environment.NewLine;
-            SpecInput += "Options: " + ClientSpec.Options.ToString().ToLowerInvariant().Replace(",", "") + Environment.NewLine;
-            SpecInput += "SubmitOptions: " + ClientSpec.SubmitOptions.ToString().ToLowerInvariant().Replace(", ", "+") + Environment.NewLine;
+			if (ClientSpec.Options != P4ClientOption.None)
+			{
+				SpecInput += "Options: " + ClientSpec.Options.ToString().ToLowerInvariant().Replace(",", "") + Environment.NewLine;
+			}
+			SpecInput += "SubmitOptions: " + ClientSpec.SubmitOptions.ToString().ToLowerInvariant().Replace(", ", "+") + Environment.NewLine;
             SpecInput += "LineEnd: " + ClientSpec.LineEnd.ToString().ToLowerInvariant() + Environment.NewLine;
 			if(ClientSpec.Stream != null)
 			{
@@ -3752,7 +3763,7 @@ namespace AutomationTool
 					string Type = RawRecord["type" + RevisionSuffix];
 					string UserName = RawRecord["user" + RevisionSuffix];
 					string ClientName = RawRecord["client" + RevisionSuffix];
-					int FileSize = RawRecord.ContainsKey("fileSize" + RevisionSuffix)? int.Parse(RawRecord["fileSize" + RevisionSuffix]) : -1;
+					long FileSize = RawRecord.ContainsKey("fileSize" + RevisionSuffix)? long.Parse(RawRecord["fileSize" + RevisionSuffix]) : -1;
 					string Digest = RawRecord.ContainsKey("digest" + RevisionSuffix)? RawRecord["digest" + RevisionSuffix] : null;
 					string Description = RawRecord["desc" + RevisionSuffix];
 
diff --git a/Engine/Source/Programs/AutomationTool/BuildGraph/Aggregate.cs b/Engine/Source/Programs/AutomationTool/BuildGraph/Aggregate.cs
new file mode 100644
--- /dev/null
+++ b/Engine/Source/Programs/AutomationTool/BuildGraph/Aggregate.cs
@@ -0,0 +1,56 @@
+// Copyright Epic Games, Inc. All Rights Reserved.
+
+using System;
+using System.Collections.Generic;
+using System.Linq;
+using System.Text;
+using System.Threading.Tasks;
+
+namespace AutomationTool
+{
+	/// <summary>
+	/// Defines a agggregate within a graph, which give the combined status of one or more job steps, and allow building several steps at once.
+	/// </summary>
+	class Aggregate
+	{
+		/// <summary>
+		/// Name of this badge
+		/// </summary>
+		public readonly string Name;
+
+		/// <summary>
+		/// Label to show for this aggregate in Horde
+		/// </summary>
+		public readonly string Label;
+
+		/// <summary>
+		/// Set of nodes that must be run for this label to be shown.
+		/// </summary>
+		public HashSet<Node> RequiredNodes = new HashSet<Node>();
+
+		/// <summary>
+		/// Set of nodes that will be included in this label if present.
+		/// </summary>
+		public HashSet<Node> IncludedNodes = new HashSet<Node>();
+
+		/// <summary>
+		/// Constructor
+		/// </summary>
+		/// <param name="InName">Name of this aggregate</param>
+		/// <param name="InLabel">Label to show for this aggregate</param>
+		public Aggregate(string InName, string InLabel)
+		{
+			Name = InName;
+			Label = InLabel;
+		}
+
+		/// <summary>
+		/// Get the name of this label
+		/// </summary>
+		/// <returns>The name of this label</returns>
+		public override string ToString()
+		{
+			return Name;
+		}
+	}
+}
diff --git a/Engine/Source/Programs/AutomationTool/BuildGraph/BuildGraph.Automation.csproj b/Engine/Source/Programs/AutomationTool/BuildGraph/BuildGraph.Automation.csproj
--- a/Engine/Source/Programs/AutomationTool/BuildGraph/BuildGraph.Automation.csproj
+++ b/Engine/Source/Programs/AutomationTool/BuildGraph/BuildGraph.Automation.csproj
@@ -67,7 +67,7 @@
     <Compile Include="Agent.cs" />
     <Compile Include="Badge.cs" />
     <Compile Include="BuildGraph.cs" />
-    <Compile Include="Label.cs" />
+    <Compile Include="Aggregate.cs" />
     <Compile Include="Graph.cs" />
     <Compile Include="Node.cs" />
     <Compile Include="Report.cs" />
diff --git a/Engine/Source/Programs/AutomationTool/BuildGraph/BuildGraph.cs b/Engine/Source/Programs/AutomationTool/BuildGraph/BuildGraph.cs
--- a/Engine/Source/Programs/AutomationTool/BuildGraph/BuildGraph.cs
+++ b/Engine/Source/Programs/AutomationTool/BuildGraph/BuildGraph.cs
@@ -77,7 +77,7 @@ namespace AutomationTool
 		{
 			// Parse the command line parameters
 			string ScriptFileName = ParseParamValue("Script", null);
-			string TargetNames = ParseParamValue("Target", null);
+			string[] TargetNames = ParseParamValues("Target").SelectMany(x => x.Split(';', '+').Select(y => y.Trim()).Where(y => y.Length > 0)).ToArray();
 			string DocumentationFileName = ParseParamValue("Documentation", null);
 			string SchemaFileName = ParseParamValue("Schema", null);
 			string ImportSchemaFileName = ParseParamValue("ImportSchema", null);
@@ -244,9 +244,9 @@ namespace AutomationTool
 
 			// Convert the supplied target references into nodes 
 			HashSet<Node> TargetNodes = new HashSet<Node>();
-			if(TargetNames == null)
+			if(TargetNames.Length == 0)
 			{
-				if(!bListOnly)
+				if (!bListOnly && SingleNodeName == null)
 				{
 					LogError("Missing -Target= parameter for BuildGraph");
 					return ExitCode.Error_Unknown;
@@ -265,7 +265,7 @@ namespace AutomationTool
 				}
 				else
 				{
-					NodesToResolve = TargetNames.Split(new char[] { '+', ';' }, StringSplitOptions.RemoveEmptyEntries).Select(x => x.Trim());
+					NodesToResolve = TargetNames;
 				}
 
 				foreach (string TargetName in NodesToResolve)
diff --git a/Engine/Source/Programs/AutomationTool/BuildGraph/Graph.cs b/Engine/Source/Programs/AutomationTool/BuildGraph/Graph.cs
--- a/Engine/Source/Programs/AutomationTool/BuildGraph/Graph.cs
+++ b/Engine/Source/Programs/AutomationTool/BuildGraph/Graph.cs
@@ -158,18 +158,13 @@ namespace AutomationTool
 		/// <summary>
 		/// Mapping of aggregate names to their respective nodes
 		/// </summary>
-		public Dictionary<string, Node[]> AggregateNameToNodes = new Dictionary<string,Node[]>(StringComparer.InvariantCultureIgnoreCase);
+		public Dictionary<string, Aggregate> NameToAggregate = new Dictionary<string, Aggregate>(StringComparer.InvariantCultureIgnoreCase);
 
 		/// <summary>
 		/// List of badges that can be displayed for this build
 		/// </summary>
 		public List<Badge> Badges = new List<Badge>();
 
-		/// <summary>
-		/// List of labels that can be displayed for this build
-		/// </summary>
-		public List<Label> Labels = new List<Label>();
-
 		/// <summary>
 		/// Diagnostic messages for this graph
 		/// </summary>
@@ -189,7 +184,7 @@ namespace AutomationTool
 		/// <returns>True if the name exists, false otherwise.</returns>
 		public bool ContainsName(string Name)
 		{
-			return NameToNode.ContainsKey(Name) || NameToReport.ContainsKey(Name) || AggregateNameToNodes.ContainsKey(Name);
+			return NameToNode.ContainsKey(Name) || NameToReport.ContainsKey(Name) || NameToAggregate.ContainsKey(Name);
 		}
 
 		/// <summary>
@@ -222,10 +217,10 @@ namespace AutomationTool
 				}
 
 				// Check if it's an aggregate name
-				Node[] Nodes;
-				if(AggregateNameToNodes.TryGetValue(Name, out Nodes))
+				Aggregate Aggregate;
+				if(NameToAggregate.TryGetValue(Name, out Aggregate))
 				{
-					OutNodes = Nodes;
+					OutNodes = Aggregate.RequiredNodes.ToArray();
 					return true;
 				}
 
@@ -273,10 +268,10 @@ namespace AutomationTool
 				}
 
 				// Check if it's an aggregate name
-				Node[] Nodes;
-				if(AggregateNameToNodes.TryGetValue(Name, out Nodes))
+				Aggregate Aggregate;
+				if(NameToAggregate.TryGetValue(Name, out Aggregate))
 				{
-					OutOutputs = Nodes.SelectMany(x => x.Outputs.Union(x.Inputs)).Distinct().ToArray();
+					OutOutputs = Aggregate.RequiredNodes.SelectMany(x => x.Outputs.Union(x.Inputs)).Distinct().ToArray();
 					return true;
 				}
 			}
@@ -327,7 +322,16 @@ namespace AutomationTool
 			NameToTrigger = RetainNodes.Where(x => x.ControllingTrigger != null).Select(x => x.ControllingTrigger).Distinct().ToDictionary(x => x.Name, x => x, StringComparer.InvariantCultureIgnoreCase);
 
 			// Create a new list of aggregates for everything that's left
-			AggregateNameToNodes = AggregateNameToNodes.Where(x => x.Value.All(y => RetainNodes.Contains(y))).ToDictionary(Pair => Pair.Key, Pair => Pair.Value, StringComparer.InvariantCultureIgnoreCase);
+			Dictionary<string, Aggregate> NewNameToAggregate = new Dictionary<string, Aggregate>(NameToAggregate.Comparer);
+			foreach(Aggregate Aggregate in NameToAggregate.Values)
+			{
+				if (Aggregate.RequiredNodes.All(x => RetainNodes.Contains(x)))
+				{
+					Aggregate.IncludedNodes.RemoveWhere(x => !RetainNodes.Contains(x));
+					NewNameToAggregate[Aggregate.Name] = Aggregate;
+				}
+			}
+			NameToAggregate = NewNameToAggregate;
 
 			// Remove any badges which do not have all their dependencies
 			Badges.RemoveAll(x => x.Nodes.Any(y => !RetainNodes.Contains(y)));
@@ -398,17 +402,35 @@ namespace AutomationTool
 					Writer.WriteEndElement();
 				}
 
-				foreach (KeyValuePair<string, Node[]> Aggregate in AggregateNameToNodes)
+				foreach (Aggregate Aggregate in NameToAggregate.Values)
 				{
 					// If the aggregate has no required elements, skip it.
-					if (!Aggregate.Value.Any())
+					if (Aggregate.RequiredNodes.Count == 0)
 					{
 						continue;
 					}
 
 					Writer.WriteStartElement("Aggregate");
-					Writer.WriteAttributeString("Name", Aggregate.Key);
-					Writer.WriteAttributeString("Requires", String.Join(";", Aggregate.Value.Select(x => x.Name)));
+					Writer.WriteAttributeString("Name", Aggregate.Name);
+					if (Aggregate.Label != null)
+					{
+						Writer.WriteAttributeString("Label", Aggregate.Label);
+					}
+					Writer.WriteAttributeString("Requires", String.Join(";", Aggregate.RequiredNodes.Select(x => x.Name)));
+
+					if (Aggregate.IncludedNodes.Count > 0)
+					{
+						Writer.WriteAttributeString("Includes", String.Join(";", Aggregate.IncludedNodes.Select(x => x.Name)));
+					}
+
+					HashSet<Node> ExcludedNodes = new HashSet<Node>(Aggregate.IncludedNodes.SelectMany(x => x.InputDependencies));
+					ExcludedNodes.ExceptWith(Aggregate.IncludedNodes);
+
+					if (ExcludedNodes.Count > 0)
+					{
+						Writer.WriteAttributeString("Excludes", String.Join(";", ExcludedNodes.Select(x => x.Name)));
+					}
+
 					Writer.WriteEndElement();
 				}
 
@@ -637,25 +659,25 @@ namespace AutomationTool
 				}
 				JsonWriter.WriteArrayEnd();
 
-				JsonWriter.WriteArrayStart("Labels");
-				foreach (Label Label in Labels)
+				JsonWriter.WriteArrayStart("Aggregates");
+				foreach (Aggregate Aggregate in NameToAggregate.Values)
 				{
 					JsonWriter.WriteObjectStart();
-					if (Label.Category != null)
+					JsonWriter.WriteValue("Name", Aggregate.Name);
+					if (Aggregate.Label != null)
 					{
-						JsonWriter.WriteValue("Category", Label.Category);
+						JsonWriter.WriteValue("Label", Aggregate.Label);
 					}
-					JsonWriter.WriteValue("Name", Label.Name);
 					JsonWriter.WriteArrayStart("RequiredNodes");
-					foreach (Node RequiredNode in Label.RequiredNodes.OrderBy(x => x.Name))
+					foreach (Node RequiredNode in Aggregate.RequiredNodes.OrderBy(x => x.Name))
 					{
 						JsonWriter.WriteValue(RequiredNode.Name);
 					}
 					JsonWriter.WriteArrayEnd();
-					if (Label.IncludedNodes.Count > 0)
+					if (Aggregate.IncludedNodes.Count > 0)
 					{
 						JsonWriter.WriteArrayStart("IncludedNodes");
-						foreach (Node IncludedNode in Label.IncludedNodes.OrderBy(x => x.Name))
+						foreach (Node IncludedNode in Aggregate.IncludedNodes.OrderBy(x => x.Name))
 						{
 							JsonWriter.WriteValue(IncludedNode.Name);
 						}
@@ -783,11 +805,11 @@ namespace AutomationTool
 			CommandUtils.LogInformation("");
 
 			// Print out all the non-empty aggregates
-			string[] AggregateNames = AggregateNameToNodes.Where(x => x.Value.Length > 0).Select(x => x.Key).OrderBy(x => x).ToArray();
-			if(AggregateNames.Length > 0)
+			Aggregate[] Aggregates = NameToAggregate.Values.OrderBy(x => x.Name).ToArray();
+			if(Aggregates.Length > 0)
 			{
 				CommandUtils.LogInformation("Aggregates:");
-				foreach(string AggregateName in AggregateNames)
+				foreach(string AggregateName in Aggregates.Select(x => x.Name))
 				{
 					CommandUtils.LogInformation("    {0}", AggregateName);
 				}
diff --git a/Engine/Source/Programs/AutomationTool/BuildGraph/Label.cs b/Engine/Source/Programs/AutomationTool/BuildGraph/Label.cs
deleted file mode 100644
diff --git a/Engine/Source/Programs/AutomationTool/BuildGraph/Schema.cs b/Engine/Source/Programs/AutomationTool/BuildGraph/Schema.cs
--- a/Engine/Source/Programs/AutomationTool/BuildGraph/Schema.cs
+++ b/Engine/Source/Programs/AutomationTool/BuildGraph/Schema.cs
@@ -155,7 +155,6 @@ namespace AutomationTool
 		Aggregate,
 		Report,
 		Badge,
-		Label,
 		Notify,
 		Include,
 		Option,
@@ -342,7 +341,6 @@ namespace AutomationTool
 			NewSchema.Items.Add(CreateAggregateType());
 			NewSchema.Items.Add(CreateReportType());
 			NewSchema.Items.Add(CreateBadgeType());
-			NewSchema.Items.Add(CreateLabelType());
 			NewSchema.Items.Add(CreateNotifyType());
 			NewSchema.Items.Add(CreateIncludeType());
 			NewSchema.Items.Add(CreateOptionType());
@@ -498,7 +496,6 @@ namespace AutomationTool
 			GraphChoice.Items.Add(CreateSchemaElement("Aggregate", ScriptSchemaStandardType.Aggregate));
 			GraphChoice.Items.Add(CreateSchemaElement("Report", ScriptSchemaStandardType.Report));
 			GraphChoice.Items.Add(CreateSchemaElement("Badge", ScriptSchemaStandardType.Badge));
-			GraphChoice.Items.Add(CreateSchemaElement("Label", ScriptSchemaStandardType.Label));
 			GraphChoice.Items.Add(CreateSchemaElement("Notify", ScriptSchemaStandardType.Notify));
 			GraphChoice.Items.Add(CreateSchemaElement("Trace", ScriptSchemaStandardType.Trace));
 			GraphChoice.Items.Add(CreateSchemaElement("Warning", ScriptSchemaStandardType.Warning));
@@ -675,7 +672,10 @@ namespace AutomationTool
 			XmlSchemaComplexType AggregateType = new XmlSchemaComplexType();
 			AggregateType.Name = GetTypeName(ScriptSchemaStandardType.Aggregate);
 			AggregateType.Attributes.Add(CreateSchemaAttribute("Name", ScriptSchemaStandardType.Name, XmlSchemaUse.Required));
+			AggregateType.Attributes.Add(CreateSchemaAttribute("Label", ScriptSchemaStandardType.BalancedString, XmlSchemaUse.Optional));
 			AggregateType.Attributes.Add(CreateSchemaAttribute("Requires", ScriptSchemaStandardType.NameOrTagList, XmlSchemaUse.Required));
+			AggregateType.Attributes.Add(CreateSchemaAttribute("Include", ScriptSchemaStandardType.NameOrTagList, XmlSchemaUse.Optional));
+			AggregateType.Attributes.Add(CreateSchemaAttribute("Exclude", ScriptSchemaStandardType.NameOrTagList, XmlSchemaUse.Optional));
 			AggregateType.Attributes.Add(CreateSchemaAttribute("If", ScriptSchemaStandardType.BalancedString, XmlSchemaUse.Optional));
 			return AggregateType;
 		}
@@ -711,23 +711,6 @@ namespace AutomationTool
 			return BadgeType;
 		}
 
-		/// <summary>
-		/// Creates the schema type representing the label type
-		/// </summary>
-		/// <returns>Type definition for a label</returns>
-		static XmlSchemaType CreateLabelType()
-		{
-			XmlSchemaComplexType LabelType = new XmlSchemaComplexType();
-			LabelType.Name = GetTypeName(ScriptSchemaStandardType.Label);
-			LabelType.Attributes.Add(CreateSchemaAttribute("Name", ScriptSchemaStandardType.Name, XmlSchemaUse.Required));
-			LabelType.Attributes.Add(CreateSchemaAttribute("Category", ScriptSchemaStandardType.Name, XmlSchemaUse.Optional));
-			LabelType.Attributes.Add(CreateSchemaAttribute("Requires", ScriptSchemaStandardType.NameOrTagList, XmlSchemaUse.Required));
-			LabelType.Attributes.Add(CreateSchemaAttribute("Include", ScriptSchemaStandardType.NameOrTagList, XmlSchemaUse.Optional));
-			LabelType.Attributes.Add(CreateSchemaAttribute("Exclude", ScriptSchemaStandardType.NameOrTagList, XmlSchemaUse.Optional));
-			LabelType.Attributes.Add(CreateSchemaAttribute("If", ScriptSchemaStandardType.BalancedString, XmlSchemaUse.Optional));
-			return LabelType;
-		}
-
 		/// <summary>
 		/// Creates the schema type representing a notifier
 		/// </summary>
diff --git a/Engine/Source/Programs/AutomationTool/BuildGraph/Script.cs b/Engine/Source/Programs/AutomationTool/BuildGraph/Script.cs
--- a/Engine/Source/Programs/AutomationTool/BuildGraph/Script.cs
+++ b/Engine/Source/Programs/AutomationTool/BuildGraph/Script.cs
@@ -377,9 +377,6 @@ namespace AutomationTool
 					case "Badge":
 						ReadBadge(ChildElement);
 						break;
-					case "Label":
-						ReadLabel(ChildElement);
-						break;
 					case "Notify":
 						ReadNotifier(ChildElement);
 						break;
@@ -949,8 +946,30 @@ namespace AutomationTool
 			string Name;
 			if (EvaluateCondition(Element) && TryReadObjectName(Element, out Name) && CheckNameIsUnique(Element, Name))
 			{
+				string Label = ReadAttribute(Element, "Label");
+
 				string[] RequiredNames = ReadListAttribute(Element, "Requires");
-				Graph.AggregateNameToNodes.Add(Name, ResolveReferences(Element, RequiredNames).ToArray());
+				string[] IncludedNames = ReadListAttribute(Element, "Include");
+				string[] ExcludedNames = ReadListAttribute(Element, "Exclude");
+
+				Aggregate NewAggregate = new Aggregate(Name, Label);
+				foreach (Node ReferencedNode in ResolveReferences(Element, RequiredNames))
+				{
+					NewAggregate.RequiredNodes.Add(ReferencedNode);
+					NewAggregate.IncludedNodes.Add(ReferencedNode);
+					NewAggregate.IncludedNodes.UnionWith(ReferencedNode.OrderDependencies);
+				}
+				foreach (Node IncludedNode in ResolveReferences(Element, IncludedNames))
+				{
+					NewAggregate.IncludedNodes.Add(IncludedNode);
+					NewAggregate.IncludedNodes.UnionWith(IncludedNode.OrderDependencies);
+				}
+				foreach (Node ExcludedNode in ResolveReferences(Element, ExcludedNames))
+				{
+					NewAggregate.IncludedNodes.Remove(ExcludedNode);
+					NewAggregate.IncludedNodes.ExceptWith(ExcludedNode.OrderDependencies);
+				}
+				Graph.NameToAggregate[Name] = NewAggregate;
 			}
 		}
 
@@ -1003,42 +1022,6 @@ namespace AutomationTool
 			}
 		}
 
-		/// <summary>
-		/// Reads the definition for a label
-		/// </summary>
-		/// <param name="Element">Xml element to read the definition from</param>
-		void ReadLabel(ScriptElement Element)
-		{
-			string Name;
-			if (EvaluateCondition(Element) && TryReadObjectName(Element, out Name))
-			{
-				string Category = ReadAttribute(Element, "Category");
-
-				string[] RequiredNames = ReadListAttribute(Element, "Requires");
-				string[] IncludedNames = ReadListAttribute(Element, "Include");
-				string[] ExcludedNames = ReadListAttribute(Element, "Exclude");
-
-				Label NewLabel = new Label(Name, Category);
-				foreach (Node ReferencedNode in ResolveReferences(Element, RequiredNames))
-				{
-					NewLabel.RequiredNodes.Add(ReferencedNode);
-					NewLabel.IncludedNodes.Add(ReferencedNode);
-					NewLabel.IncludedNodes.UnionWith(ReferencedNode.OrderDependencies);
-				}
-				foreach (Node IncludedNode in ResolveReferences(Element, IncludedNames))
-				{
-					NewLabel.IncludedNodes.Add(IncludedNode);
-					NewLabel.IncludedNodes.UnionWith(IncludedNode.OrderDependencies);
-				}
-				foreach (Node ExcludedNode in ResolveReferences(Element, ExcludedNames))
-				{
-					NewLabel.IncludedNodes.Remove(ExcludedNode);
-					NewLabel.IncludedNodes.ExceptWith(ExcludedNode.OrderDependencies);
-				}
-				Graph.Labels.Add(NewLabel);
-			}
-		}
-
 		/// <summary>
 		/// Reads the definition for a node, and adds it to the given agent
 		/// </summary>
