Correcting attribute names for Include and Exclude when writing out Aggregate elements for build graphs.

diff --git a/Engine/Source/Programs/AutomationTool/BuildGraph/Graph.cs b/Engine/Source/Programs/AutomationTool/BuildGraph/Graph.cs
--- a/Engine/Source/Programs/AutomationTool/BuildGraph/Graph.cs
+++ b/Engine/Source/Programs/AutomationTool/BuildGraph/Graph.cs
@@ -420,7 +420,7 @@ namespace AutomationTool
 
 					if (Aggregate.IncludedNodes.Count > 0)
 					{
-						Writer.WriteAttributeString("Includes", String.Join(";", Aggregate.IncludedNodes.Select(x => x.Name)));
+						Writer.WriteAttributeString("Include", String.Join(";", Aggregate.IncludedNodes.Select(x => x.Name)));
 					}
 
 					HashSet<Node> ExcludedNodes = new HashSet<Node>(Aggregate.IncludedNodes.SelectMany(x => x.InputDependencies));
@@ -428,7 +428,7 @@ namespace AutomationTool
 
 					if (ExcludedNodes.Count > 0)
 					{
-						Writer.WriteAttributeString("Excludes", String.Join(";", ExcludedNodes.Select(x => x.Name)));
+						Writer.WriteAttributeString("Exclude", String.Join(";", ExcludedNodes.Select(x => x.Name)));
 					}
 
 					Writer.WriteEndElement();
