Reintroduce support for adding BuildGraph labels separately to aggregates. Syntax is:

    <Label Category="Clients" Name="Windows" Requires="Compile UE4Client Win64" Exclude="Compile UnrealHeaderTool Win64"/>

diff --git a/Engine/Source/Programs/AutomationTool/BuildGraph/Aggregate.cs b/Engine/Source/Programs/AutomationTool/BuildGraph/Aggregate.cs
--- a/Engine/Source/Programs/AutomationTool/BuildGraph/Aggregate.cs
+++ b/Engine/Source/Programs/AutomationTool/BuildGraph/Aggregate.cs
@@ -18,30 +18,18 @@ namespace AutomationTool
 		/// </summary>
 		public readonly string Name;
 
-		/// <summary>
-		/// Label to show for this aggregate in Horde
-		/// </summary>
-		public readonly string Label;
-
 		/// <summary>
 		/// Set of nodes that must be run for this label to be shown.
 		/// </summary>
 		public HashSet<Node> RequiredNodes = new HashSet<Node>();
 
-		/// <summary>
-		/// Set of nodes that will be included in this label if present.
-		/// </summary>
-		public HashSet<Node> IncludedNodes = new HashSet<Node>();
-
 		/// <summary>
 		/// Constructor
 		/// </summary>
 		/// <param name="InName">Name of this aggregate</param>
-		/// <param name="InLabel">Label to show for this aggregate</param>
-		public Aggregate(string InName, string InLabel)
+		public Aggregate(string InName)
 		{
 			Name = InName;
-			Label = InLabel;
 		}
 
 		/// <summary>
diff --git a/Engine/Source/Programs/AutomationTool/BuildGraph/BuildGraph.Automation.csproj b/Engine/Source/Programs/AutomationTool/BuildGraph/BuildGraph.Automation.csproj
--- a/Engine/Source/Programs/AutomationTool/BuildGraph/BuildGraph.Automation.csproj
+++ b/Engine/Source/Programs/AutomationTool/BuildGraph/BuildGraph.Automation.csproj
@@ -69,6 +69,7 @@
     <Compile Include="BuildGraph.cs" />
     <Compile Include="Aggregate.cs" />
     <Compile Include="Graph.cs" />
+    <Compile Include="Label.cs" />
     <Compile Include="Node.cs" />
     <Compile Include="Report.cs" />
     <Compile Include="Task.cs" />
diff --git a/Engine/Source/Programs/AutomationTool/BuildGraph/Graph.cs b/Engine/Source/Programs/AutomationTool/BuildGraph/Graph.cs
--- a/Engine/Source/Programs/AutomationTool/BuildGraph/Graph.cs
+++ b/Engine/Source/Programs/AutomationTool/BuildGraph/Graph.cs
@@ -165,6 +165,11 @@ namespace AutomationTool
 		/// </summary>
 		public List<Badge> Badges = new List<Badge>();
 
+		/// <summary>
+		/// List of labels that can be displayed for this build
+		/// </summary>
+		public List<Label> Labels = new List<Label>();
+
 		/// <summary>
 		/// Diagnostic messages for this graph
 		/// </summary>
@@ -327,12 +332,19 @@ namespace AutomationTool
 			{
 				if (Aggregate.RequiredNodes.All(x => RetainNodes.Contains(x)))
 				{
-					Aggregate.IncludedNodes.RemoveWhere(x => !RetainNodes.Contains(x));
 					NewNameToAggregate[Aggregate.Name] = Aggregate;
 				}
 			}
 			NameToAggregate = NewNameToAggregate;
 
+			// Remove any labels that are no longer value
+			foreach (Label Label in Labels)
+			{
+				Label.RequiredNodes.RemoveWhere(x => !RetainNodes.Contains(x));
+				Label.IncludedNodes.RemoveWhere(x => !RetainNodes.Contains(x));
+			}
+			Labels.RemoveAll(x => x.RequiredNodes.Count == 0);
+
 			// Remove any badges which do not have all their dependencies
 			Badges.RemoveAll(x => x.Nodes.Any(y => !RetainNodes.Contains(y)));
 
@@ -412,25 +424,36 @@ namespace AutomationTool
 
 					Writer.WriteStartElement("Aggregate");
 					Writer.WriteAttributeString("Name", Aggregate.Name);
-					if (Aggregate.Label != null)
-					{
-						Writer.WriteAttributeString("Label", Aggregate.Label);
-					}
 					Writer.WriteAttributeString("Requires", String.Join(";", Aggregate.RequiredNodes.Select(x => x.Name)));
+					Writer.WriteEndElement();
+				}
 
-					if (Aggregate.IncludedNodes.Count > 0)
+				foreach(Label Label in Labels)
+				{
+					Writer.WriteStartElement("Label");
+					if (Label.Category != null)
 					{
-						Writer.WriteAttributeString("Include", String.Join(";", Aggregate.IncludedNodes.Select(x => x.Name)));
+						Writer.WriteAttributeString("Category", Label.Category);
 					}
+					Writer.WriteAttributeString("Name", Label.Name);
+					Writer.WriteAttributeString("Requires", String.Join(";", Label.RequiredNodes.Select(x => x.Name)));
 
-					HashSet<Node> ExcludedNodes = new HashSet<Node>(Aggregate.IncludedNodes.SelectMany(x => x.InputDependencies));
-					ExcludedNodes.ExceptWith(Aggregate.IncludedNodes);
+					HashSet<Node> IncludedNodes = new HashSet<Node>(Label.IncludedNodes);
+					IncludedNodes.ExceptWith(Label.IncludedNodes.SelectMany(x => x.InputDependencies));
+					IncludedNodes.ExceptWith(Label.RequiredNodes);
+					if (IncludedNodes.Count > 0)
+					{
+						Writer.WriteAttributeString("Include", String.Join(";", IncludedNodes.Select(x => x.Name)));
+					}
 
+					HashSet<Node> ExcludedNodes = new HashSet<Node>(Label.IncludedNodes);
+					ExcludedNodes.UnionWith(Label.IncludedNodes.SelectMany(x => x.InputDependencies));
+					ExcludedNodes.ExceptWith(Label.IncludedNodes);
+					ExcludedNodes.ExceptWith(ExcludedNodes.ToArray().SelectMany(x => x.InputDependencies));
 					if (ExcludedNodes.Count > 0)
 					{
 						Writer.WriteAttributeString("Exclude", String.Join(";", ExcludedNodes.Select(x => x.Name)));
 					}
-
 					Writer.WriteEndElement();
 				}
 
@@ -664,25 +687,37 @@ namespace AutomationTool
 				{
 					JsonWriter.WriteObjectStart();
 					JsonWriter.WriteValue("Name", Aggregate.Name);
-					if (Aggregate.Label != null)
+					JsonWriter.WriteArrayStart("Nodes");
+					foreach (Node RequiredNode in Aggregate.RequiredNodes.OrderBy(x => x.Name))
 					{
-						JsonWriter.WriteValue("Label", Aggregate.Label);
+						JsonWriter.WriteValue(RequiredNode.Name);
 					}
+					JsonWriter.WriteArrayEnd();
+					JsonWriter.WriteObjectEnd();
+				}
+				JsonWriter.WriteArrayEnd();
+
+				JsonWriter.WriteArrayStart("Labels");
+				foreach (Label Label in Labels)
+				{
+					JsonWriter.WriteObjectStart();
+					if (Label.Category != null)
+					{
+						JsonWriter.WriteValue("Category", Label.Category);
+					}
+					JsonWriter.WriteValue("Name", Label.Name);
 					JsonWriter.WriteArrayStart("RequiredNodes");
-					foreach (Node RequiredNode in Aggregate.RequiredNodes.OrderBy(x => x.Name))
+					foreach (Node RequiredNode in Label.RequiredNodes.OrderBy(x => x.Name))
 					{
 						JsonWriter.WriteValue(RequiredNode.Name);
 					}
 					JsonWriter.WriteArrayEnd();
-					if (Aggregate.IncludedNodes.Count > 0)
+					JsonWriter.WriteArrayStart("IncludedNodes");
+					foreach (Node IncludedNode in Label.IncludedNodes.OrderBy(x => x.Name))
 					{
-						JsonWriter.WriteArrayStart("IncludedNodes");
-						foreach (Node IncludedNode in Aggregate.IncludedNodes.OrderBy(x => x.Name))
-						{
-							JsonWriter.WriteValue(IncludedNode.Name);
-						}
-						JsonWriter.WriteArrayEnd();
+						JsonWriter.WriteValue(IncludedNode.Name);
 					}
+					JsonWriter.WriteArrayEnd();
 					JsonWriter.WriteObjectEnd();
 				}
 				JsonWriter.WriteArrayEnd();
diff --git a/Engine/Source/Programs/AutomationTool/BuildGraph/Schema.cs b/Engine/Source/Programs/AutomationTool/BuildGraph/Schema.cs
--- a/Engine/Source/Programs/AutomationTool/BuildGraph/Schema.cs
+++ b/Engine/Source/Programs/AutomationTool/BuildGraph/Schema.cs
@@ -155,6 +155,7 @@ namespace AutomationTool
 		Aggregate,
 		Report,
 		Badge,
+		Label,
 		Notify,
 		Include,
 		Option,
@@ -341,6 +342,7 @@ namespace AutomationTool
 			NewSchema.Items.Add(CreateAggregateType());
 			NewSchema.Items.Add(CreateReportType());
 			NewSchema.Items.Add(CreateBadgeType());
+			NewSchema.Items.Add(CreateLabelType());
 			NewSchema.Items.Add(CreateNotifyType());
 			NewSchema.Items.Add(CreateIncludeType());
 			NewSchema.Items.Add(CreateOptionType());
@@ -496,6 +498,7 @@ namespace AutomationTool
 			GraphChoice.Items.Add(CreateSchemaElement("Aggregate", ScriptSchemaStandardType.Aggregate));
 			GraphChoice.Items.Add(CreateSchemaElement("Report", ScriptSchemaStandardType.Report));
 			GraphChoice.Items.Add(CreateSchemaElement("Badge", ScriptSchemaStandardType.Badge));
+			GraphChoice.Items.Add(CreateSchemaElement("Label", ScriptSchemaStandardType.Label));
 			GraphChoice.Items.Add(CreateSchemaElement("Notify", ScriptSchemaStandardType.Notify));
 			GraphChoice.Items.Add(CreateSchemaElement("Trace", ScriptSchemaStandardType.Trace));
 			GraphChoice.Items.Add(CreateSchemaElement("Warning", ScriptSchemaStandardType.Warning));
@@ -711,6 +714,23 @@ namespace AutomationTool
 			return BadgeType;
 		}
 
+		/// <summary>
+		/// Creates the schema type representing the label type
+		/// </summary>
+		/// <returns>Type definition for a label</returns>
+		static XmlSchemaType CreateLabelType()
+		{
+			XmlSchemaComplexType LabelType = new XmlSchemaComplexType();
+			LabelType.Name = GetTypeName(ScriptSchemaStandardType.Label);
+			LabelType.Attributes.Add(CreateSchemaAttribute("Name", ScriptSchemaStandardType.Name, XmlSchemaUse.Required));
+			LabelType.Attributes.Add(CreateSchemaAttribute("Category", ScriptSchemaStandardType.Name, XmlSchemaUse.Optional));
+			LabelType.Attributes.Add(CreateSchemaAttribute("Requires", ScriptSchemaStandardType.NameOrTagList, XmlSchemaUse.Required));
+			LabelType.Attributes.Add(CreateSchemaAttribute("Include", ScriptSchemaStandardType.NameOrTagList, XmlSchemaUse.Optional));
+			LabelType.Attributes.Add(CreateSchemaAttribute("Exclude", ScriptSchemaStandardType.NameOrTagList, XmlSchemaUse.Optional));
+			LabelType.Attributes.Add(CreateSchemaAttribute("If", ScriptSchemaStandardType.BalancedString, XmlSchemaUse.Optional));
+			return LabelType;
+		}
+
 		/// <summary>
 		/// Creates the schema type representing a notifier
 		/// </summary>
diff --git a/Engine/Source/Programs/AutomationTool/BuildGraph/Script.cs b/Engine/Source/Programs/AutomationTool/BuildGraph/Script.cs
--- a/Engine/Source/Programs/AutomationTool/BuildGraph/Script.cs
+++ b/Engine/Source/Programs/AutomationTool/BuildGraph/Script.cs
@@ -377,6 +377,9 @@ namespace AutomationTool
 					case "Badge":
 						ReadBadge(ChildElement);
 						break;
+					case "Label":
+						ReadLabel(ChildElement);
+						break;
 					case "Notify":
 						ReadNotifier(ChildElement);
 						break;
@@ -949,30 +952,55 @@ namespace AutomationTool
 			string Name;
 			if (EvaluateCondition(Element) && TryReadObjectName(Element, out Name) && CheckNameIsUnique(Element, Name))
 			{
-				string Label = ReadAttribute(Element, "Label");
-
 				string[] RequiredNames = ReadListAttribute(Element, "Requires");
-				string[] IncludedNames = ReadListAttribute(Element, "Include");
-				string[] ExcludedNames = ReadListAttribute(Element, "Exclude");
 
-				Aggregate NewAggregate = new Aggregate(Name, Label);
+				Aggregate NewAggregate = new Aggregate(Name);
 				foreach (Node ReferencedNode in ResolveReferences(Element, RequiredNames))
 				{
 					NewAggregate.RequiredNodes.Add(ReferencedNode);
-					NewAggregate.IncludedNodes.Add(ReferencedNode);
-					NewAggregate.IncludedNodes.UnionWith(ReferencedNode.OrderDependencies);
-				}
-				foreach (Node IncludedNode in ResolveReferences(Element, IncludedNames))
-				{
-					NewAggregate.IncludedNodes.Add(IncludedNode);
-					NewAggregate.IncludedNodes.UnionWith(IncludedNode.OrderDependencies);
 				}
-				foreach (Node ExcludedNode in ResolveReferences(Element, ExcludedNames))
+				Graph.NameToAggregate[Name] = NewAggregate;
+
+				string LabelCategoryName = ReadAttribute(Element, "Label");
+				if (!String.IsNullOrEmpty(LabelCategoryName))
 				{
-					NewAggregate.IncludedNodes.Remove(ExcludedNode);
-					NewAggregate.IncludedNodes.ExceptWith(ExcludedNode.OrderDependencies);
+					Label Label;
+
+					// Create the label
+					int SlashIdx = LabelCategoryName.IndexOf('/');
+					if (SlashIdx != -1)
+					{
+						Label = new Label(LabelCategoryName.Substring(0, SlashIdx), LabelCategoryName.Substring(SlashIdx + 1));
+					}
+					else
+					{
+						Label = new Label("Other", LabelCategoryName);
+					}
+
+					// Find all the included nodes
+					foreach (Node RequiredNode in NewAggregate.RequiredNodes)
+					{
+						Label.RequiredNodes.Add(RequiredNode);
+						Label.IncludedNodes.Add(RequiredNode);
+						Label.IncludedNodes.UnionWith(RequiredNode.OrderDependencies);
+					}
+
+					string[] IncludedNames = ReadListAttribute(Element, "Include");
+					foreach (Node IncludedNode in ResolveReferences(Element, IncludedNames))
+					{
+						Label.IncludedNodes.Add(IncludedNode);
+						Label.IncludedNodes.UnionWith(IncludedNode.OrderDependencies);
+					}
+
+					string[] ExcludedNames = ReadListAttribute(Element, "Exclude");
+					foreach (Node ExcludedNode in ResolveReferences(Element, ExcludedNames))
+					{
+						Label.IncludedNodes.Remove(ExcludedNode);
+						Label.IncludedNodes.ExceptWith(ExcludedNode.OrderDependencies);
+					}
+
+					Graph.Labels.Add(Label);
 				}
-				Graph.NameToAggregate[Name] = NewAggregate;
 			}
 		}
 
@@ -1025,6 +1053,42 @@ namespace AutomationTool
 			}
 		}
 
+		/// <summary>
+		/// Reads the definition for a label
+		/// </summary>
+		/// <param name="Element">Xml element to read the definition from</param>
+		void ReadLabel(ScriptElement Element)
+		{
+			string Name;
+			if (EvaluateCondition(Element) && TryReadObjectName(Element, out Name))
+			{
+				string Category = ReadAttribute(Element, "Category");
+
+				string[] RequiredNames = ReadListAttribute(Element, "Requires");
+				string[] IncludedNames = ReadListAttribute(Element, "Include");
+				string[] ExcludedNames = ReadListAttribute(Element, "Exclude");
+
+				Label NewLabel = new Label(Name, Category);
+				foreach (Node ReferencedNode in ResolveReferences(Element, RequiredNames))
+				{
+					NewLabel.RequiredNodes.Add(ReferencedNode);
+					NewLabel.IncludedNodes.Add(ReferencedNode);
+					NewLabel.IncludedNodes.UnionWith(ReferencedNode.OrderDependencies);
+				}
+				foreach (Node IncludedNode in ResolveReferences(Element, IncludedNames))
+				{
+					NewLabel.IncludedNodes.Add(IncludedNode);
+					NewLabel.IncludedNodes.UnionWith(IncludedNode.OrderDependencies);
+				}
+				foreach (Node ExcludedNode in ResolveReferences(Element, ExcludedNames))
+				{
+					NewLabel.IncludedNodes.Remove(ExcludedNode);
+					NewLabel.IncludedNodes.ExceptWith(ExcludedNode.OrderDependencies);
+				}
+				Graph.Labels.Add(NewLabel);
+			}
+		}
+
 		/// <summary>
 		/// Reads the definition for a node, and adds it to the given agent
 		/// </summary>
