Add missing file.

diff --git a/Engine/Source/Programs/AutomationTool/BuildGraph/Label.cs b/Engine/Source/Programs/AutomationTool/BuildGraph/Label.cs
new file mode 100644
--- /dev/null
+++ b/Engine/Source/Programs/AutomationTool/BuildGraph/Label.cs
@@ -0,0 +1,57 @@
+// Copyright Epic Games, Inc. All Rights Reserved.
+
+using System;
+using System.Collections.Generic;
+using System.Linq;
+using System.Text;
+using System.Threading.Tasks;
+
+namespace AutomationTool
+{
+	/// <summary>
+	/// Defines a label within a graph. Labels are similar to badges, and give the combined status of one or more job steps. Unlike badges, they
+	/// separate the requirements for its status and optional nodes to be included in its status, allowing this to be handled externally.
+	/// </summary>
+	class Label
+	{
+		/// <summary>
+		/// Category for this label
+		/// </summary>
+		public readonly string Category;
+
+		/// <summary>
+		/// Name of this badge
+		/// </summary>
+		public readonly string Name;
+
+		/// <summary>
+		/// Set of nodes that must be run for this label to be shown.
+		/// </summary>
+		public HashSet<Node> RequiredNodes = new HashSet<Node>();
+
+		/// <summary>
+		/// Set of nodes that will be included in this label if present.
+		/// </summary>
+		public HashSet<Node> IncludedNodes = new HashSet<Node>();
+
+		/// <summary>
+		/// Constructor
+		/// </summary>
+		/// <param name="InCategory">Type of this label</param>
+		/// <param name="InName">Name of this label</param>
+		public Label(string InCategory, string InName)
+		{
+			Category = InCategory;
+			Name = InName;
+		}
+
+		/// <summary>
+		/// Get the name of this label
+		/// </summary>
+		/// <returns>The name of this label</returns>
+		public override string ToString()
+		{
+			return String.Format("{0}/{1}", Category, Name);
+		}
+	}
+}
