flip name and category the right way

diff --git a/Engine/Source/Programs/AutomationTool/BuildGraph/Script.cs b/Engine/Source/Programs/AutomationTool/BuildGraph/Script.cs
--- a/Engine/Source/Programs/AutomationTool/BuildGraph/Script.cs
+++ b/Engine/Source/Programs/AutomationTool/BuildGraph/Script.cs
@@ -1068,7 +1068,7 @@ namespace AutomationTool
 				string[] IncludedNames = ReadListAttribute(Element, "Include");
 				string[] ExcludedNames = ReadListAttribute(Element, "Exclude");
 
-				Label NewLabel = new Label(Name, Category);
+				Label NewLabel = new Label(Category, Name);
 				foreach (Node ReferencedNode in ResolveReferences(Element, RequiredNames))
 				{
 					NewLabel.RequiredNodes.Add(ReferencedNode);
