- Removed unnecessary MobileDeviceInterface reference from AutomationTool project

diff --git a/Engine/Source/Programs/AutomationTool/AutomationTool.csproj b/Engine/Source/Programs/AutomationTool/AutomationTool.csproj
--- a/Engine/Source/Programs/AutomationTool/AutomationTool.csproj
+++ b/Engine/Source/Programs/AutomationTool/AutomationTool.csproj
@@ -72,10 +72,6 @@
       <Name>DotNETUtilities</Name>
       <Private>False</Private>
     </ProjectReference>
-    <ProjectReference Include="..\IOS\MobileDeviceInterface\MobileDeviceInterface.csproj">
-      <Project>{a5ca2153-90b2-42e4-9547-341c28a7726a}</Project>
-      <Name>MobileDeviceInterface</Name>
-    </ProjectReference>
     <ProjectReference Include="..\UnrealBuildTool\UnrealBuildTool.csproj">
       <Project>{fd7c5e1a-cfe4-4fd5-a525-1eb1599a39ac}</Project>
       <Name>UnrealBuildTool</Name>
