- Removed MobileDeviceInterface.dll from AddUATFilesToBuildProducts ()

diff --git a/Engine/Source/Programs/AutomationTool/AutomationUtils/UE4Build.cs b/Engine/Source/Programs/AutomationTool/AutomationUtils/UE4Build.cs
--- a/Engine/Source/Programs/AutomationTool/AutomationUtils/UE4Build.cs
+++ b/Engine/Source/Programs/AutomationTool/AutomationUtils/UE4Build.cs
@@ -1326,7 +1326,6 @@ namespace AutomationTool
 						"AutomationUtils.Automation.dll",
 						"DotNETUtilities.dll",
 						"BuildUtilities.dll",
-						"MobileDeviceInterface.dll"
 					});
 
 			foreach (var UATFile in UATFiles)
