- Fixed HybridIO provider on Mac (unknown yet how well WinForms works in Linux)

diff --git a/Engine/Source/Programs/AutomationTool/Turnkey/Providers/HybridIOProvider.cs b/Engine/Source/Programs/AutomationTool/Turnkey/Providers/HybridIOProvider.cs
--- a/Engine/Source/Programs/AutomationTool/Turnkey/Providers/HybridIOProvider.cs
+++ b/Engine/Source/Programs/AutomationTool/Turnkey/Providers/HybridIOProvider.cs
@@ -6,8 +6,8 @@ using System.Linq;
 using System.Text;
 using System.Threading.Tasks;
 using System.Windows.Forms;
+using System.Text.RegularExpressions;
 using System.Drawing;
-using System.Drawing.Configuration;
 using Tools.DotNETCommon;
 
 namespace Turnkey
@@ -21,7 +21,7 @@ namespace Turnkey
 		static HybridIOProvider()
 		{
 			// make the form look good on modern displays!
-			if (Environment.OSVersion.Version.Major >= 6)
+			if (!UnrealBuildTool.Utils.IsRunningOnMono && Environment.OSVersion.Version.Major >= 6)
 			{
 				SetProcessDPIAware();
 			}
@@ -30,14 +30,45 @@ namespace Turnkey
 			Application.SetCompatibleTextRenderingDefault(false);
 		}
 
+		private string ShowMacDialog(string Prompt, string Default)
+		{
+			string Params = string.Format("-e 'display dialog \"{0}\" with title \"Turnkey Input\" default answer \"{1}\"'", Prompt.Replace("\n", "\\n").Replace("\t", "\\t"), Default);
+			string OSAOutput = UnrealBuildTool.Utils.RunLocalProcessAndReturnStdOut("osascript", Params);
+
+			// blank string means user canceled, which goes to stderr
+			if (OSAOutput == "")
+			{
+				return "";
+			}
+
+			// regex the result
+			Match Match = Regex.Match(OSAOutput, "text returned:(.*)$");
+			if (!Match.Success)
+			{
+				return "";
+			}
+
+			// return the text in the dialog box
+			return Match.Groups[1].Value;
+		}
 
-		private  string ShowDialog(string Prompt, string Default, bool bIsList)
+			private string ShowDialog(string Prompt, string Default, bool bIsList)
 		{
-			string Result = Microsoft.VisualBasic.Interaction.InputBox(Prompt, "Turnkey", Default);
+			string Result;
+			if (UnrealBuildTool.BuildHostPlatform.Current.Platform == UnrealBuildTool.UnrealTargetPlatform.Mac)
+			{
+				Result = ShowMacDialog(Prompt, Default);
+			}
+			else
+			{
+				Result = Microsoft.VisualBasic.Interaction.InputBox(Prompt, "Turnkey Input", Default);
+			}
+
 			if (Result == "" && bIsList)
 			{
 				return "0";
 			}
+
 			return Result;
 /*
 			// Create a new instance of the form.
@@ -141,9 +172,8 @@ namespace Turnkey
 			}
 
 			Form1.Dispose();
-
 			return Result;
-	*/
+*/
 		}
 
 		public override string ReadInput(string Prompt, string Default, bool bAppendNewLine)
diff --git a/Engine/Source/Runtime/Core/Private/GenericPlatform/GenericPlatformProcess.cpp b/Engine/Source/Runtime/Core/Private/GenericPlatform/GenericPlatformProcess.cpp
--- a/Engine/Source/Runtime/Core/Private/GenericPlatform/GenericPlatformProcess.cpp
+++ b/Engine/Source/Runtime/Core/Private/GenericPlatform/GenericPlatformProcess.cpp
@@ -304,7 +304,10 @@ void FGenericPlatformProcess::ModifyCreateProcParams(FString& InOutURL, FString&
 #if PLATFORM_WINDOWS
 		InOutURL = TEXT("cmd.exe");
 		InOutParams = FString::Printf(TEXT("/c %s %s"), *Command, *Arguments);
-#elif PLATFORM_MAC || PLATFORM_LINUX
+#elif PLATFORM_MAC
+		InOutURL = TEXT("/bin/sh");
+		InOutParams = FString::Printf(TEXT("%s %s"), *Command, *Arguments);
+#elif PLATFORM_LINUX
 		InOutURL = TEXT("/bin/bash");
 		InOutParams = FString::Printf(TEXT("-c '%s %s'"), *Command, *Arguments);
 #endif
