- Update IOS prerequisite handling (now fully interactive, even if not super smooth, with AppleScript to show dialog/Terminal windows)
- Added support for smb:// paths on Mac (mounts a volume and uses /Volumes, C# doesn't understand UNC paths)
- Updated path to .p12, and the signing identity from Apple Developer to Apple Development (as seen in latest certs)
- Added a new SdkInfo state (InvalidPrerequisites)
- Made IOS prerequisites not actually do any fixing if -UpdateIfNeeded is not specified

diff --git a/Engine/Config/BaseEngine.ini b/Engine/Config/BaseEngine.ini
--- a/Engine/Config/BaseEngine.ini
+++ b/Engine/Config/BaseEngine.ini
@@ -2058,7 +2058,7 @@ bUseRSync=True
 BundleDisplayName=[PROJECT_NAME]
 BundleName=[PROJECT_NAME]
 BundleIdentifier=com.YourCXompany.[PROJECT_NAME]
-DevCodeSigningIdentity=iPhone Developer:;Apple Developer:
+DevCodeSigningIdentity=iPhone Developer:;Apple Development:
 DistroCodeSigningIdentity=iPhone Distribution:;Apple Distribution:
 VersionInfo=1.0
 FrameRateLock=PUFRL_30
diff --git a/Engine/Source/Programs/AutomationTool/AutomationUtils/Platform.cs b/Engine/Source/Programs/AutomationTool/AutomationUtils/Platform.cs
--- a/Engine/Source/Programs/AutomationTool/AutomationUtils/Platform.cs
+++ b/Engine/Source/Programs/AutomationTool/AutomationUtils/Platform.cs
@@ -16,6 +16,11 @@ namespace AutomationTool
 		string RetrieveByTags(string[] RequiredTags, string[] PreferredTags, Dictionary<string,string> ExtraVariables=null);
 	}
 
+	//public interface InputOutput
+	//{
+	//	string RetrieveByTags(string[] RequiredTags, string[] PreferredTags, Dictionary<string, string> ExtraVariables = null);
+	//}
+
 	public class DeviceInfo
 	{
 		public string Name;
@@ -176,12 +181,12 @@ namespace AutomationTool
 			return true;
 		}
 
-		public virtual bool UpdateHostPrerequisites(BuildCommand Command, FileRetriever Retriever)
+		public virtual bool UpdateHostPrerequisites(BuildCommand Command, FileRetriever Retriever, bool bVerifyOnly)
 		{
 			return true;
 		}
 
-		public virtual bool UpdateDevicePrerequisites(DeviceInfo Device, BuildCommand Command, FileRetriever Retriever)
+		public virtual bool UpdateDevicePrerequisites(DeviceInfo Device, BuildCommand Command, FileRetriever Retriever, bool bVerifyOnly)
 		{
 			return true;
 		}
diff --git a/Engine/Source/Programs/AutomationTool/IOS/IOSPlatform.Automation.cs b/Engine/Source/Programs/AutomationTool/IOS/IOSPlatform.Automation.cs
--- a/Engine/Source/Programs/AutomationTool/IOS/IOSPlatform.Automation.cs
+++ b/Engine/Source/Programs/AutomationTool/IOS/IOSPlatform.Automation.cs
@@ -259,6 +259,7 @@ public class IOSPlatform : Platform
 		public string BundleId = null;
 		public string Account = null;
 		public string Password = null;
+		public string Team = null;
 
 		public string RubyScript = Path.Combine(CommandUtils.EngineDirectory.FullName, "Build/Turnkey/VerifyIOS.ru");
 
@@ -272,12 +273,14 @@ public class IOSPlatform : Platform
 			BundleId = Command.ParseParamValue("bundleid");
 			Account = Command.ParseParamValue("devcenterusername");
 			Password = Command.ParseParamValue("devcenterpassword");
+			Team = Command.ParseParamValue("teamid");
 
 			// fall back to ini for anything else
 			if (string.IsNullOrEmpty(CodeSigningIdentity)) EngineConfig.GetString("/Script/IOSRuntimeSettings.IOSRuntimeSettings", "DevCodeSigningIdentity", out CodeSigningIdentity);
 			if (string.IsNullOrEmpty(BundleId)) EngineConfig.GetString("/Script/IOSRuntimeSettings.IOSRuntimeSettings", "BundleIdentifier", out BundleId);
 			if (string.IsNullOrEmpty(Account)) EngineConfig.GetString("TurnkeySettings", "IOS_DevCenterUsername", out Account);
 			if (string.IsNullOrEmpty(Password)) EngineConfig.GetString("TurnkeySettings", "IOS_DevCenterPassword", out Password);
+			if (string.IsNullOrEmpty(Team)) EngineConfig.GetString("TurnkeySettings", "IOS_DevCenterPassword", out Password);
 
 			string Foo;
 			bool Bar = EngineConfig.GetString("/Script/IOSRuntimeSettings.IOSRuntimeSettings", "DevCodeSigningIdentity", out Foo);
@@ -289,32 +292,73 @@ public class IOSPlatform : Platform
 			}
 		}
 
-		public string MakeRubyCommandline(bool bCertOnly, string DeviceName = null)
+		public int RunRubyCommand(bool bCertOnly, bool bVerifyOnly, string DeviceName)
 		{
+			string Params;
+			int ExitCode;
+
 			if (bCertOnly)
 			{
-				return string.Format("--certonly --identity {0}", CodeSigningIdentity);
+				Params = string.Format("--certonly --identity '{0}'", CodeSigningIdentity);
 			}
+			else
+			{
+				Params = string.Format("--identity '{0}' --bundleid {1}", CodeSigningIdentity, BundleId);
 
+				if (!string.IsNullOrEmpty(Account))
+				{
+					Params += string.Format(" --login {0}", Account);
+				}
+				if (!string.IsNullOrEmpty(Password))
+				{
+					Params += string.Format(" --password {0}", Password);
+				}
+				if (!string.IsNullOrEmpty(Team))
+				{
+					Params += string.Format(" --team {0}", Team);
+				}
 
-			string Params = string.Format("--identity {0} --bundleid {1}", CodeSigningIdentity, BundleId);
 
-			if (!string.IsNullOrEmpty(Account))
-			{
-				Params += string.Format(" --login {0}", Account);
+				if (!string.IsNullOrEmpty(DeviceName))
+				{
+					Params += string.Format(" --device {0}", DeviceName);
+				}
+
+				if (bVerifyOnly)
+				{
+					Params += string.Format(" --verifyonly");
+				}
 			}
-			if (!string.IsNullOrEmpty(Password))
+
+			// if non-interactive, we can just run directly in the current shell
+			if (bVerifyOnly)
 			{
-				Params += string.Format(" --password {0}", Password);
+				UnrealBuildTool.Utils.RunLocalProcessAndReturnStdOut(RubyScript, Params, out ExitCode, true);
 			}
-
-
-			if (!string.IsNullOrEmpty(DeviceName))
+			else
 			{
-				Params += string.Format(" --device {0}", DeviceName);
+				// otherwise, run in a new Terminal window via AppleScript
+				string ReturnCodeFilename = Path.GetTempFileName();
+
+				// run potentially interactive scripts in a Terminal window
+				Params = string.Format(
+						" -e \"tell application \\\"Terminal\\\"\"" +
+						" -e   \"activate\"" +
+						" -e   \"set newTab to do script (\\\"{0} {1}; echo $? > {2}; exit\\\")\"" +
+						" -e   \"repeat\"" +
+						" -e     \"delay 1\"" +
+						" -e     \"if not busy of newTab then exit repeat\"" +
+						" -e   \"end repeat\"" +
+						" -e \"end tell\"",
+						RubyScript, Params, ReturnCodeFilename);
+
+				UnrealBuildTool.Utils.RunLocalProcessAndReturnStdOut("osascript", Params, out ExitCode, true);
+				if (ExitCode == 0)
+				{
+					ExitCode = int.Parse(File.ReadAllText(ReturnCodeFilename));
+				}
 			}
-
-			return Params;
+			return ExitCode;
 		}
 	}
 
@@ -324,11 +368,13 @@ public class IOSPlatform : Platform
 		return UnrealBuildTool.Utils.RunLocalProcessAndReturnStdOut("sh", FindCommand);
 	}
 
-	public override bool UpdateHostPrerequisites(BuildCommand Command, FileRetriever Retriever)
+	public override bool UpdateHostPrerequisites(BuildCommand Command, FileRetriever Retriever, bool bVerifyOnly)
 	{
+		int ExitCode;
+
 		if (HostPlatform.Current.HostEditorPlatform != UnrealTargetPlatform.Mac)
 		{
-			return base.UpdateHostPrerequisites(Command, Retriever);
+			return base.UpdateHostPrerequisites(Command, Retriever, bVerifyOnly);
 		}
 
 		// make sure the Configurator is installed
@@ -336,8 +382,13 @@ public class IOSPlatform : Platform
 
 		if (ConfiguratorLocation == "")
 		{
+			if (bVerifyOnly)
+			{
+				return false;
+			}
+
 			Console.WriteLine("Apple Configurator 2 is required for some automation to work. You should install it from the App Store. Launching...");
-			
+
 			// we need to install Configurator 2, and we will block until it's done
 			UnrealBuildTool.Utils.RunLocalProcessAndReturnStdOut("open", "macappstore://apps.apple.com/us/app/apple-configurator-2/id1037126344?mt=12");
 
@@ -346,38 +397,66 @@ public class IOSPlatform : Platform
 				Thread.Sleep(1000);
 			}
 		}
-		
-		return UpdateDevicePrerequisites(null, Command, Retriever);
-	}
 
-	public override bool UpdateDevicePrerequisites(DeviceInfo Device, BuildCommand Command, FileRetriever Retriever)
-	{
-		if (HostPlatform.Current.HostEditorPlatform != UnrealTargetPlatform.Mac)
+		string IsFastlaneInstalled = UnrealBuildTool.Utils.RunLocalProcessAndReturnStdOut("/usr/bin/gem", "list -ie fastlane");
+		if (IsFastlaneInstalled != "true")
 		{
-			return base.UpdateDevicePrerequisites(Device, Command, Retriever);
+			Console.WriteLine("Fastlane is not installed");
+			if (bVerifyOnly)
+			{
+				return false;
+			}
+
+			Console.WriteLine("Installing Fastlane from internet source. You may ignore the error about the bin directory not in your path.");
+
+			// install missing fastlane without needing sudo
+			UnrealBuildTool.Utils.RunLocalProcessAndReturnStdOut("/usr/bin/gem", "install fastlane --user-install --no-document", out ExitCode, true);
+
+			if (ExitCode != 0)
+			{
+				return false;
+			}
 		}
 
 		VerifyIOSSettings Settings = new VerifyIOSSettings(Command);
 
 		// look if we have a cert that matches it
-		int ExitCode;
-		string StdOut = UnrealBuildTool.Utils.RunLocalProcessAndReturnStdOut(Settings.RubyScript, Settings.MakeRubyCommandline(true), out ExitCode, true);
+		ExitCode = Settings.RunRubyCommand(true, bVerifyOnly, null);
+
+		Console.WriteLine("Host VerifyIOS exited with code {0}", ExitCode);
 
 		if (ExitCode != 0)
 		{
+			if (bVerifyOnly)
+			{
+				return false;
+			}
+
 			if (!InstallCert(Retriever))
 			{
 				return false;
 			}
 		}
 
+		return true;
+	}
+
+	public override bool UpdateDevicePrerequisites(DeviceInfo Device, BuildCommand Command, FileRetriever Retriever, bool bVerifyOnly)
+	{
+		if (HostPlatform.Current.HostEditorPlatform != UnrealTargetPlatform.Mac)
+		{
+			return base.UpdateDevicePrerequisites(Device, Command, Retriever, bVerifyOnly);
+		}
+
+		VerifyIOSSettings Settings = new VerifyIOSSettings(Command);
+
 		// @todo turnkey - better to use the device's udid if it's set properly in DeviceInfo
 		string DeviceName = Device == null ? null : Device.Name;
 
 		// now look for a provision that can be used with a (maybe newly) instally cert
-		StdOut = UnrealBuildTool.Utils.RunLocalProcessAndReturnStdOut(Settings.RubyScript, Settings.MakeRubyCommandline(false, DeviceName), out ExitCode, true);
+		int ExitCode = Settings.RunRubyCommand(false, bVerifyOnly, DeviceName);
 
-		Console.WriteLine("VerifyIOS exited with code {0}", ExitCode);
+		Console.WriteLine("Device VerifyIOS exited with code {0}", ExitCode);
 
 		return ExitCode == 0;
 	}
diff --git a/Engine/Source/Programs/AutomationTool/Turnkey/Commands/VerifySdk.cs b/Engine/Source/Programs/AutomationTool/Turnkey/Commands/VerifySdk.cs
--- a/Engine/Source/Programs/AutomationTool/Turnkey/Commands/VerifySdk.cs
+++ b/Engine/Source/Programs/AutomationTool/Turnkey/Commands/VerifySdk.cs
@@ -105,11 +105,13 @@ namespace Turnkey.Commands
 				// get the platform object
 				AutomationTool.Platform AutomationPlatform = AutomationTool.Platform.GetPlatform(Platform);
 
-				AutomationPlatform.UpdateHostPrerequisites(TurnkeyUtils.CommandUtilHelper, Retriever);
+				SdkInfo.LocalAvailability LocalState = SdkInfo.GetLocalAvailability(AutomationPlatform, bUpdateIfNeeded);
 
-				SdkInfo.LocalAvailability LocalState = SdkInfo.GetLocalAvailability(AutomationPlatform);
-
-				if ((LocalState & (SdkInfo.LocalAvailability.AutoSdk_ValidVersionExists | SdkInfo.LocalAvailability.InstalledSdk_ValidVersionExists)) == 0)
+				if ((LocalState & SdkInfo.LocalAvailability.Platform_ValidHostPrerequisites) == 0)
+				{
+					TurnkeyUtils.Report("{0}: Invalid: [Host Prerequisites are not valid]", Platform);
+				}
+				else if ((LocalState & (SdkInfo.LocalAvailability.AutoSdk_ValidVersionExists | SdkInfo.LocalAvailability.InstalledSdk_ValidVersionExists)) == 0)
 				{
 					string MinAllowedVersion, MaxAllowedVersion;
 					UEBuildPlatformSDK.GetSDKForPlatform(Platform.ToString()).GetValidVersionRange(out MinAllowedVersion, out MaxAllowedVersion);
@@ -156,10 +158,10 @@ namespace Turnkey.Commands
 
 					TurnkeyUtils.Log("Will install {0}", BestSdk.DisplayName);
 
-					bool bWasSdkInstalled = BestSdk.Install(Platform, null, bUnattended);
+					bool bWasSdkInstalled = BestSdk.Install(Platform, null, true);
 
 					// update LocalState
-					LocalState = SdkInfo.GetLocalAvailability(AutomationPlatform);
+					LocalState = SdkInfo.GetLocalAvailability(AutomationPlatform, false);
 
 					// @todo turnkey: validate!
 				}
@@ -186,7 +188,13 @@ namespace Turnkey.Commands
 					// now check software verison of each device
 					foreach (DeviceInfo Device in Devices)
 					{
-						AutomationPlatform.UpdateDevicePrerequisites(Device, TurnkeyUtils.CommandUtilHelper, Retriever);
+						bool bArePrerequisitesValid = AutomationPlatform.UpdateDevicePrerequisites(Device, TurnkeyUtils.CommandUtilHelper, Retriever, !bUpdateIfNeeded);
+
+						if (!bArePrerequisitesValid)
+						{
+							TurnkeyUtils.Report("{0}@{1}: Invalid: [Device Prerequisites are not valid]", Platform, Device.Name, Device.SoftwareVersion, AutomationPlatform.GetAllowedSoftwareVersions());
+							continue;
+						}
 
 						bool bIsSoftwareValid = TurnkeyUtils.IsValueValid(Device.SoftwareVersion, AutomationPlatform.GetAllowedSoftwareVersions(), AutomationPlatform);
 
diff --git a/Engine/Source/Programs/AutomationTool/Turnkey/Providers/NullCopyProvider.cs b/Engine/Source/Programs/AutomationTool/Turnkey/Providers/NullCopyProvider.cs
--- a/Engine/Source/Programs/AutomationTool/Turnkey/Providers/NullCopyProvider.cs
+++ b/Engine/Source/Programs/AutomationTool/Turnkey/Providers/NullCopyProvider.cs
@@ -3,6 +3,7 @@
 using System.IO;
 using System.Collections.Generic;
 using System.Linq;
+using System.Text.RegularExpressions;
 
 namespace Turnkey
 {
@@ -10,10 +11,32 @@ namespace Turnkey
 	{
 		public override string ProviderToken { get { return "file"; } }
 
-		public override string Execute(string Operation, CopyExecuteSpecialMode SpecialMode, string SpecialModeHint)
+		private void FixupOperation(ref string Operation)
 		{
 			Operation = Operation.Replace(Path.AltDirectorySeparatorChar, Path.DirectorySeparatorChar);
 
+			if (UnrealBuildTool.BuildHostPlatform.Current.Platform == UnrealBuildTool.UnrealTargetPlatform.Mac && Operation.StartsWith("smb://"))
+			{
+				// match the form smb://server.net/Foo/Bar/Baz and retrieve the Foo, which will be the name in /Volumes/Foo
+				Match Match = Regex.Match(Operation, @"^smb:\/\/.+?\/(.+?)\/(.*)$");
+
+				// make sure regex matched
+				if (!Match.Success)
+				{
+					return;
+				}
+
+				// convert smb://server.net/Foo/Bar/Baz to /Volumes/Foo/Bar/Baz
+				TurnkeyUtils.Log("SMB Before: {0}", Operation);
+				Operation = string.Format("/Volumes/{0}/{1}", Match.Groups[1].Value, Match.Groups[2].Value);
+				TurnkeyUtils.Log("SMB After: {0}", Operation);
+			}
+		}
+
+		public override string Execute(string Operation, CopyExecuteSpecialMode SpecialMode, string SpecialModeHint)
+		{
+			FixupOperation(ref Operation);
+
 			// this provider can use the file directly, so just return the input after expanding variables, if it exists
 			string OutputPath = TurnkeyUtils.ExpandVariables(Operation);
 
@@ -125,8 +148,7 @@ namespace Turnkey
 		{
 			Dictionary<string, List<string>> Output = new Dictionary<string, List<string>>();
 
-			// we want consistent slashes
-			Operation = Operation.Replace(Path.AltDirectorySeparatorChar, Path.DirectorySeparatorChar);
+			FixupOperation(ref Operation);
 
 			ExpandWildcards("", Operation, Output, Expansions == null ? null : new List<string>());
 
diff --git a/Engine/Source/Programs/AutomationTool/Turnkey/SdkInfo.cs b/Engine/Source/Programs/AutomationTool/Turnkey/SdkInfo.cs
--- a/Engine/Source/Programs/AutomationTool/Turnkey/SdkInfo.cs
+++ b/Engine/Source/Programs/AutomationTool/Turnkey/SdkInfo.cs
@@ -89,17 +89,19 @@ namespace Turnkey
 			}
 
 			// allow for shared SDK "platform" types
-			List<string> Platforms = TurnkeyUtils.ExpandVariables(Clone.PlatformString).ToLower().Split(",".ToCharArray()).ToList();
-			foreach (UEBuildPlatformSDK SDK in UEBuildPlatformSDK.AllSDKs)
+			if (!string.IsNullOrEmpty(Clone.PlatformString))
 			{
-				// if the platforms contains a AutoSDK platform that doesn't match a real platform name, add the platform to the set
-				if (Platforms.Contains(SDK.GetAutoSDKPlatformName().ToLower()) && !Platforms.Contains(SDK.PlatformName.ToLower()))
+				List<string> Platforms = TurnkeyUtils.ExpandVariables(Clone.PlatformString).ToLower().Split(",".ToCharArray()).ToList();
+				foreach (UEBuildPlatformSDK SDK in UEBuildPlatformSDK.AllSDKs)
 				{
-					Clone.PlatformString += "," + SDK.PlatformName;
-					Platforms.Add(SDK.PlatformName.ToLower()); ;
+					// if the platforms contains a AutoSDK platform that doesn't match a real platform name, add the platform to the set
+					if (Platforms.Contains(SDK.GetAutoSDKPlatformName().ToLower()) && !Platforms.Contains(SDK.PlatformName.ToLower()))
+					{
+						Clone.PlatformString += "," + SDK.PlatformName;
+						Platforms.Add(SDK.PlatformName.ToLower()); ;
+					}
+				}
 				}
-			}
-
 			Clone.PostDeserialize();
 
 			return Clone;
@@ -124,12 +126,23 @@ namespace Turnkey
 			// InstalledSdk_BuildOnlyWasInstalled = 8,
 			InstalledSdk_ValidVersionExists = 16,
 			InstalledSdk_InvalidVersionExists = 32,
+			Platform_ValidHostPrerequisites = 64,
+			Platform_InvalidHostPrerequisites = 128,
 		}
 
-		static public LocalAvailability GetLocalAvailability(AutomationTool.Platform AutomationPlatform)
+		static public LocalAvailability GetLocalAvailability(AutomationTool.Platform AutomationPlatform, bool bAllowUpdatingPrerequisites)
 		{
 			LocalAvailability Result = LocalAvailability.None;
 
+			if (AutomationPlatform.UpdateHostPrerequisites(TurnkeyUtils.CommandUtilHelper, null, !bAllowUpdatingPrerequisites))
+			{
+				Result |= LocalAvailability.Platform_ValidHostPrerequisites;
+			}
+			else
+			{
+				Result |= LocalAvailability.Platform_InvalidHostPrerequisites;
+			}
+
 			UEBuildPlatformSDK SDK = UEBuildPlatformSDK.GetSDKForPlatform(AutomationPlatform.PlatformType.ToString());
 
 			// if we have the variable at all, 
@@ -149,7 +162,7 @@ namespace Turnkey
 						// check for valid version number (auto SDK must match version exactly)
 						if (string.Compare(Dir.Name, SDK.GetDesiredVersion()) == 0)
 						{
-							// make sure it actually has bits in it
+							// make sure it actually has buits in it
 							if (File.Exists(Path.Combine(Dir.FullName, "setup.bat")))
 							{
 								bValidVersionFound = true;
