- Close the popup windows for IOS prerequisities if there were no errors

diff --git a/Engine/Source/Programs/AutomationTool/IOS/IOSPlatform.Automation.cs b/Engine/Source/Programs/AutomationTool/IOS/IOSPlatform.Automation.cs
--- a/Engine/Source/Programs/AutomationTool/IOS/IOSPlatform.Automation.cs
+++ b/Engine/Source/Programs/AutomationTool/IOS/IOSPlatform.Automation.cs
@@ -345,10 +345,15 @@ public class IOSPlatform : Platform
 						" -e \"tell application \\\"Terminal\\\"\"" +
 						" -e   \"activate\"" +
 						" -e   \"set newTab to do script (\\\"{0} {1}; echo $? > {2}; exit\\\")\"" +
+						" -e   \"set newWindow to window 1\"" +
 						" -e   \"repeat\"" +
 						" -e     \"delay 1\"" +
 						" -e     \"if not busy of newTab then exit repeat\"" +
 						" -e   \"end repeat\"" +
+						" -e   \"set exitCode to item 1 of paragraphs of (read \\\"{2}\\\")\"" +
+						" -e   \"if exitCode is equal to \\\"0\\\" then\"" +
+						" -e     \"close newWindow\"" +
+						" -e   \"end if\"" +
 						" -e \"end tell\"",
 						RubyScript, Params, ReturnCodeFilename);
 
diff --git a/Engine/Source/Programs/AutomationTool/Turnkey/SdkInfo.cs b/Engine/Source/Programs/AutomationTool/Turnkey/SdkInfo.cs
--- a/Engine/Source/Programs/AutomationTool/Turnkey/SdkInfo.cs
+++ b/Engine/Source/Programs/AutomationTool/Turnkey/SdkInfo.cs
@@ -134,7 +134,10 @@ namespace Turnkey
 		{
 			LocalAvailability Result = LocalAvailability.None;
 
-			if (AutomationPlatform.UpdateHostPrerequisites(TurnkeyUtils.CommandUtilHelper, null, !bAllowUpdatingPrerequisites))
+			// no need to retrieve anything if we aren't allowing updates
+			CopyProviderRetriever Retriever = bAllowUpdatingPrerequisites ? new CopyProviderRetriever() : null;
+
+			if (AutomationPlatform.UpdateHostPrerequisites(TurnkeyUtils.CommandUtilHelper, Retriever, !bAllowUpdatingPrerequisites))
 			{
 				Result |= LocalAvailability.Platform_ValidHostPrerequisites;
 			}
