Add a TraceSpan class that can be used to record timing information for UAT, and add tracing info for BuildGraph into it. Writes to a file in the directory specified by the UE_TELEMETRY_DIR environment variable.

diff --git a/Engine/Source/Programs/AutomationTool/AutomationUtils/Automation.cs b/Engine/Source/Programs/AutomationTool/AutomationUtils/Automation.cs
--- a/Engine/Source/Programs/AutomationTool/AutomationUtils/Automation.cs
+++ b/Engine/Source/Programs/AutomationTool/AutomationUtils/Automation.cs
@@ -505,14 +505,22 @@ AutomationTool.exe [-verbose] [-compileonly] [-p4] Command0 [-Arg0 -Arg1 -Arg2 .
 				CommandUtils.InitDefaultP4Connection();
 			}
 
-			// Find and execute commands.
-			ExitCode Result = Execute(CommandsToExecute, ScriptCompiler.Commands);
-			if (TelemetryFile != null)
+			try
 			{
-				Directory.CreateDirectory(Path.GetDirectoryName(TelemetryFile));
-				CommandUtils.Telemetry.Write(TelemetryFile);
+				// Find and execute commands.
+				ExitCode Result = Execute(CommandsToExecute, ScriptCompiler.Commands);
+				if (TelemetryFile != null)
+				{
+					Directory.CreateDirectory(Path.GetDirectoryName(TelemetryFile));
+					CommandUtils.Telemetry.Write(TelemetryFile);
+				}
+				return Result;
+			}
+			finally
+			{
+				// Flush any timing data
+				TraceSpan.Flush();
 			}
-			return Result;
 		}
 
 		/// <summary>
diff --git a/Engine/Source/Programs/AutomationTool/BuildGraph/BuildGraph.cs b/Engine/Source/Programs/AutomationTool/BuildGraph/BuildGraph.cs
--- a/Engine/Source/Programs/AutomationTool/BuildGraph/BuildGraph.cs
+++ b/Engine/Source/Programs/AutomationTool/BuildGraph/BuildGraph.cs
@@ -2,6 +2,7 @@
 
 using System;
 using System.Collections.Generic;
+using System.Diagnostics;
 using System.IO;
 using System.Linq;
 using System.Reflection;
@@ -709,10 +710,15 @@ namespace AutomationTool
 
 			// Read the manifests for all the input storage blocks
 			Dictionary<TempStorageBlock, TempStorageManifest> InputManifests = new Dictionary<TempStorageBlock, TempStorageManifest>();
-			foreach(TempStorageBlock InputStorageBlock in InputStorageBlocks)
+			using (ITraceSpan Span = TraceSpan.Create("tempstorage.get"))
 			{
-				TempStorageManifest Manifest = Storage.Retreive(InputStorageBlock.NodeName, InputStorageBlock.OutputName);
-				InputManifests[InputStorageBlock] = Manifest;
+				Span.AddMetadata("blocks", InputStorageBlocks.Count.ToString());
+				foreach (TempStorageBlock InputStorageBlock in InputStorageBlocks)
+				{
+					TempStorageManifest Manifest = Storage.Retreive(InputStorageBlock.NodeName, InputStorageBlock.OutputName);
+					InputManifests[InputStorageBlock] = Manifest;
+				}
+				Span.AddMetadata("size", string.Format("{0:n0}", InputManifests.Sum(x => x.Value.GetTotalSize())));
 			}
 
 			// Read all the input storage blocks, keeping track of which block each file came from
@@ -831,32 +837,35 @@ namespace AutomationTool
 			}
 
 			// Write all the storage blocks, and update the mapping from file to storage block
-			foreach(KeyValuePair<string, HashSet<FileReference>> Pair in OutputStorageBlockToFiles)
+			using (ITraceSpan Span = Tools.DotNETCommon.TraceSpan.Create("tempstorage.put"))
 			{
-				TempStorageBlock OutputBlock = new TempStorageBlock(Node.Name, Pair.Key);
-				foreach(FileReference File in Pair.Value)
+				foreach (KeyValuePair<string, HashSet<FileReference>> Pair in OutputStorageBlockToFiles)
 				{
-					FileToStorageBlock.Add(File, OutputBlock);
+					TempStorageBlock OutputBlock = new TempStorageBlock(Node.Name, Pair.Key);
+					foreach (FileReference File in Pair.Value)
+					{
+						FileToStorageBlock.Add(File, OutputBlock);
+					}
+					Storage.Archive(Node.Name, Pair.Key, Pair.Value.ToArray(), Pair.Value.Any(x => ReferencedOutputFiles.Contains(x)));
 				}
-				Storage.Archive(Node.Name, Pair.Key, Pair.Value.ToArray(), Pair.Value.Any(x => ReferencedOutputFiles.Contains(x)));
-			}
-
-			// Publish all the output tags
-			foreach(NodeOutput Output in Node.Outputs)
-			{
-				HashSet<FileReference> Files = TagNameToFileSet[Output.TagName];
 
-				HashSet<TempStorageBlock> StorageBlocks = new HashSet<TempStorageBlock>();
-				foreach(FileReference File in Files)
+				// Publish all the output tags
+				foreach (NodeOutput Output in Node.Outputs)
 				{
-					TempStorageBlock StorageBlock;
-					if(FileToStorageBlock.TryGetValue(File, out StorageBlock))
+					HashSet<FileReference> Files = TagNameToFileSet[Output.TagName];
+
+					HashSet<TempStorageBlock> StorageBlocks = new HashSet<TempStorageBlock>();
+					foreach (FileReference File in Files)
 					{
-						StorageBlocks.Add(StorageBlock);
+						TempStorageBlock StorageBlock;
+						if (FileToStorageBlock.TryGetValue(File, out StorageBlock))
+						{
+							StorageBlocks.Add(StorageBlock);
+						}
 					}
-				}
 
-				Storage.WriteFileList(Node.Name, Output.TagName, Files, StorageBlocks.ToArray());
+					Storage.WriteFileList(Node.Name, Output.TagName, Files, StorageBlocks.ToArray());
+				}
 			}
 
 			// Mark the node as succeeded
diff --git a/Engine/Source/Programs/AutomationTool/BuildGraph/Node.cs b/Engine/Source/Programs/AutomationTool/BuildGraph/Node.cs
--- a/Engine/Source/Programs/AutomationTool/BuildGraph/Node.cs
+++ b/Engine/Source/Programs/AutomationTool/BuildGraph/Node.cs
@@ -159,47 +159,54 @@ namespace AutomationTool
 			HashSet<FileReference> BuildProducts = TagNameToFileSet[DefaultOutput.TagName];
 			for(int Idx = 0; Idx < Tasks.Count; Idx++)
 			{
-				ITaskExecutor Executor = Tasks[Idx].GetExecutor();
-				if (Executor == null)
+				using (ITraceSpan Span = TraceSpan.Create(String.Format("task.{0}", Tasks[Idx].GetTraceName().ToLowerInvariant())))
 				{
-					// Execute this task directly
-					try
+					ITaskExecutor Executor = Tasks[Idx].GetExecutor();
+					if (Executor == null)
 					{
-						Tasks[Idx].Execute(Job, BuildProducts, TagNameToFileSet);
-					}
-					catch (Exception Ex)
-					{
-						ExceptionUtils.AddContext(Ex, "while executing task {0}", Tasks[Idx].GetTraceString());
-						if(Tasks[Idx].SourceLocation != null)
+						// Execute this task directly
+						try
 						{
-							ExceptionUtils.AddContext(Ex, "at {0}({1})", GetReadablePathForDiagnostics(Tasks[Idx].SourceLocation.Item1), Tasks[Idx].SourceLocation.Item2);
+							Tasks[Idx].GetTraceMetadata(Span, "");
+							Tasks[Idx].Execute(Job, BuildProducts, TagNameToFileSet);
+						}
+						catch (Exception Ex)
+						{
+							ExceptionUtils.AddContext(Ex, "while executing task {0}", Tasks[Idx].GetTraceString());
+							if (Tasks[Idx].SourceLocation != null)
+							{
+								ExceptionUtils.AddContext(Ex, "at {0}({1})", GetReadablePathForDiagnostics(Tasks[Idx].SourceLocation.Item1), Tasks[Idx].SourceLocation.Item2);
+							}
+							throw;
 						}
-						throw;
-					}
-				}
-				else
-				{
-					// The task has a custom executor, which may be able to execute several tasks simultaneously. Try to add the following tasks.
-					int FirstIdx = Idx; 
-					while (Idx + 1 < Tasks.Count && Executor.Add(Tasks[Idx + 1]))
-					{
-						Idx++;
-					}
-					try
-					{
-						Executor.Execute(Job, BuildProducts, TagNameToFileSet);
 					}
-					catch (Exception Ex)
+					else
 					{
-						for(int TaskIdx = FirstIdx; TaskIdx <= Idx; TaskIdx++)
+						Tasks[Idx].GetTraceMetadata(Span, "1.");
+
+						// The task has a custom executor, which may be able to execute several tasks simultaneously. Try to add the following tasks.
+						int FirstIdx = Idx;
+						while (Idx + 1 < Tasks.Count && Executor.Add(Tasks[Idx + 1]))
+						{
+							Idx++;
+							Tasks[Idx].GetTraceMetadata(Span, string.Format("{0}.", 1 + Idx - FirstIdx));
+						}
+						try
 						{
-							ExceptionUtils.AddContext(Ex, "while executing {0}", Tasks[TaskIdx].GetTraceString());
+							Executor.Execute(Job, BuildProducts, TagNameToFileSet);
 						}
-						if (Tasks[FirstIdx].SourceLocation != null)
+						catch (Exception Ex)
 						{
-							ExceptionUtils.AddContext(Ex, "at {0}({1})", GetReadablePathForDiagnostics(Tasks[FirstIdx].SourceLocation.Item1), Tasks[FirstIdx].SourceLocation.Item2);
+							for (int TaskIdx = FirstIdx; TaskIdx <= Idx; TaskIdx++)
+							{
+								ExceptionUtils.AddContext(Ex, "while executing {0}", Tasks[TaskIdx].GetTraceString());
+							}
+							if (Tasks[FirstIdx].SourceLocation != null)
+							{
+								ExceptionUtils.AddContext(Ex, "at {0}({1})", GetReadablePathForDiagnostics(Tasks[FirstIdx].SourceLocation.Item1), Tasks[FirstIdx].SourceLocation.Item2);
+							}
+							throw;
 						}
-						throw;
 					}
 				}
 			}
diff --git a/Engine/Source/Programs/AutomationTool/BuildGraph/Task.cs b/Engine/Source/Programs/AutomationTool/BuildGraph/Task.cs
--- a/Engine/Source/Programs/AutomationTool/BuildGraph/Task.cs
+++ b/Engine/Source/Programs/AutomationTool/BuildGraph/Task.cs
@@ -193,6 +193,27 @@ namespace AutomationTool
 			return Builder.ToString();
 		}
 
+		/// <summary>
+		/// Gets the name of this task for tracing
+		/// </summary>
+		/// <returns>The trace name</returns>
+		public virtual string GetTraceName()
+		{
+			TaskElementAttribute TaskElement = GetType().GetCustomAttribute<TaskElementAttribute>();
+			return (TaskElement != null)? TaskElement.Name : "unknown";
+		}
+
+		/// <summary>
+		/// Get properties to include in tracing info
+		/// </summary>
+		/// <param name="Span">The scope to add properties to</param>
+		/// <param name="Prefix">Prefix for metadata entries</param>
+		public virtual void GetTraceMetadata(ITraceSpan Span, string Prefix)
+		{
+			Span.AddMetadata(Prefix + "file", SourceLocation.Item1.MakeRelativeTo(CommandUtils.RootDirectory));
+			Span.AddMetadata(Prefix + "line", SourceLocation.Item2.ToString());
+		}
+
 		/// <summary>
 		/// Find all the tags which are used as inputs to this task
 		/// </summary>
diff --git a/Engine/Source/Programs/AutomationTool/BuildGraph/Tasks/CommandTask.cs b/Engine/Source/Programs/AutomationTool/BuildGraph/Tasks/CommandTask.cs
--- a/Engine/Source/Programs/AutomationTool/BuildGraph/Tasks/CommandTask.cs
+++ b/Engine/Source/Programs/AutomationTool/BuildGraph/Tasks/CommandTask.cs
@@ -126,6 +126,15 @@ namespace BuildGraph.Tasks
 			}
 		}
 
+		/// <summary>
+		/// Gets the name of this trace
+		/// </summary>
+		/// <returns>Name of the trace</returns>
+		public override string GetTraceName()
+		{
+			return String.Format("{0}.{1}", base.GetTraceName(), Parameters.Name.ToLowerInvariant());
+		}
+
 		/// <summary>
 		/// Output this task out to an XML writer.
 		/// </summary>
diff --git a/Engine/Source/Programs/AutomationTool/BuildGraph/Tasks/CompileTask.cs b/Engine/Source/Programs/AutomationTool/BuildGraph/Tasks/CompileTask.cs
--- a/Engine/Source/Programs/AutomationTool/BuildGraph/Tasks/CompileTask.cs
+++ b/Engine/Source/Programs/AutomationTool/BuildGraph/Tasks/CompileTask.cs
@@ -228,6 +228,20 @@ namespace AutomationTool
 			return new CompileTaskExecutor(this);
 		}
 
+		/// <summary>
+		/// Get properties to include in tracing info
+		/// </summary>
+		/// <param name="Span">The scope to add metadata to</param>
+		/// <param name="Prefix">Prefix for all metadata keys</param>
+		public override void GetTraceMetadata(ITraceSpan Span, string Prefix)
+		{
+			base.GetTraceMetadata(Span, Prefix);
+
+			Span.AddMetadata(Prefix + "target", Parameters.Target);
+			Span.AddMetadata(Prefix + "config", Parameters.Configuration.ToString());
+			Span.AddMetadata(Prefix + "platform", Parameters.Platform.ToString());
+		}
+
 		/// <summary>
 		/// Output this task out to an XML writer.
 		/// </summary>
diff --git a/Engine/Source/Programs/DotNETCommon/DotNETUtilities/DotNETUtilities.csproj b/Engine/Source/Programs/DotNETCommon/DotNETUtilities/DotNETUtilities.csproj
--- a/Engine/Source/Programs/DotNETCommon/DotNETUtilities/DotNETUtilities.csproj
+++ b/Engine/Source/Programs/DotNETCommon/DotNETUtilities/DotNETUtilities.csproj
@@ -165,6 +165,7 @@
     <Compile Include="StringUtils.cs" />
     <Compile Include="ThreadPoolWorkQueue.cs" />
     <Compile Include="HashCollection.cs" />
+    <Compile Include="TraceSpan.cs" />
     <Compile Include="Win32ExceptionWithCode.cs" />
   </ItemGroup>
   <Import Project="$(MSBuildToolsPath)\Microsoft.CSharp.targets" />
diff --git a/Engine/Source/Programs/DotNETCommon/DotNETUtilities/TraceSpan.cs b/Engine/Source/Programs/DotNETCommon/DotNETUtilities/TraceSpan.cs
new file mode 100644
--- /dev/null
+++ b/Engine/Source/Programs/DotNETCommon/DotNETUtilities/TraceSpan.cs
@@ -0,0 +1,121 @@
+using System;
+using System.Collections.Generic;
+using System.Diagnostics;
+using System.Globalization;
+using System.IO;
+using System.Linq;
+using System.Reflection;
+using System.Text;
+using System.Threading.Tasks;
+
+namespace Tools.DotNETCommon
+{
+	/// <summary>
+	/// Interface for a trace span
+	/// </summary>
+	public interface ITraceSpan : IDisposable
+	{
+		/// <summary>
+		/// Adds additional metadata to this scope
+		/// </summary>
+		/// <param name="Name">Name of the key</param>
+		/// <param name="Value">Value for this metadata</param>
+		void AddMetadata(string Name, string Value);
+	}
+
+	/// <summary>
+	/// Methods for creating ITraceScope instances
+	/// </summary>
+	public static class TraceSpan
+	{
+		/// <summary>
+		/// Concrete implementation of ITraceScope
+		/// </summary>
+		class TraceSpanImpl : ITraceSpan
+		{
+			public string Name;
+			public DateTimeOffset StartTime;
+			public DateTimeOffset? FinishTime;
+			public Dictionary<string, string> Metadata = new Dictionary<string, string>();
+
+			public TraceSpanImpl(string Name)
+			{
+				this.Name = Name;
+				this.StartTime = DateTimeOffset.Now;
+			}
+
+			public void AddMetadata(string Name, string Value)
+			{
+				Metadata[Name] = Value;
+			}
+
+			public void Dispose()
+			{
+				if(FinishTime == null)
+				{
+					FinishTime = DateTimeOffset.Now;
+				}
+			}
+		}
+
+		/// <summary>
+		/// The current scope provider
+		/// </summary>
+		static List<TraceSpanImpl> Spans = new List<TraceSpanImpl>();
+
+		/// <summary>
+		/// Creates a scope using the current provider
+		/// </summary>
+		public static ITraceSpan Create(string Name)
+		{
+			TraceSpanImpl Span = new TraceSpanImpl(Name);
+			Spans.Add(Span);
+			return Span;
+		}
+
+		/// <summary>
+		/// Saves all the scope information to a file
+		/// </summary>
+		public static void Flush()
+		{
+			string TelemetryDir = Environment.GetEnvironmentVariable("UE_TELEMETRY_DIR");
+			if (TelemetryDir != null)
+			{
+				FileReference File;
+				using (Process Process = Process.GetCurrentProcess())
+				{
+					DirectoryReference TelemetryDirRef = new DirectoryReference(TelemetryDir);
+					DirectoryReference.CreateDirectory(TelemetryDirRef);
+
+					string Name = String.Format("{0}.{1}", Path.GetFileName(Assembly.GetEntryAssembly().Location), Process.Id, Process.StartTime.Ticks);
+					File = FileReference.Combine(TelemetryDirRef, Name);
+				}
+
+				using (JsonWriter Writer = new JsonWriter(File))
+				{
+					Writer.WriteObjectStart();
+					Writer.WriteArrayStart("Scopes");
+					foreach (TraceSpanImpl Scope in Spans)
+					{
+						if (Scope.FinishTime != null)
+						{
+							Writer.WriteObjectStart();
+							Writer.WriteValue("Name", Scope.Name);
+							Writer.WriteValue("StartTime", Scope.StartTime.ToString("o", CultureInfo.InvariantCulture));
+							Writer.WriteValue("FinishTime", Scope.FinishTime.Value.ToString("o", CultureInfo.InvariantCulture));
+							Writer.WriteObjectStart("Metadata");
+							foreach (KeyValuePair<string, string> Pair in Scope.Metadata)
+							{
+								Writer.WriteValue(Pair.Key, Pair.Value);
+							}
+							Writer.WriteObjectEnd();
+							Writer.WriteObjectEnd();
+						}
+					}
+					Writer.WriteArrayEnd();
+					Writer.WriteObjectEnd();
+				}
+			}
+		}
+	}
+}
