Fix XMLdoc warnings.

diff --git a/Engine/Source/Programs/AutomationTool/BuildGraph/Tasks/CompileTask.cs b/Engine/Source/Programs/AutomationTool/BuildGraph/Tasks/CompileTask.cs
--- a/Engine/Source/Programs/AutomationTool/BuildGraph/Tasks/CompileTask.cs
+++ b/Engine/Source/Programs/AutomationTool/BuildGraph/Tasks/CompileTask.cs
@@ -231,7 +231,7 @@ namespace AutomationTool
 		/// <summary>
 		/// Get properties to include in tracing info
 		/// </summary>
-		/// <param name="Span">The scope to add metadata to</param>
+		/// <param name="Span">The span to add metadata to</param>
 		/// <param name="Prefix">Prefix for all metadata keys</param>
 		public override void GetTraceMetadata(ITraceSpan Span, string Prefix)
 		{
