Update span names for DD trace.

diff --git a/Engine/Source/Programs/AutomationTool/BuildGraph/BuildGraph.cs b/Engine/Source/Programs/AutomationTool/BuildGraph/BuildGraph.cs
--- a/Engine/Source/Programs/AutomationTool/BuildGraph/BuildGraph.cs
+++ b/Engine/Source/Programs/AutomationTool/BuildGraph/BuildGraph.cs
@@ -710,7 +710,7 @@ namespace AutomationTool
 
 			// Read the manifests for all the input storage blocks
 			Dictionary<TempStorageBlock, TempStorageManifest> InputManifests = new Dictionary<TempStorageBlock, TempStorageManifest>();
-			using (ITraceSpan Span = TraceSpan.Create("tempstorage.get"))
+			using (ITraceSpan Span = TraceSpan.Create("TempStorage", "Read"))
 			{
 				Span.AddMetadata("blocks", InputStorageBlocks.Count.ToString());
 				foreach (TempStorageBlock InputStorageBlock in InputStorageBlocks)
@@ -837,7 +837,7 @@ namespace AutomationTool
 			}
 
 			// Write all the storage blocks, and update the mapping from file to storage block
-			using (ITraceSpan Span = Tools.DotNETCommon.TraceSpan.Create("tempstorage.put"))
+			using (ITraceSpan Span = Tools.DotNETCommon.TraceSpan.Create("TempStorage", "Write"))
 			{
 				foreach (KeyValuePair<string, HashSet<FileReference>> Pair in OutputStorageBlockToFiles)
 				{
diff --git a/Engine/Source/Programs/AutomationTool/BuildGraph/Node.cs b/Engine/Source/Programs/AutomationTool/BuildGraph/Node.cs
--- a/Engine/Source/Programs/AutomationTool/BuildGraph/Node.cs
+++ b/Engine/Source/Programs/AutomationTool/BuildGraph/Node.cs
@@ -159,7 +159,7 @@ namespace AutomationTool
 			HashSet<FileReference> BuildProducts = TagNameToFileSet[DefaultOutput.TagName];
 			for(int Idx = 0; Idx < Tasks.Count; Idx++)
 			{
-				using (ITraceSpan Span = TraceSpan.Create(String.Format("task.{0}", Tasks[Idx].GetTraceName().ToLowerInvariant())))
+				using (ITraceSpan Span = TraceSpan.Create("Task", Tasks[Idx].GetTraceName()))
 				{
 					ITaskExecutor Executor = Tasks[Idx].GetExecutor();
 					if (Executor == null)
