- Added tagged retrieval to work with regex'd tags
- Added some Turnkey settings to IOSRuntimeSettings
- Added some of the IOSRuntimeSettings to be stripped out in the IniKeyBlacklist so that passwords aren't shipped
- Fixed smb: volume mounting in Turnkey
- Changed IOS prereq validation to be driven more by the bundleid, and then will hunt down a cert that matches the "best" provision that matches the bundleid. Then, uses the bundleid (with regex) to find a cert to install via Turnkey. Full wildcards can break it, so TeamID can be used to narrow down valid provisions

diff --git a/Engine/Config/BaseEngine.ini b/Engine/Config/BaseEngine.ini
--- a/Engine/Config/BaseEngine.ini
+++ b/Engine/Config/BaseEngine.ini
@@ -2057,7 +2057,7 @@ bDisableHTTPS=false
 bUseRSync=True
 BundleDisplayName=[PROJECT_NAME]
 BundleName=[PROJECT_NAME]
-BundleIdentifier=com.YourCXompany.[PROJECT_NAME]
+BundleIdentifier=com.YourCompany.[PROJECT_NAME]
 DevCodeSigningIdentity=iPhone Developer:;Apple Development:
 DistroCodeSigningIdentity=iPhone Distribution:;Apple Distribution:
 VersionInfo=1.0
diff --git a/Engine/Config/BaseGame.ini b/Engine/Config/BaseGame.ini
--- a/Engine/Config/BaseGame.ini
+++ b/Engine/Config/BaseGame.ini
@@ -109,9 +109,15 @@ bPakUsesSecondaryOrder=True
 +IniKeyBlacklist=SigningModulus
 +IniKeyBlacklist=SigningPrivateExponent
 +IniKeyBlacklist=EncryptionKey
++IniKeyBlacklist=DevCenterUsername
++IniKeyBlacklist=DevCenterPassword
++IniKeyBlacklist=IOSTeamID
++IniKeyBlacklist=SigningCertificate
++IniKeyBlacklist=MobileProvision
 +IniKeyBlacklist=IniKeyBlacklist
 +IniKeyBlacklist=IniSectionBlacklist
 
+ 
 [/Script/Engine.HUD]
 DebugDisplay=AI
 
diff --git a/Engine/Source/Programs/AutomationTool/AutomationUtils/Platform.cs b/Engine/Source/Programs/AutomationTool/AutomationUtils/Platform.cs
--- a/Engine/Source/Programs/AutomationTool/AutomationUtils/Platform.cs
+++ b/Engine/Source/Programs/AutomationTool/AutomationUtils/Platform.cs
@@ -13,7 +13,8 @@ namespace AutomationTool
 {
 	public interface FileRetriever
 	{
-		string RetrieveByTags(string[] RequiredTags, string[] PreferredTags, Dictionary<string,string> ExtraVariables=null);
+		string RetrieveByTags(string[] RequiredTags, string[] PreferredTags, Dictionary<string, string> ExtraVariables = null);
+		string GetVariable(string VariableName);
 	}
 
 	//public interface InputOutput
diff --git a/Engine/Source/Programs/AutomationTool/IOS/IOSPlatform.Automation.cs b/Engine/Source/Programs/AutomationTool/IOS/IOSPlatform.Automation.cs
--- a/Engine/Source/Programs/AutomationTool/IOS/IOSPlatform.Automation.cs
+++ b/Engine/Source/Programs/AutomationTool/IOS/IOSPlatform.Automation.cs
@@ -224,35 +224,6 @@ public class IOSPlatform : Platform
 		return true;
 	}
 
-	private bool InstallCert(FileRetriever Retriever)
-	{
-		// get the cert password from Studio settings
-		string CertPassword = "Epic123!";
-
-		string CertLoc = Retriever.RetrieveByTags(new string[] { "IOS_Cert_Generic_Development" }, null);
-
-		if (CertLoc != null)
-		{
-			Console.WriteLine("Will install cert from: '{0}'", CertLoc);
-
-			string CommandLine = string.Format("-e 'require \"cert\"; " +
-				"FastlaneCore::KeychainImporter.import_file(\"{0}\", FastlaneCore::Helper.keychain_path(\"login\"), certificate_password:\"{1}\")'",
-				CertLoc, CertPassword);
-			UnrealBuildTool.Utils.RunLocalProcessAndLogOutput("ruby", CommandLine);
-
-			// @todo turnkey verify output
-
-			return true;
-		}
-		else
-		{
-			Console.WriteLine("Unable to find a tagged source for IOS_Cert_Generic_Development");
-
-			return false;
-		}
-
-	}
-
 	private class VerifyIOSSettings
 	{
 		public string CodeSigningIdentity = null;
@@ -263,9 +234,11 @@ public class IOSPlatform : Platform
 
 		public string RubyScript = Path.Combine(CommandUtils.EngineDirectory.FullName, "Build/Turnkey/VerifyIOS.ru");
 
-		public VerifyIOSSettings(BuildCommand Command)
+		public VerifyIOSSettings(BuildCommand Command, FileRetriever Retriever)
 		{
 			FileReference ProjectPath = Command.ParseProjectParam();
+			string ProjectName = ProjectPath == null ? "" : ProjectPath.GetFileNameWithoutAnyExtensions();
+
 			ConfigHierarchy EngineConfig = ConfigCache.ReadHierarchy(ConfigHierarchyType.Engine, ProjectPath == null ? null : ProjectPath.Directory, UnrealTargetPlatform.IOS);
 
 			// first look for settings on the commandline:
@@ -275,65 +248,33 @@ public class IOSPlatform : Platform
 			Password = Command.ParseParamValue("devcenterpassword");
 			Team = Command.ParseParamValue("teamid");
 
+			if (string.IsNullOrEmpty(Team)) Team = Retriever.GetVariable("User_AppleDevCenterTeamID");
+			if (string.IsNullOrEmpty(Account)) Account = Retriever.GetVariable("User_AppleDevCenterUsername");
+
 			// fall back to ini for anything else
 			if (string.IsNullOrEmpty(CodeSigningIdentity)) EngineConfig.GetString("/Script/IOSRuntimeSettings.IOSRuntimeSettings", "DevCodeSigningIdentity", out CodeSigningIdentity);
 			if (string.IsNullOrEmpty(BundleId)) EngineConfig.GetString("/Script/IOSRuntimeSettings.IOSRuntimeSettings", "BundleIdentifier", out BundleId);
-			if (string.IsNullOrEmpty(Account)) EngineConfig.GetString("TurnkeySettings", "IOS_DevCenterUsername", out Account);
-			if (string.IsNullOrEmpty(Password)) EngineConfig.GetString("TurnkeySettings", "IOS_DevCenterPassword", out Password);
-			if (string.IsNullOrEmpty(Team)) EngineConfig.GetString("TurnkeySettings", "IOS_DevCenterPassword", out Password);
+			if (string.IsNullOrEmpty(Team)) EngineConfig.GetString("/Script/IOSRuntimeSettings.IOSRuntimeSettings", "IOSTeamID", out Team);
+			if (string.IsNullOrEmpty(Account)) EngineConfig.GetString("/Script/IOSRuntimeSettings.IOSRuntimeSettings", "DevCenterUsername", out Account);
+			if (string.IsNullOrEmpty(Password)) EngineConfig.GetString("/Script/IOSRuntimeSettings.IOSRuntimeSettings", "DevCenterPassword", out Password);
 
-			string Foo;
-			bool Bar = EngineConfig.GetString("/Script/IOSRuntimeSettings.IOSRuntimeSettings", "DevCodeSigningIdentity", out Foo);
+			BundleId.Replace("[PROJECT_NAME]", ProjectName);
 
 			// some are required
-			if (string.IsNullOrEmpty(CodeSigningIdentity) || string.IsNullOrEmpty(BundleId))
+			if (string.IsNullOrEmpty(BundleId))
 			{
-				throw new AutomationException("Turnkey IOS verification requires code signing identity (have ='{0}', ex: iPhone Developer) and a bundle id (have '{1}', ex: com.company.foo)", CodeSigningIdentity, BundleId);
+				throw new AutomationException("Turnkey IOS verification requires bundle id (have '{1}', ex: com.company.foo)", CodeSigningIdentity, BundleId);
 			}
 		}
 
-		public int RunRubyCommand(bool bCertOnly, bool bVerifyOnly, string DeviceName)
+		public int RunCommandMaybeInteractive(string Command, string Params, bool bInteractive)
 		{
-			string Params;
 			int ExitCode;
-
-			if (bCertOnly)
-			{
-				Params = string.Format("--certonly --identity '{0}'", CodeSigningIdentity);
-			}
-			else
-			{
-				Params = string.Format("--identity '{0}' --bundleid {1}", CodeSigningIdentity, BundleId);
-
-				if (!string.IsNullOrEmpty(Account))
-				{
-					Params += string.Format(" --login {0}", Account);
-				}
-				if (!string.IsNullOrEmpty(Password))
-				{
-					Params += string.Format(" --password {0}", Password);
-				}
-				if (!string.IsNullOrEmpty(Team))
-				{
-					Params += string.Format(" --team {0}", Team);
-				}
-
-
-				if (!string.IsNullOrEmpty(DeviceName))
-				{
-					Params += string.Format(" --device {0}", DeviceName);
-				}
-
-				if (bVerifyOnly)
-				{
-					Params += string.Format(" --verifyonly");
-				}
-			}
-
 			// if non-interactive, we can just run directly in the current shell
-			if (bVerifyOnly)
+			if (!bInteractive)
 			{
-				UnrealBuildTool.Utils.RunLocalProcessAndReturnStdOut(RubyScript, Params, out ExitCode, true);
+				Console.WriteLine("Running '{0} {1}'", Command, Params);
+				UnrealBuildTool.Utils.RunLocalProcessAndReturnStdOut(Command, Params, out ExitCode, true);
 			}
 			else
 			{
@@ -355,16 +296,99 @@ public class IOSPlatform : Platform
 						" -e     \"close newWindow\"" +
 						" -e   \"end if\"" +
 						" -e \"end tell\"",
-						RubyScript, Params, ReturnCodeFilename);
+						Command, Params.Replace("\"", "\\\\\\\""), ReturnCodeFilename);
 
 				UnrealBuildTool.Utils.RunLocalProcessAndReturnStdOut("osascript", Params, out ExitCode, true);
 				if (ExitCode == 0)
 				{
 					ExitCode = int.Parse(File.ReadAllText(ReturnCodeFilename));
+					File.Delete(ReturnCodeFilename);
 				}
+				
 			}
 			return ExitCode;
 		}
+
+		public int RunRubyCommand(bool bVerifyOnly, string DeviceName)
+		{
+			string Params;
+
+			Params = string.Format("--bundleid {0}", BundleId);
+
+			if (!string.IsNullOrEmpty(CodeSigningIdentity))
+			{
+				Params += string.Format(" --identity '{0}'", CodeSigningIdentity);
+			}
+
+			if (!string.IsNullOrEmpty(Account))
+			{
+				Params += string.Format(" --login {0}", Account);
+			}
+			if (!string.IsNullOrEmpty(Password))
+			{
+				Params += string.Format(" --password {0}", Password);
+			}
+			if (!string.IsNullOrEmpty(Team))
+			{
+				Params += string.Format(" --team {0}", Team);
+			}
+
+
+			if (!string.IsNullOrEmpty(DeviceName))
+			{
+				Params += string.Format(" --device {0}", DeviceName);
+			}
+
+			if (bVerifyOnly)
+			{
+				Params += string.Format(" --verifyonly");
+			}
+
+			return RunCommandMaybeInteractive(RubyScript, Params, !bVerifyOnly);
+		}
+	}
+
+	private bool InstallCert(FileRetriever Retriever, VerifyIOSSettings Settings)
+	{
+//		string ProjectName = Retriever.GetVariable("Project");
+
+		string CertLoc = null;
+
+		if (!string.IsNullOrEmpty(Settings.BundleId))
+		{
+			CertLoc = Retriever.RetrieveByTags(new string[] { "DevCert: " + Settings.BundleId }, null);
+		}
+		if (CertLoc == null)
+		{
+			CertLoc = Retriever.RetrieveByTags(new string[] { "DevCert" }, null);
+		}
+
+		if (CertLoc != null)
+		{
+			// get the cert password from Studio settings
+			string CertPassword = "Epic123!";// Retriever.GetVariable("IOSCertificatePassword");
+
+Console.WriteLine("Will install cert from: '{0}'", CertLoc);
+
+
+			// osascript -e 'Tell application "System Events" to display dialog "Enter the network password:" with hidden answer default answer ""' -e 'text returned of result' 2>/dev/null
+			string CommandLine = string.Format("-e 'require \"cert\"; " +
+				"FastlaneCore::KeychainImporter.import_file(\"{0}\", FastlaneCore::Helper.keychain_path(\"login\"), certificate_password:\"{1}\")'",
+				CertLoc, CertPassword);
+			
+			Settings.RunCommandMaybeInteractive("ruby", CommandLine, true);
+			
+			// @todo turnkey verify output
+
+			return true;
+		}
+		else
+		{
+			Console.WriteLine("Unable to find a tagged source for DevCert");
+
+			return false;
+		}
+
 	}
 
 	string GetConfiguratorLocation()
@@ -423,10 +447,10 @@ public class IOSPlatform : Platform
 			}
 		}
 
-		VerifyIOSSettings Settings = new VerifyIOSSettings(Command);
+		VerifyIOSSettings Settings = new VerifyIOSSettings(Command, Retriever);
 
 		// look if we have a cert that matches it
-		ExitCode = Settings.RunRubyCommand(true, bVerifyOnly, null);
+		ExitCode = Settings.RunRubyCommand(bVerifyOnly, null);
 
 		Console.WriteLine("Host VerifyIOS exited with code {0}", ExitCode);
 
@@ -437,7 +461,8 @@ public class IOSPlatform : Platform
 				return false;
 			}
 
-			if (!InstallCert(Retriever))
+			// ExitCode 3 means we need to install a cert
+			if (ExitCode == 3 && !InstallCert(Retriever, Settings))
 			{
 				return false;
 			}
@@ -453,16 +478,22 @@ public class IOSPlatform : Platform
 			return base.UpdateDevicePrerequisites(Device, Command, Retriever, bVerifyOnly);
 		}
 
-		VerifyIOSSettings Settings = new VerifyIOSSettings(Command);
+		VerifyIOSSettings Settings = new VerifyIOSSettings(Command, Retriever);
 
 		// @todo turnkey - better to use the device's udid if it's set properly in DeviceInfo
 		string DeviceName = Device == null ? null : Device.Name;
 
 		// now look for a provision that can be used with a (maybe newly) instally cert
-		int ExitCode = Settings.RunRubyCommand(false, bVerifyOnly, DeviceName);
+		int ExitCode = Settings.RunRubyCommand(bVerifyOnly, DeviceName);
 
 		Console.WriteLine("Device VerifyIOS exited with code {0}", ExitCode);
 
+		// ExitCode 3 means we need to install a cert
+		if (ExitCode == 3 && !InstallCert(Retriever, Settings))
+		{
+			return false;
+		}
+
 		return ExitCode == 0;
 	}
 
diff --git a/Engine/Source/Programs/AutomationTool/Turnkey/CopyProvider.cs b/Engine/Source/Programs/AutomationTool/Turnkey/CopyProvider.cs
--- a/Engine/Source/Programs/AutomationTool/Turnkey/CopyProvider.cs
+++ b/Engine/Source/Programs/AutomationTool/Turnkey/CopyProvider.cs
@@ -222,7 +222,24 @@ namespace Turnkey
 
 			List<SdkInfo> Sdks = TurnkeyManifest.GetDiscoveredSdks();
 
-			Sdks = Sdks.FindAll(x => x.Type == SdkInfo.SdkType.Misc && x.CustomSdkId == RequiredTags[0]);
+			string Tag = RequiredTags[0];
+			Sdks = Sdks.FindAll(x => x.Type == SdkInfo.SdkType.Misc);
+
+			Sdks = Sdks.FindAll(x =>
+			{
+				if (Tag.StartsWith("regex:"))
+				{
+					return TurnkeyUtils.IsValueValid(x.CustomSdkId, Tag, null);
+				}
+				// this will handle the case of x.CustomSdkId starting with regex: or just doing a case insensitive string comparison of tag and CustomSdkId
+				// range: is not supported, at least yet - we would have to check Tag with range: above, and also support range without a Platform (or pass in a platform somehow?)
+				return TurnkeyUtils.IsValueValid(Tag, x.CustomSdkId, null);
+			});
+
+			if (Sdks == null || Sdks.Count == 0)
+			{
+				return null;
+			}
 
 			SdkInfo Sdk = Sdks.First();
 
@@ -236,5 +253,10 @@ namespace Turnkey
 
 			return null;
 		}
+
+		public string GetVariable(string VariableName)
+		{
+			return TurnkeyUtils.GetVariableValue(VariableName);
+		}
 	}
 }
diff --git a/Engine/Source/Programs/AutomationTool/Turnkey/Providers/NullCopyProvider.cs b/Engine/Source/Programs/AutomationTool/Turnkey/Providers/NullCopyProvider.cs
--- a/Engine/Source/Programs/AutomationTool/Turnkey/Providers/NullCopyProvider.cs
+++ b/Engine/Source/Programs/AutomationTool/Turnkey/Providers/NullCopyProvider.cs
@@ -18,7 +18,7 @@ namespace Turnkey
 			if (UnrealBuildTool.BuildHostPlatform.Current.Platform == UnrealBuildTool.UnrealTargetPlatform.Mac && Operation.StartsWith("smb://"))
 			{
 				// match the form smb://server.net/Foo/Bar/Baz and retrieve the Foo, which will be the name in /Volumes/Foo
-				Match Match = Regex.Match(Operation, @"^smb:\/\/.+?\/(.+?)\/(.*)$");
+				Match Match = Regex.Match(Operation, @"^(smb:\/\/.+?)\/(.+?)\/(.*)$");
 
 				// make sure regex matched
 				if (!Match.Success)
@@ -26,9 +26,12 @@ namespace Turnkey
 					return;
 				}
 
+				// mount the drive
+				UnrealBuildTool.Utils.RunLocalProcessAndReturnStdOut("osascript", string.Format("-e 'mount volume \"{0}/{1}\"'", Match.Groups[1].Value, Match.Groups[2].Value));
+
 				// convert smb://server.net/Foo/Bar/Baz to /Volumes/Foo/Bar/Baz
 				TurnkeyUtils.Log("SMB Before: {0}", Operation);
-				Operation = string.Format("/Volumes/{0}/{1}", Match.Groups[1].Value, Match.Groups[2].Value);
+				Operation = string.Format("/Volumes/{0}/{1}", Match.Groups[2].Value, Match.Groups[3].Value);
 				TurnkeyUtils.Log("SMB After: {0}", Operation);
 			}
 		}
diff --git a/Engine/Source/Programs/AutomationTool/Turnkey/SdkInfo.cs b/Engine/Source/Programs/AutomationTool/Turnkey/SdkInfo.cs
--- a/Engine/Source/Programs/AutomationTool/Turnkey/SdkInfo.cs
+++ b/Engine/Source/Programs/AutomationTool/Turnkey/SdkInfo.cs
@@ -135,7 +135,7 @@ namespace Turnkey
 			LocalAvailability Result = LocalAvailability.None;
 
 			// no need to retrieve anything if we aren't allowing updates
-			CopyProviderRetriever Retriever = bAllowUpdatingPrerequisites ? new CopyProviderRetriever() : null;
+			CopyProviderRetriever Retriever = new CopyProviderRetriever();
 
 			if (AutomationPlatform.UpdateHostPrerequisites(TurnkeyUtils.CommandUtilHelper, Retriever, !bAllowUpdatingPrerequisites))
 			{
diff --git a/Engine/Source/Programs/AutomationTool/Turnkey/TurnkeySettings.cs b/Engine/Source/Programs/AutomationTool/Turnkey/TurnkeySettings.cs
--- a/Engine/Source/Programs/AutomationTool/Turnkey/TurnkeySettings.cs
+++ b/Engine/Source/Programs/AutomationTool/Turnkey/TurnkeySettings.cs
@@ -38,6 +38,8 @@ namespace Turnkey
 		{
 			new UserSetting("User_QuickSwitchSdkLocation", null, "Location for downloaded Sdks to be stored (currently just for Google Drive, perforce will sync per clientspec)"),
 			new UserSetting("User_LastPerforceClient", null, "Perforce depot clientspec to look at first for finding matching clientspecs. Will be set automatically when one is found, but this can be used to force a particular client"),
+			new UserSetting("User_AppleDevCenterUsername", null, "Sets the username to use when logging in to DevCenter for iOS/tvOS/macOS. (Supercedes the DevCenterUsername .ini setting)"),
+			new UserSetting("User_AppleDevCenterTeamID", null, "Sets the teamid to use when logging in to DevCenter for iOS/tvOS/macOS. (Supercedes the IOSTeamID .ini setting)"),
 		};
 
 		// basically same as Turnkey variables, but this only contains ones that were loaded so we can write them back out `
diff --git a/Engine/Source/Runtime/IOS/IOSRuntimeSettings/Classes/IOSRuntimeSettings.h b/Engine/Source/Runtime/IOS/IOSRuntimeSettings/Classes/IOSRuntimeSettings.h
--- a/Engine/Source/Runtime/IOS/IOSRuntimeSettings/Classes/IOSRuntimeSettings.h
+++ b/Engine/Source/Runtime/IOS/IOSRuntimeSettings/Classes/IOSRuntimeSettings.h
@@ -444,11 +444,13 @@ public:
 	UPROPERTY(GlobalConfig, EditAnywhere, Category = Online, meta = (EditCondition = "bEnableFacebookSupport"))
 	FString FacebookAppID;
     
-    // Mobile provision to utilize when signing
+    // Mobile provision to utilize when signing.
+	// This value is stripped out when making builds.
 	UPROPERTY(GlobalConfig, EditAnywhere, Category = Build)
 	FString MobileProvision;
 
-	// Signing certificate to utilize when signing
+	// Signing certificate to utilize when signing.
+	// This value is stripped out when making builds.
 	UPROPERTY(GlobalConfig, EditAnywhere, Category = Build)
 	FString SigningCertificate;
 	
@@ -456,10 +458,23 @@ public:
 	UPROPERTY(GlobalConfig, EditAnywhere, Category = Build)
 	bool bAutomaticSigning;
 
-	// The team ID of the apple developer account to be used to autmatically sign IOS builds
+	// The team ID of the apple developer account to be used to autmatically sign IOS builds.
+	// This can be overridden in Turnkey with "RunUAT Turnkey -command=ManageSettings"
+	// This value is stripped out when making builds.
 	UPROPERTY(GlobalConfig, EditAnywhere, Category = Build, meta = (ConfigHierarchyEditable))
 	FString IOSTeamID;
 
+	// The username/email to use when logging in to DevCenter with Turnkey.
+	// This can be overridden in Turnkey with "RunUAT Turnkey -command=ManageSettings"
+	// This value is stripped out when making builds.
+	UPROPERTY(GlobalConfig, EditAnywhere, Category = Build, meta = (ConfigHierarchyEditable))
+	FString DevCenterUsername;
+	
+	// The password to use when logging in to DevCenter with Turnkey. NOTE: This is saved in plaintext, and is meant for shared accounts!
+	// This value is stripped out when making builds.
+	UPROPERTY(GlobalConfig, EditAnywhere, Category = Build, meta = (ConfigHierarchyEditable))
+	FString DevCenterPassword;
+	
 	// Whether the app supports HTTPS
 	UPROPERTY(GlobalConfig, EditAnywhere, Category = Online, meta = (DisplayName = "Allow web connections to non-HTTPS websites"))
 	bool bDisableHTTPS;
