- Turnkey IOS cert usability improvements.
- Can now base on Provision Name, instead of just bundle id
- Added concept of StudioSettings, but they aren't being loaded yet from anywhere (and you can't set them with ManageSettings, by design, and am about to redo how Manifests/SdkInfo installers/etc are all managed)

diff --git a/Engine/Source/Programs/AutomationTool/IOS/IOSPlatform.Automation.cs b/Engine/Source/Programs/AutomationTool/IOS/IOSPlatform.Automation.cs
--- a/Engine/Source/Programs/AutomationTool/IOS/IOSPlatform.Automation.cs
+++ b/Engine/Source/Programs/AutomationTool/IOS/IOSPlatform.Automation.cs
@@ -231,8 +231,10 @@ public class IOSPlatform : Platform
 		public string Account = null;
 		public string Password = null;
 		public string Team = null;
+		public string Provision = null;
 
 		public string RubyScript = Path.Combine(CommandUtils.EngineDirectory.FullName, "Build/Turnkey/VerifyIOS.ru");
+		public string InstallCertScript = Path.Combine(CommandUtils.EngineDirectory.FullName, "Build/Turnkey/InstallCert.ru");
 
 		public VerifyIOSSettings(BuildCommand Command, FileRetriever Retriever)
 		{
@@ -247,9 +249,11 @@ public class IOSPlatform : Platform
 			Account = Command.ParseParamValue("devcenterusername");
 			Password = Command.ParseParamValue("devcenterpassword");
 			Team = Command.ParseParamValue("teamid");
+			Provision = Command.ParseParamValue("provision");
 
 			if (string.IsNullOrEmpty(Team)) Team = Retriever.GetVariable("User_AppleDevCenterTeamID");
 			if (string.IsNullOrEmpty(Account)) Account = Retriever.GetVariable("User_AppleDevCenterUsername");
+			if (string.IsNullOrEmpty(Provision)) Provision = Retriever.GetVariable("User_IOSProvisioningProfile");
 
 			// fall back to ini for anything else
 			if (string.IsNullOrEmpty(CodeSigningIdentity)) EngineConfig.GetString("/Script/IOSRuntimeSettings.IOSRuntimeSettings", "DevCodeSigningIdentity", out CodeSigningIdentity);
@@ -257,7 +261,8 @@ public class IOSPlatform : Platform
 			if (string.IsNullOrEmpty(Team)) EngineConfig.GetString("/Script/IOSRuntimeSettings.IOSRuntimeSettings", "IOSTeamID", out Team);
 			if (string.IsNullOrEmpty(Account)) EngineConfig.GetString("/Script/IOSRuntimeSettings.IOSRuntimeSettings", "DevCenterUsername", out Account);
 			if (string.IsNullOrEmpty(Password)) EngineConfig.GetString("/Script/IOSRuntimeSettings.IOSRuntimeSettings", "DevCenterPassword", out Password);
-
+			if (string.IsNullOrEmpty(Provision)) EngineConfig.GetString("/Script/IOSRuntimeSettings.IOSRuntimeSettings", "MobileProvision", out Provision);
+			
 			BundleId.Replace("[PROJECT_NAME]", ProjectName);
 
 			// some are required
@@ -269,11 +274,12 @@ public class IOSPlatform : Platform
 
 		public int RunCommandMaybeInteractive(string Command, string Params, bool bInteractive)
 		{
+			Console.WriteLine("Running Command '{0} {1}'", Command, Params);
+
 			int ExitCode;
 			// if non-interactive, we can just run directly in the current shell
 			if (!bInteractive)
 			{
-				Console.WriteLine("Running '{0} {1}'", Command, Params);
 				UnrealBuildTool.Utils.RunLocalProcessAndReturnStdOut(Command, Params, out ExitCode, true);
 			}
 			else
@@ -283,10 +289,14 @@ public class IOSPlatform : Platform
 
 				// run potentially interactive scripts in a Terminal window
 				Params = string.Format(
+						" -e \"tell application \\\"Finder\\\"\"" +
+						" -e   \"set desktopBounds to bounds of window of desktop\"" +
+						" -e \"end tell\"" +
 						" -e \"tell application \\\"Terminal\\\"\"" +
 						" -e   \"activate\"" +
-						" -e   \"set newTab to do script (\\\"{0} {1}; echo $? > {2}; exit\\\")\"" +
+						" -e   \"set newTab to do script (\\\"{3}; {0} {1}; echo $? > {2}; {3}; exit\\\")\"" +
 						" -e   \"set newWindow to window 1\"" +
+						" -e   \"set size of newWindow to {{ item 3 of desktopBounds / 2, item 4 of desktopBounds / 2 }}\"" +
 						" -e   \"repeat\"" +
 						" -e     \"delay 1\"" +
 						" -e     \"if not busy of newTab then exit repeat\"" +
@@ -296,7 +306,9 @@ public class IOSPlatform : Platform
 						" -e     \"close newWindow\"" +
 						" -e   \"end if\"" +
 						" -e \"end tell\"",
-						Command, Params.Replace("\"", "\\\\\\\""), ReturnCodeFilename);
+						Command, Params.Replace("\"", "\\\\\\\""), ReturnCodeFilename, "printf \\\\\\\"\\\\\\n\\\\\\n\\\\\\n\\\\\\n\\\\\\\"");
+
+				Console.WriteLine("\n\n\n{0}\n\n\n", Params);
 
 				UnrealBuildTool.Utils.RunLocalProcessAndReturnStdOut("osascript", Params, out ExitCode, true);
 				if (ExitCode == 0)
@@ -333,6 +345,10 @@ public class IOSPlatform : Platform
 				Params += string.Format(" --team {0}", Team);
 			}
 
+			if (!string.IsNullOrEmpty(Provision))
+			{
+				Params += string.Format(" --provision {0}", Provision);
+			}
 
 			if (!string.IsNullOrEmpty(DeviceName))
 			{
@@ -366,21 +382,15 @@ public class IOSPlatform : Platform
 		if (CertLoc != null)
 		{
 			// get the cert password from Studio settings
-			string CertPassword = "Epic123!";// Retriever.GetVariable("IOSCertificatePassword");
-
-Console.WriteLine("Will install cert from: '{0}'", CertLoc);
+			string CertPassword = Retriever.GetVariable("Studio_AppleSigningCertPassword");
 
+			Console.WriteLine("Will install cert from: '{0}'", CertLoc);
 
 			// osascript -e 'Tell application "System Events" to display dialog "Enter the network password:" with hidden answer default answer ""' -e 'text returned of result' 2>/dev/null
-			string CommandLine = string.Format("-e 'require \"cert\"; " +
-				"FastlaneCore::KeychainImporter.import_file(\"{0}\", FastlaneCore::Helper.keychain_path(\"login\"), certificate_password:\"{1}\")'",
-				CertLoc, CertPassword);
-			
-			Settings.RunCommandMaybeInteractive("ruby", CommandLine, true);
-			
-			// @todo turnkey verify output
+			string CommandLine = string.Format("'{0}' '{1}'", CertLoc, CertPassword);
 
-			return true;
+			// run ruby script to install cert
+			return Settings.RunCommandMaybeInteractive(Settings.InstallCertScript, CommandLine, true) == 0;
 		}
 		else
 		{
@@ -462,7 +472,7 @@ Console.WriteLine("Will install cert from: '{0}'", CertLoc);
 			}
 
 			// ExitCode 3 means we need to install a cert
-			if (ExitCode == 3 && !InstallCert(Retriever, Settings))
+			if (ExitCode != 3 || !InstallCert(Retriever, Settings))
 			{
 				return false;
 			}
diff --git a/Engine/Source/Programs/AutomationTool/Turnkey/Commands/ManageSettings.cs b/Engine/Source/Programs/AutomationTool/Turnkey/Commands/ManageSettings.cs
--- a/Engine/Source/Programs/AutomationTool/Turnkey/Commands/ManageSettings.cs
+++ b/Engine/Source/Programs/AutomationTool/Turnkey/Commands/ManageSettings.cs
@@ -20,12 +20,13 @@ namespace Turnkey.Commands
 
 		protected override void Execute(string[] CommandOptions)
 		{
-			UserSetting[] Settings = TurnkeySettings.AllUserSettings;
+			UserSetting[] UserSettings = TurnkeySettings.AllUserSettings;
+			UserSetting[] StudioSettings = TurnkeySettings.AllStudioSettings;
 
 			if (TurnkeyUtils.ParseParam("ListVars", CommandOptions))
 			{
 				TurnkeyUtils.Log("All Settings:");
-				foreach (UserSetting Setting in Settings)
+				foreach (UserSetting Setting in UserSettings.Union(StudioSettings))
 				{
 					TurnkeyUtils.Log("  {0} = '{1}'", Setting.VariableName, TurnkeyUtils.GetVariableValue(Setting.VariableName));
 					TurnkeyUtils.Log("    [{0}]", Setting.Description);
@@ -46,10 +47,10 @@ namespace Turnkey.Commands
 					return;
 				}
 
-				// make sure it's a valud variable
-				if (!Array.Exists(Settings, x => x.VariableName == Variable))
+				// make sure it's a valid variable
+				if (!Array.Exists(UserSettings, x => x.VariableName == Variable))
 				{
-					TurnkeyUtils.Log("Error: {0} is not a valid variable (variables are case sensitive)", Variable);
+					TurnkeyUtils.Log("Error: {0} is not a valid variable (variables are case sensitive, and Studio settings cannot be set)", Variable);
 					return;
 				}
 
@@ -65,7 +66,7 @@ namespace Turnkey.Commands
 			List<string> Options = new List<string>();
 
 
-			foreach (UserSetting Setting in Settings)
+			foreach (UserSetting Setting in UserSettings)
 			{
 				Options.Add(string.Format("{0} - Currently '{1}'\n    [{2}]", Setting.VariableName, TurnkeyUtils.GetVariableValue(Setting.VariableName), Setting.Description));
 			}
@@ -73,11 +74,11 @@ namespace Turnkey.Commands
 			int Choice = TurnkeyUtils.ReadInputInt("Choose a variable to set:", Options, true);
 			if (Choice > 0)
 			{
-				string NewValue = TurnkeyUtils.ReadInput(string.Format("Enter value for {0} [Currently '{1}']", Settings[Choice - 1].VariableName, TurnkeyUtils.GetVariableValue(Settings[Choice - 1].VariableName)));
+				string NewValue = TurnkeyUtils.ReadInput(string.Format("Enter value for {0} [Currently '{1}']", UserSettings[Choice - 1].VariableName, TurnkeyUtils.GetVariableValue(UserSettings[Choice - 1].VariableName)));
 				if (NewValue != "")
 				{
-					TurnkeyUtils.Log("Setting Settings[{0}].{1} = {2}", Choice - 1, Settings[Choice - 1].VariableName, NewValue);
-					TurnkeySettings.SetUserSetting(Settings[Choice - 1].VariableName, NewValue);
+					TurnkeyUtils.Log("Setting Settings[{0}].{1} = {2}", Choice - 1, UserSettings[Choice - 1].VariableName, NewValue);
+					TurnkeySettings.SetUserSetting(UserSettings[Choice - 1].VariableName, NewValue);
 					TurnkeySettings.Save();
 				}
 			}
diff --git a/Engine/Source/Programs/AutomationTool/Turnkey/TurnkeySettings.cs b/Engine/Source/Programs/AutomationTool/Turnkey/TurnkeySettings.cs
--- a/Engine/Source/Programs/AutomationTool/Turnkey/TurnkeySettings.cs
+++ b/Engine/Source/Programs/AutomationTool/Turnkey/TurnkeySettings.cs
@@ -34,15 +34,22 @@ namespace Turnkey
 	{
 		public static string UserSettingManifestLocation = Path.Combine(Environment.GetFolderPath(Environment.SpecialFolder.LocalApplicationData), "Unreal Engine", "UnrealTurnkey", "UserSettingsManifest.xml");
 
-		public static UserSetting[] AllUserSettings = 
+		public static UserSetting[] AllUserSettings =
 		{
 			new UserSetting("User_QuickSwitchSdkLocation", null, "Location for downloaded Sdks to be stored (currently just for Google Drive, perforce will sync per clientspec)"),
 			new UserSetting("User_LastPerforceClient", null, "Perforce depot clientspec to look at first for finding matching clientspecs. Will be set automatically when one is found, but this can be used to force a particular client"),
 			new UserSetting("User_AppleDevCenterUsername", null, "Sets the username to use when logging in to DevCenter for iOS/tvOS/macOS. (Supercedes the DevCenterUsername .ini setting)"),
 			new UserSetting("User_AppleDevCenterTeamID", null, "Sets the teamid to use when logging in to DevCenter for iOS/tvOS/macOS. (Supercedes the IOSTeamID .ini setting)"),
+			new UserSetting("User_IOSProvisioningProfile", null, "Sets the provisioning profile to use when setting up iOS signing. (Supercedes the MobileProvision .ini setting)"),
 		};
 
-		// basically same as Turnkey variables, but this only contains ones that were loaded so we can write them back out `
+		public static UserSetting[] AllStudioSettings =
+		{
+			new UserSetting("Studio_AppleSigningCertPassword", null, "A shared password that is used across Apple Signing Certificates"),
+			new UserSetting("Studio_AutoSDKSupported", null, "if true, your studio has set up AutoSDK, and Turnkey will offer to set up your local machine if needed."),
+		};
+
+		// basically same as Turnkey variables, but this only contains ones that were loaded so we can write them back out
 		static Dictionary<string, string> SetUserSettings = new Dictionary<string, string>();
 
 
