- Fixed a typo I made when editing on checkin, oops

diff --git a/Engine/Source/Programs/AutomationTool/Turnkey/TurnkeyManifest.cs b/Engine/Source/Programs/AutomationTool/Turnkey/TurnkeyManifest.cs
--- a/Engine/Source/Programs/AutomationTool/Turnkey/TurnkeyManifest.cs
+++ b/Engine/Source/Programs/AutomationTool/Turnkey/TurnkeyManifest.cs
@@ -32,7 +32,7 @@ namespace Turnkey
 			// load any settings set in the xml
 			if (SavedSettings != null)
 			{
-				Array.ForEach(SavedSettings, x => TurnkeyUtils.SetVariable(x.Variable, TurnkeyUtils.ExpandVariables(x.Value)); );
+				Array.ForEach(SavedSettings, x => TurnkeyUtils.SetVariable(x.Variable, TurnkeyUtils.ExpandVariables(x.Value)));
 			}
 
 			if (BuildSources != null)
