- Fixed RetrieveFIleSource as used from IOSPlatform.Automation to use the new FileSource stuff instead of Custom SdkInfo

diff --git a/Engine/Source/Programs/AutomationTool/AutomationUtils/Platform.cs b/Engine/Source/Programs/AutomationTool/AutomationUtils/Platform.cs
--- a/Engine/Source/Programs/AutomationTool/AutomationUtils/Platform.cs
+++ b/Engine/Source/Programs/AutomationTool/AutomationUtils/Platform.cs
@@ -13,8 +13,7 @@ namespace AutomationTool
 {
 	public interface FileRetriever
 	{
-		string RetrieveByTags(string[] RequiredTags, string[] PreferredTags, Dictionary<string, string> ExtraVariables = null);
-		//		string RetrieveSdk(string Platform, string Version, string Type, string SubType = null, Dictionary<string, string> ExtraVariables = null);
+		string RetrieveFileSource(string Name, string InType = "Misc", string InPlatform = null, string SubType = null);
 		string RetrieveFileSource(object HintObject);
 		string GetVariable(string VariableName);
 		bool RunExternalCommand(string Command, string Params);
diff --git a/Engine/Source/Programs/AutomationTool/IOS/IOSPlatform.Automation.cs b/Engine/Source/Programs/AutomationTool/IOS/IOSPlatform.Automation.cs
--- a/Engine/Source/Programs/AutomationTool/IOS/IOSPlatform.Automation.cs
+++ b/Engine/Source/Programs/AutomationTool/IOS/IOSPlatform.Automation.cs
@@ -364,11 +364,11 @@ public class IOSPlatform : Platform
 
 		if (!string.IsNullOrEmpty(Settings.BundleId))
 		{
-			CertLoc = Retriever.RetrieveByTags(new string[] { "DevCert: " + Settings.BundleId }, null);
+			CertLoc = Retriever.RetrieveFileSource("DevCert: " + Settings.BundleId);
 		}
 		if (CertLoc == null)
 		{
-			CertLoc = Retriever.RetrieveByTags(new string[] { "DevCert" }, null);
+			CertLoc = Retriever.RetrieveFileSource("DevCert");
 		}
 
 		if (CertLoc != null)
diff --git a/Engine/Source/Programs/AutomationTool/Turnkey/Commands/InstallSdk.cs b/Engine/Source/Programs/AutomationTool/Turnkey/Commands/InstallSdk.cs
--- a/Engine/Source/Programs/AutomationTool/Turnkey/Commands/InstallSdk.cs
+++ b/Engine/Source/Programs/AutomationTool/Turnkey/Commands/InstallSdk.cs
@@ -154,7 +154,7 @@ namespace Turnkey.Commands
 							DeviceInfo Device = GetDevice(AutomationPlatform, DeviceName);
 							Current = Device == null ? "N/A" : Device.SoftwareVersion;
 						}
-						Options.Add(string.Format("[{0}] {1} [Current: {2}]", string.Join(",", Sdk.GetPlatforms()), Sdk.DisplayName, Current));
+						Options.Add(string.Format("[{0}] {1} [Current: {2}]", string.Join(",", Sdk.GetPlatforms()), Sdk.Name, Current));
 					}
 
 					string Prompt = "Select an Sdk to install";
diff --git a/Engine/Source/Programs/AutomationTool/Turnkey/Commands/VerifySdk.cs b/Engine/Source/Programs/AutomationTool/Turnkey/Commands/VerifySdk.cs
--- a/Engine/Source/Programs/AutomationTool/Turnkey/Commands/VerifySdk.cs
+++ b/Engine/Source/Programs/AutomationTool/Turnkey/Commands/VerifySdk.cs
@@ -157,7 +157,7 @@ namespace Turnkey.Commands
 						continue;
 					}
 
-					TurnkeyUtils.Log("Will install {0}", BestSdk.DisplayName);
+					TurnkeyUtils.Log("Will install {0}", BestSdk.Name);
 
 					if (BestSdk.DownloadOrInstall(Platform) == false)
 					{
diff --git a/Engine/Source/Programs/AutomationTool/Turnkey/CopyProvider.cs b/Engine/Source/Programs/AutomationTool/Turnkey/CopyProvider.cs
--- a/Engine/Source/Programs/AutomationTool/Turnkey/CopyProvider.cs
+++ b/Engine/Source/Programs/AutomationTool/Turnkey/CopyProvider.cs
@@ -236,6 +236,47 @@ namespace Turnkey
 			return CopyAndRun.RunExternalCommand(Command, Params);
 		}
 
+		public string RetrieveFileSource(string Name, string InType, string InPlatform, string SubType)
+		{
+			UnrealTargetPlatform? Platform = null;
+			FileSource.SourceType? Type = null;
+
+			if (InPlatform != null)
+			{
+				// let this throw exception on failure, as it's a setup error
+				Platform = UnrealTargetPlatform.Parse(InPlatform);
+			}
+
+			if (InType != null)
+			{
+				// let this throw exception on failure, as it's a setup error
+				Type = (FileSource.SourceType)Enum.Parse(typeof(FileSource.SourceType), InType);
+			}
+
+			List<FileSource> Sources = TurnkeyManifest.FilterDiscoveredFileSources(Platform, Type);
+
+			Sources = Sources.FindAll(x =>
+			{
+				if (Name.StartsWith("regex:"))
+				{
+					return TurnkeyUtils.IsValueValid(x.Name, Name, null);
+				}
+				// this will handle the case of x.Name starting with regex: or just doing a case insensitive string comparison of tag and CustomSdkId
+				// range: is not supported, at least yet - we would have to check Tag with range: above, and also support range without a Platform (or pass in a platform somehow?)
+				return TurnkeyUtils.IsValueValid(Name, x.Name, null);
+			});
+
+			if (Sources.Count == 0)
+			{
+				return null;
+			}
+
+			// @todo turnkey: If > 1 in Sources, warn user
+
+			// execute the first found one
+			return CopyProvider.ExecuteCopy(Sources[0].GetCopySourceOperation());
+		}
+
 		public string RetrieveByTags(string[] RequiredTags, string[] PreferredTags, Dictionary<string, string> ExtraVariables = null)
 		{
 			return null;
diff --git a/Engine/Source/Programs/AutomationTool/Turnkey/FileSource.cs b/Engine/Source/Programs/AutomationTool/Turnkey/FileSource.cs
--- a/Engine/Source/Programs/AutomationTool/Turnkey/FileSource.cs
+++ b/Engine/Source/Programs/AutomationTool/Turnkey/FileSource.cs
@@ -103,7 +103,7 @@ namespace Turnkey
 		public CopySource[] Sources = null;
 
 		public string Version = null;
-		public string DisplayName = null;
+		public string Name = null;
 		public SourceType Type = SourceType.Full;
 		public string AllowedFlashDeviceTypes = null;
 
@@ -114,7 +114,7 @@ namespace Turnkey
 			FileSource Clone = new FileSource();
 			Clone.PlatformString = PlatformString;
 			Clone.Version = Version;
-			Clone.DisplayName = DisplayName;
+			Clone.Name = Name;
 			Clone.AllowedFlashDeviceTypes = AllowedFlashDeviceTypes;
 			Clone.Type = Type;
 
@@ -292,7 +292,7 @@ namespace Turnkey
 
 				// if unable to pick one automatically, ask the user
 				int BestIndex = Sdks.IndexOf(Best);
-				int Choice = TurnkeyUtils.ReadInputInt("Multiple Sdks found that could be installed. Please select one:", Sdks.Select(x => x.DisplayName + (x == Best ? " [Best Choice]" : "")).ToList(), true, BestIndex >= 0 ? BestIndex + 1 : -1);
+				int Choice = TurnkeyUtils.ReadInputInt("Multiple Sdks found that could be installed. Please select one:", Sdks.Select(x => x.Name + (x == Best ? " [Best Choice]" : "")).ToList(), true, BestIndex >= 0 ? BestIndex + 1 : -1);
 
 				// we take canceling to mean to try the next type
 				if (Choice == 0)
@@ -722,7 +722,7 @@ namespace Turnkey
 				{
 					if (ExpandingSource != null)
 					{
-						throw new AutomationTool.AutomationException("FileSource {0} had multiple expansions active on this platform. This is not allowed", DisplayName);
+						throw new AutomationTool.AutomationException("FileSource {0} had multiple expansions active on this platform. This is not allowed", Name);
 					}
 					ExpandingSource = Source;
 				}
@@ -826,11 +826,11 @@ namespace Turnkey
 			// validate
 			if (Version == null && Type == SourceType.Misc)
 			{
-				throw new AutomationTool.AutomationException("FileSource {0} needs to either be a Misc type, or have a version specified", DisplayName);
+				throw new AutomationTool.AutomationException("FileSource {0} needs to either be a Misc type, or have a version specified", Name);
 			}
 			if (Sources == null)
 			{
-				throw new AutomationTool.AutomationException("FileSource {0} has no acutal Source operations specified! This is a setup error.", DisplayName);
+				throw new AutomationTool.AutomationException("FileSource {0} has no acutal Source operations specified! This is a setup error.", Name);
 			}
 
 			PlatformString = TurnkeyUtils.ExpandVariables(PlatformString, true);
@@ -864,11 +864,11 @@ namespace Turnkey
 
 
 			Version = TurnkeyUtils.ExpandVariables(Version, true);
-			DisplayName = TurnkeyUtils.ExpandVariables(DisplayName, true);
+			Name = TurnkeyUtils.ExpandVariables(Name, true);
 
-			if (string.IsNullOrEmpty(DisplayName))
+			if (string.IsNullOrEmpty(Name))
 			{
-				DisplayName = string.Format("{0} {1}", PlatformString, Version);
+				Name = string.Format("{0} {1}", PlatformString, Version);
 			}
 
 			AllowedFlashDeviceTypes = TurnkeyUtils.ExpandVariables(AllowedFlashDeviceTypes, true);
@@ -889,7 +889,7 @@ namespace Turnkey
 		public string ToString(int BaseIndent)
 		{
 			StringBuilder Builder = new StringBuilder();
-			Builder.AppendLine("{1}Name: {0}", DisplayName, Indent(BaseIndent));
+			Builder.AppendLine("{1}Name: {0}", Name, Indent(BaseIndent));
 			Builder.AppendLine("{1}Version: {0}", Version == null ? "<Any>" : Version, Indent(BaseIndent + 2));
 			Builder.AppendLine("{1}Platform: {0}", Platforms == null ? "" : string.Join(",", Platforms), Indent(BaseIndent + 2));
 			Builder.AppendLine("{1}Type: {0}", Type, Indent(BaseIndent + 2));
diff --git a/Engine/Source/Programs/AutomationTool/Turnkey/SdkInfo.cs b/Engine/Source/Programs/AutomationTool/Turnkey/SdkInfo.cs
--- a/Engine/Source/Programs/AutomationTool/Turnkey/SdkInfo.cs
+++ b/Engine/Source/Programs/AutomationTool/Turnkey/SdkInfo.cs
@@ -188,7 +188,7 @@ namespace Turnkey
 				}
 
 				// and now activate it in case we need it this run
-				TurnkeyUtils.Log("Re-activating AutoSDK '{0}'...", Source.DisplayName);
+				TurnkeyUtils.Log("Re-activating AutoSDK '{0}'...", Source.Name);
 
 				UEBuildPlatformSDK.GetSDKForPlatform(Platform.ToString()).ReactivateAutoSDK();
 
