Include UGS badge info in exported Horde data.

#rnx

diff --git a/Engine/Source/Programs/AutomationTool/BuildGraph/Graph.cs b/Engine/Source/Programs/AutomationTool/BuildGraph/Graph.cs
--- a/Engine/Source/Programs/AutomationTool/BuildGraph/Graph.cs
+++ b/Engine/Source/Programs/AutomationTool/BuildGraph/Graph.cs
@@ -722,6 +722,35 @@ namespace AutomationTool
 				}
 				JsonWriter.WriteArrayEnd();
 
+				JsonWriter.WriteArrayStart("Badges");
+				foreach (Badge Badge in Badges)
+				{
+					HashSet<Node> Dependencies = Badge.Nodes;
+					if (Dependencies.Count > 0)
+					{
+						// Reduce that list to the smallest subset of direct dependencies
+						HashSet<Node> DirectDependencies = new HashSet<Node>(Dependencies);
+						foreach (Node Dependency in Dependencies)
+						{
+							DirectDependencies.ExceptWith(Dependency.OrderDependencies);
+						}
+
+						JsonWriter.WriteObjectStart();
+						JsonWriter.WriteValue("Name", Badge.Name);
+						if (!String.IsNullOrEmpty(Badge.Project))
+						{
+							JsonWriter.WriteValue("Project", Badge.Project);
+						}
+						if (Badge.Change != 0)
+						{
+							JsonWriter.WriteValue("Change", Badge.Change);
+						}
+						JsonWriter.WriteValue("Dependencies", String.Join(";", DirectDependencies.Select(x => x.Name)));
+						JsonWriter.WriteObjectEnd();
+					}
+				}
+				JsonWriter.WriteArrayEnd();
+
 				JsonWriter.WriteObjectEnd();
 			}
 		}
