Allow specifying UGS badges via the Label attribute.

#rnx

diff --git a/Engine/Source/Programs/AutomationTool/BuildGraph/Graph.cs b/Engine/Source/Programs/AutomationTool/BuildGraph/Graph.cs
--- a/Engine/Source/Programs/AutomationTool/BuildGraph/Graph.cs
+++ b/Engine/Source/Programs/AutomationTool/BuildGraph/Graph.cs
@@ -431,11 +431,11 @@ namespace AutomationTool
 				foreach(Label Label in Labels)
 				{
 					Writer.WriteStartElement("Label");
-					if (Label.Category != null)
+					if (Label.DashboardCategory != null)
 					{
-						Writer.WriteAttributeString("Category", Label.Category);
+						Writer.WriteAttributeString("Category", Label.DashboardCategory);
 					}
-					Writer.WriteAttributeString("Name", Label.Name);
+					Writer.WriteAttributeString("Name", Label.DashboardName);
 					Writer.WriteAttributeString("Requires", String.Join(";", Label.RequiredNodes.Select(x => x.Name)));
 
 					HashSet<Node> IncludedNodes = new HashSet<Node>(Label.IncludedNodes);
@@ -701,11 +701,27 @@ namespace AutomationTool
 				foreach (Label Label in Labels)
 				{
 					JsonWriter.WriteObjectStart();
-					if (Label.Category != null)
+					if (!String.IsNullOrEmpty(Label.DashboardName))
 					{
-						JsonWriter.WriteValue("Category", Label.Category);
+						JsonWriter.WriteValue("Name", Label.DashboardName);
 					}
-					JsonWriter.WriteValue("Name", Label.Name);
+					if (!String.IsNullOrEmpty(Label.DashboardCategory))
+					{
+						JsonWriter.WriteValue("Category", Label.DashboardCategory);
+					}
+					if (!String.IsNullOrEmpty(Label.UgsBadge))
+					{
+						JsonWriter.WriteValue("UgsBadge", Label.UgsBadge);
+					}
+					if (!String.IsNullOrEmpty(Label.UgsProject))
+					{
+						JsonWriter.WriteValue("UgsProject", Label.UgsProject);
+					}
+					if (Label.Change != LabelChange.Current)
+					{
+						JsonWriter.WriteValue("Change", Label.Change.ToString());
+					}
+
 					JsonWriter.WriteArrayStart("RequiredNodes");
 					foreach (Node RequiredNode in Label.RequiredNodes.OrderBy(x => x.Name))
 					{
diff --git a/Engine/Source/Programs/AutomationTool/BuildGraph/Label.cs b/Engine/Source/Programs/AutomationTool/BuildGraph/Label.cs
--- a/Engine/Source/Programs/AutomationTool/BuildGraph/Label.cs
+++ b/Engine/Source/Programs/AutomationTool/BuildGraph/Label.cs
@@ -8,21 +8,52 @@ using System.Threading.Tasks;
 
 namespace AutomationTool
 {
+	/// <summary>
+	/// Which changelist to show a UGS badge for
+	/// </summary>
+	enum LabelChange
+	{
+		/// <summary>
+		/// The current changelist being built
+		/// </summary>
+		Current,
+
+		/// <summary>
+		/// The last code changelist
+		/// </summary>
+		Code,
+	}
+
 	/// <summary>
 	/// Defines a label within a graph. Labels are similar to badges, and give the combined status of one or more job steps. Unlike badges, they
 	/// separate the requirements for its status and optional nodes to be included in its status, allowing this to be handled externally.
 	/// </summary>
 	class Label
 	{
+		/// <summary>
+		/// Name of this badge
+		/// </summary>
+		public readonly string DashboardName;
+
 		/// <summary>
 		/// Category for this label
 		/// </summary>
-		public readonly string Category;
+		public readonly string DashboardCategory;
 
 		/// <summary>
-		/// Name of this badge
+		/// Name of the badge in UGS
+		/// </summary>
+		public readonly string UgsBadge;
+
+		/// <summary>
+		/// Path to the project folder in UGS
+		/// </summary>
+		public readonly string UgsProject;
+
+		/// <summary>
+		/// Which change to show the badge for
 		/// </summary>
-		public readonly string Name;
+		public readonly LabelChange Change;
 
 		/// <summary>
 		/// Set of nodes that must be run for this label to be shown.
@@ -37,12 +68,18 @@ namespace AutomationTool
 		/// <summary>
 		/// Constructor
 		/// </summary>
-		/// <param name="InCategory">Type of this label</param>
-		/// <param name="InName">Name of this label</param>
-		public Label(string InCategory, string InName)
+		/// <param name="InDashboardName">Name of this label</param>
+		/// <param name="InDashboardCategory">Type of this label</param>
+		/// <param name="InUgsBadge">The UGS badge name</param>
+		/// <param name="InUgsProject">Project to display this badge for</param>
+		/// <param name="InChange">The change to show this badge on in UGS</param>
+		public Label(string InDashboardName, string InDashboardCategory, string InUgsBadge, string InUgsProject, LabelChange InChange)
 		{
-			Category = InCategory;
-			Name = InName;
+			DashboardName = InDashboardName;
+			DashboardCategory = InDashboardCategory;
+			UgsBadge = InUgsBadge;
+			UgsProject = InUgsProject;
+			Change = InChange;
 		}
 
 		/// <summary>
@@ -51,7 +88,14 @@ namespace AutomationTool
 		/// <returns>The name of this label</returns>
 		public override string ToString()
 		{
-			return String.Format("{0}/{1}", Category, Name);
+			if (!String.IsNullOrEmpty(DashboardName))
+			{
+				return String.Format("{0}/{1}", DashboardCategory, DashboardName);
+			}
+			else
+			{
+				return UgsBadge;
+			}
 		}
 	}
 }
diff --git a/Engine/Source/Programs/AutomationTool/BuildGraph/Schema.cs b/Engine/Source/Programs/AutomationTool/BuildGraph/Schema.cs
--- a/Engine/Source/Programs/AutomationTool/BuildGraph/Schema.cs
+++ b/Engine/Source/Programs/AutomationTool/BuildGraph/Schema.cs
@@ -178,6 +178,7 @@ namespace AutomationTool
 		BalancedString,
 		Boolean,
 		Integer,
+		LabelChange
 	}
 
 	/// <summary>
@@ -343,6 +344,7 @@ namespace AutomationTool
 			NewSchema.Items.Add(CreateReportType());
 			NewSchema.Items.Add(CreateBadgeType());
 			NewSchema.Items.Add(CreateLabelType());
+			NewSchema.Items.Add(CreateEnumType(GetTypeName(ScriptSchemaStandardType.LabelChange), typeof(LabelChange)));
 			NewSchema.Items.Add(CreateNotifyType());
 			NewSchema.Items.Add(CreateIncludeType());
 			NewSchema.Items.Add(CreateOptionType());
@@ -722,8 +724,11 @@ namespace AutomationTool
 		{
 			XmlSchemaComplexType LabelType = new XmlSchemaComplexType();
 			LabelType.Name = GetTypeName(ScriptSchemaStandardType.Label);
-			LabelType.Attributes.Add(CreateSchemaAttribute("Name", ScriptSchemaStandardType.Name, XmlSchemaUse.Required));
+			LabelType.Attributes.Add(CreateSchemaAttribute("Name", ScriptSchemaStandardType.Name, XmlSchemaUse.Optional));
 			LabelType.Attributes.Add(CreateSchemaAttribute("Category", ScriptSchemaStandardType.Name, XmlSchemaUse.Optional));
+			LabelType.Attributes.Add(CreateSchemaAttribute("UgsBadge", ScriptSchemaStandardType.Name, XmlSchemaUse.Optional));
+			LabelType.Attributes.Add(CreateSchemaAttribute("UgsProject", ScriptSchemaStandardType.BalancedString, XmlSchemaUse.Optional));
+			LabelType.Attributes.Add(CreateSchemaAttribute("Change", ScriptSchemaStandardType.LabelChange, XmlSchemaUse.Optional));
 			LabelType.Attributes.Add(CreateSchemaAttribute("Requires", ScriptSchemaStandardType.NameOrTagList, XmlSchemaUse.Required));
 			LabelType.Attributes.Add(CreateSchemaAttribute("Include", ScriptSchemaStandardType.NameOrTagList, XmlSchemaUse.Optional));
 			LabelType.Attributes.Add(CreateSchemaAttribute("Exclude", ScriptSchemaStandardType.NameOrTagList, XmlSchemaUse.Optional));
diff --git a/Engine/Source/Programs/AutomationTool/BuildGraph/Script.cs b/Engine/Source/Programs/AutomationTool/BuildGraph/Script.cs
--- a/Engine/Source/Programs/AutomationTool/BuildGraph/Script.cs
+++ b/Engine/Source/Programs/AutomationTool/BuildGraph/Script.cs
@@ -970,11 +970,11 @@ namespace AutomationTool
 					int SlashIdx = LabelCategoryName.IndexOf('/');
 					if (SlashIdx != -1)
 					{
-						Label = new Label(LabelCategoryName.Substring(0, SlashIdx), LabelCategoryName.Substring(SlashIdx + 1));
+						Label = new Label(LabelCategoryName.Substring(SlashIdx + 1), LabelCategoryName.Substring(0, SlashIdx), null, null, LabelChange.Current);
 					}
 					else
 					{
-						Label = new Label("Other", LabelCategoryName);
+						Label = new Label(LabelCategoryName, "Other", null, null, LabelChange.Current);
 					}
 
 					// Find all the included nodes
@@ -1059,16 +1059,26 @@ namespace AutomationTool
 		/// <param name="Element">Xml element to read the definition from</param>
 		void ReadLabel(ScriptElement Element)
 		{
-			string Name;
-			if (EvaluateCondition(Element) && TryReadObjectName(Element, out Name))
+			if (EvaluateCondition(Element))
 			{
+				string Name = ReadAttribute(Element, "Name");
+				if (!String.IsNullOrEmpty(Name))
+				{
+					ValidateName(Element, Name);
+				}
+
 				string Category = ReadAttribute(Element, "Category");
 
 				string[] RequiredNames = ReadListAttribute(Element, "Requires");
 				string[] IncludedNames = ReadListAttribute(Element, "Include");
 				string[] ExcludedNames = ReadListAttribute(Element, "Exclude");
 
-				Label NewLabel = new Label(Category, Name);
+				string UgsBadge = ReadAttribute(Element, "UgsBadge");
+				string UgsProject = ReadAttribute(Element, "UgsProject");
+
+				LabelChange Change = ReadEnumAttribute<LabelChange>(Element, "Change", LabelChange.Current);
+
+				Label NewLabel = new Label(Name, Category, UgsBadge, UgsProject, Change);
 				foreach (Node ReferencedNode in ResolveReferences(Element, RequiredNames))
 				{
 					NewLabel.RequiredNodes.Add(ReferencedNode);
