Send automated test results to Horde

#rnx

diff --git a/Engine/Source/Programs/AutomationTool/Gauntlet/Framework/Utils/Gauntlet.HordeReport.cs b/Engine/Source/Programs/AutomationTool/Gauntlet/Framework/Utils/Gauntlet.HordeReport.cs
new file mode 100644
--- /dev/null
+++ b/Engine/Source/Programs/AutomationTool/Gauntlet/Framework/Utils/Gauntlet.HordeReport.cs
@@ -0,0 +1,359 @@
+// Copyright Epic Games, Inc. All Rights Reserved.
+
+using System;
+using System.Collections.Generic;
+using System.IO;
+using Tools.DotNETCommon;
+using AutomationTool;
+
+namespace Gauntlet
+{
+	public static class HordeReport
+	{
+		/// <summary>
+		/// Default location to store Test Data 
+		/// </summary>
+		public static string DefaultTestDataDir
+		{
+			get
+			{
+				return Environment.GetEnvironmentVariable("UE_TESTDATA_DIR") ?? Path.Combine(CommandUtils.CmdEnv.EngineSavedFolder, "TestData");
+			}
+		}
+		/// <summary>
+		/// Default location to store test Artifacts 
+		/// </summary>
+		public static string DefaultArtifactsDir
+		{
+			get
+			{
+				return Environment.GetEnvironmentVariable("UE_ARTIFACTS_DIR") ?? CommandUtils.CmdEnv.LogFolder;
+			}
+		}
+
+		/// <summary>
+		/// Contains reference to files used or generated for file comparison
+		/// </summary>
+		public class ComparisonFiles
+		{
+			public string Difference;
+			public string Approved;
+			public string Unapproved;
+		}
+		/// <summary>
+		/// Contains information about test artifact
+		/// </summary>
+		public class Artifact
+		{
+			public Artifact()
+			{
+				Files = new ComparisonFiles();
+			}
+
+			public string Id;
+			public string Name;
+			public string Type;
+			public ComparisonFiles Files;
+		}
+		/// <summary>
+		/// Contains information about test entry event
+		/// </summary>
+		public class Event
+		{
+			public string Type;
+			public string Message;
+			public string Context;
+			public string Artifact;
+		}
+		/// <summary>
+		/// Contains information about test entry
+		/// </summary>
+		public class Entry
+		{
+			public Entry()
+			{
+				Event = new Event();
+			}
+
+			public Event Event;
+			public string Filename;
+			public int LineNumber;
+			public string Timestamp;
+		}
+		/// <summary>
+		/// Contains detailed information about test result. This is to what TestPassResult refere to for each test result. 
+		/// </summary>
+		public class TestResultDetailed
+		{
+			public TestResultDetailed()
+			{
+				Artifacts = new List<Artifact>();
+				Entries = new List<Entry>();
+			}
+
+			public string TestDisplayName;
+			public string FullTestPath;
+			public string State;
+			public int Warnings;
+			public int Errors;
+			public List<Artifact> Artifacts;
+			public List<Entry> Entries;
+
+			/// <summary>
+			/// Add a new Artifact to the test result and return it 
+			/// </summary>
+			public Artifact AddNewArtifact()
+			{
+				Artifact NewArtifact = new Artifact();
+				Artifacts.Add(NewArtifact);
+
+				return NewArtifact;
+			}
+
+			/// <summary>
+			/// Add a new Entry to the test result and return it 
+			/// </summary>
+			public Entry AddNewEntry()
+			{
+				Entry NewEntry = new Entry();
+				Entries.Add(NewEntry);
+
+				return NewEntry;
+			}
+		}
+		/// <summary>
+		/// Contains a brief information about test result.
+		/// </summary>
+		public class TestResult
+		{
+			public TestResult()
+			{
+				TestDetailed = new TestResultDetailed();
+			}
+
+			public string TestDisplayName
+			{
+				get { return TestDetailed.TestDisplayName; }
+				set { TestDetailed.TestDisplayName = value; }
+			}
+			public string FullTestPath
+			{
+				get { return TestDetailed.FullTestPath; }
+				set { TestDetailed.FullTestPath = value; }
+			}
+			public string State
+			{
+				get { return TestDetailed.State; }
+				set { TestDetailed.State = value; }
+			}
+			public string ArtifactName;
+
+			private TestResultDetailed TestDetailed;
+
+			/// <summary>
+			/// Return the underlying TestResultDetailed 
+			/// </summary>
+			public TestResultDetailed GetTestResultDetailed()
+			{
+				return TestDetailed;
+			}
+			/// <summary>
+			/// Set the underlying TestResultDetailed
+			/// </summary>
+			public void SetTestResultDetailed(TestResultDetailed InTestResultDetailed)
+			{
+				TestDetailed = InTestResultDetailed;
+			}
+		}
+
+		/// <summary>
+		/// Contains information about an entire test pass 
+		/// </summary>
+		public class TestPassResults
+		{
+			public TestPassResults()
+			{
+				Tests = new List<TestResult>();
+			}
+
+			public string ClientDescriptor;
+			public string ReportCreatedOn;
+			public int SucceededCount;
+			public int SucceededWithWarningsCount;
+			public int FailedCount;
+			public int NotRunCount;
+			public float TotalDurationSeconds;
+			public List<TestResult> Tests;
+
+			/// <summary>
+			/// Add a new TestResult to the pass results and return it 
+			/// </summary>
+			private TestResult AddNewTestResult()
+			{
+				TestResult NewTestResult = new TestResult();
+				Tests.Add(NewTestResult);
+
+				return NewTestResult;
+			}
+
+			/// <summary>
+			/// Convert UnrealAutomatedTestPassResults to Horde data model
+			/// </summary>
+			/// <param name="InTestPassResults"></param>
+			/// <param name="ReportPath"></param>
+			public static TestPassResults FromUnrealAutomatedTests(UnrealAutomatedTestPassResults InTestPassResults, string ReportPath)
+			{
+				TestPassResults OutTestPassResults = new TestPassResults();
+				OutTestPassResults.ClientDescriptor = InTestPassResults.clientDescriptor;
+				OutTestPassResults.ReportCreatedOn = InTestPassResults.reportCreatedOn;
+				OutTestPassResults.SucceededCount = InTestPassResults.succeeded;
+				OutTestPassResults.SucceededWithWarningsCount = InTestPassResults.succeededWithWarnings;
+				OutTestPassResults.FailedCount = InTestPassResults.failed;
+				OutTestPassResults.NotRunCount = InTestPassResults.notRun;
+				OutTestPassResults.TotalDurationSeconds = InTestPassResults.totalDuration;
+				if (InTestPassResults.tests != null)
+				{
+					foreach (UnrealAutomatedTestResult InTestResult in InTestPassResults.tests)
+					{
+						TestResult ConvertedTestResult = OutTestPassResults.AddNewTestResult();
+						ConvertedTestResult.TestDisplayName = InTestResult.testDisplayName;
+						ConvertedTestResult.FullTestPath = InTestResult.fullTestPath;
+						ConvertedTestResult.State = InTestResult.state;
+						Guid TestGuid = Guid.NewGuid();
+						ConvertedTestResult.ArtifactName = TestGuid + ".json";
+						InTestResult.artifactName = ConvertedTestResult.ArtifactName;
+						// Copy Test Result Detail
+						TestResultDetailed ConvertedTestResultDetailed = ConvertedTestResult.GetTestResultDetailed();
+						foreach (UnrealAutomationArtifact InTestArtifact in InTestResult.artifacts)
+						{
+							Artifact NewArtifact = ConvertedTestResultDetailed.AddNewArtifact();
+							NewArtifact.Id = InTestArtifact.id;
+							NewArtifact.Name = InTestArtifact.name;
+							NewArtifact.Type = InTestArtifact.type;
+							ComparisonFiles ArtifactFiles = NewArtifact.Files;
+							ArtifactFiles.Difference = Path.Combine(ReportPath, InTestArtifact.files.difference);
+							ArtifactFiles.Approved = Path.Combine(ReportPath, InTestArtifact.files.approved);
+							ArtifactFiles.Unapproved = Path.Combine(ReportPath, InTestArtifact.files.unapproved);
+						}
+						foreach (UnrealAutomationEntry InTestEntry in InTestResult.entries)
+						{
+							Entry NewEntry = ConvertedTestResultDetailed.AddNewEntry();
+							NewEntry.Filename = InTestEntry.filename;
+							NewEntry.LineNumber = InTestEntry.lineNumber;
+							NewEntry.Timestamp = InTestEntry.timestamp;
+							Event EntryEvent = NewEntry.Event;
+							EntryEvent.Artifact = [REDACTED_EMAIL];
+							EntryEvent.Context = [REDACTED_EMAIL];
+							EntryEvent.Message = [REDACTED_EMAIL];
+							EntryEvent.Type = [REDACTED_EMAIL];
+						}
+					}
+				}
+				return OutTestPassResults;
+			}
+
+			/// <summary>
+			/// Copy Test Results Artifacts
+			/// </summary>
+			/// <param name="OutputArtifactPath"></param>
+			public void CopyTestResultsArtifacts(string OutputArtifactPath)
+			{
+
+				if (!Directory.Exists(OutputArtifactPath))
+				{
+					Directory.CreateDirectory(OutputArtifactPath);
+				}
+				int ArtifactsCount = 0;
+				foreach (TestResult OutputTestResult in Tests)
+				{
+					TestResultDetailed OutputTestResultDetailed = OutputTestResult.GetTestResultDetailed();
+					// copy artifacts
+					foreach (Artifact TestArtifact in OutputTestResultDetailed.Artifacts)
+					{
+						string[] ArtifactPaths= { TestArtifact.Files.Difference, TestArtifact.Files.Approved, TestArtifact.Files.Unapproved };
+						foreach (string ArtifactPath in ArtifactPaths)
+						{
+							if (File.Exists(ArtifactPath))
+							{
+								try
+								{
+									File.Copy(ArtifactPath, Path.Combine(OutputArtifactPath, Path.GetFileName(ArtifactPath)));
+									ArtifactsCount++;
+								}
+								catch (Exception Ex)
+								{
+									Log.Error("Failed to copy artifact '{0}'. {1}", Path.GetFileName(ArtifactPath), Ex);
+								}
+							}
+						}
+						TestArtifact.Files.Difference = Path.GetFileName(TestArtifact.Files.Difference);
+						TestArtifact.Files.Approved = Path.GetFileName(TestArtifact.Files.Approved);
+						TestArtifact.Files.Unapproved = Path.GetFileName(TestArtifact.Files.Unapproved);
+					}
+					// write test json blob
+					string TestResultFilePath = Path.Combine(OutputArtifactPath, OutputTestResult.ArtifactName);
+					try
+					{
+						Json.Save(new FileReference(TestResultFilePath), OutputTestResultDetailed);
+						ArtifactsCount++;
+					}
+					catch (Exception Ex)
+					{
+						Log.Error("Failed to save Test Result for '{0}'. {1}", OutputTestResult.TestDisplayName, Ex);
+					}
+				}
+				Log.Verbose("Copied {0} artifacts for Horde to {1}", ArtifactsCount, OutputArtifactPath);
+			}
+		}
+
+		/// <summary>
+		/// Container for Test Data items 
+		/// </summary>
+		public class TestDataCollection
+		{
+			public class DataItem
+			{
+				public string Key;
+				public object Data;
+			}
+			public TestDataCollection()
+			{
+				Items = new List<DataItem>();
+			}
+
+			public DataItem AddNewItem(string InKey, object InData)
+			{
+				DataItem NewDataItem = new DataItem();
+				NewDataItem.Key = InKey;
+				NewDataItem.Data = InData;
+				Items.Add(NewDataItem);
+				return NewDataItem;
+			}
+
+			public List<DataItem> Items;
+
+			/// <summary>
+			/// Write Test Data Collection to json
+			/// </summary>
+			/// <param name="OutputTestDataFilePath"></param>
+			public void WriteToJson(string OutputTestDataFilePath)
+			{
+				string OutputTestDataDir = Path.GetDirectoryName(OutputTestDataFilePath);
+				if (!Directory.Exists(OutputTestDataDir))
+				{
+					Directory.CreateDirectory(OutputTestDataDir);
+				}
+				// write test pass summary
+				Log.Verbose("Writing Test Data Collection for Horde at {0}", OutputTestDataFilePath);
+				try
+				{
+					Json.Save(new FileReference(OutputTestDataFilePath), this);
+				}
+				catch (Exception Ex)
+				{
+					Log.Error("Failed to save Test Data Collection for Horde. {0}", Ex);
+				}
+			}
+		}
+	}
+}
diff --git a/Engine/Source/Programs/AutomationTool/Gauntlet/Gauntlet.Automation.csproj b/Engine/Source/Programs/AutomationTool/Gauntlet/Gauntlet.Automation.csproj
--- a/Engine/Source/Programs/AutomationTool/Gauntlet/Gauntlet.Automation.csproj
+++ b/Engine/Source/Programs/AutomationTool/Gauntlet/Gauntlet.Automation.csproj
@@ -81,6 +81,7 @@
     <Compile Include="Framework\Gauntlet.Exception.cs" />
     <Compile Include="Framework\Gauntlet.Params.cs" />
     <Compile Include="Framework\Gauntlet.TestGroup.cs" />
+    <Compile Include="Framework\Utils\Gauntlet.HordeReport.cs" />
     <Compile Include="Framework\Utils\Gauntlet.IniConfigUtil.cs" />
     <Compile Include="Framework\Utils\Gauntlet.MarkdownBuilder.cs" />
     <Compile Include="Framework\Utils\Gauntlet.HtmlBuilder.cs" />
@@ -136,6 +137,7 @@
     <Compile Include="Unreal\Utils\Gauntlet.UnrealLogParser.cs" />
     <Compile Include="Unreal\BuildSource\Gauntlet.StagedBuildSource.cs" />
     <Compile Include="Unreal\Utils\Gauntlet.UnrealSnapshotParser.cs" />
+    <Compile Include="Unreal\Utils\Gauntlet.UnrealTestResults.cs" />
   </ItemGroup>
   <ItemGroup>
     <ProjectReference Include="..\..\DotNETCommon\DotNETUtilities\DotNETUtilities.csproj">
diff --git a/Engine/Source/Programs/AutomationTool/Gauntlet/Unreal/Utils/Gauntlet.UnrealTestResults.cs b/Engine/Source/Programs/AutomationTool/Gauntlet/Unreal/Utils/Gauntlet.UnrealTestResults.cs
new file mode 100644
--- /dev/null
+++ b/Engine/Source/Programs/AutomationTool/Gauntlet/Unreal/Utils/Gauntlet.UnrealTestResults.cs
@@ -0,0 +1,73 @@
+// Copyright Epic Games, Inc. All Rights Reserved.
+
+using System.Collections.Generic;
+using Tools.DotNETCommon;
+
+namespace Gauntlet
+{
+	// The following classes are used for json deserialization of the Unreal Automation Test Results.
+	// Because of a bug in the json deserialization used here(Tools.DotNETCommon.Json),
+	// the names of the properties need to match the case of the json properties.
+
+	public class UnrealAutomationComparisonFiles
+	{
+		public string difference;
+		public string approved;
+		public string unapproved;
+	}
+	public class UnrealAutomationArtifact
+	{
+		public string id;
+		public string name;
+		public string type;
+		public UnrealAutomationComparisonFiles files;
+	}
+	public class UnrealAutomationEvent
+	{
+		public string type;
+		public string message;
+		public string context;
+		public string artifact;
+	}
+	public class UnrealAutomationEntry
+	{
+		public UnrealAutomationEvent @event;
+		public string filename;
+		public int lineNumber;
+		public string timestamp;
+	}
+	public class UnrealAutomatedTestResult
+	{
+		public string testDisplayName;
+		public string fullTestPath;
+		public string artifactName;
+		public string state;
+		public int warnings;
+		public int errors;
+		public List<UnrealAutomationArtifact> artifacts;
+		public List<UnrealAutomationEntry> entries;
+	}
+	public class UnrealAutomatedTestPassResults
+	{
+		public string clientDescriptor;
+		public string reportCreatedOn;
+		public int succeeded;
+		public int succeededWithWarnings;
+		public int failed;
+		public int notRun;
+		public float totalDuration;
+		public bool comparisonExported;
+		public string comparisonExportDirectory;
+		public List<UnrealAutomatedTestResult> tests;
+
+		/// <summary>
+		/// Load Unreal Automated Test Results from json report
+		/// </summary>
+		/// <param name="FilePath"></param>
+		public static UnrealAutomatedTestPassResults LoadFromJson(string FilePath)
+		{
+			UnrealAutomatedTestPassResults JsonTestPassResults = Json.Load<UnrealAutomatedTestPassResults>(new FileReference(FilePath));
+			return JsonTestPassResults;
+		}
+	}
+}
