- Updated some turnkey installations
- Removed BuildSource type, and just use the generic FileSource for finding builds
- Improved support for $[x] captures

diff --git a/Engine/Source/Programs/AutomationTool/Turnkey/Commands/InstallBuild.cs b/Engine/Source/Programs/AutomationTool/Turnkey/Commands/InstallBuild.cs
--- a/Engine/Source/Programs/AutomationTool/Turnkey/Commands/InstallBuild.cs
+++ b/Engine/Source/Programs/AutomationTool/Turnkey/Commands/InstallBuild.cs
@@ -13,12 +13,8 @@ namespace Turnkey.Commands
 	{
 		protected override void Execute(string[] CommandOptions)
 		{
-		
-			// get all projects
-			List<BuildSource> AllSources = TurnkeyManifest.GetDiscoveredBuildSources();
-
 			// get distinct project names
-			List<string> AllProjects = AllSources.Select(x => x.Project).Distinct().ToList();
+			List<string> AllProjects = TurnkeyManifest.GetProjectsWithBuilds();
 
 			int Choice = TurnkeyUtils.ReadInputInt("Select a project:", AllProjects, true);
 			if (Choice == 0)
@@ -29,50 +25,26 @@ namespace Turnkey.Commands
 			string Project = AllProjects[Choice - 1];
 			
 			// go over sources for this project
-			List<BuildSource> ProjectSources = AllSources.Where(x => x.Project == Project).ToList();
-
+			List<FileSource> ProjectSources = TurnkeyManifest.FilterDiscoveredBuilds(Project);
 
-			// go over all possible sources, assume there is an wildcard with an expansion in it
-			// @todo turnkey: if there is no wildcard, then we could just use the final part of pathname in the results
-			Dictionary<string, Tuple<string, string>> BuildSourceToLocation = new Dictionary<string, Tuple<string, string>>();
-			foreach (BuildSource Source in ProjectSources)
-			{
-				foreach (CopyAndRun Copy in Source.Sources)
-				{
-					List<List<string>> Expansions = new List<List<string>>();
-					string[] EnumeratedSources = CopyProvider.ExecuteEnumerate(Copy.Copy, Expansions);
-					for (int Index = 0; Index < EnumeratedSources.Length; Index++)
-					{
-						// get a mapping of description (expanded bit) to location
-						BuildSourceToLocation.Add(Expansions[Index][0], new Tuple<string,string>(EnumeratedSources[Index], Source.BuildEnumerationSuffix));
-					}
-				}
-			}
+//			List<UnrealTargetPlatform> PossiblePlatforms = ProjectSources.SelectMany(x => x.GetPlatforms()).Distinct().ToList();
+			List<string> PossibleVersions = ProjectSources.Select(x => x.Version).Distinct().ToList();
 
-			if (BuildSourceToLocation.Count == 0)
-			{
-				TurnkeyUtils.Log("No builds found!");
-				return;
-			}
-
-			List<string> BuildKeys = BuildSourceToLocation.Keys.ToList();
-			Choice = TurnkeyUtils.ReadInputInt("Choose a build version", BuildKeys, true);
+			Choice = TurnkeyUtils.ReadInputInt("Choose a build version", PossibleVersions, true);
 			if (Choice == 0)
 			{
 				return;
 			}
 
-			// get the chosen path
-			string BuildLocation = BuildSourceToLocation[BuildKeys[Choice - 1]].Item1;
-			string EnumerationSuffix = BuildSourceToLocation[BuildKeys[Choice - 1]].Item2;
+			FileSource OriginalSource = ProjectSources[Choice - 1];
 
 			// now find platforms under it
 			List<List<string>> PlatformExpansion = new List<List<string>>();
-			string[] PlatformSources = CopyProvider.ExecuteEnumerate(BuildLocation + EnumerationSuffix, PlatformExpansion);
+			string[] PlatformSources = CopyProvider.ExecuteEnumerate(OriginalSource.GetCopySourceOperation() + OriginalSource.BuildPlatformEnumerationSuffix, PlatformExpansion);
 
 			if (PlatformSources.Length == 0)
 			{
-				TurnkeyUtils.Log("No platform builds found in {0}!", BuildLocation);
+				TurnkeyUtils.Log("No platform builds found in {0}!", OriginalSource.GetCopySourceOperation());
 				return;
 			}
 
@@ -90,11 +62,9 @@ namespace Turnkey.Commands
 			UnrealTargetPlatform Platform;
 			if (!UnrealTargetPlatform.TryParse(TargetPlatformString, out Platform))
 			{
-				Action<string> FindAndStrip = x => { int Loc = TargetPlatformString.IndexOf(x); if (Loc > 0) { TargetPlatformString = TargetPlatformString.Substring(0, Loc); } };
-
-				FindAndStrip("NoEditor");
-				FindAndStrip("Client");
-				FindAndStrip("Server");
+				TargetPlatformString = TargetPlatformString.Replace("Editor", "");
+				TargetPlatformString = TargetPlatformString.Replace("Client", "");
+				TargetPlatformString = TargetPlatformString.Replace("Server", "");
 
 				// Windows hack
 				if (TargetPlatformString == "Windows")
diff --git a/Engine/Source/Programs/AutomationTool/Turnkey/Commands/InstallSdk.cs b/Engine/Source/Programs/AutomationTool/Turnkey/Commands/InstallSdk.cs
--- a/Engine/Source/Programs/AutomationTool/Turnkey/Commands/InstallSdk.cs
+++ b/Engine/Source/Programs/AutomationTool/Turnkey/Commands/InstallSdk.cs
@@ -37,14 +37,16 @@ namespace Turnkey.Commands
 			bool bNeededOnly = bBestAvailable || TurnkeyUtils.ParseParam("NeededOnly", CommandOptions);
 
 
-			FileSource.SourceType DesiredType = FileSource.SourceType.Misc;
+			FileSource.SourceType? DesiredType = null;
 			if (SdkTypeString != null)
 			{
-				if (!Enum.TryParse(SdkTypeString, out DesiredType))
+				FileSource.SourceType OutType;
+				if (!Enum.TryParse(SdkTypeString, out OutType))
 				{
 					TurnkeyUtils.Log("Invalid SdkType given with -SdkType={0}", SdkTypeString);
 					return;
 				}
+				DesiredType = OutType;
 			}
 
 			// we need all sdks we can find
@@ -72,11 +74,8 @@ namespace Turnkey.Commands
 				AutomationTool.Platform AutomationPlatform = AutomationTool.Platform.GetPlatform(Platform);
 				UEBuildPlatformSDK SDK = UEBuildPlatformSDK.GetSDKForPlatform(Platform.ToString());
 
-				// filter on type
-				FileSource.SourceType? OptionalType = SdkTypeString != null ? (FileSource.SourceType?)DesiredType : null;
-
-				// filter the Sdks if a platform was given
-				List<FileSource> Sdks = TurnkeyManifest.FilterDiscoveredFileSources(Platform, OptionalType);
+				// filter the Sdks if a platform was given, and type if it was given
+				List<FileSource> Sdks = TurnkeyManifest.FilterDiscoveredFileSources(Platform, DesiredType);
 
 				// strip out flash Sdks where there are no devices
 				int SdkCount = Sdks.Count;
@@ -84,7 +83,7 @@ namespace Turnkey.Commands
 				bool bStrippedDevices = Sdks.Count != SdkCount;
 
 				// skip Misc FileSources, as we dont know how they are sued
-				Sdks = Sdks.FindAll(x => x.Type != FileSource.SourceType.Misc);
+				Sdks = Sdks.FindAll(x => x.IsSdkType());
 
 				if (bUpdateOnly)
 				{
@@ -219,7 +218,7 @@ namespace Turnkey.Commands
 						}
 					}
 
-					Sdk.DownloadOrInstall(Platform);
+					Sdk.DownloadOrInstall(Platform, InstallDevice);
 				}
 			}
 		}
diff --git a/Engine/Source/Programs/AutomationTool/Turnkey/CopyProvider.cs b/Engine/Source/Programs/AutomationTool/Turnkey/CopyProvider.cs
--- a/Engine/Source/Programs/AutomationTool/Turnkey/CopyProvider.cs
+++ b/Engine/Source/Programs/AutomationTool/Turnkey/CopyProvider.cs
@@ -277,54 +277,6 @@ namespace Turnkey
 			return CopyProvider.ExecuteCopy(Sources[0].GetCopySourceOperation());
 		}
 
-		public string RetrieveByTags(string[] RequiredTags, string[] PreferredTags, Dictionary<string, string> ExtraVariables = null)
-		{
-			return null;
-// 			// @todo turnkey - delete, or replace with a more "retrieve by Platform and Name/Version something" function
-// 
-// 			// @todo turnkey: unset these
-// 			if (ExtraVariables != null)
-// 			{
-// 				foreach (var Pair in ExtraVariables)
-// 				{
-// 					TurnkeyUtils.SetVariable(Pair.Key, Pair.Value);
-// 				}
-// 			}
-// 
-// 			List<SdkInfo> Sdks = TurnkeyManifest.GetAllDiscoveredSdks();
-// 
-// 			string Tag = RequiredTags[0];
-// 			Sdks = Sdks.FindAll(x => x.Type == FileSource.SourceType.Misc);
-// 
-// 			Sdks = Sdks.FindAll(x =>
-// 			{
-// 				if (Tag.StartsWith("regex:"))
-// 				{
-// 					return TurnkeyUtils.IsValueValid(x.CustomSdkId, Tag, null);
-// 				}
-// 				// this will handle the case of x.CustomSdkId starting with regex: or just doing a case insensitive string comparison of tag and CustomSdkId
-// 				// range: is not supported, at least yet - we would have to check Tag with range: above, and also support range without a Platform (or pass in a platform somehow?)
-// 				return TurnkeyUtils.IsValueValid(Tag, x.CustomSdkId, null);
-// 			});
-// 
-// 			if (Sdks == null || Sdks.Count == 0)
-// 			{
-// 				return null;
-// 			}
-// 
-// 			SdkInfo Sdk = Sdks.First();
-// 
-// 			foreach (CopyAndRun Copy in Sdk.CustomSdkInputFiles)
-// 			{
-// 				if (Copy.ShouldExecute())
-// 				{
-// 					return CopyProvider.ExecuteCopy(Copy.Copy);
-// 				}
-// 			}
-// 
-// 			return null;
-		}
-
 		public string GetVariable(string VariableName)
 		{
 			return TurnkeyUtils.GetVariableValue(VariableName);
diff --git a/Engine/Source/Programs/AutomationTool/Turnkey/FileSource.cs b/Engine/Source/Programs/AutomationTool/Turnkey/FileSource.cs
--- a/Engine/Source/Programs/AutomationTool/Turnkey/FileSource.cs
+++ b/Engine/Source/Programs/AutomationTool/Turnkey/FileSource.cs
@@ -89,6 +89,7 @@ namespace Turnkey
 			Full,
 			Flash,
 			Misc,
+			Build,
 		};
 
 		#region Fields
@@ -106,6 +107,11 @@ namespace Turnkey
 		public string Name = null;
 		public SourceType Type = SourceType.Full;
 		public string AllowedFlashDeviceTypes = null;
+		
+		// Project for Build types to filter on
+		public string Project;
+		public string BuildPlatformEnumerationSuffix;
+
 
 
 
@@ -117,6 +123,8 @@ namespace Turnkey
 			Clone.Name = Name;
 			Clone.AllowedFlashDeviceTypes = AllowedFlashDeviceTypes;
 			Clone.Type = Type;
+			Clone.Project = Project;
+			Clone.BuildPlatformEnumerationSuffix = BuildPlatformEnumerationSuffix;
 
 			if (NewValue != null)
 			{
@@ -178,10 +186,15 @@ namespace Turnkey
 			return Platforms.Contains(Platform);
 		}
 
+		public bool IsSdkType()
+		{
+			return Type == SourceType.BuildOnly || Type == SourceType.AutoSdk || Type == SourceType.RunOnly || Type == SourceType.Full || Type == SourceType.Flash;
+		}
+
 		public bool IsVersionValid(UnrealTargetPlatform Platform, DeviceInfo Device=null)
 		{
-			// Misc ones are always valid
-			if (Type == SourceType.Misc)
+			// Non-Sdk types are always valid, although this isn't meant for them
+			if (!IsSdkType())
 			{
 				return true;
 			}
@@ -700,6 +713,7 @@ namespace Turnkey
 			Expansions = CheckForListExpansions(Expansions, x => x.PlatformString, (x,y) => x.PlatformString = y);
 			Expansions = CheckForListExpansions(Expansions, x => x.Version, (x, y) => x.Version = y);
 			Expansions = CheckForListExpansions(Expansions, x => x.AllowedFlashDeviceTypes, (x, y) => x.AllowedFlashDeviceTypes = y);
+			Expansions = CheckForListExpansions(Expansions, x => x.Project, (x, y) => x.Project = y);
 
 			// return null if nothing was actually expanded!
 			if (Expansions.Count == 1 && Expansions[0] == this)
@@ -824,9 +838,9 @@ namespace Turnkey
 		internal void PostDeserialize()
 		{
 			// validate
-			if (Version == null && Type == SourceType.Misc)
+			if (Version == null && IsSdkType())
 			{
-				throw new AutomationTool.AutomationException("FileSource {0} needs to either be a Misc type, or have a version specified", Name);
+				throw new AutomationTool.AutomationException("FileSource {0} needs to have a version specified, since it's an Sdk type", Name);
 			}
 			if (Sources == null)
 			{
diff --git a/Engine/Source/Programs/AutomationTool/Turnkey/Providers/NullCopyProvider.cs b/Engine/Source/Programs/AutomationTool/Turnkey/Providers/NullCopyProvider.cs
--- a/Engine/Source/Programs/AutomationTool/Turnkey/Providers/NullCopyProvider.cs
+++ b/Engine/Source/Programs/AutomationTool/Turnkey/Providers/NullCopyProvider.cs
@@ -100,6 +100,9 @@ namespace Turnkey
 				// get the path component that has a wildcard
 				string Wildcard = (NextSlash == -1) ? PathString.Substring(PrevSlash + 1) : PathString.Substring(PrevSlash + 1, (NextSlash - PrevSlash) - 1);
 
+				// replace * with a non-greedy capture so *-foo-*.zip can get the best values for the *'s
+				Regex Regex = new Regex(string.Format("^{0}$", Wildcard.Replace("*", "(.+?)")), RegexOptions.IgnoreCase);
+
 				// track what's before and after the * to return what it expanded to
 				int StarLoc = Wildcard.IndexOf('*');
 				int PrefixLen = StarLoc;
@@ -111,8 +114,20 @@ namespace Turnkey
 					{
 						NewExpansionSet = new List<string>(ExpansionSet);
 						string PathComponent = Path.GetFileName(Dirname);
-						// the match is the part of the filename that was not the *, so removing that will give us what we wanted to match
-						NewExpansionSet.Add(PathComponent.Substring(PrefixLen, PathComponent.Length - (PrefixLen + SuffixLen)));
+
+						Match Match = Regex.Match(PathComponent);
+						if (!Match.Success)
+						{
+							continue;
+						}
+
+						for (int GroupIndex = 1; GroupIndex < Match.Groups.Count; GroupIndex++)
+						{
+							NewExpansionSet.Add(Match.Groups[GroupIndex].Value);
+						}
+
+//						// the match is the part of the filename that was not the *, so removing that will give us what we wanted to match
+//						NewExpansionSet.Add(PathComponent.Substring(PrefixLen, PathComponent.Length - (PrefixLen + SuffixLen)));
 					}
 
 					if (bIsLastComponent)
@@ -130,15 +145,27 @@ namespace Turnkey
 				// if the path ends with a slash, then we are only looking for directories (D:\\Sdks\*\ would only want to return directories)
 				if (bIsLastComponent && NextSlash == -1)
 				{
-					foreach (string Filename in Directory.EnumerateFiles(FullPathBeforeWildcard, Wildcard))
+					foreach (string Filename in Directory.EnumerateFiles(FullPathBeforeWildcard))//, Wildcard))
 					{
+						string PathComponent = Path.GetFileName(Filename);
+						Match Match = Regex.Match(PathComponent);
+						if (!Match.Success)
+						{
+							continue;
+						}
+
 						List<string> NewExpansionSet = null;
 						if (ExpansionSet != null)
 						{
 							NewExpansionSet = new List<string>(ExpansionSet);
-							string PathComponent = Path.GetFileName(Filename);
+
+							for (int GroupIndex = 1; GroupIndex < Match.Groups.Count; GroupIndex++)
+							{
+								NewExpansionSet.Add(Match.Groups[GroupIndex].Value);
+							}
+
 							// the match is the part of the filename that was not the *, so removing that will give us what we wanted to match
-							NewExpansionSet.Add(PathComponent.Substring(PrefixLen, PathComponent.Length - (PrefixLen + SuffixLen)));
+//							NewExpansionSet.Add(PathComponent.Substring(PrefixLen, PathComponent.Length - (PrefixLen + SuffixLen)));
 						}
 
 						Output.Add(ProviderToken + ":" + Filename, NewExpansionSet);
diff --git a/Engine/Source/Programs/AutomationTool/Turnkey/TurnkeyManifest.cs b/Engine/Source/Programs/AutomationTool/Turnkey/TurnkeyManifest.cs
--- a/Engine/Source/Programs/AutomationTool/Turnkey/TurnkeyManifest.cs
+++ b/Engine/Source/Programs/AutomationTool/Turnkey/TurnkeyManifest.cs
@@ -20,8 +20,6 @@ namespace Turnkey
 		[XmlElement("FileSource")]
 		public FileSource[] FileSources = null;
 
-		public BuildSource[] BuildSources = null;
-
 		[XmlArrayItem(ElementName = "Manifest")]
 		public string[] AdditionalManifests = null;
 
@@ -35,11 +33,6 @@ namespace Turnkey
 				Array.ForEach(SavedSettings, x => TurnkeyUtils.SetVariable(x.Variable, TurnkeyUtils.ExpandVariables(x.Value)));
 			}
 
-			if (BuildSources != null)
-			{
-				Array.ForEach(BuildSources, x => x.PostDeserialize());
-			}
-
 			if (FileSources != null)
 			{
 				Array.ForEach(FileSources, x => x.PostDeserialize());
@@ -69,25 +62,60 @@ namespace Turnkey
 		}
 
 		static List<FileSource> DiscoveredFileSources = null;
-		static List<BuildSource> DiscoveredBuildSources = null;
+
 		public static List<UnrealTargetPlatform> GetPlatformsWithSdks()
 		{
 			DiscoverManifests();
 
 			// this can handle FileSources with pending Expansions, so get the unique set of Platforms represented by pre or post expansion sources
-			// skip over Misc type, since this wants just Sdk types
+			// skip over non-Sdk types, since this wants just Sdks
 			List<UnrealTargetPlatform> Platforms = new List<UnrealTargetPlatform>();
-			DiscoveredFileSources.FindAll(x => x.Type != FileSource.SourceType.Misc).ForEach(x => Platforms.AddRange(x.GetPlatforms()));
+			DiscoveredFileSources.FindAll(x => x.IsSdkType()).ForEach(x => Platforms.AddRange(x.GetPlatforms()));
 
 			return Platforms.Distinct().ToList();
 		}
 
+		public static List<string> GetProjectsWithBuilds()
+		{
+			DiscoverManifests();
+
+			// this can handle FileSources with pending Expansions, so get the unique set of Platforms represented by pre or post expansion sources
+			// skip over non-Sdk types, since this wants just Sdks
+			List<string> Projects = new List<string>();
+			DiscoveredFileSources.FindAll(x => x.Type == FileSource.SourceType.Build).ForEach(x => Projects.Add(x.Project));
+
+			return Projects.Distinct(StringComparer.InvariantCultureIgnoreCase).ToList();
+		}
+
 
 		public static List<FileSource> GetAllDiscoveredFileSources()
 		{
 			return FilterDiscoveredFileSources(null, null);
 		}
 		
+		private static List<FileSource> ExpandFilteredSources(List<FileSource> Sources)
+		{
+			// get the set that needs to expand
+			List<FileSource> NeedsExpansion = Sources.FindAll(x => x.NeedsFileExpansion());
+			List<FileSource> NotNeedsExpansion = Sources.FindAll(x => !x.NeedsFileExpansion());
+
+
+			// remove them, then we add expansions below
+			DiscoveredFileSources = DiscoveredFileSources.FindAll(x => !x.NeedsFileExpansion());
+
+			foreach (FileSource Source in NeedsExpansion)
+			{
+				List<FileSource> ExpansionResult = Source.ExpandCopySource();
+
+				// add them to the full set of sources, and our filtered sources
+				NotNeedsExpansion.AddRange(ExpansionResult);
+				DiscoveredFileSources.AddRange(ExpansionResult);
+			}
+
+			// NotNeedsExpansion now contains the expanded sources, and the sources that didn't need to be expanded
+			return NotNeedsExpansion;
+		}
+
 		public static List<FileSource> FilterDiscoveredFileSources(UnrealTargetPlatform? Platform, FileSource.SourceType? Type)
 		{
 			// hunt down manifests if needed
@@ -104,30 +132,17 @@ namespace Turnkey
 				Matching = DiscoveredFileSources.FindAll(x => (Platform == null || (x.PlatformString?.StartsWith("$(")).GetValueOrDefault() || x.SupportsPlatform(Platform.Value)) && (Type == null || x.Type == Type.Value));
 			}
 
-			// get the set that needs to expand
-			List<FileSource> WorkingSet = Matching.FindAll(x => x.NeedsFileExpansion());
-
-			// remove them, then we add expansions below
-			DiscoveredFileSources = DiscoveredFileSources.FindAll(x => !WorkingSet.Contains(x));
-
-			foreach (FileSource Source in WorkingSet)
-			{
-				DiscoveredFileSources.AddRange(Source.ExpandCopySource());
-			}
-
-			if (Platform == null && Type == null)
-			{
-				return DiscoveredFileSources;
-			}
-			return DiscoveredFileSources.FindAll(x => (Platform == null || x.SupportsPlatform(Platform.Value)) && (Type == null || x.Type == Type.Value));
+			return ExpandFilteredSources(Matching);
 		}
 
-		public static List<BuildSource> GetDiscoveredBuildSources()
+		public static List<FileSource> FilterDiscoveredBuilds(string Project)
 		{
 			// hunt down manifests if needed
 			DiscoverManifests();
 
-			return DiscoveredBuildSources;
+			List<FileSource> Filtered = DiscoveredFileSources.FindAll(x => x.Type == FileSource.SourceType.Build && x.Project.Equals(Project, StringComparison.InvariantCultureIgnoreCase));
+
+			return ExpandFilteredSources(Filtered);
 		}
 
 		public static void DiscoverManifests()
@@ -136,7 +151,6 @@ namespace Turnkey
 			if (DiscoveredFileSources == null)
 			{
 				DiscoveredFileSources = new List<FileSource>();
-				DiscoveredBuildSources = new List<BuildSource>();
 
 				// known location to branch from, this will include a few other locations
 				string RootOperation = "file:$(EngineDir)/Build/Turnkey/TurnkeyManifest.xml";
@@ -147,10 +161,6 @@ namespace Turnkey
 					{
 						DiscoveredFileSources.AddRange(x.FileSources);
 					}
-					if (x.BuildSources != null)
-					{
-						DiscoveredBuildSources.AddRange(x.BuildSources);
-					}
 				});
 			}
 		}
diff --git a/Engine/Source/Programs/AutomationTool/Turnkey/TurnkeyUtils.cs b/Engine/Source/Programs/AutomationTool/Turnkey/TurnkeyUtils.cs
--- a/Engine/Source/Programs/AutomationTool/Turnkey/TurnkeyUtils.cs
+++ b/Engine/Source/Programs/AutomationTool/Turnkey/TurnkeyUtils.cs
@@ -274,10 +274,14 @@ namespace Turnkey
 
 		public static bool IsValueValid(string Value, string AllowedValues, AutomationTool.Platform Platform)
 		{
-			if (string.IsNullOrEmpty(Value) || string.IsNullOrEmpty(AllowedValues))
+			if (string.IsNullOrEmpty(Value))
 			{
 				return false;
 			}
+			if (string.IsNullOrEmpty(AllowedValues))
+			{
+				return true;
+			}
 
 			// use a regex if the allowed string starts with "regex:"
 			if (AllowedValues.StartsWith("regex:", StringComparison.InvariantCultureIgnoreCase))
