Turnkey:
- Added SDK installation pre/post-ambles so the platform can give user more
information on what's happening, and what to do on a failure (or a success)
- Added some defaults for Studio settings in TurnkeyManifest.xml, to show how to set in a studio's .xml file
- Added/updated some Studio settings/names
- Added bCanConnect flag to DeviceInfo (in case a platform can list devices even if they are unable to flash/run/etc - in which case the fields may or may not be correct)
- Added $(AllPlatforms) and $(AutoSDKPlatforms) built-in variables

diff --git a/Engine/Build/Turnkey/TurnkeyManifest.xml b/Engine/Build/Turnkey/TurnkeyManifest.xml
--- a/Engine/Build/Turnkey/TurnkeyManifest.xml
+++ b/Engine/Build/Turnkey/TurnkeyManifest.xml
@@ -1,7 +1,19 @@
 <?xml version="1.0" encoding="utf-8" ?>
 <TurnkeyManifest>
+	<!-- Some default values for Studio-wide settings. It's recommended you set these appropriately in one of the .xml files referenced below. "RunuAT Turnkey -command=ManageSettings -listvars" for more information -->
+	<SavedSettings>
+		<SavedSetting>
+			<Variable>Studio_AppleSigningCertPassword</Variable>
+			<Value></Value>
+		</SavedSetting>
+		<SavedSetting>
+			<Variable>Studio_GoogleDriveCredentials</Variable>
+			<Value>file:$(EngineDir)/Build/Turnkey/GoogleDriveCredentials.json</Value>
+		</SavedSetting>
+	</SavedSettings>
 	<AdditionalManifests>
 		<Manifest>$(ThisManifestDir)/StudioTurnkeyManifest.xml</Manifest>
+		<Manifest>file:$(EngineDir)/Restricted/NotForLicensees/Build/Turnkey/TurnkeyManifest.xml</Manifest>
 		<Manifest>file:$(EngineDir)/Restricted/NoRedist/Build/Turnkey/TurnkeyManifest.xml</Manifest>
 		<Manifest>$(UE_STUDIO_TURNKEY_LOCATION)</Manifest>
 	</AdditionalManifests>
diff --git a/Engine/Source/Programs/AutomationTool/AutomationUtils/Platform.cs b/Engine/Source/Programs/AutomationTool/AutomationUtils/Platform.cs
--- a/Engine/Source/Programs/AutomationTool/AutomationUtils/Platform.cs
+++ b/Engine/Source/Programs/AutomationTool/AutomationUtils/Platform.cs
@@ -16,7 +16,7 @@ namespace AutomationTool
 		string RetrieveFileSource(string Name, string InType = "Misc", string InPlatform = null, string SubType = null);
 		string RetrieveFileSource(object HintObject);
 		string GetVariable(string VariableName);
-		bool RunExternalCommand(string Command, string Params);
+		bool RunExternalCommand(string Command, string Params, string Preamble = "", string SuccessPostamble = "", string FailurePostamble = "");
 	}
 
 	//public interface InputOutput
@@ -31,6 +31,9 @@ namespace AutomationTool
 		public string SoftwareVersion;
 		public string Type;
 		public bool bIsDefault = false;
+		// is the device able to be connected to (this is more about able to flash SDK or run, not about matching SDK version)
+		// if false, any of the above fields are suspect, especually SoftwareVersion
+		public bool bCanConnect = true;
 	}
 
 	/// <summary>
@@ -185,28 +188,28 @@ namespace AutomationTool
 				return false;
 			}
 
-			string Command, Params;
+			string Command, Params, Preamble="", SuccessPostamble="", FailurePostamble="";
 			
-			if (Device != null && GetDeviceUpdateSoftwareCommand(out Command, out Params, Retriever, Device))
+			if (Device != null && GetDeviceUpdateSoftwareCommand(out Command, out Params, ref Preamble, ref SuccessPostamble, ref FailurePostamble, Retriever, Device))
 			{
-				return Retriever.RunExternalCommand(Command, Params);
+				return Retriever.RunExternalCommand(Command, Params, Preamble, SuccessPostamble, FailurePostamble);
 			}
-			else if (Device == null && GetSDKInstallCommand(out Command, out Params, Retriever))
+			else if (Device == null && GetSDKInstallCommand(out Command, out Params, ref Preamble, ref SuccessPostamble, ref FailurePostamble, Retriever))
 			{
-				return Retriever.RunExternalCommand(Command, Params);
+				return Retriever.RunExternalCommand(Command, Params, Preamble, SuccessPostamble, FailurePostamble);
 			}
 
 			return false;
 		}
 
-		public virtual bool GetSDKInstallCommand(out string Command, out string Params, FileRetriever Retriever)
+		public virtual bool GetSDKInstallCommand(out string Command, out string Params, ref string Preamble, ref string SuccessPostamble, ref string FailurePostamble, FileRetriever Retriever)
 		{
 			Command = null;
 			Params = null;
 			return false;
 		}
 
-		public virtual bool GetDeviceUpdateSoftwareCommand(out string Command, out string Params, FileRetriever Retriever, DeviceInfo Device = null)
+		public virtual bool GetDeviceUpdateSoftwareCommand(out string Command, out string Params, ref string Preamble, ref string SuccessPostamble, ref string FailurePostamble, FileRetriever Retriever, DeviceInfo Device = null)
 		{
 			Command = null;
 			Params = null;
diff --git a/Engine/Source/Programs/AutomationTool/IOS/IOSPlatform.Automation.cs b/Engine/Source/Programs/AutomationTool/IOS/IOSPlatform.Automation.cs
--- a/Engine/Source/Programs/AutomationTool/IOS/IOSPlatform.Automation.cs
+++ b/Engine/Source/Programs/AutomationTool/IOS/IOSPlatform.Automation.cs
@@ -174,8 +174,11 @@ public class IOSPlatform : Platform
 		return @"regex:^13\.5.*";
 	}
 
-	public override bool GetSDKInstallCommand(out string Command, out string Params, FileRetriever Retriever)
+	public override bool GetSDKInstallCommand(out string Command, out string Params, ref string Preamble, ref string SuccessPostamble, ref string FailurePostamble, FileRetriever Retriever)
 	{
+		Preamble = "Moving your original Xcode application from /Applications to the Trash, and unzipping the new version into /Applications";
+		SuccessPostamble = "If you had Xcode in your Dock, you will need to remove it and add the new one (even though it was in the same location). macOS follows the move to the Trash for the Dock icon";
+
 		// put current Xcode in the trash, and unzip a new one. Xcode in the dock will have to be fixed up tho!
 		Command = "osascript";
 		Params = 
@@ -190,8 +193,10 @@ public class IOSPlatform : Platform
 		return true;
 	}
 
-	public override bool GetDeviceUpdateSoftwareCommand(out string Command, out string Params, FileRetriever Retriever, DeviceInfo Device)
+	public override bool GetDeviceUpdateSoftwareCommand(out string Command, out string Params, ref string Preamble, ref string SuccessPostamble, ref string FailurePostamble, FileRetriever Retriever, DeviceInfo Device)
 	{
+		Preamble = "Installing an offline downloaded .ipsw onto your device using the Apple Configurator application.";
+
 		// cfgtool needs ECID, not UDID, so find it
 		string Configurator = Path.Combine(GetConfiguratorLocation().Replace(" ", "\\ "), "Contents/MacOS/cfgutil");
 
diff --git a/Engine/Source/Programs/AutomationTool/Turnkey/Commands/ManageSettings.cs b/Engine/Source/Programs/AutomationTool/Turnkey/Commands/ManageSettings.cs
--- a/Engine/Source/Programs/AutomationTool/Turnkey/Commands/ManageSettings.cs
+++ b/Engine/Source/Programs/AutomationTool/Turnkey/Commands/ManageSettings.cs
@@ -23,6 +23,9 @@ namespace Turnkey.Commands
 			UserSetting[] UserSettings = TurnkeySettings.AllUserSettings;
 			UserSetting[] StudioSettings = TurnkeySettings.AllStudioSettings;
 
+			// need to load the manifests to get StudioSettings!
+			TurnkeyManifest.DiscoverManifests();
+
 			if (TurnkeyUtils.ParseParam("ListVars", CommandOptions))
 			{
 				TurnkeyUtils.Log("All Settings:");
diff --git a/Engine/Source/Programs/AutomationTool/Turnkey/CopyProvider.cs b/Engine/Source/Programs/AutomationTool/Turnkey/CopyProvider.cs
--- a/Engine/Source/Programs/AutomationTool/Turnkey/CopyProvider.cs
+++ b/Engine/Source/Programs/AutomationTool/Turnkey/CopyProvider.cs
@@ -231,9 +231,43 @@ namespace Turnkey
 			return null;
 		}
 
-		public bool RunExternalCommand(string Command, string Params)
+		public bool RunExternalCommand(string Command, string Params, string Preamble, string SuccessPostAmble, string FailurePostamble)
 		{
-			return CopyAndRun.RunExternalCommand(Command, Params);
+			TurnkeyUtils.Log("----------------------------------------------");
+			TurnkeyUtils.Log("Running '{0} {1}'", Command, Params);
+
+			if (!string.IsNullOrEmpty(Preamble))
+			{
+				TurnkeyUtils.Log("");
+				TurnkeyUtils.Log(Preamble);
+			}
+			TurnkeyUtils.Log("----------------------------------------------", Command);
+
+			bool bSuccess = CopyAndRun.RunExternalCommand(Command, Params);
+
+			TurnkeyUtils.Log("----------------------------------------------");
+			TurnkeyUtils.Log("Finished with {0}", bSuccess ? "Success" : "Failure");
+
+			if (bSuccess)
+			{
+				if (!string.IsNullOrEmpty(SuccessPostAmble))
+				{
+					TurnkeyUtils.Log("");
+					TurnkeyUtils.Log(SuccessPostAmble);
+				}
+			}
+			else
+			{
+				if (!string.IsNullOrEmpty(FailurePostamble))
+				{
+					TurnkeyUtils.Log("");
+					TurnkeyUtils.Log(FailurePostamble);
+				}
+			}
+
+			TurnkeyUtils.Log("----------------------------------------------", Command);
+
+			return bSuccess;
 		}
 
 		public string RetrieveFileSource(string Name, string InType, string InPlatform, string SubType)
diff --git a/Engine/Source/Programs/AutomationTool/Turnkey/Providers/GoogleDriveCopyProvider.cs b/Engine/Source/Programs/AutomationTool/Turnkey/Providers/GoogleDriveCopyProvider.cs
--- a/Engine/Source/Programs/AutomationTool/Turnkey/Providers/GoogleDriveCopyProvider.cs
+++ b/Engine/Source/Programs/AutomationTool/Turnkey/Providers/GoogleDriveCopyProvider.cs
@@ -224,16 +224,16 @@ namespace Turnkey
 		{
 			if (ServiceHelper == null)
 			{
-				if (!TurnkeyUtils.HasVariable("GoogleDrive_Credentials"))
+				if (!TurnkeyUtils.HasVariable("Studio_GoogleDriveCredentials"))
 				{
-					TurnkeyUtils.Log("ERROR: Unable to use GoogleDrive without GoogleDrive_Credentials being set first!");
+					TurnkeyUtils.Log("ERROR: Unable to use GoogleDrive without Studio_GoogleDriveCredentials being set first!");
 					return false;
 				}
 
-				string SecretsFile = CopyProvider.ExecuteCopy("$(GoogleDrive_Credentials)");
+				string SecretsFile = CopyProvider.ExecuteCopy("$(Studio_GoogleDriveCredentials)");
 				if (SecretsFile == null)
 				{
-					TurnkeyUtils.Log("ERROR: Unable to get GoogleDrive secrets file from {0}", TurnkeyUtils.GetVariableValue("GoogleDrive_Credentials"));
+					TurnkeyUtils.Log("ERROR: Unable to get GoogleDrive secrets file from {0}", TurnkeyUtils.GetVariableValue("Studio_GoogleDriveCredentials"));
 					return false;
 				}
 
diff --git a/Engine/Source/Programs/AutomationTool/Turnkey/SdkInfo.cs b/Engine/Source/Programs/AutomationTool/Turnkey/SdkInfo.cs
--- a/Engine/Source/Programs/AutomationTool/Turnkey/SdkInfo.cs
+++ b/Engine/Source/Programs/AutomationTool/Turnkey/SdkInfo.cs
@@ -108,8 +108,7 @@ namespace Turnkey
 			{
 				if (!bUnattended)
 				{
-					// @todo turnkey - have studio settings 
-					string Response = TurnkeyUtils.ReadInput("AutoSdks are not setup, but your studio has support. Would you like to set it up now? [Y/n]", "Y");
+					string Response = TurnkeyUtils.ReadInput("The AutoSdk system is not setup on this machine. Would you like to set it up now? [Y/n]", "Y");
 					if (string.Compare(Response, "Y", true) == 0)
 					{
 						bAttemptAutoSdkSetup = true;
diff --git a/Engine/Source/Programs/AutomationTool/Turnkey/TurnkeySettings.cs b/Engine/Source/Programs/AutomationTool/Turnkey/TurnkeySettings.cs
--- a/Engine/Source/Programs/AutomationTool/Turnkey/TurnkeySettings.cs
+++ b/Engine/Source/Programs/AutomationTool/Turnkey/TurnkeySettings.cs
@@ -46,7 +46,7 @@ namespace Turnkey
 		public static UserSetting[] AllStudioSettings =
 		{
 			new UserSetting("Studio_AppleSigningCertPassword", null, "A shared password that is used across Apple Signing Certificates"),
-			new UserSetting("Studio_AutoSDKSupported", null, "if true, your studio has set up AutoSDK, and Turnkey will offer to set up your local machine if needed."),
+			new UserSetting("Studio_GoogleDriveCredentials", null, "The location of credentials file needed for GoogleDrive integration. The contents should start with: '{\"installed\":{\"client_id\":'"),
 		};
 
 		// basically same as Turnkey variables, but this only contains ones that were loaded so we can write them back out
diff --git a/Engine/Source/Programs/AutomationTool/Turnkey/TurnkeyUtils.cs b/Engine/Source/Programs/AutomationTool/Turnkey/TurnkeyUtils.cs
--- a/Engine/Source/Programs/AutomationTool/Turnkey/TurnkeyUtils.cs
+++ b/Engine/Source/Programs/AutomationTool/Turnkey/TurnkeyUtils.cs
@@ -27,6 +27,16 @@ namespace Turnkey
 			IOProvider = InIOProvider;
 			CommandUtilHelper = InCommandUtilHelper;
 
+			// set up some lists
+			SetVariable("AllPlatforms", string.Join(",", UnrealTargetPlatform.GetValidPlatformNames()));
+
+			// walk over all the SDKs and get their AutoSDK string
+			IEnumerable<string> AutoSDKPlatforms = UEBuildPlatformSDK.AllSDKs.Select(x => x.GetAutoSDKPlatformName()).Distinct();
+			SetVariable("AutoSDKPlatforms", string.Join(",", AutoSDKPlatforms));
+
+			TurnkeyUtils.Log("AllPlatforms = {0}", GetVariableValue("AllPlatforms"));
+			TurnkeyUtils.Log("AutoSDKPlatforms = {0}", GetVariableValue("AutoSDKPlatforms"));
+
 			SetVariable("HOST_PLATFORM_NAME", HostPlatform.Current.HostEditorPlatform.ToString());
 		}
 
