Include the flag for whether to show warnings in the exported Horde graph definition.

diff --git a/Engine/Source/Programs/AutomationTool/BuildGraph/Graph.cs b/Engine/Source/Programs/AutomationTool/BuildGraph/Graph.cs
--- a/Engine/Source/Programs/AutomationTool/BuildGraph/Graph.cs
+++ b/Engine/Source/Programs/AutomationTool/BuildGraph/Graph.cs
@@ -660,6 +660,7 @@ namespace AutomationTool
 						JsonWriter.WriteObjectStart();
 						JsonWriter.WriteValue("Name", Node.Name);
 						JsonWriter.WriteValue("RunEarly", Node.bRunEarly);
+						JsonWriter.WriteValue("Warnings", Node.bNotifyOnWarnings);
 
 						JsonWriter.WriteArrayStart("InputDependencies");
 						foreach (string InputDependency in Node.GetDirectInputDependencies().Select(x => x.Name))
