Fix missing errors and warnings count not passed to HordeReport.TestResultDetailed
+ add external report URL + enforce specifying a type to testdata key

#rnx

diff --git a/Engine/Source/Programs/AutomationTool/Gauntlet/Framework/Utils/Gauntlet.HordeReport.cs b/Engine/Source/Programs/AutomationTool/Gauntlet/Framework/Utils/Gauntlet.HordeReport.cs
--- a/Engine/Source/Programs/AutomationTool/Gauntlet/Framework/Utils/Gauntlet.HordeReport.cs
+++ b/Engine/Source/Programs/AutomationTool/Gauntlet/Framework/Utils/Gauntlet.HordeReport.cs
@@ -178,6 +178,7 @@ namespace Gauntlet
 
 			public string ClientDescriptor;
 			public string ReportCreatedOn;
+			public string ReportURL;
 			public int SucceededCount;
 			public int SucceededWithWarningsCount;
 			public int FailedCount;
@@ -201,11 +202,13 @@ namespace Gauntlet
 			/// </summary>
 			/// <param name="InTestPassResults"></param>
 			/// <param name="ReportPath"></param>
-			public static TestPassResults FromUnrealAutomatedTests(UnrealAutomatedTestPassResults InTestPassResults, string ReportPath)
+			/// <param name="ReportURL"></param>
+			public static TestPassResults FromUnrealAutomatedTests(UnrealAutomatedTestPassResults InTestPassResults, string ReportPath, string ReportURL)
 			{
 				TestPassResults OutTestPassResults = new TestPassResults();
 				OutTestPassResults.ClientDescriptor = InTestPassResults.clientDescriptor;
 				OutTestPassResults.ReportCreatedOn = InTestPassResults.reportCreatedOn;
+				OutTestPassResults.ReportURL = ReportURL;
 				OutTestPassResults.SucceededCount = InTestPassResults.succeeded;
 				OutTestPassResults.SucceededWithWarningsCount = InTestPassResults.succeededWithWarnings;
 				OutTestPassResults.FailedCount = InTestPassResults.failed;
@@ -224,6 +227,8 @@ namespace Gauntlet
 						InTestResult.artifactName = ConvertedTestResult.ArtifactName;
 						// Copy Test Result Detail
 						TestResultDetailed ConvertedTestResultDetailed = ConvertedTestResult.GetTestResultDetailed();
+						ConvertedTestResultDetailed.Errors = InTestResult.errors;
+						ConvertedTestResultDetailed.Warnings = InTestResult.warnings;
 						foreach (UnrealAutomationArtifact InTestArtifact in InTestResult.artifacts)
 						{
 							Artifact NewArtifact = ConvertedTestResultDetailed.AddNewArtifact();
@@ -321,10 +326,10 @@ namespace Gauntlet
 				Items = new List<DataItem>();
 			}
 
-			public DataItem AddNewItem(string InKey, object InData)
+			public DataItem AddNewItem(string InType, string InKey, object InData)
 			{
 				DataItem NewDataItem = new DataItem();
-				NewDataItem.Key = InKey;
+				NewDataItem.Key = InType +"::"+ InKey;
 				NewDataItem.Data = InData;
 				Items.Add(NewDataItem);
 				return NewDataItem;
