- Added -WaitForUATMutex commandline parameter to UAT so that it will block on another UAT process instead of erroring out

diff --git a/Engine/Source/Programs/AutomationTool/AutomationUtils/Utils.cs b/Engine/Source/Programs/AutomationTool/AutomationUtils/Utils.cs
--- a/Engine/Source/Programs/AutomationTool/AutomationUtils/Utils.cs
+++ b/Engine/Source/Programs/AutomationTool/AutomationUtils/Utils.cs
@@ -717,7 +717,7 @@ namespace AutomationTool
 		/// </summary>
 		/// <param name="Main"></param>
 		/// <param name="Param"></param>
-		public static ExitCode RunSingleInstance(Func<ExitCode> Main)
+		public static ExitCode RunSingleInstance(string[] Arguments, Func<ExitCode> Main)
 		{
 			bool AllowMultipleInsances = (Environment.GetEnvironmentVariable("uebp_UATMutexNoWait") == "1");
 	
@@ -730,8 +730,21 @@ namespace AutomationTool
 				IsSoleInstance = bCreatedMutex;
 
 				if (!IsSoleInstance && AllowMultipleInsances == false)
-				{ 
-					throw new AutomationException("A conflicting instance of AutomationTool is already running. Curent location: {0}. A process manager may be used to determine the conflicting process and what tool may have launched it", EntryAssemblyLocation);
+				{
+					if (CommandUtils.ParseParam(Arguments, "-WaitForUATMutex"))
+					{
+						CommandUtils.LogWarning("Another instance of UAT at '{0}' is running, and the -WaitForUATMutex parameter has been used. Waiting for other UAT to finish...", EntryAssemblyLocation);
+						int Seconds = 0;
+						while (SingleInstanceMutex.WaitOne(15 * 1000) == false)
+						{
+							Seconds += 15;
+							CommandUtils.LogInformation("Still waiting for Mutex. {0} seconds has passed...", Seconds);
+						}
+					}
+					else
+					{
+						throw new AutomationException("A conflicting instance of AutomationTool is already running. Curent location: {0}. A process manager may be used to determine the conflicting process and what tool may have launched it", EntryAssemblyLocation);
+					}
 				}
 
 				ExitCode Result = Main();
diff --git a/Engine/Source/Programs/AutomationTool/Program.cs b/Engine/Source/Programs/AutomationTool/Program.cs
--- a/Engine/Source/Programs/AutomationTool/Program.cs
+++ b/Engine/Source/Programs/AutomationTool/Program.cs
@@ -89,7 +89,7 @@ namespace AutomationTool
 				Log.TraceVerbose("{0} ver. {1}", Version.ProductName, Version.ProductVersion);
 
 				// Don't allow simultaneous execution of AT (in the same branch)
-				ReturnCode = InternalUtils.RunSingleInstance(() => MainProc(Arguments, StartupListener));
+				ReturnCode = InternalUtils.RunSingleInstance(Arguments, () => MainProc(Arguments, StartupListener));
 			}
 			catch (AutomationException Ex)
 			{
