Fixed warning treated as error in UAT, added documentation for FindProducedTagNames

diff --git a/Engine/Source/Programs/AutomationTool/BuildGraph/Tasks/UploadTask.cs b/Engine/Source/Programs/AutomationTool/BuildGraph/Tasks/UploadTask.cs
--- a/Engine/Source/Programs/AutomationTool/BuildGraph/Tasks/UploadTask.cs
+++ b/Engine/Source/Programs/AutomationTool/BuildGraph/Tasks/UploadTask.cs
@@ -185,6 +185,10 @@ namespace BuildGraph.Tasks
 			return FindTagNamesFromFilespec(Parameters.Files);
 		}
 
+		/// <summary>
+		/// Finds the tags which are produced by this task
+		/// </summary>
+		/// <returns>The tag names which are produced by this task</returns>
 		public override IEnumerable<string> FindProducedTagNames()
 		{
 			return new string[] { };
