- Deprecating XboxOne (removing from builds, and no longer registering it in UBT)

diff --git a/Engine/Source/Programs/AutomationTool/Turnkey/Commands/ListPlatforms.cs b/Engine/Source/Programs/AutomationTool/Turnkey/Commands/ListPlatforms.cs
--- a/Engine/Source/Programs/AutomationTool/Turnkey/Commands/ListPlatforms.cs
+++ b/Engine/Source/Programs/AutomationTool/Turnkey/Commands/ListPlatforms.cs
@@ -19,7 +19,7 @@ namespace Turnkey.Commands
 			foreach (UnrealTargetPlatform TargetPlatform in UnrealTargetPlatform.GetValidPlatforms())
 			{
 				// HACK UNTIUL WIN32 IS GONE COMPLETELY
-				if (TargetPlatform == UnrealTargetPlatform.Win32)
+				if (TargetPlatform == UnrealTargetPlatform.Win32 || TargetPlatform == UnrealTargetPlatform.XboxOne)
 				{
 					continue;
 				}
diff --git a/Engine/Source/Programs/AutomationTool/Turnkey/TurnkeyUtils.cs b/Engine/Source/Programs/AutomationTool/Turnkey/TurnkeyUtils.cs
--- a/Engine/Source/Programs/AutomationTool/Turnkey/TurnkeyUtils.cs
+++ b/Engine/Source/Programs/AutomationTool/Turnkey/TurnkeyUtils.cs
@@ -109,6 +109,9 @@ namespace Turnkey
 			string PlatformString = TurnkeyUtils.ParseParamValue("Platform", null, CommandOptions);
 			bool bUnattended = TurnkeyUtils.ParseParam("Unattended", CommandOptions);
 
+			// HACK UNTIL WIN32 IS GONE
+			PossiblePlatforms = PossiblePlatforms.Where(x => x != UnrealTargetPlatform.Win32 && x != UnrealTargetPlatform.XboxOne).ToList();
+
 			List<UnrealTargetPlatform> Platforms = new List<UnrealTargetPlatform>();
 			// prompt user for a platform
 			if (PlatformString == null)
@@ -130,8 +133,7 @@ namespace Turnkey
 				// All platforms means to install every platform with an installer
 				if (PlatformChoice == PlatformOptions.Count)
 				{
-					// HACK UNTIL WIN32 IS GONE
-					Platforms = PossiblePlatforms.Where(x => x != UnrealTargetPlatform.Win32).ToList();
+					Platforms = PossiblePlatforms;
 				}
 				else
 				{
@@ -141,7 +143,7 @@ namespace Turnkey
 			else if (PlatformString == "All")
 			{
 				// HACK UNTIL WIN32 IS GONE
-				Platforms = PossiblePlatforms.Where(x => x != UnrealTargetPlatform.Win32).ToList();
+				Platforms = PossiblePlatforms;
 			}
 			else
 			{
