- Fixed some turnkey ios profile sorting issues
- Allow clearing of Turnkey User settings, or setting to empty string

diff --git a/Engine/Source/Programs/AutomationTool/Turnkey/Commands/ManageSettings.cs b/Engine/Source/Programs/AutomationTool/Turnkey/Commands/ManageSettings.cs
--- a/Engine/Source/Programs/AutomationTool/Turnkey/Commands/ManageSettings.cs
+++ b/Engine/Source/Programs/AutomationTool/Turnkey/Commands/ManageSettings.cs
@@ -78,7 +78,22 @@ namespace Turnkey.Commands
 			if (Choice > 0)
 			{
 				string NewValue = TurnkeyUtils.ReadInput(string.Format("Enter value for {0} [Currently '{1}']", UserSettings[Choice - 1].VariableName, TurnkeyUtils.GetVariableValue(UserSettings[Choice - 1].VariableName)));
-				if (NewValue != "")
+				if (NewValue == "")
+				{
+					int EmptyChoice = TurnkeyUtils.ReadInputInt("Empty response was given - what did you want to do?", new List<string> { "Delete setting", "Set to blank", "Leave existing value alone" }, false);
+					if (EmptyChoice == 0)
+					{
+						TurnkeyUtils.Log("Clearing Settings[{0}].{1}", Choice - 1, UserSettings[Choice - 1].VariableName);
+						TurnkeySettings.SetUserSetting(UserSettings[Choice - 1].VariableName, null);
+						TurnkeySettings.Save();
+					}
+					else
+					{
+						NewValue = EmptyChoice == 1 ? "" : null;
+					}
+				}
+
+				if (NewValue != null)
 				{
 					TurnkeyUtils.Log("Setting Settings[{0}].{1} = {2}", Choice - 1, UserSettings[Choice - 1].VariableName, NewValue);
 					TurnkeySettings.SetUserSetting(UserSettings[Choice - 1].VariableName, NewValue);
diff --git a/Engine/Source/Programs/AutomationTool/Turnkey/Providers/ConsoleIOProvider.cs b/Engine/Source/Programs/AutomationTool/Turnkey/Providers/ConsoleIOProvider.cs
--- a/Engine/Source/Programs/AutomationTool/Turnkey/Providers/ConsoleIOProvider.cs
+++ b/Engine/Source/Programs/AutomationTool/Turnkey/Providers/ConsoleIOProvider.cs
@@ -66,7 +66,7 @@ namespace Turnkey
 			// go until good choice
 			while (true)
 			{
-				string ChoiceString = ReadInput(FullPrompt, DefaultValue >= 0 ? DefaultValue.ToString() : "", false);
+				string ChoiceString = ReadInput(FullPrompt, DefaultValue >= 0 ? DefaultValue.ToString() : null, false);
 
 				int Choice;
 				if (Int32.TryParse(ChoiceString, out Choice) == false || Choice < 0 || Choice >= Options.Count + (bIsCancellable ? 1 : 0))
diff --git a/Engine/Source/Programs/AutomationTool/Turnkey/Providers/HybridIOProvider.cs b/Engine/Source/Programs/AutomationTool/Turnkey/Providers/HybridIOProvider.cs
--- a/Engine/Source/Programs/AutomationTool/Turnkey/Providers/HybridIOProvider.cs
+++ b/Engine/Source/Programs/AutomationTool/Turnkey/Providers/HybridIOProvider.cs
@@ -38,21 +38,21 @@ namespace Turnkey
 			// blank string means user canceled, which goes to stderr
 			if (OSAOutput == "")
 			{
-				return "";
+				return null;
 			}
 
 			// regex the result
 			Match Match = Regex.Match(OSAOutput, "text returned:(.*)$");
 			if (!Match.Success)
 			{
-				return "";
+				return null;
 			}
 
 			// return the text in the dialog box
 			return Match.Groups[1].Value;
 		}
 
-			private string ShowDialog(string Prompt, string Default, bool bIsList)
+		private string ShowDialog(string Prompt, string Default, bool bIsList)
 		{
 			string Result;
 			if (UnrealBuildTool.BuildHostPlatform.Current.Platform == UnrealBuildTool.UnrealTargetPlatform.Mac)
@@ -64,7 +64,7 @@ namespace Turnkey
 				Result = Microsoft.VisualBasic.Interaction.InputBox(Prompt, "Turnkey Input", Default);
 			}
 
-			if (Result == "" && bIsList)
+			if (string.IsNullOrEmpty(Result) && bIsList)
 			{
 				return "0";
 			}
@@ -204,7 +204,7 @@ namespace Turnkey
 			// go until good choice
 			while (true)
 			{
-				string ChoiceString = ShowDialog(FullPrompt, DefaultValue >= 0 ? DefaultValue.ToString() : "", true);
+				string ChoiceString = ShowDialog(FullPrompt, DefaultValue >= 0 ? DefaultValue.ToString() : null, true);
 
 				if (ChoiceString == null)
 				{
