- Improved the Xcode installation experience, for accepting the license. Requires admin privileges with a password or sudo, so buildmachines not running sudo jobs will not work

diff --git a/Engine/Source/Programs/AutomationTool/IOS/IOSPlatform.Automation.cs b/Engine/Source/Programs/AutomationTool/IOS/IOSPlatform.Automation.cs
--- a/Engine/Source/Programs/AutomationTool/IOS/IOSPlatform.Automation.cs
+++ b/Engine/Source/Programs/AutomationTool/IOS/IOSPlatform.Automation.cs
@@ -181,13 +181,13 @@ public class IOSPlatform : Platform
 
 		// put current Xcode in the trash, and unzip a new one. Xcode in the dock will have to be fixed up tho!
 		Command = "osascript";
-		Params = 
+		Params =
 			" -e \"try\"" +
 			" -e   \"tell application \\\"Finder\\\" to delete POSIX file \\\"/Applications/Xcode.app\\\"\"" +
 			" -e \"end try\"" +
 			" -e \"do shell script \\\"cd /Applications; xip --expand $(CopyOutputPath);\\\"\"" +
 			" -e \"try\"" +
-			" -e   \"do shell script \\\"xcode-select -s /Applications/Xcode.app; xcode-select --install\\\" with administrator privileges\"" +
+			" -e   \"do shell script \\\"xcode-select -s /Applications/Xcode.app; xcode-select --install; xcodebuild -license accept; xcodebuild -runFirstLaunch\\\" with administrator privileges\"" +
 			" -e \"end try\"";
 
 		return true;
