Fix UAT compile error.

diff --git a/Engine/Source/Programs/AutomationTool/AutomationUtils/ProjectUtils.cs b/Engine/Source/Programs/AutomationTool/AutomationUtils/ProjectUtils.cs
--- a/Engine/Source/Programs/AutomationTool/AutomationUtils/ProjectUtils.cs
+++ b/Engine/Source/Programs/AutomationTool/AutomationUtils/ProjectUtils.cs
@@ -401,13 +401,6 @@ namespace AutomationTool
 						continue;
 					}
 
-					// Skip loading Enterprise plugins when project is not an Enterprise project
-					if (Plugin.Type == PluginType.Enterprise && ProjectDescriptor != null && !ProjectDescriptor.IsEnterpriseProject)
-					{
-						Log.TraceLog("Ignoring plugin '{0}' due to not being an enterprise project", Reference.Name);
-						continue;
-					}
-
 					// Add references to all its dependencies
 					if (Plugin.Descriptor.Plugins != null)
 					{
@@ -683,7 +676,7 @@ namespace AutomationTool
 				DirPushed = true;
 			}
 			List<DirectoryReference> ExtraSearchDirectories = (ExtraSearchPaths == null)? null : ExtraSearchPaths.Select(x => new DirectoryReference(x)).ToList();
-			List<FileReference> TargetScripts = RulesCompiler.FindAllRulesSourceFiles(RulesCompiler.RulesFileType.Target, GameFolders: GameFolders, ForeignPlugins: null, AdditionalSearchPaths: ExtraSearchDirectories, bIncludeEnterprise: false);
+			List<FileReference> TargetScripts = RulesCompiler.FindAllRulesSourceFiles(RulesCompiler.RulesFileType.Target, GameFolders: GameFolders, ForeignPlugins: null, AdditionalSearchPaths: ExtraSearchDirectories);
 			if (DirPushed)
 			{
 				CommandUtils.PopDir();
