- Fixed up AutoSDK on Mac in UE5

diff --git a/Engine/Source/Programs/AutomationTool/Turnkey/SdkInfo.cs b/Engine/Source/Programs/AutomationTool/Turnkey/SdkInfo.cs
--- a/Engine/Source/Programs/AutomationTool/Turnkey/SdkInfo.cs
+++ b/Engine/Source/Programs/AutomationTool/Turnkey/SdkInfo.cs
@@ -64,7 +64,7 @@ namespace Turnkey
 						if (string.Compare(Dir.Name, SDK.GetAutoSDKDirectoryForMasterVersion()) == 0)
 						{
 							// make sure it actually has buits in it
-							if (File.Exists(Path.Combine(Dir.FullName, "setup.bat")))
+							if (File.Exists(Path.Combine(Dir.FullName, "setup.bat")) || File.Exists(Path.Combine(Dir.FullName, "setup.sh")))
 							{
 								bValidVersionFound = true;
 								break;
diff --git a/Engine/Source/Programs/DotNETCommon/BuildUtilities/UEBuildPlatformSDK.cs b/Engine/Source/Programs/DotNETCommon/BuildUtilities/UEBuildPlatformSDK.cs
--- a/Engine/Source/Programs/DotNETCommon/BuildUtilities/UEBuildPlatformSDK.cs
+++ b/Engine/Source/Programs/DotNETCommon/BuildUtilities/UEBuildPlatformSDK.cs
@@ -325,7 +325,7 @@ namespace Tools.DotNETCommon
 			{
 				return "Win64";
 			}
-			else if (Environment.OSVersion.Platform == PlatformID.MacOSX)
+			else if (IsMac())
 			{
 				return "Mac";
 			}
