- Removed some debug logging

diff --git a/Engine/Source/Programs/AutomationTool/Turnkey/Commands/VerifySdk.cs b/Engine/Source/Programs/AutomationTool/Turnkey/Commands/VerifySdk.cs
--- a/Engine/Source/Programs/AutomationTool/Turnkey/Commands/VerifySdk.cs
+++ b/Engine/Source/Programs/AutomationTool/Turnkey/Commands/VerifySdk.cs
@@ -115,11 +115,6 @@ namespace Turnkey.Commands
 				PlatformSDK.GetValidVersionRange(out MinAllowedVersion, out MaxAllowedVersion);
 				PlatformSDK.GetValidSoftwareVersionRange(out MinSoftwareAllowedVersion, out MaxSoftwareAllowedVersion);
 
-//				if (ManualSDKVersion == "1")
-				{
-					TurnkeyUtils.Report("How did i get {0} {1}", ManualSDKVersion, AutoSDKVersion);
-				}
-
 				SdkUtils.LocalAvailability ReportedState = LocalState;
 				string StatusString;
 
