- Locally copied remote files are now remembered so won't be re-downloaded within one run of Turnkey (helps if updating multiple devkits with windows share sources)
- Don't attempt to install an flash SDK on a device that can't be connected to

diff --git a/Engine/Source/Programs/AutomationTool/Turnkey/Commands/VerifySdk.cs b/Engine/Source/Programs/AutomationTool/Turnkey/Commands/VerifySdk.cs
--- a/Engine/Source/Programs/AutomationTool/Turnkey/Commands/VerifySdk.cs
+++ b/Engine/Source/Programs/AutomationTool/Turnkey/Commands/VerifySdk.cs
@@ -241,17 +241,24 @@ namespace Turnkey.Commands
 
 							if (bUpdateIfNeeded)
 							{
-								FileSource MatchingInstallableSdk = FileSource.FindMatchingSdk(AutomationPlatform, new FileSource.SourceType[] { FileSource.SourceType.Flash }, bSelectBest: bUnattended, DeviceType: Device.Type);
-
-								if (MatchingInstallableSdk == null)
+								if (Device.bCanConnect)
 								{
-									TurnkeyUtils.Log("ERROR: {0}: Unable top find any Sdks that could be installed on {1}", Platform, Device.Name);
-									TurnkeyUtils.ExitCode = AutomationTool.ExitCode.Error_SDKNotFound;
+									FileSource MatchingInstallableSdk = FileSource.FindMatchingSdk(AutomationPlatform, new FileSource.SourceType[] { FileSource.SourceType.Flash }, bSelectBest: bUnattended, DeviceType: Device.Type);
+
+									if (MatchingInstallableSdk == null)
+									{
+										TurnkeyUtils.Log("ERROR: {0}: Unable top find any Sdks that could be installed on {1}", Platform, Device.Name);
+										TurnkeyUtils.ExitCode = AutomationTool.ExitCode.Error_SDKNotFound;
+									}
+									else
+									{
+										MatchingInstallableSdk.DownloadOrInstall(Platform, Device, bUnattended);
+										TurnkeyUtils.ExitCode = AutomationTool.ExitCode.Success;
+									}
 								}
 								else
 								{
-									MatchingInstallableSdk.DownloadOrInstall(Platform, Device, bUnattended);
-									TurnkeyUtils.ExitCode = AutomationTool.ExitCode.Success;
+									TurnkeyUtils.Log("Skipping device {0} because it cannot connect.", Device.Name);
 								}
 							}
 						}
diff --git a/Engine/Source/Programs/AutomationTool/Turnkey/Providers/NullCopyProvider.cs b/Engine/Source/Programs/AutomationTool/Turnkey/Providers/NullCopyProvider.cs
--- a/Engine/Source/Programs/AutomationTool/Turnkey/Providers/NullCopyProvider.cs
+++ b/Engine/Source/Programs/AutomationTool/Turnkey/Providers/NullCopyProvider.cs
@@ -65,6 +65,15 @@ namespace Turnkey
 
 			if (bIsRemotePath && (bIsLargeFile || bIsDirectory))
 			{
+				// look in out local cache
+				string OperationTag = string.Format("file_op:{0}", Operation);
+				string CachedLocation = LocalCache.GetCachedPathByTag(OperationTag);
+
+				if (CachedLocation != null)
+				{
+					return CachedLocation;
+				}
+
 				string CopyLocation = Path.Combine(LocalCache.CreateTempDirectory(), Path.GetFileName(OutputPath));
 
 				if (bIsDirectory)
@@ -87,6 +96,7 @@ namespace Turnkey
 					}
 				}
 
+				LocalCache.CacheLocationByTag(OperationTag, CopyLocation);
 				TurnkeyUtils.Log("Done!");
 
 				return CopyLocation;
