- Turnkey can now delete User setting by passing it null

diff --git a/Engine/Source/Programs/AutomationTool/Turnkey/TurnkeySettings.cs b/Engine/Source/Programs/AutomationTool/Turnkey/TurnkeySettings.cs
--- a/Engine/Source/Programs/AutomationTool/Turnkey/TurnkeySettings.cs
+++ b/Engine/Source/Programs/AutomationTool/Turnkey/TurnkeySettings.cs
@@ -115,6 +115,10 @@ namespace Turnkey
 			string PreviousValue = null;
 
 			SetUserSettings.TryGetValue(VariableName, out PreviousValue);
+			if (Value == null)
+			{
+				SetUserSettings.Remove(VariableName);
+			}
 			SetUserSettings[VariableName] = Value;
 			Save();
 
