Fix bug in exporting build graph schemas when the output folder doesn't exist.

#rnx

diff --git a/Engine/Source/Programs/AutomationTool/BuildGraph/Schema.cs b/Engine/Source/Programs/AutomationTool/BuildGraph/Schema.cs
--- a/Engine/Source/Programs/AutomationTool/BuildGraph/Schema.cs
+++ b/Engine/Source/Programs/AutomationTool/BuildGraph/Schema.cs
@@ -410,6 +410,7 @@ namespace AutomationTool
 			Settings.IndentChars = "  ";
 			Settings.NewLineChars = "\n";
 
+			DirectoryReference.CreateDirectory(File.Directory);
 			using(XmlWriter Writer = XmlWriter.Create(File.FullName, Settings))
 			{
 				CompiledSchema.Write(Writer);
