- Added Turnkey Studio setting for GoogleDrive app name, and added a log to show user that they may need to check their web browser for Google authentication

diff --git a/Engine/Source/Programs/AutomationTool/Turnkey/Providers/GoogleDriveCopyProvider.cs b/Engine/Source/Programs/AutomationTool/Turnkey/Providers/GoogleDriveCopyProvider.cs
--- a/Engine/Source/Programs/AutomationTool/Turnkey/Providers/GoogleDriveCopyProvider.cs
+++ b/Engine/Source/Programs/AutomationTool/Turnkey/Providers/GoogleDriveCopyProvider.cs
@@ -224,20 +224,24 @@ namespace Turnkey
 		{
 			if (ServiceHelper == null)
 			{
-				if (!TurnkeyUtils.HasVariable("Studio_GoogleDriveCredentials"))
+				string GoogleDriveCredentials = TurnkeyUtils.GetVariableValue("Studio_GoogleDriveCredentials");
+				string GoogleDriveAppName = TurnkeyUtils.GetVariableValue("Studio_GoogleDriveAppName");
+
+				if (string.IsNullOrEmpty(GoogleDriveCredentials) || string.IsNullOrEmpty(GoogleDriveAppName))
 				{
-					TurnkeyUtils.Log("ERROR: Unable to use GoogleDrive without Studio_GoogleDriveCredentials being set first!");
+					TurnkeyUtils.Log("ERROR: Unable to use GoogleDrive without 'Studio_GoogleDriveCredentials' and 'Studio_GoogleDriveAppName' being set first!");
 					return false;
 				}
 
-				string SecretsFile = CopyProvider.ExecuteCopy("$(Studio_GoogleDriveCredentials)");
+				string SecretsFile = CopyProvider.ExecuteCopy(GoogleDriveCredentials);
 				if (SecretsFile == null)
 				{
-					TurnkeyUtils.Log("ERROR: Unable to get GoogleDrive secrets file from {0}", TurnkeyUtils.GetVariableValue("Studio_GoogleDriveCredentials"));
+					TurnkeyUtils.Log("ERROR: Unable to get GoogleDrive secrets/credentials file from {0}", GoogleDriveCredentials);
 					return false;
 				}
 
-				ServiceHelper = new DriveServiceHelper("Quickstart", SecretsFile, Path.GetDirectoryName(SecretsFile));
+				TurnkeyUtils.Log("Connecting to GoogleDrive app '{0}'. If this pauses here, check your web browser for an authentication tab. This is required to be able to connect to Google Drive", GoogleDriveAppName);
+				ServiceHelper = new DriveServiceHelper(GoogleDriveAppName, SecretsFile, Path.GetDirectoryName(SecretsFile));
 
 				PathToFileCache = new Dictionary<string, Google.Apis.Drive.v3.Data.File>();
 			}
diff --git a/Engine/Source/Programs/AutomationTool/Turnkey/Providers/NullCopyProvider.cs b/Engine/Source/Programs/AutomationTool/Turnkey/Providers/NullCopyProvider.cs
--- a/Engine/Source/Programs/AutomationTool/Turnkey/Providers/NullCopyProvider.cs
+++ b/Engine/Source/Programs/AutomationTool/Turnkey/Providers/NullCopyProvider.cs
@@ -146,7 +146,7 @@ namespace Turnkey
 				string Wildcard = (NextSlash == -1) ? PathString.Substring(PrevSlash + 1) : PathString.Substring(PrevSlash + 1, (NextSlash - PrevSlash) - 1);
 
 				// replace * with a non-greedy capture so *-foo-*.zip can get the best values for the *'s
-				Regex Regex = new Regex(string.Format("^{0}$", Wildcard.Replace("*", "(.+?)")), RegexOptions.IgnoreCase);
+				Regex Regex = new Regex(string.Format("^{0}$", Wildcard.Replace("+", "\\+").Replace(".", "\\.").Replace("*", "(.+?)")), RegexOptions.IgnoreCase);
 
 				// track what's before and after the * to return what it expanded to
 				int StarLoc = Wildcard.IndexOf('*');
diff --git a/Engine/Source/Programs/AutomationTool/Turnkey/TurnkeySettings.cs b/Engine/Source/Programs/AutomationTool/Turnkey/TurnkeySettings.cs
--- a/Engine/Source/Programs/AutomationTool/Turnkey/TurnkeySettings.cs
+++ b/Engine/Source/Programs/AutomationTool/Turnkey/TurnkeySettings.cs
@@ -47,6 +47,7 @@ namespace Turnkey
 		{
 			new UserSetting("Studio_AppleSigningCertPassword", null, "A shared password that is used across Apple Signing Certificates"),
 			new UserSetting("Studio_GoogleDriveCredentials", null, "The location of credentials file needed for GoogleDrive integration. The contents should start with: '{\"installed\":{\"client_id\":'"),
+			new UserSetting("Studio_GoogleDriveAppName", null, "The name of the application your studio needs to create (see documentation for help)"),
 		};
 
 		// basically same as Turnkey variables, but this only contains ones that were loaded so we can write them back out
