Attempt to silence warning about COM references when building c# projects.

diff --git a/Engine/Source/Programs/AutomationTool/Gauntlet/Gauntlet.Automation.csproj b/Engine/Source/Programs/AutomationTool/Gauntlet/Gauntlet.Automation.csproj
--- a/Engine/Source/Programs/AutomationTool/Gauntlet/Gauntlet.Automation.csproj
+++ b/Engine/Source/Programs/AutomationTool/Gauntlet/Gauntlet.Automation.csproj
@@ -13,6 +13,7 @@
     <FileAlignment>512</FileAlignment>
     <WarningsNotAsErrors>612,618</WarningsNotAsErrors>
     <TargetFrameworkProfile />
+    <ResolveComReferenceSilent>True</ResolveComReferenceSilent>
   </PropertyGroup>
   <PropertyGroup Condition=" '$(Configuration)|$(Platform)' == 'Debug|AnyCPU' ">
     <DebugSymbols>true</DebugSymbols>
