Adding the ability to allow verbose P4 output for SubmitTask.

#rnx

diff --git a/Engine/Source/Programs/AutomationTool/BuildGraph/Tasks/SubmitTask.cs b/Engine/Source/Programs/AutomationTool/BuildGraph/Tasks/SubmitTask.cs
--- a/Engine/Source/Programs/AutomationTool/BuildGraph/Tasks/SubmitTask.cs
+++ b/Engine/Source/Programs/AutomationTool/BuildGraph/Tasks/SubmitTask.cs
@@ -63,6 +63,12 @@ namespace AutomationTool.Tasks
 		/// </summary>
 		[TaskParameter(Optional = true)]
 		public bool Force;
+
+		/// <summary>
+		/// Allow verbose P4 output (spew).
+		/// </summary>
+		[TaskParameter(Optional = false)]
+		public bool P4Verbose;
 	}
 
 	/// <summary>
@@ -117,7 +123,7 @@ namespace AutomationTool.Tasks
 					Client.Name = Parameters.Workspace;
 					Client.Options = P4ClientOption.NoAllWrite | P4ClientOption.Clobber | P4ClientOption.NoCompress | P4ClientOption.Unlocked | P4ClientOption.NoModTime | P4ClientOption.RmDir;
 					Client.LineEnd = P4LineEnd.Local;
-					CommandUtils.P4.CreateClient(Client, AllowSpew: false);
+					CommandUtils.P4.CreateClient(Client, AllowSpew: Parameters.P4Verbose);
 
 					// Create a new connection for it
 					SubmitP4 = new P4Connection(Client.Owner, Client.Name);
@@ -127,13 +133,13 @@ namespace AutomationTool.Tasks
 				int NewCL = SubmitP4.CreateChange(Description: Parameters.Description.Replace("\\n", "\n"));
 				foreach(FileReference File in Files)
 				{
-					SubmitP4.Revert(String.Format("-k \"{0}\"", File.FullName));
-					SubmitP4.Sync(String.Format("-k \"{0}\"", File.FullName), AllowSpew: false);
+					SubmitP4.Revert(String.Format("-k \"{0}\"", File.FullName), AllowSpew: Parameters.P4Verbose);
+					SubmitP4.Sync(String.Format("-k \"{0}\"", File.FullName), AllowSpew: Parameters.P4Verbose);
 					SubmitP4.Add(NewCL, String.Format("\"{0}\"", File.FullName));
-					SubmitP4.Edit(NewCL, String.Format("\"{0}\"", File.FullName));
+					SubmitP4.Edit(NewCL, String.Format("\"{0}\"", File.FullName), AllowSpew: Parameters.P4Verbose);
 					if (Parameters.FileType != null)
 					{
-						SubmitP4.P4(String.Format("reopen -t \"{0}\" \"{1}\"", Parameters.FileType, File.FullName), AllowSpew: false);
+						SubmitP4.P4(String.Format("reopen -t \"{0}\" \"{1}\"", Parameters.FileType, File.FullName), AllowSpew: Parameters.P4Verbose);
 					}
 				}
 
