Fix P4Verbose being required on SubmitTask.

#rnx

diff --git a/Engine/Source/Programs/AutomationTool/BuildGraph/Tasks/SubmitTask.cs b/Engine/Source/Programs/AutomationTool/BuildGraph/Tasks/SubmitTask.cs
--- a/Engine/Source/Programs/AutomationTool/BuildGraph/Tasks/SubmitTask.cs
+++ b/Engine/Source/Programs/AutomationTool/BuildGraph/Tasks/SubmitTask.cs
@@ -67,7 +67,7 @@ namespace AutomationTool.Tasks
 		/// <summary>
 		/// Allow verbose P4 output (spew).
 		/// </summary>
-		[TaskParameter(Optional = false)]
+		[TaskParameter(Optional = true)]
 		public bool P4Verbose;
 	}
 
