Correct calculation of OBB per file cost

#ue4
#android

diff --git a/Engine/Source/Programs/AutomationTool/Android/AndroidPlatform.Automation.cs b/Engine/Source/Programs/AutomationTool/Android/AndroidPlatform.Automation.cs
--- a/Engine/Source/Programs/AutomationTool/Android/AndroidPlatform.Automation.cs
+++ b/Engine/Source/Programs/AutomationTool/Android/AndroidPlatform.Automation.cs
@@ -548,20 +548,21 @@ public class AndroidPlatform : Platform
 				FilesToObb = new List<FileReference>();
 
 				// Collect the filesize and place into Obb or Patch list
-				Int64 PaddingSize = 16384;
-				Int64 MainObbSize = 22 + 10 + PaddingSize;			// EOCD wit comment (store version) + padding
-				Int64 PatchObbSize = 22 + 10 + PaddingSize;		// EOCD wit comment (store version) + padding
-				Int64 Overflow1ObbSize = 22 + 10 + PaddingSize;	// EOCD wit comment (store version) + padding
-				Int64 Overflow2ObbSize = 22 + 10 + PaddingSize;	// EOCD wit comment (store version) + padding
+				Int64 MinimumObbSize = 22 + 10;		// EOCD with comment (store version)
+				Int64 MainObbSize = MinimumObbSize;
+				Int64 PatchObbSize = MinimumObbSize;
+				Int64 Overflow1ObbSize = MinimumObbSize;
+				Int64 Overflow2ObbSize = MinimumObbSize;
 				foreach (FileReference FileRef in FilesForObb)
 				{
 					FileInfo LocalFileInfo = new FileInfo(FileRef.FullName);
 					Int64 LocalFileLength = LocalFileInfo.Length;
 
 					string DestinationPath = Path.GetDirectoryName(FileRef.FullName).Replace(StageDirectoryPath, SC.ShortProjectName);
-					int FilenameLength = DestinationPath.Length;
-					Int64 Overhead = (30 + 36 + FilenameLength) + (46 + 36 + FilenameLength); // local file descriptor + central directory cost per file
-					Int64 FileRequirements = LocalFileLength + Overhead;
+					Int64 FilenameLength = DestinationPath.Length + Path.GetFileName(FileRef.FullName).Length + 1;
+					Int64 LocalOverhead = (30 + FilenameLength + 36);		// local file descriptor
+					Int64 GlobalOverhead = (46 + FilenameLength + 36);		// central directory cost
+					Int64 FileRequirements = LocalFileLength + LocalOverhead + GlobalOverhead;
 
 					if (MainObbSize + FileRequirements < OBBSizeAllowed)
 					{
