UAT - Minor fixes (spelling, removing extra semicolons)

diff --git a/Engine/Source/Programs/AutomationTool/AutomationUtils/ScriptCompiler.cs b/Engine/Source/Programs/AutomationTool/AutomationUtils/ScriptCompiler.cs
--- a/Engine/Source/Programs/AutomationTool/AutomationUtils/ScriptCompiler.cs
+++ b/Engine/Source/Programs/AutomationTool/AutomationUtils/ScriptCompiler.cs
@@ -411,7 +411,7 @@ namespace AutomationTool
 			else
 			{
 				NewHashes.SaveToFile(DependencyFile);
-				Log.TraceVerbose("Wrote depencencies to {0}", DependencyFile);
+				Log.TraceVerbose("Wrote dependencies to {0}", DependencyFile);
 			}
 		}
 
