UAT Net Core - Copied project files into net core compatible format

diff --git a/Engine/Source/Programs/AutomationTool/Android/Android.Automation.Core.csproj b/Engine/Source/Programs/AutomationTool/Android/Android.Automation.Core.csproj
new file mode 100644
--- /dev/null
+++ b/Engine/Source/Programs/AutomationTool/Android/Android.Automation.Core.csproj
@@ -0,0 +1,66 @@
+﻿<?xml version="1.0" encoding="utf-8"?>
+<Project>
+  <PropertyGroup>
+    <!-- Set a custom temp directory so we can live side by side with the .net framework output -->
+    <BaseIntermediateOutputPath>$(MSBuildThisFileDirectory)obj-core\$(MSBuildProjectName)\</BaseIntermediateOutputPath>
+    <BaseOutputPath>$(MSBuildThisFileDirectory)bin-core\$(MSBuildProjectName)\</BaseOutputPath>
+  </PropertyGroup>
+
+  <PropertyGroup>
+    <TargetFramework>netcoreapp3.1</TargetFramework>
+    <Configuration Condition=" '$(Configuration)' == '' ">Development</Configuration>
+    <OutputType>Library</OutputType>
+    <AppendTargetFrameworkToOutputPath>false</AppendTargetFrameworkToOutputPath>
+    <GenerateAssemblyInfo>false</GenerateAssemblyInfo>
+    <GenerateTargetFrameworkAttribute>false</GenerateTargetFrameworkAttribute>
+    <Configurations>Debug;Release;Development</Configurations>
+    <RootNamespace>Android.Automation</RootNamespace>
+    <AssemblyName>Android.Automation</AssemblyName>
+    <WarningsNotAsErrors>612,618</WarningsNotAsErrors>
+    <OutputPath>..\..\..\..\Binaries\DotNET\AutomationTool\AutomationScripts\</OutputPath>
+    <DefineConstants>NET_CORE</DefineConstants>
+    <TreatWarningsAsErrors>true</TreatWarningsAsErrors>
+  </PropertyGroup>
+  <PropertyGroup Condition=" '$(Configuration)|$(Platform)' == 'Development|AnyCPU' ">
+    <DefineConstants>NET_CORE;TRACE</DefineConstants>
+    <Optimize>true</Optimize>
+  </PropertyGroup>
+  <PropertyGroup Condition=" '$(Configuration)|$(Platform)' == 'Release|AnyCPU' ">
+    <DefineConstants>NET_CORE;TRACE</DefineConstants>
+    <Optimize>true</Optimize>
+  </PropertyGroup>
+  <PropertyGroup Condition=" '$(Configuration)|$(Platform)' == 'Debug|AnyCPU' ">
+    <DefineConstants>NET_CORE;DEBUG;TRACE</DefineConstants>
+  </PropertyGroup>
+  <ItemGroup>
+    <Reference Include="Ionic.Zip.Reduced">
+      <HintPath>..\..\..\..\Binaries\DotNET\Ionic.Zip.Reduced.dll</HintPath>
+      <Private>False</Private>
+    </Reference>
+  </ItemGroup>
+  <Import Project="Sdk.props" Sdk="Microsoft.NET.Sdk" />
+  <Import Project="Sdk.targets" Sdk="Microsoft.NET.Sdk" />
+
+  <ItemGroup>
+    <Compile Remove="obj\**" />
+    <EmbeddedResource Remove="obj\**" />
+    <None Remove="obj\**" />
+  </ItemGroup>
+
+  <ItemGroup>
+    <Compile Include="..\..\DotNETCommon\MetaData.cs">
+      <Link>Properties\MetaData.cs</Link>
+    </Compile>
+  </ItemGroup>
+  <ItemGroup>
+    <ProjectReference Include="..\..\DotNETCommon\DotNETUtilities\DotNETUtilitiesCore.csproj" />
+    <ProjectReference Include="..\..\UnrealBuildTool\UnrealBuildToolCore.csproj" />
+    <ProjectReference Include="..\AutomationUtils\AutomationUtils.Automation.Core.csproj" />
+    <ProjectReference Include="..\Scripts\AutomationScripts.Automation.Core.csproj" />
+  </ItemGroup>
+  <ItemGroup>
+    <PackageReference Include="Microsoft.CSharp" Version="4.7.0" />
+    <PackageReference Include="System.Data.DataSetExtensions" Version="4.5.0" />
+    <PackageReference Include="System.Net.Http" Version="4.3.4" />
+  </ItemGroup>
+</Project>
\ No newline at end of file
diff --git a/Engine/Source/Programs/AutomationTool/AutomationToolCore.csproj b/Engine/Source/Programs/AutomationTool/AutomationToolCore.csproj
new file mode 100644
--- /dev/null
+++ b/Engine/Source/Programs/AutomationTool/AutomationToolCore.csproj
@@ -0,0 +1,68 @@
+﻿<Project>
+  <PropertyGroup>
+    <!-- Set a custom temp directory so we can live side by side with the .net framework output -->
+    <BaseIntermediateOutputPath>$(MSBuildThisFileDirectory)obj-core\$(MSBuildProjectName)\</BaseIntermediateOutputPath>
+    <BaseOutputPath>$(MSBuildThisFileDirectory)bin-core\$(MSBuildProjectName)\</BaseOutputPath>
+  </PropertyGroup>
+  <Import Project="Sdk.props" Sdk="Microsoft.NET.Sdk" />
+  <PropertyGroup>
+    <TargetFramework>netcoreapp3.1</TargetFramework>
+    <Configuration Condition=" '$(Configuration)' == '' ">Development</Configuration>
+    <OutputType>Exe</OutputType>
+    <AppendTargetFrameworkToOutputPath>false</AppendTargetFrameworkToOutputPath>
+    <GenerateAssemblyInfo>false</GenerateAssemblyInfo>
+    <GenerateTargetFrameworkAttribute>false</GenerateTargetFrameworkAttribute>
+    <Configurations>Debug;Release;Development</Configurations>
+    <ApplicationManifest>app.manifest</ApplicationManifest>
+    <RootNamespace>AutomationTool</RootNamespace>
+    <AssemblyName>AutomationTool</AssemblyName>
+    <WarningsNotAsErrors>612,618</WarningsNotAsErrors>
+    <OutputPath>..\..\..\Binaries\DotNET\AutomationTool</OutputPath>
+    <EnableDefaultItems>false</EnableDefaultItems>
+    <DefineConstants>NET_CORE</DefineConstants>
+  </PropertyGroup>
+  <PropertyGroup Condition=" '$(Configuration)|$(Platform)' == 'Development|AnyCPU' ">
+    <DefineConstants>NET_CORE;TRACE</DefineConstants>
+    <Optimize>true</Optimize>
+  </PropertyGroup>
+  <PropertyGroup Condition=" '$(Configuration)|$(Platform)' == 'Release|AnyCPU' ">
+    <DefineConstants>NET_CORE;TRACE</DefineConstants>
+    <Optimize>true</Optimize>
+  </PropertyGroup>
+  <PropertyGroup Condition=" '$(Configuration)|$(Platform)' == 'Debug|AnyCPU' ">
+    <DefineConstants>NET_CORE;DEBUG;TRACE</DefineConstants>
+  </PropertyGroup>
+  <ItemGroup>
+    <Reference Include="Ionic.Zip.Reduced">
+      <HintPath>..\..\..\Binaries\DotNET\Ionic.Zip.Reduced.dll</HintPath>
+      <Private>true</Private>
+    </Reference>
+  </ItemGroup>
+  <Import Project="Sdk.targets" Sdk="Microsoft.NET.Sdk" />
+
+  <ItemGroup>
+    <Compile Remove="obj\**" />
+    <EmbeddedResource Remove="obj\**" />
+    <None Remove="obj\**" />
+  </ItemGroup>
+
+  <ItemGroup>
+    <Compile Include="..\DotNETCommon\MetaData.cs">
+      <Link>Properties\MetaData.cs</Link>
+    </Compile>
+    <Compile Include="Program.cs" />
+    <Compile Include="Properties\AssemblyInfo.cs" />
+    <None Include="app.manifest">
+      <SubType>Designer</SubType>
+    </None>
+  </ItemGroup>
+  
+  <!--     <None Include="app.config" /> -->
+  <ItemGroup>
+    <ProjectReference Include="..\DotNETCommon\DotNETUtilities\DotNETUtilitiesCore.csproj" />
+    <ProjectReference Include="..\UnrealBuildTool\UnrealBuildToolCore.csproj" />
+    <ProjectReference Include="AutomationUtils\AutomationUtils.Automation.Core.csproj" />
+  </ItemGroup>
+
+  <Import Project="AutomationToolCore.csproj.References" Condition="Exists('AutomationToolCore.csproj.References') And '$(AutomationToolProjectOnly)' == ''" />
+</Project>
\ No newline at end of file
diff --git a/Engine/Source/Programs/AutomationTool/AutomationUtils/AutomationUtils.Automation.Core.csproj b/Engine/Source/Programs/AutomationTool/AutomationUtils/AutomationUtils.Automation.Core.csproj
new file mode 100644
--- /dev/null
+++ b/Engine/Source/Programs/AutomationTool/AutomationUtils/AutomationUtils.Automation.Core.csproj
@@ -0,0 +1,68 @@
+﻿<Project>
+  <PropertyGroup>
+    <!-- Set a custom temp directory so we can live side by side with the .net framework output -->
+    <BaseIntermediateOutputPath>$(MSBuildThisFileDirectory)obj-core\$(MSBuildProjectName)\</BaseIntermediateOutputPath>
+    <BaseOutputPath>$(MSBuildThisFileDirectory)bin-core\$(MSBuildProjectName)\</BaseOutputPath>
+  </PropertyGroup>
+  <Import Project="Sdk.props" Sdk="Microsoft.NET.Sdk" />
+  <PropertyGroup>
+    <!--<TargetFrameworks>net472;netstandard2.0</TargetFrameworks>-->
+    <TargetFramework>netcoreapp3.1</TargetFramework>
+    <Configuration Condition=" '$(Configuration)' == '' ">Development</Configuration>
+    <OutputType>Library</OutputType>
+    <RootNamespace>AutomationUtils.Automation</RootNamespace>
+    <AssemblyName>AutomationUtils.Automation</AssemblyName>
+    <GenerateAssemblyInfo>false</GenerateAssemblyInfo>
+    <GenerateTargetFrameworkAttribute>false</GenerateTargetFrameworkAttribute>
+    <FileUpgradeFlags>
+    </FileUpgradeFlags>
+    <UpgradeBackupLocation>
+    </UpgradeBackupLocation>
+    <OldToolsVersion>2.0</OldToolsVersion>
+    <OutputPath>..\..\..\..\Binaries\DotNET\AutomationTool</OutputPath>
+    <Configurations>Debug;Release;Development</Configurations>
+    <TreatWarningsAsErrors>true</TreatWarningsAsErrors>
+  </PropertyGroup>
+    <ItemGroup>
+    <Reference Include="Ionic.Zip.Reduced">
+      <HintPath>..\..\..\..\Binaries\DotNET\Ionic.Zip.Reduced.dll</HintPath>
+      <Private>False</Private>
+    </Reference>
+  </ItemGroup>
+  <Import Project="Sdk.targets" Sdk="Microsoft.NET.Sdk" />
+  <!-- We do not output these binaries to Binaries as it targets multiple frameworks. They would have overloaded names for each framework, let the application copy it to its output directory instead -->
+  <PropertyGroup Condition=" '$(Configuration)|$(Platform)' == 'Debug|AnyCPU' ">
+    <DefineConstants>TRACE;NET_CORE</DefineConstants>
+  </PropertyGroup>
+  <PropertyGroup Condition="'$(Configuration)|$(Platform)' == 'Development|AnyCPU'">
+    <Optimize>true</Optimize>
+    <DefineConstants>TRACE;NET_CORE</DefineConstants>
+  </PropertyGroup>
+  <PropertyGroup Condition="'$(Configuration)|$(Platform)' == 'Release|AnyCPU'">
+    <Optimize>true</Optimize>
+    <DefineConstants>TRACE;NET_CORE</DefineConstants>
+  </PropertyGroup>
+  <ItemGroup>
+    <Compile Remove="obj\**" />
+    <EmbeddedResource Remove="obj\**" />
+    <None Remove="obj\**" />
+  </ItemGroup>
+
+  <ItemGroup>
+    <ProjectReference Include="..\..\DotNETCommon\DotNETUtilities\DotNETUtilitiesCore.csproj" />
+    <ProjectReference Include="..\..\UnrealBuildTool\UnrealBuildToolCore.csproj" />
+  </ItemGroup>
+  <ItemGroup>
+    <PackageReference Include="Microsoft.CSharp" Version="4.7.0" />
+    <PackageReference Include="System.Data.DataSetExtensions" Version="4.5.0" />
+    <PackageReference Include="System.Net.Http" Version="4.3.4" />
+    <PackageReference Include="System.Drawing.Common" Version="4.7.0" />
+
+    <!--
+    <PackageReference Include="Google.Apis" Version="1.33.1" />
+    <PackageReference Include="Google.Apis.Auth" Version="1.33.1" />
+    <PackageReference Include="Google.Apis.Core" Version="1.33.1" />-->
+    <PackageReference Include="Google.Apis.Drive.v3" Version="1.33.1.1225" />
+    <PackageReference Include="Google.Apis.Sheets.v4" Version="1.33.1.1229" />
+  </ItemGroup>
+</Project>
\ No newline at end of file
diff --git a/Engine/Source/Programs/AutomationTool/BuildGraph/BuildGraph.Automation.Core.csproj b/Engine/Source/Programs/AutomationTool/BuildGraph/BuildGraph.Automation.Core.csproj
new file mode 100644
--- /dev/null
+++ b/Engine/Source/Programs/AutomationTool/BuildGraph/BuildGraph.Automation.Core.csproj
@@ -0,0 +1,65 @@
+﻿<Project>
+  <PropertyGroup>
+    <!-- Set a custom temp directory so we can live side by side with the .net framework output -->
+    <BaseIntermediateOutputPath>$(MSBuildThisFileDirectory)obj-core\$(MSBuildProjectName)\</BaseIntermediateOutputPath>
+    <BaseOutputPath>$(MSBuildThisFileDirectory)bin-core\$(MSBuildProjectName)\</BaseOutputPath>
+  </PropertyGroup>
+
+  <PropertyGroup>
+    <TargetFramework>netcoreapp3.1</TargetFramework>
+    <Configuration Condition=" '$(Configuration)' == '' ">Development</Configuration>
+    <OutputType>Library</OutputType>
+    <AppendTargetFrameworkToOutputPath>false</AppendTargetFrameworkToOutputPath>
+    <GenerateAssemblyInfo>false</GenerateAssemblyInfo>
+    <GenerateTargetFrameworkAttribute>false</GenerateTargetFrameworkAttribute>
+    <Configurations>Debug;Release;Development</Configurations>
+    <RootNamespace>AutomationTool</RootNamespace>
+    <AssemblyName>BuildGraph.Automation</AssemblyName>
+    <WarningsNotAsErrors>612,618</WarningsNotAsErrors>
+    <OutputPath>..\..\..\..\Binaries\DotNET\AutomationTool</OutputPath>
+    <DefineConstants>NET_CORE</DefineConstants>
+    <TreatWarningsAsErrors>true</TreatWarningsAsErrors>
+  </PropertyGroup>
+  <PropertyGroup Condition=" '$(Configuration)|$(Platform)' == 'Development|AnyCPU' ">
+    <DefineConstants>NET_CORE;TRACE</DefineConstants>
+    <Optimize>true</Optimize>
+  </PropertyGroup>
+  <PropertyGroup Condition=" '$(Configuration)|$(Platform)' == 'Release|AnyCPU' ">
+    <DefineConstants>NET_CORE;TRACE</DefineConstants>
+    <Optimize>true</Optimize>
+  </PropertyGroup>
+  <PropertyGroup Condition=" '$(Configuration)|$(Platform)' == 'Debug|AnyCPU' ">
+    <DefineConstants>NET_CORE;DEBUG;TRACE</DefineConstants>
+  </PropertyGroup>
+  <ItemGroup>
+    <Reference Include="Ionic.Zip.Reduced">
+      <HintPath>..\..\..\..\Binaries\DotNET\Ionic.Zip.Reduced.dll</HintPath>
+      <Private>False</Private>
+    </Reference>
+  </ItemGroup>
+  <Import Project="Sdk.props" Sdk="Microsoft.NET.Sdk" />
+  <Import Project="Sdk.targets" Sdk="Microsoft.NET.Sdk" />
+
+  <ItemGroup>
+    <Compile Remove="obj\**" />
+    <EmbeddedResource Remove="obj\**" />
+    <None Remove="obj\**" />
+  </ItemGroup>
+
+  <ItemGroup>
+    <Compile Include="..\..\DotNETCommon\MetaData.cs">
+      <Link>Properties\MetaData.cs</Link>
+    </Compile>
+  </ItemGroup>
+  <ItemGroup>
+    <ProjectReference Include="..\..\DotNETCommon\DotNETUtilities\DotNETUtilitiesCore.csproj" />
+    <ProjectReference Include="..\..\UnrealBuildTool\UnrealBuildToolCore.csproj" />
+    <ProjectReference Include="..\AutomationUtils\AutomationUtils.Automation.Core.csproj" />
+  </ItemGroup>
+  <ItemGroup>
+    <PackageReference Include="Microsoft.CSharp" Version="4.7.0" />
+    <PackageReference Include="System.Data.DataSetExtensions" Version="4.5.0" />
+    <PackageReference Include="System.Net.Http" Version="4.3.4" />
+    <PackageReference Include="System.Drawing.Common" Version="4.7.0" />
+  </ItemGroup>
+</Project>
\ No newline at end of file
diff --git a/Engine/Source/Programs/AutomationTool/Gauntlet/Gauntlet.Automation.Core.csproj b/Engine/Source/Programs/AutomationTool/Gauntlet/Gauntlet.Automation.Core.csproj
new file mode 100644
--- /dev/null
+++ b/Engine/Source/Programs/AutomationTool/Gauntlet/Gauntlet.Automation.Core.csproj
@@ -0,0 +1,82 @@
+﻿<Project>
+  <PropertyGroup>
+    <!-- Set a custom temp directory so we can live side by side with the .net framework output -->
+    <BaseIntermediateOutputPath>$(MSBuildThisFileDirectory)obj-core\$(MSBuildProjectName)\</BaseIntermediateOutputPath>
+    <BaseOutputPath>$(MSBuildThisFileDirectory)bin-core\$(MSBuildProjectName)\</BaseOutputPath>
+  </PropertyGroup>
+
+  <PropertyGroup>
+    <TargetFramework>netcoreapp3.1</TargetFramework>
+    <Configuration Condition=" '$(Configuration)' == '' ">Development</Configuration>
+    <OutputType>Library</OutputType>
+    <AppendTargetFrameworkToOutputPath>false</AppendTargetFrameworkToOutputPath>
+    <GenerateAssemblyInfo>false</GenerateAssemblyInfo>
+    <GenerateTargetFrameworkAttribute>false</GenerateTargetFrameworkAttribute>
+    <Configurations>Debug;Release;Development</Configurations>
+    <RootNamespace>Gauntlet.Automation</RootNamespace>
+    <AssemblyName>Gauntlet.Automation</AssemblyName>
+    <WarningsNotAsErrors>612,618</WarningsNotAsErrors>
+    <OutputPath>..\..\..\..\Binaries\DotNET\AutomationTool</OutputPath>
+    <DefineConstants>NET_CORE</DefineConstants>
+    <TreatWarningsAsErrors>true</TreatWarningsAsErrors>
+  </PropertyGroup>
+  <PropertyGroup Condition=" '$(Configuration)|$(Platform)' == 'Development|AnyCPU' ">
+    <DefineConstants>NET_CORE;TRACE</DefineConstants>
+    <Optimize>true</Optimize>
+  </PropertyGroup>
+  <PropertyGroup Condition=" '$(Configuration)|$(Platform)' == 'Release|AnyCPU' ">
+    <DefineConstants>NET_CORE;TRACE</DefineConstants>
+    <Optimize>true</Optimize>
+  </PropertyGroup>
+  <PropertyGroup Condition=" '$(Configuration)|$(Platform)' == 'Debug|AnyCPU' ">
+    <DefineConstants>NET_CORE;DEBUG;TRACE</DefineConstants>
+  </PropertyGroup>
+  <ItemGroup>
+    <Reference Include="Ionic.Zip.Reduced">
+      <HintPath>..\..\..\..\Binaries\DotNET\Ionic.Zip.Reduced.dll</HintPath>
+      <Private>False</Private>
+    </Reference>
+  </ItemGroup>
+  <Import Project="Sdk.props" Sdk="Microsoft.NET.Sdk" />
+  <Import Project="Sdk.targets" Sdk="Microsoft.NET.Sdk" />
+
+  <ItemGroup>
+    <Compile Remove="obj\**" />
+    <EmbeddedResource Remove="obj\**" />
+    <None Remove="obj\**" />
+  </ItemGroup>
+
+    <!-- The following will glob all platform extensions' Gauntlet source
+  files and include them in the build. They will also appear in VisualStudio -->
+  <ItemGroup>
+    <Compile Include="../../../../Platforms/*/Source/Programs/AutomationTool/Gauntlet/**/*.cs">
+      <!-- RecursiveDir metadata expands to [PlatformName]/Source/Programs/[etc]. The Replace() will
+      replace everything from the first slash to the end leaving just the platform name. Should the
+      Replace() fail then RecursiveDir is still an agreeable value -->
+      <Link>Platform/$([System.Text.RegularExpressions.Regex]::Replace(%(Compile.RecursiveDir), [\\/].+$, ``))/_</Link>
+    </Compile>
+  </ItemGroup>
+  
+  <!-- TODO: Disabled to avoid com references which are unsupported in net core 
+  <Import Project="../../../../Platforms/*/Source/Programs/AutomationTool/Gauntlet/*.Gauntlet.targets" Condition="'$(OS)' == 'Windows_NT'" />
+  -->
+
+  <ItemGroup>
+    <ProjectReference Include="..\..\DotNETCommon\DotNETUtilities\DotNETUtilitiesCore.csproj" />
+    <ProjectReference Include="..\..\UnrealBuildTool\UnrealBuildToolCore.csproj" />
+    <ProjectReference Include="..\AutomationUtils\AutomationUtils.Automation.Core.csproj" />
+    <ProjectReference Include="..\Scripts\AutomationScripts.Automation.Core.csproj" />
+    <ProjectReference Include="..\OneSkyLocalization\OneSkyLocalization.Automation.Core.csproj" />
+    <ProjectReference Include="..\XLocLocalization\XLocLocalization.Automation.Core.csproj" />
+    <ProjectReference Include="..\Localization\Localization.Automation.Core.csproj" />
+    <ProjectReference Include="..\Lumin\Lumin.Automation.Core.csproj" />
+    <ProjectReference Include="..\Android\Android.Automation.Core.csproj" />
+  </ItemGroup>
+  <ItemGroup>
+    <PackageReference Include="Microsoft.CSharp" Version="4.7.0" />
+    <PackageReference Include="System.Data.DataSetExtensions" Version="4.5.0" />
+    <PackageReference Include="System.Net.Http" Version="4.3.4" />
+
+    <PackageReference Include="Magick.NET-Q16-HDRI-AnyCPU" Version="7.22.0" />
+  </ItemGroup>
+</Project>
\ No newline at end of file
diff --git a/Engine/Source/Programs/AutomationTool/HoloLens/HoloLens.Automation.Core.csproj b/Engine/Source/Programs/AutomationTool/HoloLens/HoloLens.Automation.Core.csproj
new file mode 100644
--- /dev/null
+++ b/Engine/Source/Programs/AutomationTool/HoloLens/HoloLens.Automation.Core.csproj
@@ -0,0 +1,103 @@
+﻿<Project>
+  <PropertyGroup>
+    <!-- Set a custom temp directory so we can live side by side with the .net framework output -->
+    <BaseIntermediateOutputPath>$(MSBuildThisFileDirectory)obj-core\$(MSBuildProjectName)\</BaseIntermediateOutputPath>
+    <BaseOutputPath>$(MSBuildThisFileDirectory)bin-core\$(MSBuildProjectName)\</BaseOutputPath>
+  </PropertyGroup>
+
+  <PropertyGroup>
+    <TargetFramework>netcoreapp3.1</TargetFramework>
+    <Configuration Condition=" '$(Configuration)' == '' ">Development</Configuration>
+    <OutputType>Library</OutputType>
+    <AppendTargetFrameworkToOutputPath>false</AppendTargetFrameworkToOutputPath>
+    <GenerateAssemblyInfo>false</GenerateAssemblyInfo>
+    <GenerateTargetFrameworkAttribute>false</GenerateTargetFrameworkAttribute>
+    <Configurations>Debug;Release;Development</Configurations>
+    <RootNamespace>HoloLens.Automation</RootNamespace>
+    <AssemblyName>HoloLens.Automation</AssemblyName>
+    <WarningsNotAsErrors>612,618</WarningsNotAsErrors>
+    <OutputPath>..\..\..\..\Binaries\DotNET\AutomationTool</OutputPath>
+    <DefineConstants>NET_CORE</DefineConstants>
+    <TreatWarningsAsErrors>true</TreatWarningsAsErrors>
+    <IsWindows Condition="'$([System.Runtime.InteropServices.RuntimeInformation]::IsOSPlatform($([System.Runtime.InteropServices.OSPlatform]::Windows)))' == 'true'">true</IsWindows> 
+    <IsOSX Condition="'$([System.Runtime.InteropServices.RuntimeInformation]::IsOSPlatform($([System.Runtime.InteropServices.OSPlatform]::OSX)))' == 'true'">true</IsOSX> 
+    <IsLinux Condition="'$([System.Runtime.InteropServices.RuntimeInformation]::IsOSPlatform($([System.Runtime.InteropServices.OSPlatform]::Linux)))' == 'true'">true</IsLinux> 
+
+  </PropertyGroup>
+  <PropertyGroup Condition=" '$(Configuration)|$(Platform)' == 'Development|AnyCPU' ">
+    <DefineConstants>NET_CORE;TRACE</DefineConstants>
+    <Optimize>true</Optimize>
+  </PropertyGroup>
+  <PropertyGroup Condition=" '$(Configuration)|$(Platform)' == 'Release|AnyCPU' ">
+    <DefineConstants>NET_CORE;TRACE</DefineConstants>
+    <Optimize>true</Optimize>
+  </PropertyGroup>
+  <PropertyGroup Condition=" '$(Configuration)|$(Platform)' == 'Debug|AnyCPU' ">
+    <DefineConstants>NET_CORE;DEBUG;TRACE</DefineConstants>
+  </PropertyGroup>
+  <ItemGroup>
+    <Reference Include="WindowsDevicePortalWrapper">
+      <HintPath>..\..\..\..\..\Engine\Binaries\ThirdParty\WindowsDevicePortalWrapper\WindowsDevicePortalWrapper.dll</HintPath>
+    </Reference>
+  </ItemGroup>
+  <Choose>
+    <When Condition="'$(IsWindows)'!='true'">
+      <PropertyGroup>
+        <DefineConstants>$(DefineConstants);__MonoCS__</DefineConstants>
+      </PropertyGroup>
+    </When>
+  </Choose>
+  <Choose>
+    <When Condition="'$(IsWindows)'=='true'">
+      <PropertyGroup>
+        <UseWindowsForms>true</UseWindowsForms>
+        <Windows10SdkPath>$(Registry:HKEY_LOCAL_MACHINE\SOFTWARE\Wow6432Node\Microsoft\Microsoft SDKs\Windows\v10.0@InstallationFolder)</Windows10SdkPath>
+        <Windows10SdkLatestVersion>$(Registry:HKEY_LOCAL_MACHINE\SOFTWARE\Wow6432Node\Microsoft\Microsoft SDKs\Windows\v10.0@ProductVersion)</Windows10SdkLatestVersion>
+      </PropertyGroup>
+      <ItemGroup>
+        <Reference Include="System.Runtime.WindowsRuntime">
+          <HintPath>$(SystemDrive)\Program Files (x86)\Reference Assemblies\Microsoft\Framework\.NETCore\$(TargetFrameworkVersion)\System.Runtime.WindowsRuntime.dll</HintPath>
+          <HintPath Condition="'$([System.IO.File]::Exists(%(HintPath)))'=='false'">C:\Program Files (x86)\Reference Assemblies\Microsoft\Framework\.NETCore\v4.5\System.Runtime.WindowsRuntime.dll</HintPath>
+          <Private>True</Private>
+        </Reference>
+        <!-- Disable windows reference in net core as this causes infinite rebuilds
+        <Reference Include="Windows">
+          <HintPath>$(Windows10SdkPath)\UnionMetadata\$(Windows10SdkLatestVersion)\Windows.winmd</HintPath>
+          <SDKsWinmdPath Condition="'$(UE_SDKS_ROOT)' != ''">$(UE_SDKS_ROOT)\HostWin64\Win64\Windows Kits\10\UnionMetadata\10.0.18362.0\Windows.winmd</SDKsWinmdPath>
+          <HintPath Condition="'$(UE_SDKS_ROOT)' != '' And '$([System.IO.File]::Exists(%(SDKsWinmdPath)))'!='false'">$(UE_SDKS_ROOT)\HostWin64\Win64\Windows Kits\10\UnionMetadata\10.0.18362.0\Windows.winmd</HintPath>
+          <HintPath Condition="'$([System.IO.File]::Exists(%(HintPath)))'=='false'">$(Windows10SdkPath)\UnionMetadata\$(Windows10SdkLatestVersion).0\Windows.winmd</HintPath>
+          <HintPath Condition="'$([System.IO.File]::Exists(%(HintPath)))'=='false'">$(Windows10SdkPath)\UnionMetadata\Windows.winmd</HintPath>
+          <Private>false</Private>
+        </Reference>
+        -->
+      </ItemGroup>
+    </When>
+  </Choose>
+
+  <Import Project="Sdk.props" Sdk="Microsoft.NET.Sdk.WindowsDesktop" Condition="'$(IsWindows)'=='true'"/>
+  <Import Project="Sdk.targets" Sdk="Microsoft.NET.Sdk.WindowsDesktop" Condition="'$(IsWindows)'=='true'"/>
+  <Import Project="Sdk.props" Sdk="Microsoft.NET.Sdk" Condition="'$(IsWindows)'!='true'"/>
+  <Import Project="Sdk.targets" Sdk="Microsoft.NET.Sdk" Condition="'$(IsWindows)'!='true'"/>
+
+  <ItemGroup>
+    <Compile Remove="obj\**" />
+    <EmbeddedResource Remove="obj\**" />
+    <None Remove="obj\**" />
+  </ItemGroup>
+
+  <ItemGroup>
+    <Compile Include="..\..\DotNETCommon\MetaData.cs">
+      <Link>Properties\MetaData.cs</Link>
+    </Compile>
+  </ItemGroup>
+  <ItemGroup>
+    <ProjectReference Include="..\..\DotNETCommon\DotNETUtilities\DotNETUtilitiesCore.csproj" />
+    <ProjectReference Include="..\..\UnrealBuildTool\UnrealBuildToolCore.csproj" />
+    <ProjectReference Include="..\AutomationUtils\AutomationUtils.Automation.Core.csproj" />
+    <ProjectReference Include="..\Scripts\AutomationScripts.Automation.Core.csproj" />
+  </ItemGroup>
+  <ItemGroup>
+    <PackageReference Include="Microsoft.CSharp" Version="4.7.0" />
+    <PackageReference Include="System.Data.DataSetExtensions" Version="4.5.0" />
+  </ItemGroup>
+</Project>
\ No newline at end of file
diff --git a/Engine/Source/Programs/AutomationTool/IOS/IOS.Automation.Core.csproj b/Engine/Source/Programs/AutomationTool/IOS/IOS.Automation.Core.csproj
new file mode 100644
--- /dev/null
+++ b/Engine/Source/Programs/AutomationTool/IOS/IOS.Automation.Core.csproj
@@ -0,0 +1,72 @@
+﻿<Project>
+  <PropertyGroup>
+    <!-- Set a custom temp directory so we can live side by side with the .net framework output -->
+    <BaseIntermediateOutputPath>$(MSBuildThisFileDirectory)obj-core\$(MSBuildProjectName)\</BaseIntermediateOutputPath>
+    <BaseOutputPath>$(MSBuildThisFileDirectory)bin-core\$(MSBuildProjectName)\</BaseOutputPath>
+  </PropertyGroup>
+
+  <PropertyGroup>
+    <TargetFramework>netcoreapp3.1</TargetFramework>
+    <Configuration Condition=" '$(Configuration)' == '' ">Development</Configuration>
+    <OutputType>Library</OutputType>
+    <AppendTargetFrameworkToOutputPath>false</AppendTargetFrameworkToOutputPath>
+    <GenerateAssemblyInfo>false</GenerateAssemblyInfo>
+    <GenerateTargetFrameworkAttribute>false</GenerateTargetFrameworkAttribute>
+    <Configurations>Debug;Release;Development</Configurations>
+    <RootNamespace>IOS.Automation</RootNamespace>
+    <AssemblyName>IOS.Automation</AssemblyName>
+    <WarningsNotAsErrors>612,618</WarningsNotAsErrors>
+    <OutputPath>..\..\..\..\Binaries\DotNET\AutomationTool</OutputPath>
+    <DefineConstants>NET_CORE</DefineConstants>
+    <TreatWarningsAsErrors>true</TreatWarningsAsErrors>
+  </PropertyGroup>
+  <PropertyGroup Condition=" '$(Configuration)|$(Platform)' == 'Development|AnyCPU' ">
+    <DefineConstants>NET_CORE;TRACE</DefineConstants>
+    <Optimize>true</Optimize>
+  </PropertyGroup>
+  <PropertyGroup Condition=" '$(Configuration)|$(Platform)' == 'Release|AnyCPU' ">
+    <DefineConstants>NET_CORE;TRACE</DefineConstants>
+    <Optimize>true</Optimize>
+  </PropertyGroup>
+  <PropertyGroup Condition=" '$(Configuration)|$(Platform)' == 'Debug|AnyCPU' ">
+    <DefineConstants>NET_CORE;DEBUG;TRACE</DefineConstants>
+  </PropertyGroup>
+  <ItemGroup>
+    <Reference Include="Ionic.Zip.Reduced">
+      <HintPath>..\..\..\..\Binaries\DotNET\Ionic.Zip.Reduced.dll</HintPath>
+      <Private>False</Private>
+    </Reference>
+    <Reference Include="MobileDeviceInterface">
+      <HintPath>..\..\..\..\Binaries\DotNET\IOS\MobileDeviceInterface.dll</HintPath>
+      <Private>False</Private>
+    </Reference>
+
+  </ItemGroup>
+  <Import Project="Sdk.props" Sdk="Microsoft.NET.Sdk" />
+  <Import Project="Sdk.targets" Sdk="Microsoft.NET.Sdk" />
+
+  <ItemGroup>
+    <Compile Remove="obj\**" />
+    <EmbeddedResource Remove="obj\**" />
+    <None Remove="obj\**" />
+  </ItemGroup>
+
+  <ItemGroup>
+    <Compile Include="..\..\DotNETCommon\MetaData.cs">
+      <Link>Properties\MetaData.cs</Link>
+    </Compile>
+  </ItemGroup>
+  <ItemGroup>
+    <ProjectReference Include="..\..\DotNETCommon\DotNETUtilities\DotNETUtilitiesCore.csproj" />
+    <ProjectReference Include="..\..\UnrealBuildTool\UnrealBuildToolCore.csproj" />
+    <ProjectReference Include="..\AutomationUtils\AutomationUtils.Automation.Core.csproj" />
+  </ItemGroup>
+  <ItemGroup>
+    <PackageReference Include="Microsoft.CSharp" Version="4.7.0" />
+    <PackageReference Include="System.Data.DataSetExtensions" Version="4.5.0" />
+    <PackageReference Include="System.Net.Http" Version="4.3.4" />
+    <PackageReference Include="System.Drawing.Common" Version="4.7.0" />
+    <PackageReference Include="System.Resources.Extensions" Version="4.7.1" />
+    <PackageReference Include="Portable.BouncyCastle" Version="1.8.6.7" />
+  </ItemGroup>
+</Project>
\ No newline at end of file
diff --git a/Engine/Source/Programs/AutomationTool/Linux/Linux.Automation.Core.csproj b/Engine/Source/Programs/AutomationTool/Linux/Linux.Automation.Core.csproj
new file mode 100644
--- /dev/null
+++ b/Engine/Source/Programs/AutomationTool/Linux/Linux.Automation.Core.csproj
@@ -0,0 +1,59 @@
+﻿<Project>
+  <PropertyGroup>
+    <!-- Set a custom temp directory so we can live side by side with the .net framework output -->
+    <BaseIntermediateOutputPath>$(MSBuildThisFileDirectory)obj-core\$(MSBuildProjectName)\</BaseIntermediateOutputPath>
+    <BaseOutputPath>$(MSBuildThisFileDirectory)bin-core\$(MSBuildProjectName)\</BaseOutputPath>
+  </PropertyGroup>
+
+  <PropertyGroup>
+    <TargetFramework>netcoreapp3.1</TargetFramework>
+    <Configuration Condition=" '$(Configuration)' == '' ">Development</Configuration>
+    <OutputType>Library</OutputType>
+    <AppendTargetFrameworkToOutputPath>false</AppendTargetFrameworkToOutputPath>
+    <GenerateAssemblyInfo>false</GenerateAssemblyInfo>
+    <GenerateTargetFrameworkAttribute>false</GenerateTargetFrameworkAttribute>
+    <Configurations>Debug;Release;Development</Configurations>
+    <RootNamespace>Linux.Automation</RootNamespace>
+    <AssemblyName>Linux.Automation</AssemblyName>
+    <WarningsNotAsErrors>612,618</WarningsNotAsErrors>
+    <OutputPath>..\..\..\..\Binaries\DotNET\AutomationTool</OutputPath>
+    <DefineConstants>NET_CORE</DefineConstants>
+    <TreatWarningsAsErrors>true</TreatWarningsAsErrors>
+  </PropertyGroup>
+  <PropertyGroup Condition=" '$(Configuration)|$(Platform)' == 'Development|AnyCPU' ">
+    <DefineConstants>NET_CORE;TRACE</DefineConstants>
+    <Optimize>true</Optimize>
+  </PropertyGroup>
+  <PropertyGroup Condition=" '$(Configuration)|$(Platform)' == 'Release|AnyCPU' ">
+    <DefineConstants>NET_CORE;TRACE</DefineConstants>
+    <Optimize>true</Optimize>
+  </PropertyGroup>
+  <PropertyGroup Condition=" '$(Configuration)|$(Platform)' == 'Debug|AnyCPU' ">
+    <DefineConstants>NET_CORE;DEBUG;TRACE</DefineConstants>
+  </PropertyGroup>
+  <Import Project="Sdk.props" Sdk="Microsoft.NET.Sdk" />
+  <Import Project="Sdk.targets" Sdk="Microsoft.NET.Sdk" />
+
+  <ItemGroup>
+    <Compile Remove="obj\**" />
+    <EmbeddedResource Remove="obj\**" />
+    <None Remove="obj\**" />
+  </ItemGroup>
+
+  <ItemGroup>
+    <Compile Include="..\..\DotNETCommon\MetaData.cs">
+      <Link>Properties\MetaData.cs</Link>
+    </Compile>
+  </ItemGroup>
+  <ItemGroup>
+    <ProjectReference Include="..\..\DotNETCommon\DotNETUtilities\DotNETUtilitiesCore.csproj" />
+    <ProjectReference Include="..\..\UnrealBuildTool\UnrealBuildToolCore.csproj" />
+    <ProjectReference Include="..\AutomationUtils\AutomationUtils.Automation.Core.csproj" />
+  </ItemGroup>
+  <ItemGroup>
+    <PackageReference Include="Microsoft.CSharp" Version="4.7.0" />
+    <PackageReference Include="System.Data.DataSetExtensions" Version="4.5.0" />
+    <PackageReference Include="System.Net.Http" Version="4.3.4" />
+    <PackageReference Include="System.Drawing.Common" Version="4.7.0" />
+  </ItemGroup>
+</Project>
\ No newline at end of file
diff --git a/Engine/Source/Programs/AutomationTool/Localization/Localization.Automation.Core.csproj b/Engine/Source/Programs/AutomationTool/Localization/Localization.Automation.Core.csproj
new file mode 100644
--- /dev/null
+++ b/Engine/Source/Programs/AutomationTool/Localization/Localization.Automation.Core.csproj
@@ -0,0 +1,56 @@
+﻿<Project>
+  <PropertyGroup>
+    <!-- Set a custom temp directory so we can live side by side with the .net framework output -->
+    <BaseIntermediateOutputPath>$(MSBuildThisFileDirectory)obj-core\$(MSBuildProjectName)\</BaseIntermediateOutputPath>
+    <BaseOutputPath>$(MSBuildThisFileDirectory)bin-core\$(MSBuildProjectName)\</BaseOutputPath>
+  </PropertyGroup>
+
+  <PropertyGroup>
+    <TargetFramework>netcoreapp3.1</TargetFramework>
+    <Configuration Condition=" '$(Configuration)' == '' ">Development</Configuration>
+    <OutputType>Library</OutputType>
+    <AppendTargetFrameworkToOutputPath>false</AppendTargetFrameworkToOutputPath>
+    <GenerateAssemblyInfo>false</GenerateAssemblyInfo>
+    <GenerateTargetFrameworkAttribute>false</GenerateTargetFrameworkAttribute>
+    <Configurations>Debug;Release;Development</Configurations>
+    <RootNamespace>AutomationTool</RootNamespace>
+    <AssemblyName>Localization.Automation</AssemblyName>
+    <WarningsNotAsErrors>612,618</WarningsNotAsErrors>
+    <OutputPath>..\..\..\..\Binaries\DotNET\AutomationTool</OutputPath>
+    <DefineConstants>NET_CORE</DefineConstants>
+    <TreatWarningsAsErrors>true</TreatWarningsAsErrors>
+  </PropertyGroup>
+  <PropertyGroup Condition=" '$(Configuration)|$(Platform)' == 'Development|AnyCPU' ">
+    <DefineConstants>NET_CORE;TRACE</DefineConstants>
+    <Optimize>true</Optimize>
+  </PropertyGroup>
+  <PropertyGroup Condition=" '$(Configuration)|$(Platform)' == 'Release|AnyCPU' ">
+    <DefineConstants>NET_CORE;TRACE</DefineConstants>
+    <Optimize>true</Optimize>
+  </PropertyGroup>
+  <PropertyGroup Condition=" '$(Configuration)|$(Platform)' == 'Debug|AnyCPU' ">
+    <DefineConstants>NET_CORE;DEBUG;TRACE</DefineConstants>
+  </PropertyGroup>
+  <Import Project="Sdk.props" Sdk="Microsoft.NET.Sdk" />
+  <Import Project="Sdk.targets" Sdk="Microsoft.NET.Sdk" />
+
+  <ItemGroup>
+    <Compile Remove="obj\**" />
+    <EmbeddedResource Remove="obj\**" />
+    <None Remove="obj\**" />
+  </ItemGroup>
+
+  <ItemGroup>
+    <Compile Include="..\..\DotNETCommon\MetaData.cs">
+      <Link>Properties\MetaData.cs</Link>
+    </Compile>
+  </ItemGroup>
+  <ItemGroup>
+    <ProjectReference Include="..\..\DotNETCommon\DotNETUtilities\DotNETUtilitiesCore.csproj" />
+    <ProjectReference Include="..\..\UnrealBuildTool\UnrealBuildToolCore.csproj" />
+    <ProjectReference Include="..\AutomationUtils\AutomationUtils.Automation.Core.csproj" />
+  </ItemGroup>
+  <ItemGroup>
+    <PackageReference Include="System.Data.DataSetExtensions" Version="4.5.0" />
+  </ItemGroup>
+</Project>
diff --git a/Engine/Source/Programs/AutomationTool/Lumin/Lumin.Automation.Core.csproj b/Engine/Source/Programs/AutomationTool/Lumin/Lumin.Automation.Core.csproj
new file mode 100644
--- /dev/null
+++ b/Engine/Source/Programs/AutomationTool/Lumin/Lumin.Automation.Core.csproj
@@ -0,0 +1,65 @@
+﻿<Project>
+  <PropertyGroup>
+    <!-- Set a custom temp directory so we can live side by side with the .net framework output -->
+    <BaseIntermediateOutputPath>$(MSBuildThisFileDirectory)obj-core\$(MSBuildProjectName)\</BaseIntermediateOutputPath>
+    <BaseOutputPath>$(MSBuildThisFileDirectory)bin-core\$(MSBuildProjectName)\</BaseOutputPath>
+  </PropertyGroup>
+
+  <PropertyGroup>
+    <TargetFramework>netcoreapp3.1</TargetFramework>
+    <Configuration Condition=" '$(Configuration)' == '' ">Development</Configuration>
+    <OutputType>Library</OutputType>
+    <AppendTargetFrameworkToOutputPath>false</AppendTargetFrameworkToOutputPath>
+    <GenerateAssemblyInfo>false</GenerateAssemblyInfo>
+    <GenerateTargetFrameworkAttribute>false</GenerateTargetFrameworkAttribute>
+    <Configurations>Debug;Release;Development</Configurations>
+    <RootNamespace>Lumin.Automation</RootNamespace>
+    <AssemblyName>Lumin.Automation</AssemblyName>
+    <WarningsNotAsErrors>612,618</WarningsNotAsErrors>
+    <OutputPath>..\..\..\..\Binaries\DotNET\AutomationTool</OutputPath>
+    <DefineConstants>NET_CORE</DefineConstants>
+    <TreatWarningsAsErrors>true</TreatWarningsAsErrors>
+  </PropertyGroup>
+  <PropertyGroup Condition=" '$(Configuration)|$(Platform)' == 'Development|AnyCPU' ">
+    <DefineConstants>NET_CORE;TRACE</DefineConstants>
+    <Optimize>true</Optimize>
+  </PropertyGroup>
+  <PropertyGroup Condition=" '$(Configuration)|$(Platform)' == 'Release|AnyCPU' ">
+    <DefineConstants>NET_CORE;TRACE</DefineConstants>
+    <Optimize>true</Optimize>
+  </PropertyGroup>
+  <PropertyGroup Condition=" '$(Configuration)|$(Platform)' == 'Debug|AnyCPU' ">
+    <DefineConstants>NET_CORE;DEBUG;TRACE</DefineConstants>
+  </PropertyGroup>
+  <ItemGroup>
+    <Reference Include="Ionic.Zip.Reduced">
+      <HintPath>..\..\..\..\Binaries\DotNET\Ionic.Zip.Reduced.dll</HintPath>
+      <Private>False</Private>
+    </Reference>
+  </ItemGroup>
+  <Import Project="Sdk.props" Sdk="Microsoft.NET.Sdk" />
+  <Import Project="Sdk.targets" Sdk="Microsoft.NET.Sdk" />
+
+  <ItemGroup>
+    <Compile Remove="obj\**" />
+    <EmbeddedResource Remove="obj\**" />
+    <None Remove="obj\**" />
+  </ItemGroup>
+
+  <ItemGroup>
+    <Compile Include="..\..\DotNETCommon\MetaData.cs">
+      <Link>Properties\MetaData.cs</Link>
+    </Compile>
+  </ItemGroup>
+  <ItemGroup>
+    <ProjectReference Include="..\..\DotNETCommon\DotNETUtilities\DotNETUtilitiesCore.csproj" />
+    <ProjectReference Include="..\..\UnrealBuildTool\UnrealBuildToolCore.csproj" />
+    <ProjectReference Include="..\AutomationUtils\AutomationUtils.Automation.Core.csproj" />
+    <ProjectReference Include="..\Scripts\AutomationScripts.Automation.Core.csproj" />
+  </ItemGroup>
+  <ItemGroup>
+    <PackageReference Include="Microsoft.CSharp" Version="4.7.0" />
+    <PackageReference Include="System.Data.DataSetExtensions" Version="4.5.0" />
+    <PackageReference Include="System.Net.Http" Version="4.3.4" />
+  </ItemGroup>
+</Project>
diff --git a/Engine/Source/Programs/AutomationTool/Mac/Mac.Automation.Core.csproj b/Engine/Source/Programs/AutomationTool/Mac/Mac.Automation.Core.csproj
new file mode 100644
--- /dev/null
+++ b/Engine/Source/Programs/AutomationTool/Mac/Mac.Automation.Core.csproj
@@ -0,0 +1,59 @@
+﻿<Project>
+  <PropertyGroup>
+    <!-- Set a custom temp directory so we can live side by side with the .net framework output -->
+    <BaseIntermediateOutputPath>$(MSBuildThisFileDirectory)obj-core\$(MSBuildProjectName)\</BaseIntermediateOutputPath>
+    <BaseOutputPath>$(MSBuildThisFileDirectory)bin-core\$(MSBuildProjectName)\</BaseOutputPath>
+  </PropertyGroup>
+
+  <PropertyGroup>
+    <TargetFramework>netcoreapp3.1</TargetFramework>
+    <Configuration Condition=" '$(Configuration)' == '' ">Development</Configuration>
+    <OutputType>Library</OutputType>
+    <AppendTargetFrameworkToOutputPath>false</AppendTargetFrameworkToOutputPath>
+    <GenerateAssemblyInfo>false</GenerateAssemblyInfo>
+    <GenerateTargetFrameworkAttribute>false</GenerateTargetFrameworkAttribute>
+    <Configurations>Debug;Release;Development</Configurations>
+    <RootNamespace>Mac.Automation</RootNamespace>
+    <AssemblyName>Mac.Automation</AssemblyName>
+    <WarningsNotAsErrors>612,618</WarningsNotAsErrors>
+    <OutputPath>..\..\..\..\Binaries\DotNET\AutomationTool</OutputPath>
+    <DefineConstants>NET_CORE</DefineConstants>
+    <TreatWarningsAsErrors>true</TreatWarningsAsErrors>
+  </PropertyGroup>
+  <PropertyGroup Condition=" '$(Configuration)|$(Platform)' == 'Development|AnyCPU' ">
+    <DefineConstants>NET_CORE;TRACE</DefineConstants>
+    <Optimize>true</Optimize>
+  </PropertyGroup>
+  <PropertyGroup Condition=" '$(Configuration)|$(Platform)' == 'Release|AnyCPU' ">
+    <DefineConstants>NET_CORE;TRACE</DefineConstants>
+    <Optimize>true</Optimize>
+  </PropertyGroup>
+  <PropertyGroup Condition=" '$(Configuration)|$(Platform)' == 'Debug|AnyCPU' ">
+    <DefineConstants>NET_CORE;DEBUG;TRACE</DefineConstants>
+  </PropertyGroup>
+  <Import Project="Sdk.props" Sdk="Microsoft.NET.Sdk" />
+  <Import Project="Sdk.targets" Sdk="Microsoft.NET.Sdk" />
+
+  <ItemGroup>
+    <Compile Remove="obj\**" />
+    <EmbeddedResource Remove="obj\**" />
+    <None Remove="obj\**" />
+  </ItemGroup>
+
+  <ItemGroup>
+    <Compile Include="..\..\DotNETCommon\MetaData.cs">
+      <Link>Properties\MetaData.cs</Link>
+    </Compile>
+  </ItemGroup>
+  <ItemGroup>
+    <ProjectReference Include="..\..\DotNETCommon\DotNETUtilities\DotNETUtilitiesCore.csproj" />
+    <ProjectReference Include="..\..\UnrealBuildTool\UnrealBuildToolCore.csproj" />
+    <ProjectReference Include="..\AutomationUtils\AutomationUtils.Automation.Core.csproj" />
+  </ItemGroup>
+  <ItemGroup>
+    <PackageReference Include="Microsoft.CSharp" Version="4.7.0" />
+    <PackageReference Include="System.Data.DataSetExtensions" Version="4.5.0" />
+    <PackageReference Include="System.Net.Http" Version="4.3.4" />
+    <PackageReference Include="System.Drawing.Common" Version="4.7.0" />
+  </ItemGroup>
+</Project>
\ No newline at end of file
diff --git a/Engine/Source/Programs/AutomationTool/OneSkyLocalization/OneSkyLocalization.Automation.Core.csproj b/Engine/Source/Programs/AutomationTool/OneSkyLocalization/OneSkyLocalization.Automation.Core.csproj
new file mode 100644
--- /dev/null
+++ b/Engine/Source/Programs/AutomationTool/OneSkyLocalization/OneSkyLocalization.Automation.Core.csproj
@@ -0,0 +1,66 @@
+﻿<Project>
+  <PropertyGroup>
+    <!-- Set a custom temp directory so we can live side by side with the .net framework output -->
+    <BaseIntermediateOutputPath>$(MSBuildThisFileDirectory)obj-core\$(MSBuildProjectName)\</BaseIntermediateOutputPath>
+    <BaseOutputPath>$(MSBuildThisFileDirectory)bin-core\$(MSBuildProjectName)\</BaseOutputPath>
+  </PropertyGroup>
+
+  <PropertyGroup>
+    <TargetFramework>netcoreapp3.1</TargetFramework>
+    <Configuration Condition=" '$(Configuration)' == '' ">Development</Configuration>
+    <OutputType>Library</OutputType>
+    <AppendTargetFrameworkToOutputPath>false</AppendTargetFrameworkToOutputPath>
+    <GenerateAssemblyInfo>false</GenerateAssemblyInfo>
+    <GenerateTargetFrameworkAttribute>false</GenerateTargetFrameworkAttribute>
+    <Configurations>Debug;Release;Development</Configurations>
+    <RootNamespace>AutomationTool</RootNamespace>
+    <AssemblyName>OneSkyLocalization.Automation</AssemblyName>
+    <WarningsNotAsErrors>612,618</WarningsNotAsErrors>
+    <OutputPath>..\..\..\..\Binaries\DotNET\AutomationTool</OutputPath>
+    <DefineConstants>NET_CORE</DefineConstants>
+    <TreatWarningsAsErrors>true</TreatWarningsAsErrors>
+  </PropertyGroup>
+  <PropertyGroup Condition=" '$(Configuration)|$(Platform)' == 'Development|AnyCPU' ">
+    <DefineConstants>NET_CORE;TRACE</DefineConstants>
+    <Optimize>true</Optimize>
+  </PropertyGroup>
+  <PropertyGroup Condition=" '$(Configuration)|$(Platform)' == 'Release|AnyCPU' ">
+    <DefineConstants>NET_CORE;TRACE</DefineConstants>
+    <Optimize>true</Optimize>
+  </PropertyGroup>
+  <PropertyGroup Condition=" '$(Configuration)|$(Platform)' == 'Debug|AnyCPU' ">
+    <DefineConstants>NET_CORE;DEBUG;TRACE</DefineConstants>
+  </PropertyGroup>
+  <ItemGroup>
+    <Reference Include="OneSky, Version=1.0.0.0, Culture=neutral, processorArchitecture=MSIL">
+      <HintPath>..\..\..\..\Binaries\DotNET\OneSky.dll</HintPath>
+      <SpecificVersion>False</SpecificVersion>
+      <Private>False</Private>
+    </Reference>
+  </ItemGroup>
+  <Import Project="Sdk.props" Sdk="Microsoft.NET.Sdk" />
+  <Import Project="Sdk.targets" Sdk="Microsoft.NET.Sdk" />
+
+  <ItemGroup>
+    <Compile Remove="obj\**" />
+    <EmbeddedResource Remove="obj\**" />
+    <None Remove="obj\**" />
+  </ItemGroup>
+
+  <ItemGroup>
+    <Compile Include="..\..\DotNETCommon\MetaData.cs">
+      <Link>Properties\MetaData.cs</Link>
+    </Compile>
+  </ItemGroup>
+  <ItemGroup>
+    <ProjectReference Include="..\..\DotNETCommon\DotNETUtilities\DotNETUtilitiesCore.csproj" />
+    <ProjectReference Include="..\..\UnrealBuildTool\UnrealBuildToolCore.csproj" />
+    <ProjectReference Include="..\AutomationUtils\AutomationUtils.Automation.Core.csproj" />
+    <ProjectReference Include="..\Localization\Localization.Automation.Core.csproj" />
+  </ItemGroup>
+  <ItemGroup>
+    <PackageReference Include="Microsoft.CSharp" Version="4.7.0" />
+    <PackageReference Include="System.Data.DataSetExtensions" Version="4.5.0" />
+    <PackageReference Include="System.Net.Http" Version="4.3.4" />
+  </ItemGroup>
+</Project>
\ No newline at end of file
diff --git a/Engine/Source/Programs/AutomationTool/Scripts/AutomationScripts.Automation.Core.csproj b/Engine/Source/Programs/AutomationTool/Scripts/AutomationScripts.Automation.Core.csproj
new file mode 100644
--- /dev/null
+++ b/Engine/Source/Programs/AutomationTool/Scripts/AutomationScripts.Automation.Core.csproj
@@ -0,0 +1,72 @@
+﻿<Project>
+  <PropertyGroup>
+    <!-- Set a custom temp directory so we can live side by side with the .net framework output -->
+    <BaseIntermediateOutputPath>$(MSBuildThisFileDirectory)obj-core\$(MSBuildProjectName)\</BaseIntermediateOutputPath>
+    <BaseOutputPath>$(MSBuildThisFileDirectory)bin-core\$(MSBuildProjectName)\</BaseOutputPath>
+  </PropertyGroup>
+
+  <PropertyGroup>
+    <TargetFramework>netcoreapp3.1</TargetFramework>
+    <Configuration Condition=" '$(Configuration)' == '' ">Development</Configuration>
+    <OutputType>Library</OutputType>
+    <AppendTargetFrameworkToOutputPath>false</AppendTargetFrameworkToOutputPath>
+    <GenerateAssemblyInfo>false</GenerateAssemblyInfo>
+    <GenerateTargetFrameworkAttribute>false</GenerateTargetFrameworkAttribute>
+    <Configurations>Debug;Release;Development</Configurations>
+    <RootNamespace>AutomationTool</RootNamespace>
+    <AssemblyName>AutomationScripts.Automation</AssemblyName>
+    <WarningsNotAsErrors>612,618</WarningsNotAsErrors>
+    <OutputPath>..\..\..\..\Binaries\DotNET\AutomationTool</OutputPath>
+    <DefineConstants>NET_CORE</DefineConstants>
+    <TreatWarningsAsErrors>true</TreatWarningsAsErrors>
+  </PropertyGroup>
+  <PropertyGroup Condition=" '$(Configuration)|$(Platform)' == 'Development|AnyCPU' ">
+    <DefineConstants>NET_CORE;TRACE</DefineConstants>
+    <Optimize>true</Optimize>
+  </PropertyGroup>
+  <PropertyGroup Condition=" '$(Configuration)|$(Platform)' == 'Release|AnyCPU' ">
+    <DefineConstants>NET_CORE;TRACE</DefineConstants>
+    <Optimize>true</Optimize>
+  </PropertyGroup>
+  <PropertyGroup Condition=" '$(Configuration)|$(Platform)' == 'Debug|AnyCPU' ">
+    <DefineConstants>NET_CORE;DEBUG;TRACE</DefineConstants>
+  </PropertyGroup>
+  <ItemGroup>
+    <Reference Include="Ionic.Zip.Reduced">
+      <HintPath>..\..\..\..\Binaries\DotNET\Ionic.Zip.Reduced.dll</HintPath>
+      <Private>False</Private>
+    </Reference>
+    <Reference Include="OneSky, Version=1.0.0.0, Culture=neutral, processorArchitecture=MSIL">
+      <HintPath>..\..\..\..\Binaries\DotNET\OneSky.dll</HintPath>
+      <SpecificVersion>False</SpecificVersion>
+      <Private>False</Private>
+    </Reference>
+  </ItemGroup>
+  <Import Project="Sdk.props" Sdk="Microsoft.NET.Sdk" />
+  <Import Project="Sdk.targets" Sdk="Microsoft.NET.Sdk" />
+
+  <ItemGroup>
+    <Compile Remove="obj\**" />
+    <EmbeddedResource Remove="obj\**" />
+    <None Remove="obj\**" />
+  </ItemGroup>
+
+  <ItemGroup>
+    <Compile Include="..\..\DotNETCommon\MetaData.cs">
+      <Link>Properties\MetaData.cs</Link>
+    </Compile>
+  </ItemGroup>
+  <ItemGroup>
+    <ProjectReference Include="..\..\DotNETCommon\DotNETUtilities\DotNETUtilitiesCore.csproj" />
+    <ProjectReference Include="..\..\UnrealBuildTool\UnrealBuildToolCore.csproj" />
+    <ProjectReference Include="..\AutomationUtils\AutomationUtils.Automation.Core.csproj" />
+    <ProjectReference Include="..\Localization\Localization.Automation.Core.csproj" />
+  </ItemGroup>
+  <ItemGroup>
+    <PackageReference Include="Microsoft.CSharp" Version="4.7.0" />
+    <PackageReference Include="System.Data.DataSetExtensions" Version="4.5.0" />
+    <PackageReference Include="System.Net.Http" Version="4.3.4" />
+    <PackageReference Include="System.Drawing.Common" Version="4.7.0" />
+    <PackageReference Include="AWSSDK.S3" Version="3.3.113.2" />
+  </ItemGroup>
+</Project>
\ No newline at end of file
diff --git a/Engine/Source/Programs/AutomationTool/TVOS/TVOS.Automation.Core.csproj b/Engine/Source/Programs/AutomationTool/TVOS/TVOS.Automation.Core.csproj
new file mode 100644
--- /dev/null
+++ b/Engine/Source/Programs/AutomationTool/TVOS/TVOS.Automation.Core.csproj
@@ -0,0 +1,66 @@
+﻿<Project>
+  <PropertyGroup>
+    <!-- Set a custom temp directory so we can live side by side with the .net framework output -->
+    <BaseIntermediateOutputPath>$(MSBuildThisFileDirectory)obj-core\$(MSBuildProjectName)\</BaseIntermediateOutputPath>
+    <BaseOutputPath>$(MSBuildThisFileDirectory)bin-core\$(MSBuildProjectName)\</BaseOutputPath>
+  </PropertyGroup>
+
+  <PropertyGroup>
+    <TargetFramework>netcoreapp3.1</TargetFramework>
+    <Configuration Condition=" '$(Configuration)' == '' ">Development</Configuration>
+    <OutputType>Library</OutputType>
+    <AppendTargetFrameworkToOutputPath>false</AppendTargetFrameworkToOutputPath>
+    <GenerateAssemblyInfo>false</GenerateAssemblyInfo>
+    <GenerateTargetFrameworkAttribute>false</GenerateTargetFrameworkAttribute>
+    <Configurations>Debug;Release;Development</Configurations>
+    <RootNamespace>TVOS.Automation</RootNamespace>
+    <AssemblyName>TVOS.Automation</AssemblyName>
+    <WarningsNotAsErrors>612,618</WarningsNotAsErrors>
+    <OutputPath>..\..\..\..\Binaries\DotNET\AutomationTool</OutputPath>
+    <DefineConstants>NET_CORE</DefineConstants>
+    <TreatWarningsAsErrors>true</TreatWarningsAsErrors>
+  </PropertyGroup>
+  <PropertyGroup Condition=" '$(Configuration)|$(Platform)' == 'Development|AnyCPU' ">
+    <DefineConstants>NET_CORE;TRACE</DefineConstants>
+    <Optimize>true</Optimize>
+  </PropertyGroup>
+  <PropertyGroup Condition=" '$(Configuration)|$(Platform)' == 'Release|AnyCPU' ">
+    <DefineConstants>NET_CORE;TRACE</DefineConstants>
+    <Optimize>true</Optimize>
+  </PropertyGroup>
+  <PropertyGroup Condition=" '$(Configuration)|$(Platform)' == 'Debug|AnyCPU' ">
+    <DefineConstants>NET_CORE;DEBUG;TRACE</DefineConstants>
+  </PropertyGroup>
+  <ItemGroup>
+    <Reference Include="Ionic.Zip.Reduced">
+      <HintPath>..\..\..\..\Binaries\DotNET\Ionic.Zip.Reduced.dll</HintPath>
+      <Private>False</Private>
+    </Reference>
+  </ItemGroup>
+  <Import Project="Sdk.props" Sdk="Microsoft.NET.Sdk" />
+  <Import Project="Sdk.targets" Sdk="Microsoft.NET.Sdk" />
+
+  <ItemGroup>
+    <Compile Remove="obj\**" />
+    <EmbeddedResource Remove="obj\**" />
+    <None Remove="obj\**" />
+  </ItemGroup>
+
+  <ItemGroup>
+    <Compile Include="..\..\DotNETCommon\MetaData.cs">
+      <Link>Properties\MetaData.cs</Link>
+    </Compile>
+  </ItemGroup>
+  <ItemGroup>
+    <ProjectReference Include="..\..\DotNETCommon\DotNETUtilities\DotNETUtilitiesCore.csproj" />
+    <ProjectReference Include="..\..\UnrealBuildTool\UnrealBuildToolCore.csproj" />
+    <ProjectReference Include="..\AutomationUtils\AutomationUtils.Automation.Core.csproj" />
+    <ProjectReference Include="..\IOS\IOS.Automation.Core.csproj" />
+  </ItemGroup>
+  <ItemGroup>
+    <PackageReference Include="Microsoft.CSharp" Version="4.7.0" />
+    <PackageReference Include="System.Data.DataSetExtensions" Version="4.5.0" />
+    <PackageReference Include="System.Net.Http" Version="4.3.4" />
+    <PackageReference Include="System.Drawing.Common" Version="4.7.0" />
+  </ItemGroup>
+</Project>
\ No newline at end of file
diff --git a/Engine/Source/Programs/AutomationTool/Turnkey/Turnkey.Automation.Core.csproj b/Engine/Source/Programs/AutomationTool/Turnkey/Turnkey.Automation.Core.csproj
new file mode 100644
--- /dev/null
+++ b/Engine/Source/Programs/AutomationTool/Turnkey/Turnkey.Automation.Core.csproj
@@ -0,0 +1,90 @@
+﻿<Project>
+  <PropertyGroup>
+    <!-- Set a custom temp directory so we can live side by side with the .net framework output -->
+    <BaseIntermediateOutputPath>$(MSBuildThisFileDirectory)obj-core\$(MSBuildProjectName)\</BaseIntermediateOutputPath>
+    <BaseOutputPath>$(MSBuildThisFileDirectory)bin-core\$(MSBuildProjectName)\</BaseOutputPath>
+  </PropertyGroup>
+
+  <PropertyGroup>
+    <TargetFramework>netcoreapp3.1</TargetFramework>
+    <Configuration Condition=" '$(Configuration)' == '' ">Development</Configuration>
+    <OutputType>Library</OutputType>
+    <AppendTargetFrameworkToOutputPath>false</AppendTargetFrameworkToOutputPath>
+    <GenerateAssemblyInfo>false</GenerateAssemblyInfo>
+    <GenerateTargetFrameworkAttribute>false</GenerateTargetFrameworkAttribute>
+    <Configurations>Debug;Release;Development</Configurations>
+    <RootNamespace>Turnkey.Automation</RootNamespace>
+    <AssemblyName>Turnkey.Automation</AssemblyName>
+    <WarningsNotAsErrors>612,618</WarningsNotAsErrors>
+    <OutputPath>..\..\..\..\Binaries\DotNET\AutomationTool</OutputPath>
+    <DefineConstants>NET_CORE</DefineConstants>
+    <TreatWarningsAsErrors>true</TreatWarningsAsErrors>
+    <IsWindows Condition="'$([System.Runtime.InteropServices.RuntimeInformation]::IsOSPlatform($([System.Runtime.InteropServices.OSPlatform]::Windows)))' == 'true'">true</IsWindows> 
+    <IsOSX Condition="'$([System.Runtime.InteropServices.RuntimeInformation]::IsOSPlatform($([System.Runtime.InteropServices.OSPlatform]::OSX)))' == 'true'">true</IsOSX> 
+    <IsLinux Condition="'$([System.Runtime.InteropServices.RuntimeInformation]::IsOSPlatform($([System.Runtime.InteropServices.OSPlatform]::Linux)))' == 'true'">true</IsLinux> 
+  </PropertyGroup>
+  <PropertyGroup Condition="'$(Configuration)|$(Platform)' == 'Development|AnyCPU' ">
+    <DefineConstants>$(DefineConstants);TRACE</DefineConstants>
+    <Optimize>true</Optimize>
+  </PropertyGroup>
+  <PropertyGroup Condition="'$(Configuration)|$(Platform)' == 'Release|AnyCPU' ">
+    <DefineConstants>$(DefineConstants);TRACE</DefineConstants>
+    <Optimize>true</Optimize>
+  </PropertyGroup>
+  <PropertyGroup Condition="'$(Configuration)|$(Platform)' == 'Debug|AnyCPU' ">
+    <DefineConstants>$(DefineConstants);DEBUG;TRACE</DefineConstants>
+  </PropertyGroup>
+  <PropertyGroup Condition="'$(IsWindows)'=='true'">
+    <DefineConstants>$(DefineConstants);WINDOWS</DefineConstants>
+    <UseWindowsForms>true</UseWindowsForms>
+  </PropertyGroup>
+  <PropertyGroup Condition="'$(IsOSX)'=='true'">
+    <DefineConstants>$(DefineConstants);OSX</DefineConstants>
+  </PropertyGroup>
+  <PropertyGroup Condition="'$(IsLinux)'=='true'">
+    <DefineConstants>$(DefineConstants);LINUX</DefineConstants>
+  </PropertyGroup>
+
+  <ItemGroup>
+    <Reference Include="Ionic.Zip.Reduced">
+      <HintPath>..\..\..\..\Binaries\DotNET\Ionic.Zip.Reduced.dll</HintPath>
+      <Private>False</Private>
+    </Reference>
+    <!--<Reference Include="OneSky, Version=1.0.0.0, Culture=neutral, processorArchitecture=MSIL">
+      <HintPath>..\..\..\..\Binaries\DotNET\OneSky.dll</HintPath>
+      <SpecificVersion>False</SpecificVersion>
+      <Private>False</Private>
+    </Reference>-->
+  </ItemGroup>
+
+  <Import Project="Sdk.props" Sdk="Microsoft.NET.Sdk.WindowsDesktop" Condition="'$(IsWindows)'=='true'"/>
+  <Import Project="Sdk.targets" Sdk="Microsoft.NET.Sdk.WindowsDesktop" Condition="'$(IsWindows)'=='true'"/>
+  <Import Project="Sdk.props" Sdk="Microsoft.NET.Sdk" Condition="'$(IsWindows)'!='true'"/>
+  <Import Project="Sdk.targets" Sdk="Microsoft.NET.Sdk" Condition="'$(IsWindows)'!='true'"/>
+
+  <ItemGroup>
+    <Compile Remove="obj\**" />
+    <EmbeddedResource Remove="obj\**" />
+    <None Remove="obj\**" />
+  </ItemGroup>
+
+  <ItemGroup>
+    <Compile Include="..\..\DotNETCommon\MetaData.cs">
+      <Link>Properties\MetaData.cs</Link>
+    </Compile>
+  </ItemGroup>
+  <ItemGroup>
+    <ProjectReference Include="..\..\DotNETCommon\BuildUtilities\BuildUtilitiesCore.csproj" />
+    <ProjectReference Include="..\..\DotNETCommon\DotNETUtilities\DotNETUtilitiesCore.csproj" />
+    <ProjectReference Include="..\..\UnrealBuildTool\UnrealBuildToolCore.csproj" />
+    <ProjectReference Include="..\AutomationUtils\AutomationUtils.Automation.Core.csproj" />
+    <ProjectReference Include="..\Gauntlet\Gauntlet.Automation.Core.csproj" />
+  </ItemGroup>
+  <ItemGroup>
+    <PackageReference Include="Microsoft.CSharp" Version="4.7.0" />
+    <PackageReference Include="System.Data.DataSetExtensions" Version="4.5.0" />
+    <PackageReference Include="System.Net.Http" Version="4.3.4" />
+    <PackageReference Include="System.Drawing.Common" Version="4.7.0" />
+    <PackageReference Include="Microsoft.VisualBasic" Version="10.3.0" />
+  </ItemGroup>
+</Project>
\ No newline at end of file
diff --git a/Engine/Source/Programs/AutomationTool/Win/Win.Automation.Core.csproj b/Engine/Source/Programs/AutomationTool/Win/Win.Automation.Core.csproj
new file mode 100644
--- /dev/null
+++ b/Engine/Source/Programs/AutomationTool/Win/Win.Automation.Core.csproj
@@ -0,0 +1,71 @@
+﻿<Project>
+  <PropertyGroup>
+    <!-- Set a custom temp directory so we can live side by side with the .net framework output -->
+    <BaseIntermediateOutputPath>$(MSBuildThisFileDirectory)obj-core\$(MSBuildProjectName)\</BaseIntermediateOutputPath>
+    <BaseOutputPath>$(MSBuildThisFileDirectory)bin-core\$(MSBuildProjectName)\</BaseOutputPath>
+  </PropertyGroup>
+
+  <PropertyGroup>
+    <TargetFramework>netcoreapp3.1</TargetFramework>
+    <Configuration Condition=" '$(Configuration)' == '' ">Development</Configuration>
+    <OutputType>Library</OutputType>
+    <AppendTargetFrameworkToOutputPath>false</AppendTargetFrameworkToOutputPath>
+    <GenerateAssemblyInfo>false</GenerateAssemblyInfo>
+    <GenerateTargetFrameworkAttribute>false</GenerateTargetFrameworkAttribute>
+    <Configurations>Debug;Release;Development</Configurations>
+    <RootNamespace>Win.Automation</RootNamespace>
+    <AssemblyName>Win.Automation</AssemblyName>
+    <WarningsNotAsErrors>612,618</WarningsNotAsErrors>
+    <OutputPath>..\..\..\..\Binaries\DotNET\AutomationTool</OutputPath>
+    <DefineConstants>NET_CORE</DefineConstants>
+    <TreatWarningsAsErrors>true</TreatWarningsAsErrors>
+  </PropertyGroup>
+  <PropertyGroup Condition=" '$(Configuration)|$(Platform)' == 'Development|AnyCPU' ">
+    <DefineConstants>NET_CORE;TRACE</DefineConstants>
+    <Optimize>true</Optimize>
+  </PropertyGroup>
+  <PropertyGroup Condition=" '$(Configuration)|$(Platform)' == 'Release|AnyCPU' ">
+    <DefineConstants>NET_CORE;TRACE</DefineConstants>
+    <Optimize>true</Optimize>
+  </PropertyGroup>
+  <PropertyGroup Condition=" '$(Configuration)|$(Platform)' == 'Debug|AnyCPU' ">
+    <DefineConstants>NET_CORE;DEBUG;TRACE</DefineConstants>
+  </PropertyGroup>
+  <ItemGroup>
+    <Reference Include="Ionic.Zip.Reduced">
+      <HintPath>..\..\..\..\Binaries\DotNET\Ionic.Zip.Reduced.dll</HintPath>
+      <Private>False</Private>
+    </Reference>
+    <!--<Reference Include="OneSky, Version=1.0.0.0, Culture=neutral, processorArchitecture=MSIL">
+      <HintPath>..\..\..\..\Binaries\DotNET\OneSky.dll</HintPath>
+      <SpecificVersion>False</SpecificVersion>
+      <Private>False</Private>
+    </Reference>-->
+  </ItemGroup>
+  <Import Project="Sdk.props" Sdk="Microsoft.NET.Sdk" />
+  <Import Project="Sdk.targets" Sdk="Microsoft.NET.Sdk" />
+
+  <ItemGroup>
+    <Compile Remove="obj\**" />
+    <EmbeddedResource Remove="obj\**" />
+    <None Remove="obj\**" />
+  </ItemGroup>
+
+  <ItemGroup>
+    <Compile Include="..\..\DotNETCommon\MetaData.cs">
+      <Link>Properties\MetaData.cs</Link>
+    </Compile>
+  </ItemGroup>
+  <ItemGroup>
+    <ProjectReference Include="..\..\DotNETCommon\DotNETUtilities\DotNETUtilitiesCore.csproj" />
+    <ProjectReference Include="..\..\UnrealBuildTool\UnrealBuildToolCore.csproj" />
+    <ProjectReference Include="..\AutomationUtils\AutomationUtils.Automation.Core.csproj" />
+    <ProjectReference Include="..\BuildGraph\BuildGraph.Automation.Core.csproj" />
+  </ItemGroup>
+  <ItemGroup>
+    <PackageReference Include="Microsoft.CSharp" Version="4.7.0" />
+    <PackageReference Include="System.Data.DataSetExtensions" Version="4.5.0" />
+    <PackageReference Include="System.Net.Http" Version="4.3.4" />
+    <PackageReference Include="System.Drawing.Common" Version="4.7.0" />
+  </ItemGroup>
+</Project>
\ No newline at end of file
diff --git a/Engine/Source/Programs/AutomationTool/XLocLocalization/XLocLocalization.Automation.Core.csproj b/Engine/Source/Programs/AutomationTool/XLocLocalization/XLocLocalization.Automation.Core.csproj
new file mode 100644
--- /dev/null
+++ b/Engine/Source/Programs/AutomationTool/XLocLocalization/XLocLocalization.Automation.Core.csproj
@@ -0,0 +1,55 @@
+﻿<Project>
+  <PropertyGroup>
+    <!-- Set a custom temp directory so we can live side by side with the .net framework output -->
+    <BaseIntermediateOutputPath>$(MSBuildThisFileDirectory)obj-core\$(MSBuildProjectName)\</BaseIntermediateOutputPath>
+    <BaseOutputPath>$(MSBuildThisFileDirectory)bin-core\$(MSBuildProjectName)\</BaseOutputPath>
+  </PropertyGroup>
+
+  <PropertyGroup>
+    <TargetFramework>netcoreapp3.1</TargetFramework>
+    <Configuration Condition=" '$(Configuration)' == '' ">Development</Configuration>
+    <OutputType>Library</OutputType>
+    <AppendTargetFrameworkToOutputPath>false</AppendTargetFrameworkToOutputPath>
+    <GenerateAssemblyInfo>false</GenerateAssemblyInfo>
+    <GenerateTargetFrameworkAttribute>false</GenerateTargetFrameworkAttribute>
+    <Configurations>Debug;Release;Development</Configurations>
+    <RootNamespace>AutomationTool</RootNamespace>
+    <AssemblyName>XLocLocalization.Automation</AssemblyName>
+    <WarningsNotAsErrors>612,618</WarningsNotAsErrors>
+    <OutputPath>..\..\..\..\Binaries\DotNET\AutomationTool</OutputPath>
+    <DefineConstants>NET_CORE</DefineConstants>
+    <TreatWarningsAsErrors>true</TreatWarningsAsErrors>
+  </PropertyGroup>
+  <PropertyGroup Condition=" '$(Configuration)|$(Platform)' == 'Development|AnyCPU' ">
+    <DefineConstants>NET_CORE;TRACE</DefineConstants>
+    <Optimize>true</Optimize>
+  </PropertyGroup>
+  <PropertyGroup Condition=" '$(Configuration)|$(Platform)' == 'Release|AnyCPU' ">
+    <DefineConstants>NET_CORE;TRACE</DefineConstants>
+    <Optimize>true</Optimize>
+  </PropertyGroup>
+  <PropertyGroup Condition=" '$(Configuration)|$(Platform)' == 'Debug|AnyCPU' ">
+    <DefineConstants>NET_CORE;DEBUG;TRACE</DefineConstants>
+  </PropertyGroup>
+
+  <Import Project="Sdk.props" Sdk="Microsoft.NET.Sdk" />
+  <Import Project="Sdk.targets" Sdk="Microsoft.NET.Sdk" />
+
+  <ItemGroup>
+    <Compile Remove="obj\**" />
+    <EmbeddedResource Remove="obj\**" />
+    <None Remove="obj\**" />
+  </ItemGroup>
+
+  <ItemGroup>
+    <ProjectReference Include="..\..\DotNETCommon\DotNETUtilities\DotNETUtilitiesCore.csproj" />
+    <ProjectReference Include="..\..\UnrealBuildTool\UnrealBuildToolCore.csproj" />
+    <ProjectReference Include="..\AutomationUtils\AutomationUtils.Automation.Core.csproj" />
+    <ProjectReference Include="..\Localization\Localization.Automation.Core.csproj" />
+  </ItemGroup>
+  <ItemGroup>
+    <PackageReference Include="Microsoft.CSharp" Version="4.7.0" />
+    <PackageReference Include="System.Data.DataSetExtensions" Version="4.5.0" />
+    <PackageReference Include="System.ServiceModel.Http" Version="4.5.3" />
+  </ItemGroup>
+</Project>
\ No newline at end of file
diff --git a/Engine/Source/Programs/AutomationToolLauncher/AutomationToolLauncherCore.csproj b/Engine/Source/Programs/AutomationToolLauncher/AutomationToolLauncherCore.csproj
new file mode 100644
--- /dev/null
+++ b/Engine/Source/Programs/AutomationToolLauncher/AutomationToolLauncherCore.csproj
@@ -0,0 +1,48 @@
+﻿<Project>
+  <PropertyGroup>
+    <!-- Set a custom temp directory so we can live side by side with the .net framework output -->
+    <BaseIntermediateOutputPath>$(MSBuildThisFileDirectory)obj-core\$(MSBuildProjectName)\</BaseIntermediateOutputPath>
+    <BaseOutputPath>$(MSBuildThisFileDirectory)bin-core\$(MSBuildProjectName)\</BaseOutputPath>
+  </PropertyGroup>
+  <Import Project="Sdk.props" Sdk="Microsoft.NET.Sdk" />
+  <PropertyGroup>
+    <TargetFramework>netcoreapp3.1</TargetFramework>
+    <Configuration Condition=" '$(Configuration)' == '' ">Development</Configuration>
+    <OutputType>Exe</OutputType>
+    <AppendTargetFrameworkToOutputPath>false</AppendTargetFrameworkToOutputPath>
+    <GenerateAssemblyInfo>false</GenerateAssemblyInfo>
+    <GenerateTargetFrameworkAttribute>false</GenerateTargetFrameworkAttribute>
+    <Configurations>Debug;Release;Development</Configurations>
+    <ApplicationManifest>app.manifest</ApplicationManifest>
+    <RootNamespace>AutomationToolLauncher</RootNamespace>
+    <AssemblyName>AutomationToolLauncher</AssemblyName>
+    <WarningsNotAsErrors>612,618</WarningsNotAsErrors>
+    <OutputPath>..\..\..\Binaries\DotNET\AutomationToolLauncher</OutputPath>
+    <DefineConstants>NET_CORE</DefineConstants>
+  </PropertyGroup>
+  <PropertyGroup Condition=" '$(Configuration)|$(Platform)' == 'Development|AnyCPU' ">
+    <DefineConstants>$(DefineConstants);TRACE</DefineConstants>
+    <Optimize>true</Optimize>
+  </PropertyGroup>
+  <PropertyGroup Condition=" '$(Configuration)|$(Platform)' == 'Release|AnyCPU' ">
+    <DefineConstants>$(DefineConstants);TRACE</DefineConstants>
+    <Optimize>true</Optimize>
+  </PropertyGroup>
+  <PropertyGroup Condition=" '$(Configuration)|$(Platform)' == 'Debug|AnyCPU' ">
+    <DefineConstants>$(DefineConstants);DEBUG;TRACE</DefineConstants>
+  </PropertyGroup>
+
+  <Import Project="Sdk.targets" Sdk="Microsoft.NET.Sdk" />
+
+  <ItemGroup>
+    <Compile Remove="obj\**" />
+    <EmbeddedResource Remove="obj\**" />
+    <None Remove="obj\**" />
+  </ItemGroup>
+
+  <ItemGroup>
+    <Compile Include="..\DotNETCommon\MetaData.cs">
+      <Link>Properties\MetaData.cs</Link>
+    </Compile>
+  </ItemGroup>
+</Project>
\ No newline at end of file
