UAT - Converted from system.web to use our json serialization (fastJson) and removed references to system.web that were not used.
This because we do not need this added dependency and it is not available by default in net core.

diff --git a/Engine/Source/Programs/AutomationTool/AutomationUtils/UE4Build.cs b/Engine/Source/Programs/AutomationTool/AutomationUtils/UE4Build.cs
--- a/Engine/Source/Programs/AutomationTool/AutomationUtils/UE4Build.cs
+++ b/Engine/Source/Programs/AutomationTool/AutomationUtils/UE4Build.cs
@@ -5,7 +5,6 @@ using System.Text;
 using System.IO;
 using System.Xml;
 using System.Diagnostics;
-using System.Web.Script.Serialization;
 using UnrealBuildTool;
 using Tools.DotNETCommon;
 using System.Linq;
@@ -1365,7 +1364,11 @@ namespace AutomationTool
 
 		public static string GetUBTExecutable()
 		{
+#if NET_CORE
+			return CommandUtils.CombinePaths(CommandUtils.CmdEnv.LocalRoot, @"Engine/Binaries/DotNET/UnrealBuildTool/UnrealBuildTool.exe");
+#else
 			return CommandUtils.CombinePaths(CommandUtils.CmdEnv.LocalRoot, @"Engine/Binaries/DotNET/UnrealBuildTool.exe");
+#endif
 		}
 
 		public string UBTExecutable
