UAT can now build as a netcore application.

Added a NET_CORE define to allow us to have changes side by side.
The AWS S3 changes are required due to us requiring to upgrade the S3 assembly version to get net core support (which made all methods async).
The ACL checks for files are not available in the system libraries of net core, as such the api is a bit different.

AutomationToolLauncher now just spawns a subprocess when used in netcore, as netcore does not support custom AppDomains and shadow copying. We will generally need to revisit this for netcore as this whole feature of building the source for UAT in UAT is not really possible.

To enable this set environment variable "UE_USE_DOTNET=1", note that with netcore all applications change their output path so this will likely break a bit of tooling when enabled.

diff --git a/Engine/Build/BatchFiles/RunUAT.bat b/Engine/Build/BatchFiles/RunUAT.bat
--- a/Engine/Build/BatchFiles/RunUAT.bat
+++ b/Engine/Build/BatchFiles/RunUAT.bat
@@ -11,6 +11,10 @@ echo Running AutomationTool...
 
 set UATExecutable=AutomationToolLauncher.exe
 set UATDirectory=Binaries\DotNET\
+if "%UE_USE_DOTNET%" == "1" (
+	set UATDirectory=Binaries\DotNET\AutomationToolLauncher
+)
+
 set UATCompileArg=-compile
 
 rem ## Change the CWD to /Engine. 
@@ -31,6 +35,11 @@ if not exist Source\Programs\AutomationTool\AutomationTool.csproj goto RunPrecom
 if not exist Source\Programs\AutomationToolLauncher\AutomationToolLauncher.csproj goto RunPrecompiled
 
 rem ## Get the path to MSBuild
+if "%UE_USE_DOTNET%" == "1" (
+	goto BuildNetCore
+)
+
+:BuildNetFW
 call "%~dp0GetMSBuildPath.bat"
 if errorlevel 1 goto RunPrecompiled
 %MSBUILD_EXE% /nologo /verbosity:quiet Source\Programs\AutomationToolLauncher\AutomationToolLauncher.csproj /property:Configuration=Development /property:Platform=AnyCPU
@@ -39,15 +48,36 @@ if errorlevel 1 goto Error_UATCompileFailed
 if errorlevel 1 goto Error_UATCompileFailed
 goto DoRunUAT
 
+:BuildNetCore
+rem ## Verify that dotnet is present
+call "%~dp0GetDotnetPath.bat"
+if errorlevel 1 goto Error_NoDotnetSDK
+
+dotnet build Source\Programs\AutomationToolLauncher\AutomationToolLauncherCore.csproj -c Development 
+if errorlevel 1 goto Error_UATCompileFailed
+dotnet build Source\Programs\AutomationTool\AutomationToolCore.csproj -c Development 
+if errorlevel 1 goto Error_UATCompileFailed
+goto DoRunUAT
 
 rem ## ok, well it doesn't look like visual studio is installed, let's try running the precompiled one.
 :RunPrecompiled
-if not exist Binaries\DotNET\AutomationTool.exe goto Error_NoFallbackExecutable
+
 set UATCompileArg=
-if not exist Binaries\DotNET\AutomationToolLauncher.exe set UATExecutable=AutomationTool.exe
+if "%UE_USE_DOTNET%" == "1" (
+	if not exist Binaries\DotNET\AutomationTool\AutomationTool.exe goto Error_NoFallbackExecutable
+	if not exist Binaries\DotNET\AutomationToolLauncher\AutomationToolLauncher.exe (
+		set UATExecutable=AutomationTool.exe
+		set UATDirectory=Binaries\DotNET\AutomationTool
+	) 
+) else (
+	if not exist Binaries\DotNET\AutomationTool.exe goto Error_NoFallbackExecutable
+	if not exist Binaries\DotNET\AutomationToolLauncher.exe set UATExecutable=AutomationTool.exe
+)
 goto DoRunUAT
 
 
+
+
 rem ## Run AutomationTool
 :DoRunUAT
 pushd %UATDirectory%
@@ -84,6 +114,11 @@ echo RunUAT.bat ERROR: A valid version of Visual Studio 2015 does not appear to
 set RUNUAT_EXITCODE=1
 goto Exit_Failure
 
+:Error_NoDotnetSDK
+echo RunUAT.bat ERROR: Unable to find a install of Dotnet SDK.  Please make sure you have it installed and that `dotnet` is a globally available command.
+set RUNUAT_EXITCODE=1
+goto Exit_Failure
+
 :Error_NoFallbackExecutable
 echo RunUAT.bat ERROR: Visual studio and/or AutomationTool.csproj was not found, nor was Engine\Binaries\DotNET\AutomationTool.exe. Can't run the automation tool.
 set RUNUAT_EXITCODE=1
diff --git a/Engine/Source/Programs/AutomationTool/AutomationUtils/CommandEnvironment.cs b/Engine/Source/Programs/AutomationTool/AutomationUtils/CommandEnvironment.cs
--- a/Engine/Source/Programs/AutomationTool/AutomationUtils/CommandEnvironment.cs
+++ b/Engine/Source/Programs/AutomationTool/AutomationUtils/CommandEnvironment.cs
@@ -81,7 +81,11 @@ namespace AutomationTool
 			LocalRoot = CommandUtils.GetEnvVar(EnvVarNames.LocalRoot);
 			if(String.IsNullOrEmpty(LocalRoot))
 			{
+#if NET_CORE
+				LocalRoot = CommandUtils.ConvertSeparators(PathSeparator.Slash, Path.GetFullPath(Path.Combine(Path.GetDirectoryName(UATExe), "..", "..", "..", "..")));
+#else
 				LocalRoot = CommandUtils.ConvertSeparators(PathSeparator.Slash, Path.GetFullPath(Path.Combine(Path.GetDirectoryName(UATExe), "..", "..", "..")));
+#endif
 				CommandUtils.ConditionallySetEnvVar(EnvVarNames.LocalRoot, LocalRoot);
 			}
 
diff --git a/Engine/Source/Programs/AutomationTool/AutomationUtils/CommandUtils.cs b/Engine/Source/Programs/AutomationTool/AutomationUtils/CommandUtils.cs
--- a/Engine/Source/Programs/AutomationTool/AutomationUtils/CommandUtils.cs
+++ b/Engine/Source/Programs/AutomationTool/AutomationUtils/CommandUtils.cs
@@ -1971,7 +1971,11 @@ namespace AutomationTool
 		/// <summary>
 		/// Path to the root directory
 		/// </summary>
+#if NET_CORE
+		public static readonly DirectoryReference RootDirectory = new DirectoryReference(Path.Combine(Path.GetDirectoryName(Assembly.GetExecutingAssembly().GetOriginalLocation()), "..", "..", "..", ".."));
+#else
 		public static readonly DirectoryReference RootDirectory = new DirectoryReference(Path.Combine(Path.GetDirectoryName(Assembly.GetExecutingAssembly().GetOriginalLocation()), "..", "..", ".."));
+#endif
 
 		/// <summary>
 		/// Path to the engine directory
diff --git a/Engine/Source/Programs/AutomationTool/BuildGraph/Tasks/UploadTask.cs b/Engine/Source/Programs/AutomationTool/BuildGraph/Tasks/UploadTask.cs
--- a/Engine/Source/Programs/AutomationTool/BuildGraph/Tasks/UploadTask.cs
+++ b/Engine/Source/Programs/AutomationTool/BuildGraph/Tasks/UploadTask.cs
@@ -9,7 +9,6 @@ using System.Net.Http;
 using System.Security.Cryptography;
 using System.Text;
 using System.Threading.Tasks;
-using System.Web.UI.HtmlControls;
 using System.Xml;
 using Tools.DotNETCommon;
 using Tools.DotNETCommon.Jupiter;
diff --git a/Engine/Source/Programs/AutomationTool/Gauntlet/Unreal/Base/Gauntlet.UnrealSession.cs b/Engine/Source/Programs/AutomationTool/Gauntlet/Unreal/Base/Gauntlet.UnrealSession.cs
--- a/Engine/Source/Programs/AutomationTool/Gauntlet/Unreal/Base/Gauntlet.UnrealSession.cs
+++ b/Engine/Source/Programs/AutomationTool/Gauntlet/Unreal/Base/Gauntlet.UnrealSession.cs
@@ -817,7 +817,8 @@ namespace Gauntlet
 				bool InstallSuccess = true;
 
 
-#if !__MonoCS__
+#if (!__MonoCS__) && (!NET_CORE)
+				// netcore: this uses windows forms which would require us to add a big dependency as such just skipping this
 				// count how many desktop clients
 				IEnumerable<UnrealSessionRole> DesktopClients = SessionRoles.Where(R => R.Platform == BuildHostPlatform.Current.Platform).Where(R => R.RoleType.IsClient());
 
diff --git a/Engine/Source/Programs/AutomationTool/HoloLens/HoloLensPlatform.Automation.cs b/Engine/Source/Programs/AutomationTool/HoloLens/HoloLensPlatform.Automation.cs
--- a/Engine/Source/Programs/AutomationTool/HoloLens/HoloLensPlatform.Automation.cs
+++ b/Engine/Source/Programs/AutomationTool/HoloLens/HoloLensPlatform.Automation.cs
@@ -13,7 +13,7 @@ using Tools.DotNETCommon;
 using System.Xml;
 using System.Text.RegularExpressions;
 
-#if !__MonoCS__
+#if !__MonoCS__ && !NET_CORE
 using Windows.Management.Deployment;
 using Windows.Foundation;
 #if false
@@ -196,7 +196,7 @@ namespace HoloLens.Automation
 		}
 	}
 
-#if !__MonoCS__
+#if !__MonoCS__ && !NET_CORE
 	class HoloLensDevicePortalCreatedProcess : IProcessResult
 	{
 		object StateLock;
@@ -1224,7 +1224,7 @@ namespace HoloLens.Automation
 			return !result;
 		}
 
-#if !__MonoCS__
+#if !__MonoCS__ && !NET_CORE
 		private void WaitFor(IAsyncOperationWithProgress<DeploymentResult, DeploymentProgress> deploymentOperation)
 		{
 			// This event is signaled when the operation completes
@@ -1262,7 +1262,7 @@ namespace HoloLens.Automation
 
 		private void DeployToLocalDevice(ProjectParams Params, DeploymentContext SC)
 		{
-#if !__MonoCS__
+#if !__MonoCS__ && !NET_CORE
             if (Utils.IsRunningOnMono)
             {
                 return;
@@ -1564,7 +1564,7 @@ namespace HoloLens.Automation
 
 		private IProcessResult RunUsingDevicePortal(string DeviceAddress, ERunOptions ClientRunFlags, string ClientApp, string ClientCmdLine, ProjectParams Params)
 		{
-#if !__MonoCS__
+#if !__MonoCS__ && !NET_CORE
             if (Utils.IsRunningOnMono)
             {
                 return null;
diff --git a/Engine/Source/Programs/AutomationTool/Program.cs b/Engine/Source/Programs/AutomationTool/Program.cs
--- a/Engine/Source/Programs/AutomationTool/Program.cs
+++ b/Engine/Source/Programs/AutomationTool/Program.cs
@@ -52,7 +52,13 @@ namespace AutomationTool
 			try
 			{
 				// Set the working directory to the UE4 root
+#if NET_CORE
+				// net core projects our output to a subfolder in the binaries path e.g. Engine/Binaries/Dotnet/AutomationTool
+				Environment.CurrentDirectory = Path.GetFullPath(Path.Combine(Path.GetDirectoryName(Assembly.GetExecutingAssembly().GetOriginalLocation()), "..", "..", "..", ".."));
+#else
 				Environment.CurrentDirectory = Path.GetFullPath(Path.Combine(Path.GetDirectoryName(Assembly.GetExecutingAssembly().GetOriginalLocation()), "..", "..", ".."));
+#endif
+
 
 				// Ensure we can resolve any external assemblies as necessary.
 				string PathToBinariesDotNET = Path.GetDirectoryName(Assembly.GetEntryAssembly().GetOriginalLocation());
@@ -60,10 +66,12 @@ namespace AutomationTool
 				AssemblyUtils.InstallRecursiveAssemblyResolver(PathToBinariesDotNET);
 
 				// Ensure that any third-party libraries marked as CopyLocal=false have their folders added as well (if the resolver can't be done locally where the class is used)
+#if !NET_CORE
+				// for net core we just copy the dependent assemblies to their output location
 				string PathToBinariesThirdParty = Path.Combine(PathToBinariesDotNET, "..", "ThirdParty");
 				AssemblyUtils.InstallRecursiveAssemblyResolver(Path.Combine(PathToBinariesThirdParty, "Google"));
 				AssemblyUtils.InstallRecursiveAssemblyResolver(Path.Combine(PathToBinariesThirdParty, "AWSSDK"));
-
+#endif
 				// Initialize the host platform layer
 				HostPlatform.Initialize();
 
diff --git a/Engine/Source/Programs/AutomationTool/Scripts/UploadDDCToAWS.cs b/Engine/Source/Programs/AutomationTool/Scripts/UploadDDCToAWS.cs
--- a/Engine/Source/Programs/AutomationTool/Scripts/UploadDDCToAWS.cs
+++ b/Engine/Source/Programs/AutomationTool/Scripts/UploadDDCToAWS.cs
@@ -602,6 +602,16 @@ namespace AutomationTool
 				Request.BucketName = BucketName;
 				Request.Key = Key;
 
+#if NET_CORE
+				using (GetObjectResponse Response = Client.GetObjectAsync(Request).Result)
+				{
+					if ((int)Response.HttpStatusCode >= 200 && (int)Response.HttpStatusCode <= 299)
+					{
+						Response.WriteResponseStreamToFileAsync(OutputFile.FullName, false, CancellationToken.None).Wait();
+						return true;
+					}
+				}
+#else
 				using (GetObjectResponse Response = Client.GetObject(Request))
 				{
 					if ((int)Response.HttpStatusCode >= 200 && (int)Response.HttpStatusCode <= 299)
@@ -610,6 +620,7 @@ namespace AutomationTool
 						return true;
 					}
 				}
+#endif
 				return false;
 			}
 			catch (Exception Ex)
@@ -724,7 +735,11 @@ namespace AutomationTool
 					Request.Key = Key;
 					Request.FilePath = File.FullName;
 					Request.CannedACL = S3CannedACL.NoACL;
+#if NET_CORE
+					Client.PutObjectAsync(Request).Wait();
+#else
 					Client.PutObject(Request);
+#endif
 				}
 
 				if (State != null)
@@ -760,8 +775,11 @@ namespace AutomationTool
 				DeleteObjectRequest Request = new DeleteObjectRequest();
 				Request.BucketName = BucketName;
 				Request.Key = Key;
+#if NET_CORE
+				Client.DeleteObjectAsync(Request).Wait();
+#else
 				Client.DeleteObject(Request);
-
+#endif
 				if (State != null)
 				{
 					Interlocked.Increment(ref State.NumFiles);
diff --git a/Engine/Source/Programs/AutomationTool/Turnkey/CopyProvider.cs b/Engine/Source/Programs/AutomationTool/Turnkey/CopyProvider.cs
--- a/Engine/Source/Programs/AutomationTool/Turnkey/CopyProvider.cs
+++ b/Engine/Source/Programs/AutomationTool/Turnkey/CopyProvider.cs
@@ -10,7 +10,6 @@ using System.IO;
 using AutomationTool;
 using Ionic.Zip;
 using UnrealBuildTool;
-using System.Windows.Forms;
 
 namespace Turnkey
 {
diff --git a/Engine/Source/Programs/AutomationTool/Turnkey/Providers/HybridIOProvider.cs b/Engine/Source/Programs/AutomationTool/Turnkey/Providers/HybridIOProvider.cs
--- a/Engine/Source/Programs/AutomationTool/Turnkey/Providers/HybridIOProvider.cs
+++ b/Engine/Source/Programs/AutomationTool/Turnkey/Providers/HybridIOProvider.cs
@@ -5,11 +5,16 @@ using System.Collections.Generic;
 using System.Linq;
 using System.Text;
 using System.Threading.Tasks;
-using System.Windows.Forms;
 using System.Text.RegularExpressions;
 using System.Drawing;
 using Tools.DotNETCommon;
 
+#if NET_CORE && WINDOWS
+using System.Windows.Forms;
+#elif !NET_CORE
+using System.Windows.Forms;
+#endif
+
 namespace Turnkey
 {
 	class HybridIOProvider : ConsoleIOProvider
@@ -17,7 +22,6 @@ namespace Turnkey
 		// @todo turnkey: if this is really bad on Mono, then we can set it in the manifest, but that means changing AutomationTool's and ATLaunchers' manifests
 		[System.Runtime.InteropServices.DllImport("user32.dll")]
 		private static extern bool SetProcessDPIAware();
-
 		static HybridIOProvider()
 		{
 			// make the form look good on modern displays!
@@ -26,8 +30,16 @@ namespace Turnkey
 				SetProcessDPIAware();
 			}
 
+#if NET_CORE
+#if WINDOWS
+			Application.EnableVisualStyles();
+			Application.SetCompatibleTextRenderingDefault(false);
+#endif
+#else
 			Application.EnableVisualStyles();
 			Application.SetCompatibleTextRenderingDefault(false);
+#endif
+
 		}
 
 		private string ShowMacDialog(string Prompt, string Default)
@@ -52,9 +64,71 @@ namespace Turnkey
 			return Match.Groups[1].Value;
 		}
 
+#if NET_CORE
+#if WINDOWS
+		private string ShowFormsDialogs(string Prompt, string Default)
+		{
+			// TODO: Make sure this dialog actually looks like we expect
+			Size size = new Size(200, 70);
+			Form inputBox = new Form();
+
+			inputBox.FormBorderStyle = FormBorderStyle.FixedDialog;
+			inputBox.ClientSize = size;
+			inputBox.Text = "Turnkey Input";
+
+			TextBox textBox = new TextBox();
+			textBox.Size = new Size(size.Width - 10, 23);
+			textBox.Location = new Point(5, 5);
+			textBox.Text = Prompt;
+			inputBox.Controls.Add(textBox);
+
+			Button okButton = new Button();
+			okButton.DialogResult = DialogResult.OK;
+			okButton.Name = "okButton";
+			okButton.Size = new Size(75, 23);
+			okButton.Text = "&OK";
+			okButton.Location = new Point(size.Width - 80 - 80, 39);
+			inputBox.Controls.Add(okButton);
+
+			Button cancelButton = new Button();
+			cancelButton.DialogResult = DialogResult.Cancel;
+			cancelButton.Name = "cancelButton";
+			cancelButton.Size = new Size(75, 23);
+			cancelButton.Text = "&Cancel";
+			cancelButton.Location = new Point(size.Width - 80, 39);
+			inputBox.Controls.Add(cancelButton);
+
+			inputBox.AcceptButton = okButton;
+			inputBox.CancelButton = cancelButton;
+
+			DialogResult result = inputBox.ShowDialog();
+			if (result == DialogResult.OK)
+			{
+				return textBox.Text;
+			}
+			else
+			{
+				return Default;
+			}
+		}
+#endif
+#endif
+
 		private string ShowDialog(string Prompt, string Default, bool bIsList)
 		{
+// disable unreachable code warning as this fails on linux builds due to this method not being implemented
+#pragma warning disable 0162
 			string Result;
+#if NET_CORE
+#if WINDOWS
+			Result = ShowFormsDialogs(Prompt, Default);
+#elif  OSX
+			Result = ShowMacDialog(Prompt, Default);
+#else
+			throw new NotImplementedException("Linux dialog not implemented");
+#endif
+
+#else // !NET_CORE
 			if (UnrealBuildTool.BuildHostPlatform.Current.Platform == UnrealBuildTool.UnrealTargetPlatform.Mac)
 			{
 				Result = ShowMacDialog(Prompt, Default);
@@ -63,6 +137,7 @@ namespace Turnkey
 			{
 				Result = Microsoft.VisualBasic.Interaction.InputBox(Prompt, "Turnkey Input", Default);
 			}
+#endif
 
 			if (string.IsNullOrEmpty(Result) && bIsList)
 			{
@@ -70,6 +145,7 @@ namespace Turnkey
 			}
 
 			return Result;
+#pragma warning restore
 /*
 			// Create a new instance of the form.
 			Form Form1 = new Form();
diff --git a/Engine/Source/Programs/AutomationTool/XLocLocalization/XLocLocalizationProvider.cs b/Engine/Source/Programs/AutomationTool/XLocLocalization/XLocLocalizationProvider.cs
--- a/Engine/Source/Programs/AutomationTool/XLocLocalization/XLocLocalizationProvider.cs
+++ b/Engine/Source/Programs/AutomationTool/XLocLocalization/XLocLocalizationProvider.cs
@@ -317,11 +317,13 @@ namespace EpicGames.XLocLocalization
 			Binding.SendTimeout = new TimeSpan(0, 1, 0);
 			Binding.AllowCookies = false;
 			Binding.BypassProxyOnLocal = false;
+#if !NET_CORE
 			Binding.HostNameComparisonMode = HostNameComparisonMode.StrongWildcard;
+			Binding.MessageEncoding = WSMessageEncoding.Text;
+#endif
 			Binding.MaxBufferSize = 2147483647;
 			Binding.MaxBufferPoolSize = 5242880;
 			Binding.MaxReceivedMessageSize = 2147483647;
-			Binding.MessageEncoding = WSMessageEncoding.Text;
 			Binding.TextEncoding = Encoding.UTF8;
 			Binding.TransferMode = TransferMode.Buffered;
 			Binding.UseDefaultWebProxy = true;
diff --git a/Engine/Source/Programs/AutomationToolLauncher/Launcher.cs b/Engine/Source/Programs/AutomationToolLauncher/Launcher.cs
--- a/Engine/Source/Programs/AutomationToolLauncher/Launcher.cs
+++ b/Engine/Source/Programs/AutomationToolLauncher/Launcher.cs
@@ -12,15 +12,17 @@ namespace AutomationToolLauncher
 	{
 		static int Main(string[] Arguments)
 		{
+			// net core does not support shadow copying so we just have to run the executable
+#if !NET_CORE
 			if (Arguments.Contains("-compile", StringComparer.OrdinalIgnoreCase))
 			{
 				return RunInAppDomain(Arguments);
 			}
-
+#endif
 			return Run(Arguments);
 
 		}
-
+#if !NET_CORE
 		static int RunInAppDomain(string[] Arguments)
 		{
 			// Create application domain setup information.
@@ -52,12 +54,17 @@ namespace AutomationToolLauncher
 
 			return ExitCode;
 		}
+#endif
 
 		static int Run(string[] Arguments)
 		{
 
 			string ApplicationBase = Path.GetDirectoryName(Assembly.GetExecutingAssembly().Location);
+#if NET_CORE
+			string UATExecutable = Path.Combine(ApplicationBase, "..\\AutomationTool", "AutomationTool.exe");
+#else
 			string UATExecutable = Path.Combine(ApplicationBase, "AutomationTool.exe");
+#endif
 
 			if (!File.Exists(UATExecutable))
 			{
@@ -67,8 +74,19 @@ namespace AutomationToolLauncher
 
 			try
 			{
+#if NET_CORE
+				ProcessStartInfo StartInfo = new ProcessStartInfo(UATExecutable);
+				foreach (string s in Arguments)
+				{
+					StartInfo.ArgumentList.Add(s);
+				}
+				Process uatProcess = Process.Start(StartInfo);
+				uatProcess.WaitForExit();
+				Environment.Exit(uatProcess.ExitCode);
+#else
 				Assembly UAT = Assembly.LoadFile(UATExecutable);
 				Environment.Exit((int) UAT.EntryPoint.Invoke(null, new object[] { Arguments }));
+#endif
 			}
 			catch (Exception Ex)
 			{
