Reenabled reading Gautlet.target files for netcore projects as the COM references that caused this to fail have been addressed.

diff --git a/Engine/Source/Programs/AutomationTool/Gauntlet/Gauntlet.Automation.Core.csproj b/Engine/Source/Programs/AutomationTool/Gauntlet/Gauntlet.Automation.Core.csproj
--- a/Engine/Source/Programs/AutomationTool/Gauntlet/Gauntlet.Automation.Core.csproj
+++ b/Engine/Source/Programs/AutomationTool/Gauntlet/Gauntlet.Automation.Core.csproj
@@ -56,10 +56,8 @@
       <Link>Platform/$([System.Text.RegularExpressions.Regex]::Replace(%(Compile.RecursiveDir), [\\/].+$, ``))/_</Link>
     </Compile>
   </ItemGroup>
-  
-  <!-- TODO: Disabled to avoid com references which are unsupported in net core 
+
   <Import Project="../../../../Platforms/*/Source/Programs/AutomationTool/Gauntlet/*.Gauntlet.targets" Condition="'$(OS)' == 'Windows_NT'" />
-  -->
 
   <ItemGroup>
     <ProjectReference Include="..\..\DotNETCommon\DotNETUtilities\DotNETUtilitiesCore.csproj" />
