Moved logic for finding automation tool projects into msbuild, making the generated .reference file from UBT redundant.
This makes it possible to build UAT with these references without having to run UBT first. Thus also fixes net core building of the project.

This adds about 5 seconds to the build time of automation tool but eliminates the need to run UBT and should avoid the need to implement -compile in UAT (in the end saving more then these 5 seconds)

This essentially does the same as UBT did before, create a .reference file on disk that can then be read by msbuild, it just does this in msbuild as well so we have fewer components involved.

diff --git a/Engine/Source/Programs/AutomationTool/AutomationToolCore.csproj b/Engine/Source/Programs/AutomationTool/AutomationToolCore.csproj
--- a/Engine/Source/Programs/AutomationTool/AutomationToolCore.csproj
+++ b/Engine/Source/Programs/AutomationTool/AutomationToolCore.csproj
@@ -22,15 +22,15 @@
     <DefineConstants>NET_CORE</DefineConstants>
   </PropertyGroup>
   <PropertyGroup Condition=" '$(Configuration)|$(Platform)' == 'Development|AnyCPU' ">
-    <DefineConstants>NET_CORE;TRACE</DefineConstants>
+    <DefineConstants>$(DefineConstants);TRACE</DefineConstants>
     <Optimize>true</Optimize>
   </PropertyGroup>
   <PropertyGroup Condition=" '$(Configuration)|$(Platform)' == 'Release|AnyCPU' ">
-    <DefineConstants>NET_CORE;TRACE</DefineConstants>
+    <DefineConstants>$(DefineConstants);TRACE</DefineConstants>
     <Optimize>true</Optimize>
   </PropertyGroup>
   <PropertyGroup Condition=" '$(Configuration)|$(Platform)' == 'Debug|AnyCPU' ">
-    <DefineConstants>NET_CORE;DEBUG;TRACE</DefineConstants>
+    <DefineConstants>$(DefineConstants);DEBUG;TRACE</DefineConstants>
   </PropertyGroup>
   <ItemGroup>
     <Reference Include="Ionic.Zip.Reduced">
@@ -57,12 +57,14 @@
     </None>
   </ItemGroup>
   
-  <!--     <None Include="app.config" /> -->
   <ItemGroup>
     <ProjectReference Include="..\DotNETCommon\DotNETUtilities\DotNETUtilitiesCore.csproj" />
     <ProjectReference Include="..\UnrealBuildTool\UnrealBuildToolCore.csproj" />
     <ProjectReference Include="AutomationUtils\AutomationUtils.Automation.Core.csproj" />
   </ItemGroup>
 
+  <!-- If building inside visual studio it means we generate projects, thus UBT will have output this file and we do not need to add the relativley slow task of globbing this in msbuild -->
+  <Import Project="AutomationTool.proj" Condition="'$(AutomationToolProjectOnly)' == '' AND '$(BuildingInsideVisualStudio)' != 'true'" />
   <Import Project="AutomationToolCore.csproj.References" Condition="Exists('AutomationToolCore.csproj.References') And '$(AutomationToolProjectOnly)' == ''" />
+  
 </Project>
\ No newline at end of file
diff --git a/Engine/Source/Programs/UnrealBuildTool/ProjectFiles/ProjectFileGenerator.cs b/Engine/Source/Programs/UnrealBuildTool/ProjectFiles/ProjectFileGenerator.cs
--- a/Engine/Source/Programs/UnrealBuildTool/ProjectFiles/ProjectFileGenerator.cs
+++ b/Engine/Source/Programs/UnrealBuildTool/ProjectFiles/ProjectFileGenerator.cs
@@ -2638,25 +2638,6 @@ namespace UnrealBuildTool
 				Progress.Write(TotalProjectFileCount, TotalProjectFileCount);
 			}
 
-            // Update the dotnet reference files so that the tooling knows which projects to build
-            // TODO: Should UAT build like this for net core or should it use custom logic to determine which projects to build?
-            if (AutomationProjectFiles.Any())
-            {
-	            XNamespace NS = XNamespace.Get("http://schemas.microsoft.com/developer/msbuild/2003");
-
-	            DirectoryReference AutomationToolDir = DirectoryReference.Combine(UnrealBuildTool.EngineSourceDirectory, "Programs", "AutomationTool");
-	            new XDocument(
-		            new XElement(NS + "Project",
-			            new XAttribute("DefaultTargets", "Build"),
-			            new XElement(NS + "ItemGroup",
-				            from AutomationProject in AutomationProjectFiles
-				            select new XElement(NS + "ProjectReference",
-					            new XAttribute("Include", AutomationProject.ProjectFilePath.MakeRelativeTo(AutomationToolDir))
-				            )
-			            )
-		            )
-	            ).Save(FileReference.Combine(AutomationToolDir, "AutomationToolCore.csproj.References").FullName);
-            }
 			return true;
 		}
 
