UAT Create a executable on Mac under netcore to avoid needing to boostrap via dotnet run

diff --git a/Engine/Source/Programs/AutomationTool/AutomationToolCore.csproj b/Engine/Source/Programs/AutomationTool/AutomationToolCore.csproj
--- a/Engine/Source/Programs/AutomationTool/AutomationToolCore.csproj
+++ b/Engine/Source/Programs/AutomationTool/AutomationToolCore.csproj
@@ -20,6 +20,7 @@
     <OutputPath>..\..\..\Binaries\DotNET\AutomationTool</OutputPath>
     <EnableDefaultItems>false</EnableDefaultItems>
     <DefineConstants>NET_CORE</DefineConstants>
+    <AppendRuntimeIdentifierToOutputPath>false</AppendRuntimeIdentifierToOutputPath>
   </PropertyGroup>
   <PropertyGroup Condition=" '$(Configuration)|$(Platform)' == 'Development|AnyCPU' ">
     <DefineConstants>$(DefineConstants);TRACE</DefineConstants>
@@ -32,6 +33,14 @@
   <PropertyGroup Condition=" '$(Configuration)|$(Platform)' == 'Debug|AnyCPU' ">
     <DefineConstants>$(DefineConstants);DEBUG;TRACE</DefineConstants>
   </PropertyGroup>
+  <Choose>
+    <When Condition="$([MSBuild]::IsOSPlatform('OSX'))">
+      <PropertyGroup>
+        <!-- Set runtime identifier for osx so that it produces a app host executable (consistency with the other platforms) -->
+        <RuntimeIdentifier>osx-x64</RuntimeIdentifier>
+      </PropertyGroup>
+    </When>
+  </Choose>
   <ItemGroup>
     <Reference Include="Ionic.Zip.Reduced">
       <HintPath>..\..\..\Binaries\DotNET\Ionic.Zip.Reduced.dll</HintPath>
