Disabled copy local on project references for netcore projects, fixes build warnigns caused by multiple copies trying to write to the same output files.

Also reduces footprint on disk and speedsup builds slightly.

Also took a pass over all projects to fixup their output paths as there were some inconsitencies from my port to netcore.
Lastly I disabled the resource folders for everything but english to reduce some of the noise in the output folder, this just impacts which languages you get code analysis.

diff --git a/Engine/Source/Programs/AutomationTool/Android/Android.Automation.Core.csproj b/Engine/Source/Programs/AutomationTool/Android/Android.Automation.Core.csproj
--- a/Engine/Source/Programs/AutomationTool/Android/Android.Automation.Core.csproj
+++ b/Engine/Source/Programs/AutomationTool/Android/Android.Automation.Core.csproj
@@ -17,7 +17,7 @@
     <RootNamespace>Android.Automation</RootNamespace>
     <AssemblyName>Android.Automation</AssemblyName>
     <WarningsNotAsErrors>612,618</WarningsNotAsErrors>
-    <OutputPath>..\..\..\..\Binaries\DotNET\AutomationTool\AutomationScripts\</OutputPath>
+    <OutputPath>..\..\..\..\Binaries\DotNET\AutomationTool\AutomationScripts\Android</OutputPath>
     <DefineConstants>NET_CORE</DefineConstants>
     <TreatWarningsAsErrors>true</TreatWarningsAsErrors>
   </PropertyGroup>
@@ -53,10 +53,10 @@
     </Compile>
   </ItemGroup>
   <ItemGroup>
-    <ProjectReference Include="..\..\DotNETCommon\DotNETUtilities\DotNETUtilitiesCore.csproj" />
-    <ProjectReference Include="..\..\UnrealBuildTool\UnrealBuildToolCore.csproj" />
-    <ProjectReference Include="..\AutomationUtils\AutomationUtils.Automation.Core.csproj" />
-    <ProjectReference Include="..\Scripts\AutomationScripts.Automation.Core.csproj" />
+    <ProjectReference Include="..\..\DotNETCommon\DotNETUtilities\DotNETUtilitiesCore.csproj"><Private>false</Private></ProjectReference>
+    <ProjectReference Include="..\..\UnrealBuildTool\UnrealBuildToolCore.csproj"><Private>false</Private></ProjectReference>
+    <ProjectReference Include="..\AutomationUtils\AutomationUtils.Automation.Core.csproj"><Private>false</Private></ProjectReference>
+    <ProjectReference Include="..\Scripts\AutomationScripts.Automation.Core.csproj"><Private>false</Private></ProjectReference>
   </ItemGroup>
   <ItemGroup>
     <PackageReference Include="Microsoft.CSharp" Version="4.7.0" />
diff --git a/Engine/Source/Programs/AutomationTool/AutomationToolCore.csproj b/Engine/Source/Programs/AutomationTool/AutomationToolCore.csproj
--- a/Engine/Source/Programs/AutomationTool/AutomationToolCore.csproj
+++ b/Engine/Source/Programs/AutomationTool/AutomationToolCore.csproj
@@ -21,6 +21,7 @@
     <EnableDefaultItems>false</EnableDefaultItems>
     <DefineConstants>NET_CORE</DefineConstants>
     <AppendRuntimeIdentifierToOutputPath>false</AppendRuntimeIdentifierToOutputPath>
+    <SatelliteResourceLanguages>en</SatelliteResourceLanguages> <!-- remove non english resource languages -->
     <IsWindows Condition="'$([System.Runtime.InteropServices.RuntimeInformation]::IsOSPlatform($([System.Runtime.InteropServices.OSPlatform]::Windows)))' == 'true'">true</IsWindows> 
     <IsOSX Condition="'$([System.Runtime.InteropServices.RuntimeInformation]::IsOSPlatform($([System.Runtime.InteropServices.OSPlatform]::OSX)))' == 'true'">true</IsOSX> 
     <IsLinux Condition="'$([System.Runtime.InteropServices.RuntimeInformation]::IsOSPlatform($([System.Runtime.InteropServices.OSPlatform]::Linux)))' == 'true'">true</IsLinux> 
diff --git a/Engine/Source/Programs/AutomationTool/AutomationUtils/AutomationUtils.Automation.Core.csproj b/Engine/Source/Programs/AutomationTool/AutomationUtils/AutomationUtils.Automation.Core.csproj
--- a/Engine/Source/Programs/AutomationTool/AutomationUtils/AutomationUtils.Automation.Core.csproj
+++ b/Engine/Source/Programs/AutomationTool/AutomationUtils/AutomationUtils.Automation.Core.csproj
@@ -49,8 +49,8 @@
   </ItemGroup>
 
   <ItemGroup>
-    <ProjectReference Include="..\..\DotNETCommon\DotNETUtilities\DotNETUtilitiesCore.csproj" />
-    <ProjectReference Include="..\..\UnrealBuildTool\UnrealBuildToolCore.csproj" />
+    <ProjectReference Include="..\..\DotNETCommon\DotNETUtilities\DotNETUtilitiesCore.csproj"><Private>false</Private></ProjectReference>
+    <ProjectReference Include="..\..\UnrealBuildTool\UnrealBuildToolCore.csproj"><Private>false</Private></ProjectReference>
   </ItemGroup>
   <ItemGroup>
     <PackageReference Include="Microsoft.CSharp" Version="4.7.0" />
diff --git a/Engine/Source/Programs/AutomationTool/BuildGraph/BuildGraph.Automation.Core.csproj b/Engine/Source/Programs/AutomationTool/BuildGraph/BuildGraph.Automation.Core.csproj
--- a/Engine/Source/Programs/AutomationTool/BuildGraph/BuildGraph.Automation.Core.csproj
+++ b/Engine/Source/Programs/AutomationTool/BuildGraph/BuildGraph.Automation.Core.csproj
@@ -16,7 +16,7 @@
     <RootNamespace>AutomationTool</RootNamespace>
     <AssemblyName>BuildGraph.Automation</AssemblyName>
     <WarningsNotAsErrors>612,618</WarningsNotAsErrors>
-    <OutputPath>..\..\..\..\Binaries\DotNET\AutomationTool</OutputPath>
+    <OutputPath>..\..\..\..\Binaries\DotNET\AutomationTool\AutomationScripts</OutputPath>
     <DefineConstants>NET_CORE</DefineConstants>
     <TreatWarningsAsErrors>true</TreatWarningsAsErrors>
   </PropertyGroup>
@@ -52,9 +52,9 @@
     </Compile>
   </ItemGroup>
   <ItemGroup>
-    <ProjectReference Include="..\..\DotNETCommon\DotNETUtilities\DotNETUtilitiesCore.csproj" />
-    <ProjectReference Include="..\..\UnrealBuildTool\UnrealBuildToolCore.csproj" />
-    <ProjectReference Include="..\AutomationUtils\AutomationUtils.Automation.Core.csproj" />
+    <ProjectReference Include="..\..\DotNETCommon\DotNETUtilities\DotNETUtilitiesCore.csproj"><Private>false</Private></ProjectReference>
+    <ProjectReference Include="..\..\UnrealBuildTool\UnrealBuildToolCore.csproj"><Private>false</Private></ProjectReference>
+    <ProjectReference Include="..\AutomationUtils\AutomationUtils.Automation.Core.csproj"><Private>false</Private></ProjectReference>
   </ItemGroup>
   <ItemGroup>
     <PackageReference Include="Microsoft.CSharp" Version="4.7.0" />
diff --git a/Engine/Source/Programs/AutomationTool/Gauntlet/Gauntlet.Automation.Core.csproj b/Engine/Source/Programs/AutomationTool/Gauntlet/Gauntlet.Automation.Core.csproj
--- a/Engine/Source/Programs/AutomationTool/Gauntlet/Gauntlet.Automation.Core.csproj
+++ b/Engine/Source/Programs/AutomationTool/Gauntlet/Gauntlet.Automation.Core.csproj
@@ -16,7 +16,7 @@
     <RootNamespace>Gauntlet.Automation</RootNamespace>
     <AssemblyName>Gauntlet.Automation</AssemblyName>
     <WarningsNotAsErrors>612,618</WarningsNotAsErrors>
-    <OutputPath>..\..\..\..\Binaries\DotNET\AutomationTool</OutputPath>
+    <OutputPath>..\..\..\..\Binaries\DotNET\AutomationTool\AutomationScripts</OutputPath>
     <DefineConstants>NET_CORE</DefineConstants>
     <TreatWarningsAsErrors>true</TreatWarningsAsErrors>
   </PropertyGroup>
@@ -60,15 +60,15 @@
   <Import Project="../../../../Platforms/*/Source/Programs/AutomationTool/Gauntlet/*.Gauntlet.targets" Condition="'$(OS)' == 'Windows_NT'" />
 
   <ItemGroup>
-    <ProjectReference Include="..\..\DotNETCommon\DotNETUtilities\DotNETUtilitiesCore.csproj" />
-    <ProjectReference Include="..\..\UnrealBuildTool\UnrealBuildToolCore.csproj" />
-    <ProjectReference Include="..\AutomationUtils\AutomationUtils.Automation.Core.csproj" />
-    <ProjectReference Include="..\Scripts\AutomationScripts.Automation.Core.csproj" />
-    <ProjectReference Include="..\OneSkyLocalization\OneSkyLocalization.Automation.Core.csproj" />
-    <ProjectReference Include="..\XLocLocalization\XLocLocalization.Automation.Core.csproj" />
-    <ProjectReference Include="..\Localization\Localization.Automation.Core.csproj" />
-    <ProjectReference Include="..\Lumin\Lumin.Automation.Core.csproj" />
-    <ProjectReference Include="..\Android\Android.Automation.Core.csproj" />
+    <ProjectReference Include="..\..\DotNETCommon\DotNETUtilities\DotNETUtilitiesCore.csproj"><Private>false</Private></ProjectReference>
+    <ProjectReference Include="..\..\UnrealBuildTool\UnrealBuildToolCore.csproj"><Private>false</Private></ProjectReference>
+    <ProjectReference Include="..\AutomationUtils\AutomationUtils.Automation.Core.csproj"><Private>false</Private></ProjectReference>
+    <ProjectReference Include="..\Scripts\AutomationScripts.Automation.Core.csproj"><Private>false</Private></ProjectReference>
+    <ProjectReference Include="..\OneSkyLocalization\OneSkyLocalization.Automation.Core.csproj"><Private>false</Private></ProjectReference>
+    <ProjectReference Include="..\XLocLocalization\XLocLocalization.Automation.Core.csproj"><Private>false</Private></ProjectReference>
+    <ProjectReference Include="..\Localization\Localization.Automation.Core.csproj"><Private>false</Private></ProjectReference>
+    <ProjectReference Include="..\Lumin\Lumin.Automation.Core.csproj"><Private>false</Private></ProjectReference>
+    <ProjectReference Include="..\Android\Android.Automation.Core.csproj"><Private>false</Private></ProjectReference>
   </ItemGroup>
   <ItemGroup>
     <PackageReference Include="Microsoft.CSharp" Version="4.7.0" />
diff --git a/Engine/Source/Programs/AutomationTool/HoloLens/HoloLens.Automation.Core.csproj b/Engine/Source/Programs/AutomationTool/HoloLens/HoloLens.Automation.Core.csproj
--- a/Engine/Source/Programs/AutomationTool/HoloLens/HoloLens.Automation.Core.csproj
+++ b/Engine/Source/Programs/AutomationTool/HoloLens/HoloLens.Automation.Core.csproj
@@ -16,7 +16,7 @@
     <RootNamespace>HoloLens.Automation</RootNamespace>
     <AssemblyName>HoloLens.Automation</AssemblyName>
     <WarningsNotAsErrors>612,618</WarningsNotAsErrors>
-    <OutputPath>..\..\..\..\Binaries\DotNET\AutomationTool</OutputPath>
+    <OutputPath>..\..\..\..\Binaries\DotNET\AutomationTool\AutomationScripts\HoloLens\</OutputPath>
     <DefineConstants>NET_CORE</DefineConstants>
     <TreatWarningsAsErrors>true</TreatWarningsAsErrors>
     <IsWindows Condition="'$([System.Runtime.InteropServices.RuntimeInformation]::IsOSPlatform($([System.Runtime.InteropServices.OSPlatform]::Windows)))' == 'true'">true</IsWindows> 
@@ -73,7 +73,6 @@
       </ItemGroup>
     </When>
   </Choose>
-
   <Import Project="Sdk.props" Sdk="Microsoft.NET.Sdk.WindowsDesktop" Condition="'$(IsWindows)'=='true'"/>
   <Import Project="Sdk.targets" Sdk="Microsoft.NET.Sdk.WindowsDesktop" Condition="'$(IsWindows)'=='true'"/>
   <Import Project="Sdk.props" Sdk="Microsoft.NET.Sdk" Condition="'$(IsWindows)'!='true'"/>
@@ -91,10 +90,10 @@
     </Compile>
   </ItemGroup>
   <ItemGroup>
-    <ProjectReference Include="..\..\DotNETCommon\DotNETUtilities\DotNETUtilitiesCore.csproj" />
-    <ProjectReference Include="..\..\UnrealBuildTool\UnrealBuildToolCore.csproj" />
-    <ProjectReference Include="..\AutomationUtils\AutomationUtils.Automation.Core.csproj" />
-    <ProjectReference Include="..\Scripts\AutomationScripts.Automation.Core.csproj" />
+    <ProjectReference Include="..\..\DotNETCommon\DotNETUtilities\DotNETUtilitiesCore.csproj"><Private>false</Private></ProjectReference>
+    <ProjectReference Include="..\..\UnrealBuildTool\UnrealBuildToolCore.csproj"><Private>false</Private></ProjectReference>
+    <ProjectReference Include="..\AutomationUtils\AutomationUtils.Automation.Core.csproj"><Private>false</Private></ProjectReference>
+    <ProjectReference Include="..\Scripts\AutomationScripts.Automation.Core.csproj"><Private>false</Private></ProjectReference>
   </ItemGroup>
   <ItemGroup>
     <PackageReference Include="Microsoft.CSharp" Version="4.7.0" />
diff --git a/Engine/Source/Programs/AutomationTool/IOS/IOS.Automation.Core.csproj b/Engine/Source/Programs/AutomationTool/IOS/IOS.Automation.Core.csproj
--- a/Engine/Source/Programs/AutomationTool/IOS/IOS.Automation.Core.csproj
+++ b/Engine/Source/Programs/AutomationTool/IOS/IOS.Automation.Core.csproj
@@ -16,7 +16,7 @@
     <RootNamespace>IOS.Automation</RootNamespace>
     <AssemblyName>IOS.Automation</AssemblyName>
     <WarningsNotAsErrors>612,618</WarningsNotAsErrors>
-    <OutputPath>..\..\..\..\Binaries\DotNET\AutomationTool</OutputPath>
+    <OutputPath>..\..\..\..\Binaries\DotNET\AutomationTool\AutomationScripts\IOS</OutputPath>
     <DefineConstants>NET_CORE</DefineConstants>
     <TreatWarningsAsErrors>true</TreatWarningsAsErrors>
   </PropertyGroup>
@@ -57,9 +57,9 @@
     </Compile>
   </ItemGroup>
   <ItemGroup>
-    <ProjectReference Include="..\..\DotNETCommon\DotNETUtilities\DotNETUtilitiesCore.csproj" />
-    <ProjectReference Include="..\..\UnrealBuildTool\UnrealBuildToolCore.csproj" />
-    <ProjectReference Include="..\AutomationUtils\AutomationUtils.Automation.Core.csproj" />
+    <ProjectReference Include="..\..\DotNETCommon\DotNETUtilities\DotNETUtilitiesCore.csproj"><Private>false</Private></ProjectReference>
+    <ProjectReference Include="..\..\UnrealBuildTool\UnrealBuildToolCore.csproj"><Private>false</Private></ProjectReference>
+    <ProjectReference Include="..\AutomationUtils\AutomationUtils.Automation.Core.csproj"><Private>false</Private></ProjectReference>
   </ItemGroup>
   <ItemGroup>
     <PackageReference Include="Microsoft.CSharp" Version="4.7.0" />
diff --git a/Engine/Source/Programs/AutomationTool/Linux/Linux.Automation.Core.csproj b/Engine/Source/Programs/AutomationTool/Linux/Linux.Automation.Core.csproj
--- a/Engine/Source/Programs/AutomationTool/Linux/Linux.Automation.Core.csproj
+++ b/Engine/Source/Programs/AutomationTool/Linux/Linux.Automation.Core.csproj
@@ -16,7 +16,7 @@
     <RootNamespace>Linux.Automation</RootNamespace>
     <AssemblyName>Linux.Automation</AssemblyName>
     <WarningsNotAsErrors>612,618</WarningsNotAsErrors>
-    <OutputPath>..\..\..\..\Binaries\DotNET\AutomationTool</OutputPath>
+    <OutputPath>..\..\..\..\Binaries\DotNET\AutomationTool\AutomationScripts\Linux\</OutputPath>
     <DefineConstants>NET_CORE</DefineConstants>
     <TreatWarningsAsErrors>true</TreatWarningsAsErrors>
   </PropertyGroup>
@@ -46,9 +46,9 @@
     </Compile>
   </ItemGroup>
   <ItemGroup>
-    <ProjectReference Include="..\..\DotNETCommon\DotNETUtilities\DotNETUtilitiesCore.csproj" />
-    <ProjectReference Include="..\..\UnrealBuildTool\UnrealBuildToolCore.csproj" />
-    <ProjectReference Include="..\AutomationUtils\AutomationUtils.Automation.Core.csproj" />
+    <ProjectReference Include="..\..\DotNETCommon\DotNETUtilities\DotNETUtilitiesCore.csproj"><Private>false</Private></ProjectReference>
+    <ProjectReference Include="..\..\UnrealBuildTool\UnrealBuildToolCore.csproj"><Private>false</Private></ProjectReference>
+    <ProjectReference Include="..\AutomationUtils\AutomationUtils.Automation.Core.csproj"><Private>false</Private></ProjectReference>
   </ItemGroup>
   <ItemGroup>
     <PackageReference Include="Microsoft.CSharp" Version="4.7.0" />
diff --git a/Engine/Source/Programs/AutomationTool/Localization/Localization.Automation.Core.csproj b/Engine/Source/Programs/AutomationTool/Localization/Localization.Automation.Core.csproj
--- a/Engine/Source/Programs/AutomationTool/Localization/Localization.Automation.Core.csproj
+++ b/Engine/Source/Programs/AutomationTool/Localization/Localization.Automation.Core.csproj
@@ -16,7 +16,7 @@
     <RootNamespace>AutomationTool</RootNamespace>
     <AssemblyName>Localization.Automation</AssemblyName>
     <WarningsNotAsErrors>612,618</WarningsNotAsErrors>
-    <OutputPath>..\..\..\..\Binaries\DotNET\AutomationTool</OutputPath>
+    <OutputPath>..\..\..\..\Binaries\DotNET\AutomationTool\AutomationScripts</OutputPath>
     <DefineConstants>NET_CORE</DefineConstants>
     <TreatWarningsAsErrors>true</TreatWarningsAsErrors>
   </PropertyGroup>
@@ -46,9 +46,9 @@
     </Compile>
   </ItemGroup>
   <ItemGroup>
-    <ProjectReference Include="..\..\DotNETCommon\DotNETUtilities\DotNETUtilitiesCore.csproj" />
-    <ProjectReference Include="..\..\UnrealBuildTool\UnrealBuildToolCore.csproj" />
-    <ProjectReference Include="..\AutomationUtils\AutomationUtils.Automation.Core.csproj" />
+    <ProjectReference Include="..\..\DotNETCommon\DotNETUtilities\DotNETUtilitiesCore.csproj"><Private>false</Private></ProjectReference>
+    <ProjectReference Include="..\..\UnrealBuildTool\UnrealBuildToolCore.csproj"><Private>false</Private></ProjectReference>
+    <ProjectReference Include="..\AutomationUtils\AutomationUtils.Automation.Core.csproj"><Private>false</Private></ProjectReference>
   </ItemGroup>
   <ItemGroup>
     <PackageReference Include="System.Data.DataSetExtensions" Version="4.5.0" />
diff --git a/Engine/Source/Programs/AutomationTool/Lumin/Lumin.Automation.Core.csproj b/Engine/Source/Programs/AutomationTool/Lumin/Lumin.Automation.Core.csproj
--- a/Engine/Source/Programs/AutomationTool/Lumin/Lumin.Automation.Core.csproj
+++ b/Engine/Source/Programs/AutomationTool/Lumin/Lumin.Automation.Core.csproj
@@ -16,7 +16,7 @@
     <RootNamespace>Lumin.Automation</RootNamespace>
     <AssemblyName>Lumin.Automation</AssemblyName>
     <WarningsNotAsErrors>612,618</WarningsNotAsErrors>
-    <OutputPath>..\..\..\..\Binaries\DotNET\AutomationTool</OutputPath>
+    <OutputPath>..\..\..\..\Binaries\DotNET\AutomationTool\AutomationScripts\Lumin\</OutputPath>
     <DefineConstants>NET_CORE</DefineConstants>
     <TreatWarningsAsErrors>true</TreatWarningsAsErrors>
   </PropertyGroup>
@@ -52,10 +52,10 @@
     </Compile>
   </ItemGroup>
   <ItemGroup>
-    <ProjectReference Include="..\..\DotNETCommon\DotNETUtilities\DotNETUtilitiesCore.csproj" />
-    <ProjectReference Include="..\..\UnrealBuildTool\UnrealBuildToolCore.csproj" />
-    <ProjectReference Include="..\AutomationUtils\AutomationUtils.Automation.Core.csproj" />
-    <ProjectReference Include="..\Scripts\AutomationScripts.Automation.Core.csproj" />
+    <ProjectReference Include="..\..\DotNETCommon\DotNETUtilities\DotNETUtilitiesCore.csproj"><Private>false</Private></ProjectReference>
+    <ProjectReference Include="..\..\UnrealBuildTool\UnrealBuildToolCore.csproj"><Private>false</Private></ProjectReference>
+    <ProjectReference Include="..\AutomationUtils\AutomationUtils.Automation.Core.csproj"><Private>false</Private></ProjectReference>
+    <ProjectReference Include="..\Scripts\AutomationScripts.Automation.Core.csproj"><Private>false</Private></ProjectReference>
   </ItemGroup>
   <ItemGroup>
     <PackageReference Include="Microsoft.CSharp" Version="4.7.0" />
diff --git a/Engine/Source/Programs/AutomationTool/Mac/Mac.Automation.Core.csproj b/Engine/Source/Programs/AutomationTool/Mac/Mac.Automation.Core.csproj
--- a/Engine/Source/Programs/AutomationTool/Mac/Mac.Automation.Core.csproj
+++ b/Engine/Source/Programs/AutomationTool/Mac/Mac.Automation.Core.csproj
@@ -16,7 +16,7 @@
     <RootNamespace>Mac.Automation</RootNamespace>
     <AssemblyName>Mac.Automation</AssemblyName>
     <WarningsNotAsErrors>612,618</WarningsNotAsErrors>
-    <OutputPath>..\..\..\..\Binaries\DotNET\AutomationTool</OutputPath>
+    <OutputPath>..\..\..\..\Binaries\DotNET\AutomationTool\AutomationScripts\Mac\</OutputPath>
     <DefineConstants>NET_CORE</DefineConstants>
     <TreatWarningsAsErrors>true</TreatWarningsAsErrors>
   </PropertyGroup>
@@ -46,9 +46,9 @@
     </Compile>
   </ItemGroup>
   <ItemGroup>
-    <ProjectReference Include="..\..\DotNETCommon\DotNETUtilities\DotNETUtilitiesCore.csproj" />
-    <ProjectReference Include="..\..\UnrealBuildTool\UnrealBuildToolCore.csproj" />
-    <ProjectReference Include="..\AutomationUtils\AutomationUtils.Automation.Core.csproj" />
+    <ProjectReference Include="..\..\DotNETCommon\DotNETUtilities\DotNETUtilitiesCore.csproj"><Private>false</Private></ProjectReference>
+    <ProjectReference Include="..\..\UnrealBuildTool\UnrealBuildToolCore.csproj"><Private>false</Private></ProjectReference>
+    <ProjectReference Include="..\AutomationUtils\AutomationUtils.Automation.Core.csproj"><Private>false</Private></ProjectReference>
   </ItemGroup>
   <ItemGroup>
     <PackageReference Include="Microsoft.CSharp" Version="4.7.0" />
diff --git a/Engine/Source/Programs/AutomationTool/OneSkyLocalization/OneSkyLocalization.Automation.Core.csproj b/Engine/Source/Programs/AutomationTool/OneSkyLocalization/OneSkyLocalization.Automation.Core.csproj
--- a/Engine/Source/Programs/AutomationTool/OneSkyLocalization/OneSkyLocalization.Automation.Core.csproj
+++ b/Engine/Source/Programs/AutomationTool/OneSkyLocalization/OneSkyLocalization.Automation.Core.csproj
@@ -16,7 +16,7 @@
     <RootNamespace>AutomationTool</RootNamespace>
     <AssemblyName>OneSkyLocalization.Automation</AssemblyName>
     <WarningsNotAsErrors>612,618</WarningsNotAsErrors>
-    <OutputPath>..\..\..\..\Binaries\DotNET\AutomationTool</OutputPath>
+    <OutputPath>..\..\..\..\Binaries\DotNET\AutomationTool\AutomationScripts</OutputPath>
     <DefineConstants>NET_CORE</DefineConstants>
     <TreatWarningsAsErrors>true</TreatWarningsAsErrors>
   </PropertyGroup>
@@ -53,10 +53,10 @@
     </Compile>
   </ItemGroup>
   <ItemGroup>
-    <ProjectReference Include="..\..\DotNETCommon\DotNETUtilities\DotNETUtilitiesCore.csproj" />
-    <ProjectReference Include="..\..\UnrealBuildTool\UnrealBuildToolCore.csproj" />
-    <ProjectReference Include="..\AutomationUtils\AutomationUtils.Automation.Core.csproj" />
-    <ProjectReference Include="..\Localization\Localization.Automation.Core.csproj" />
+    <ProjectReference Include="..\..\DotNETCommon\DotNETUtilities\DotNETUtilitiesCore.csproj"><Private>false</Private></ProjectReference>
+    <ProjectReference Include="..\..\UnrealBuildTool\UnrealBuildToolCore.csproj"><Private>false</Private></ProjectReference>
+    <ProjectReference Include="..\AutomationUtils\AutomationUtils.Automation.Core.csproj"><Private>false</Private></ProjectReference>
+    <ProjectReference Include="..\Localization\Localization.Automation.Core.csproj"><Private>false</Private></ProjectReference>
   </ItemGroup>
   <ItemGroup>
     <PackageReference Include="Microsoft.CSharp" Version="4.7.0" />
diff --git a/Engine/Source/Programs/AutomationTool/Scripts/AutomationScripts.Automation.Core.csproj b/Engine/Source/Programs/AutomationTool/Scripts/AutomationScripts.Automation.Core.csproj
--- a/Engine/Source/Programs/AutomationTool/Scripts/AutomationScripts.Automation.Core.csproj
+++ b/Engine/Source/Programs/AutomationTool/Scripts/AutomationScripts.Automation.Core.csproj
@@ -16,7 +16,7 @@
     <RootNamespace>AutomationTool</RootNamespace>
     <AssemblyName>AutomationScripts.Automation</AssemblyName>
     <WarningsNotAsErrors>612,618</WarningsNotAsErrors>
-    <OutputPath>..\..\..\..\Binaries\DotNET\AutomationTool</OutputPath>
+    <OutputPath>..\..\..\..\Binaries\DotNET\AutomationTool\AutomationScripts</OutputPath>
     <DefineConstants>NET_CORE</DefineConstants>
     <TreatWarningsAsErrors>true</TreatWarningsAsErrors>
   </PropertyGroup>
@@ -57,10 +57,10 @@
     </Compile>
   </ItemGroup>
   <ItemGroup>
-    <ProjectReference Include="..\..\DotNETCommon\DotNETUtilities\DotNETUtilitiesCore.csproj" />
-    <ProjectReference Include="..\..\UnrealBuildTool\UnrealBuildToolCore.csproj" />
-    <ProjectReference Include="..\AutomationUtils\AutomationUtils.Automation.Core.csproj" />
-    <ProjectReference Include="..\Localization\Localization.Automation.Core.csproj" />
+    <ProjectReference Include="..\..\DotNETCommon\DotNETUtilities\DotNETUtilitiesCore.csproj"><Private>false</Private></ProjectReference>
+    <ProjectReference Include="..\..\UnrealBuildTool\UnrealBuildToolCore.csproj"><Private>false</Private></ProjectReference>
+    <ProjectReference Include="..\AutomationUtils\AutomationUtils.Automation.Core.csproj"><Private>false</Private></ProjectReference>
+    <ProjectReference Include="..\Localization\Localization.Automation.Core.csproj"><Private>false</Private></ProjectReference>
   </ItemGroup>
   <ItemGroup>
     <PackageReference Include="Microsoft.CSharp" Version="4.7.0" />
diff --git a/Engine/Source/Programs/AutomationTool/TVOS/TVOS.Automation.Core.csproj b/Engine/Source/Programs/AutomationTool/TVOS/TVOS.Automation.Core.csproj
--- a/Engine/Source/Programs/AutomationTool/TVOS/TVOS.Automation.Core.csproj
+++ b/Engine/Source/Programs/AutomationTool/TVOS/TVOS.Automation.Core.csproj
@@ -16,7 +16,7 @@
     <RootNamespace>TVOS.Automation</RootNamespace>
     <AssemblyName>TVOS.Automation</AssemblyName>
     <WarningsNotAsErrors>612,618</WarningsNotAsErrors>
-    <OutputPath>..\..\..\..\Binaries\DotNET\AutomationTool</OutputPath>
+    <OutputPath>..\..\..\..\Binaries\DotNET\AutomationTool\AutomationScripts\TVOS\</OutputPath>
     <DefineConstants>NET_CORE</DefineConstants>
     <TreatWarningsAsErrors>true</TreatWarningsAsErrors>
   </PropertyGroup>
@@ -52,10 +52,10 @@
     </Compile>
   </ItemGroup>
   <ItemGroup>
-    <ProjectReference Include="..\..\DotNETCommon\DotNETUtilities\DotNETUtilitiesCore.csproj" />
-    <ProjectReference Include="..\..\UnrealBuildTool\UnrealBuildToolCore.csproj" />
-    <ProjectReference Include="..\AutomationUtils\AutomationUtils.Automation.Core.csproj" />
-    <ProjectReference Include="..\IOS\IOS.Automation.Core.csproj" />
+    <ProjectReference Include="..\..\DotNETCommon\DotNETUtilities\DotNETUtilitiesCore.csproj"><Private>false</Private></ProjectReference>
+    <ProjectReference Include="..\..\UnrealBuildTool\UnrealBuildToolCore.csproj"><Private>false</Private></ProjectReference>
+    <ProjectReference Include="..\AutomationUtils\AutomationUtils.Automation.Core.csproj"><Private>false</Private></ProjectReference>
+    <ProjectReference Include="..\IOS\IOS.Automation.Core.csproj"><Private>false</Private></ProjectReference>
   </ItemGroup>
   <ItemGroup>
     <PackageReference Include="Microsoft.CSharp" Version="4.7.0" />
diff --git a/Engine/Source/Programs/AutomationTool/Turnkey/Turnkey.Automation.Core.csproj b/Engine/Source/Programs/AutomationTool/Turnkey/Turnkey.Automation.Core.csproj
--- a/Engine/Source/Programs/AutomationTool/Turnkey/Turnkey.Automation.Core.csproj
+++ b/Engine/Source/Programs/AutomationTool/Turnkey/Turnkey.Automation.Core.csproj
@@ -16,7 +16,7 @@
     <RootNamespace>Turnkey.Automation</RootNamespace>
     <AssemblyName>Turnkey.Automation</AssemblyName>
     <WarningsNotAsErrors>612,618</WarningsNotAsErrors>
-    <OutputPath>..\..\..\..\Binaries\DotNET\AutomationTool</OutputPath>
+    <OutputPath>..\..\..\..\Binaries\DotNET\AutomationTool\AutomationScripts\Turnkey\</OutputPath>
     <DefineConstants>NET_CORE</DefineConstants>
     <TreatWarningsAsErrors>true</TreatWarningsAsErrors>
     <IsWindows Condition="'$([System.Runtime.InteropServices.RuntimeInformation]::IsOSPlatform($([System.Runtime.InteropServices.OSPlatform]::Windows)))' == 'true'">true</IsWindows> 
@@ -50,11 +50,6 @@
       <HintPath>..\..\..\..\Binaries\DotNET\Ionic.Zip.Reduced.dll</HintPath>
       <Private>False</Private>
     </Reference>
-    <!--<Reference Include="OneSky, Version=1.0.0.0, Culture=neutral, processorArchitecture=MSIL">
-      <HintPath>..\..\..\..\Binaries\DotNET\OneSky.dll</HintPath>
-      <SpecificVersion>False</SpecificVersion>
-      <Private>False</Private>
-    </Reference>-->
   </ItemGroup>
 
   <Import Project="Sdk.props" Sdk="Microsoft.NET.Sdk.WindowsDesktop" Condition="'$(IsWindows)'=='true'"/>
@@ -74,11 +69,11 @@
     </Compile>
   </ItemGroup>
   <ItemGroup>
-    <ProjectReference Include="..\..\DotNETCommon\BuildUtilities\BuildUtilitiesCore.csproj" />
-    <ProjectReference Include="..\..\DotNETCommon\DotNETUtilities\DotNETUtilitiesCore.csproj" />
-    <ProjectReference Include="..\..\UnrealBuildTool\UnrealBuildToolCore.csproj" />
-    <ProjectReference Include="..\AutomationUtils\AutomationUtils.Automation.Core.csproj" />
-    <ProjectReference Include="..\Gauntlet\Gauntlet.Automation.Core.csproj" />
+    <ProjectReference Include="..\..\DotNETCommon\BuildUtilities\BuildUtilitiesCore.csproj"><Private>false</Private></ProjectReference>
+    <ProjectReference Include="..\..\DotNETCommon\DotNETUtilities\DotNETUtilitiesCore.csproj"><Private>false</Private></ProjectReference>
+    <ProjectReference Include="..\..\UnrealBuildTool\UnrealBuildToolCore.csproj"><Private>false</Private></ProjectReference>
+    <ProjectReference Include="..\AutomationUtils\AutomationUtils.Automation.Core.csproj"><Private>false</Private></ProjectReference>
+    <ProjectReference Include="..\Gauntlet\Gauntlet.Automation.Core.csproj"><Private>false</Private></ProjectReference>
   </ItemGroup>
   <ItemGroup>
     <PackageReference Include="Microsoft.CSharp" Version="4.7.0" />
diff --git a/Engine/Source/Programs/AutomationTool/Win/Win.Automation.Core.csproj b/Engine/Source/Programs/AutomationTool/Win/Win.Automation.Core.csproj
--- a/Engine/Source/Programs/AutomationTool/Win/Win.Automation.Core.csproj
+++ b/Engine/Source/Programs/AutomationTool/Win/Win.Automation.Core.csproj
@@ -16,7 +16,7 @@
     <RootNamespace>Win.Automation</RootNamespace>
     <AssemblyName>Win.Automation</AssemblyName>
     <WarningsNotAsErrors>612,618</WarningsNotAsErrors>
-    <OutputPath>..\..\..\..\Binaries\DotNET\AutomationTool</OutputPath>
+    <OutputPath>..\..\..\..\Binaries\DotNET\AutomationTool\AutomationScripts</OutputPath>
     <DefineConstants>NET_CORE</DefineConstants>
     <TreatWarningsAsErrors>true</TreatWarningsAsErrors>
   </PropertyGroup>
@@ -57,10 +57,10 @@
     </Compile>
   </ItemGroup>
   <ItemGroup>
-    <ProjectReference Include="..\..\DotNETCommon\DotNETUtilities\DotNETUtilitiesCore.csproj" />
-    <ProjectReference Include="..\..\UnrealBuildTool\UnrealBuildToolCore.csproj" />
-    <ProjectReference Include="..\AutomationUtils\AutomationUtils.Automation.Core.csproj" />
-    <ProjectReference Include="..\BuildGraph\BuildGraph.Automation.Core.csproj" />
+    <ProjectReference Include="..\..\DotNETCommon\DotNETUtilities\DotNETUtilitiesCore.csproj"><Private>false</Private></ProjectReference>
+    <ProjectReference Include="..\..\UnrealBuildTool\UnrealBuildToolCore.csproj"><Private>false</Private></ProjectReference>
+    <ProjectReference Include="..\AutomationUtils\AutomationUtils.Automation.Core.csproj"><Private>false</Private></ProjectReference>
+    <ProjectReference Include="..\BuildGraph\BuildGraph.Automation.Core.csproj"><Private>false</Private></ProjectReference>
   </ItemGroup>
   <ItemGroup>
     <PackageReference Include="Microsoft.CSharp" Version="4.7.0" />
diff --git a/Engine/Source/Programs/AutomationTool/XLocLocalization/XLocLocalization.Automation.Core.csproj b/Engine/Source/Programs/AutomationTool/XLocLocalization/XLocLocalization.Automation.Core.csproj
--- a/Engine/Source/Programs/AutomationTool/XLocLocalization/XLocLocalization.Automation.Core.csproj
+++ b/Engine/Source/Programs/AutomationTool/XLocLocalization/XLocLocalization.Automation.Core.csproj
@@ -16,7 +16,7 @@
     <RootNamespace>AutomationTool</RootNamespace>
     <AssemblyName>XLocLocalization.Automation</AssemblyName>
     <WarningsNotAsErrors>612,618</WarningsNotAsErrors>
-    <OutputPath>..\..\..\..\Binaries\DotNET\AutomationTool</OutputPath>
+    <OutputPath>..\..\..\..\Binaries\DotNET\AutomationTool\AutomationScripts</OutputPath>
     <DefineConstants>NET_CORE</DefineConstants>
     <TreatWarningsAsErrors>true</TreatWarningsAsErrors>
   </PropertyGroup>
@@ -42,10 +42,10 @@
   </ItemGroup>
 
   <ItemGroup>
-    <ProjectReference Include="..\..\DotNETCommon\DotNETUtilities\DotNETUtilitiesCore.csproj" />
-    <ProjectReference Include="..\..\UnrealBuildTool\UnrealBuildToolCore.csproj" />
-    <ProjectReference Include="..\AutomationUtils\AutomationUtils.Automation.Core.csproj" />
-    <ProjectReference Include="..\Localization\Localization.Automation.Core.csproj" />
+    <ProjectReference Include="..\..\DotNETCommon\DotNETUtilities\DotNETUtilitiesCore.csproj"><Private>false</Private></ProjectReference>
+    <ProjectReference Include="..\..\UnrealBuildTool\UnrealBuildToolCore.csproj"><Private>false</Private></ProjectReference>
+    <ProjectReference Include="..\AutomationUtils\AutomationUtils.Automation.Core.csproj"><Private>false</Private></ProjectReference>
+    <ProjectReference Include="..\Localization\Localization.Automation.Core.csproj"><Private>false</Private></ProjectReference>
   </ItemGroup>
   <ItemGroup>
     <PackageReference Include="Microsoft.CSharp" Version="4.7.0" />
