- Worked around p4.exe having a bad interaction with Windows parsing the commandline with a * in it in a server path, by passing commands with a * in the parameters through in a file and the -x global option. (Windows would treat //UE4/Main as \\UE4\Main and take forever to timeout looking for a server named UE4)

diff --git a/Engine/Source/Programs/AutomationTool/AutomationUtils/P4Utils.cs b/Engine/Source/Programs/AutomationTool/AutomationUtils/P4Utils.cs
--- a/Engine/Source/Programs/AutomationTool/AutomationUtils/P4Utils.cs
+++ b/Engine/Source/Programs/AutomationTool/AutomationUtils/P4Utils.cs
@@ -1209,6 +1209,27 @@ namespace AutomationTool
 		/// <returns>Exit code</returns>
 		public IProcessResult P4(string CommandLine, string Input = null, bool AllowSpew = true, bool WithClient = true, bool SpewIsVerbose = false)
 		{
+			return P4("", CommandLine, Input, AllowSpew, WithClient, SpewIsVerbose);
+		}
+
+		/// <summary>
+		/// Shortcut to Run but with P4.exe as the program name.
+		/// </summary>
+		/// <param name="GlobalOptions">Extra global options just for this command</param>
+		/// <param name="CommandLine">Command line</param>
+		/// <param name="Input">Stdin</param>
+		/// <param name="AllowSpew">true for spew</param>
+		/// <returns>Exit code</returns>
+		public IProcessResult P4(string ExtraGlobalOptions, string CommandLine, string Input, bool AllowSpew = true, bool WithClient = true, bool SpewIsVerbose = false)
+		{
+			CommandLine = CommandLine.Trim();
+
+			// we need the first token to be a command ("files") and not a global option ("-c foo")
+			if (CommandLine.StartsWith("-"))
+			{
+				throw new AutomationException("Fix your call to P4 to put global options into the GlobalOptions parameter. The first token should be a p4 command: {0}", CommandLine);
+			}
+
 			CommandUtils.ERunOptions RunOptions = AllowSpew ? CommandUtils.ERunOptions.AllowSpew : CommandUtils.ERunOptions.NoLoggingOfRunCommand;
 			if( SpewIsVerbose )
 			{
@@ -1217,7 +1238,41 @@ namespace AutomationTool
 
 			var SpewDelegate = AllowSpew ? null : new ProcessResult.SpewFilterCallbackType(NoSpewFilter);
 
-			return CommandUtils.Run(HostPlatform.Current.P4Exe, (WithClient ? GlobalOptions : GlobalOptionsWithoutClient) + CommandLine, Input, Options:RunOptions, SpewFilterCallback:SpewDelegate);
+			IProcessResult Result;
+			// if there's a star anywhere in the commandline, p4.exe, when parsing command line on Windows, will internally perform a find-files to expand the *,
+			// and Windows thinks a p4 path is a UNC path (//Depot/Stream/Foo/*/bar.txt). It has been seen that this can stall for for seconds (over 20
+			// seconds potentially). This can be seen with just "dir \\fake\server" on a Windows command prompt, and some machines will take forever
+			if (CommandLine.Contains("*"))
+			{
+				// we can bypass the problem by putting the params for the command into a file, and using the -x <paramsfile> optjon. So:
+				//   p4 -c clientspec files //Depot/Stream/Foo/*/bar.txt
+				// would be converted to this awkward format (with params BEFORE the command):
+				//   p4 -c clientspec -x tempfile.txt files
+				// where tempfile.txt contains:
+				//   //Depot/Stream/Foo/*/bar.txt
+
+				// pull the command out ("files" in the above case)
+				string CommandToken = CommandLine.Trim().Split(" ".ToCharArray())[0];
+
+				// make a temp file, and write the params, minus the command, to it
+				string ParamsFile = Path.Combine(Path.GetTempPath(), Path.GetTempFileName());
+				File.WriteAllText(ParamsFile, CommandLine.Substring(CommandToken.Length + 1));
+
+				// run with -x
+				string FinalCommandline = string.Format("{0} {1} -x \"{3}\" {2}", (WithClient ? GlobalOptions : GlobalOptionsWithoutClient), ExtraGlobalOptions, CommandToken, ParamsFile).Trim();
+				Result = CommandUtils.Run(HostPlatform.Current.P4Exe, FinalCommandline, Input, Options: RunOptions, SpewFilterCallback: SpewDelegate);
+
+				// delete the temp file
+				File.Delete(ParamsFile);
+			}
+			else
+			{
+				string FinalCommandline = string.Format("{0} {1} {2}", (WithClient ? GlobalOptions : GlobalOptionsWithoutClient), ExtraGlobalOptions, CommandLine);
+				Result = CommandUtils.Run(HostPlatform.Current.P4Exe, FinalCommandline, Input, Options: RunOptions, SpewFilterCallback: SpewDelegate);
+			}
+
+
+			return Result;
 		}
 
 		/// <summary>
@@ -3278,7 +3333,7 @@ namespace AutomationTool
 				CommandUtils.LogLog("Checking if client {0} exists", ClientName);
 			}
 
-            var P4Result = P4(String.Format("-c {0} where //...", ClientName), AllowSpew: false, WithClient: false);
+            var P4Result = P4(String.Format("-c {0}", ClientName), "where //...", Input: null, AllowSpew: false, WithClient: false);
             return P4Result.Output.IndexOf("unknown - use 'client' command", StringComparison.InvariantCultureIgnoreCase) < 0 && P4Result.Output.IndexOf("doesn't exist", StringComparison.InvariantCultureIgnoreCase) < 0;
 		}
 
diff --git a/Engine/Source/Programs/AutomationTool/Turnkey/Providers/PerforceCopyProvider.cs b/Engine/Source/Programs/AutomationTool/Turnkey/Providers/PerforceCopyProvider.cs
--- a/Engine/Source/Programs/AutomationTool/Turnkey/Providers/PerforceCopyProvider.cs
+++ b/Engine/Source/Programs/AutomationTool/Turnkey/Providers/PerforceCopyProvider.cs
@@ -127,7 +127,7 @@ namespace Turnkey
 		private P4ClientInfo CanClientHandleOperation(string ClientName, string Operation, string Hostname)
 		{
 			// make sure the Operation can be supported by the client (Tokens[1] is the clientspec name)
-			IProcessResult P4Result = PerforceConnection.P4(string.Format("-c {0} where {1}", ClientName, Operation), AllowSpew: false, WithClient: false);
+			IProcessResult P4Result = PerforceConnection.P4(string.Format("-c {0}", ClientName), string.Format("where {0}", Operation), Input: null, AllowSpew: false, WithClient: false);
 
 			if (P4Result.ExitCode == 0)
 			{
