Fixed UAT Game specific UAT projects output location

Changed so that automation projects are loaded from their file path (LoadFrom rules) in net core as this is in practise what happened in net fw. Fixes issue were UAT projects not under Engine/Binaries couldnt be loaded.

diff --git a/Engine/Source/Programs/AutomationTool/AutomationUtils/ScriptCompiler.cs b/Engine/Source/Programs/AutomationTool/AutomationUtils/ScriptCompiler.cs
--- a/Engine/Source/Programs/AutomationTool/AutomationUtils/ScriptCompiler.cs
+++ b/Engine/Source/Programs/AutomationTool/AutomationUtils/ScriptCompiler.cs
@@ -474,7 +474,11 @@ namespace AutomationTool
 				CommandUtils.LogLog("Loading script DLL: {0}", AssemblyLocation);
 				try
 				{
+#if NET_CORE
+					Assembly Assembly = Assembly.LoadFile(AssemblyLocation.FullName);
+#else
 					Assembly Assembly = AppDomain.CurrentDomain.Load(AssemblyName.GetAssemblyName(AssemblyLocation.FullName));
+#endif
 					Assemblies.Add(Assembly);
 				}
 				catch (Exception Ex)
