[Turnkey] Added null check when AllowedFlashDeviceTypes is unset (and null)

diff --git a/Engine/Source/Programs/AutomationTool/Turnkey/Commands/InstallSdk.cs b/Engine/Source/Programs/AutomationTool/Turnkey/Commands/InstallSdk.cs
--- a/Engine/Source/Programs/AutomationTool/Turnkey/Commands/InstallSdk.cs
+++ b/Engine/Source/Programs/AutomationTool/Turnkey/Commands/InstallSdk.cs
@@ -119,7 +119,10 @@ namespace Turnkey.Commands
 						{
 							return Val;
 						}
-						return a.AllowedFlashDeviceTypes.CompareTo(b.AllowedFlashDeviceTypes);
+						if (a.AllowedFlashDeviceTypes != null && b.AllowedFlashDeviceTypes != null)
+						{
+							return a.AllowedFlashDeviceTypes.CompareTo(b.AllowedFlashDeviceTypes);
+						}
 					}
 
 					return 0;
