Added encodings not present in net core by default to fix " 'IBM437' is not a supported encoding name"

diff --git a/Engine/Source/Programs/AutomationTool/Android/Android.Automation.Core.csproj b/Engine/Source/Programs/AutomationTool/Android/Android.Automation.Core.csproj
--- a/Engine/Source/Programs/AutomationTool/Android/Android.Automation.Core.csproj
+++ b/Engine/Source/Programs/AutomationTool/Android/Android.Automation.Core.csproj
@@ -62,5 +62,6 @@
     <PackageReference Include="Microsoft.CSharp" Version="4.7.0" />
     <PackageReference Include="System.Data.DataSetExtensions" Version="4.5.0" />
     <PackageReference Include="System.Net.Http" Version="4.3.4" />
+    <PackageReference Include="System.Text.Encoding.CodePages" Version="4.7.1" />
   </ItemGroup>
 </Project>
\ No newline at end of file
diff --git a/Engine/Source/Programs/AutomationTool/Android/AndroidPlatform.Automation.cs b/Engine/Source/Programs/AutomationTool/Android/AndroidPlatform.Automation.cs
--- a/Engine/Source/Programs/AutomationTool/Android/AndroidPlatform.Automation.cs
+++ b/Engine/Source/Programs/AutomationTool/Android/AndroidPlatform.Automation.cs
@@ -26,6 +26,9 @@ public class AndroidPlatform : Platform
 	public AndroidPlatform()
 		: base(UnrealTargetPlatform.Android)
 	{
+#if NET_CORE
+		Encoding.RegisterProvider(CodePagesEncodingProvider.Instance);
+#endif
 	}
 
 	private static string GetSONameWithoutArchitecture(ProjectParams Params, string DecoratedExeName)
