[Turnkey] Ignore files in google drive that are trashed.

diff --git a/Engine/Source/Programs/AutomationTool/Turnkey/Providers/GoogleDriveCopyProvider.cs b/Engine/Source/Programs/AutomationTool/Turnkey/Providers/GoogleDriveCopyProvider.cs
--- a/Engine/Source/Programs/AutomationTool/Turnkey/Providers/GoogleDriveCopyProvider.cs
+++ b/Engine/Source/Programs/AutomationTool/Turnkey/Providers/GoogleDriveCopyProvider.cs
@@ -452,8 +452,8 @@ namespace Turnkey
 				// convert to a wildcard with capturing to get what the matches are
 				Regex Regex = new Regex(string.Format("^{0}$", Wildcard.Replace("*", "(.+?)")));
 
-				// get files that could match the wildcard
-				string Query = string.Format("'{0}' in parents", FolderBeforeWildcard.Id);
+				// get files that could match the wildcard and are not trashed
+				string Query = string.Format("'{0}' in parents and trashed = false", FolderBeforeWildcard.Id);
 
 				List<Google.Apis.Drive.v3.Data.File> Results = ServiceHelper.SearchFiles(Query);
 
