[Turnkey] When using cached files in GoogleDriveCopyProvider, don't overwrite OutputPath with the last file path if a wildcard is used.

diff --git a/Engine/Source/Programs/AutomationTool/Turnkey/Providers/GoogleDriveCopyProvider.cs b/Engine/Source/Programs/AutomationTool/Turnkey/Providers/GoogleDriveCopyProvider.cs
--- a/Engine/Source/Programs/AutomationTool/Turnkey/Providers/GoogleDriveCopyProvider.cs
+++ b/Engine/Source/Programs/AutomationTool/Turnkey/Providers/GoogleDriveCopyProvider.cs
@@ -95,7 +95,7 @@ namespace Turnkey
 					}
 					// if we are removing our one and only file from the list to be deleted, we need to change our OutputPath to be that single file
 					// since we won't be going through the lower loop
-					else
+					else if (!bHasWildcard)
 					{
 						OutputPath = ExistingPath;
 					}
