UAT: fix SyncProject's handling of p4 commands and improvements to global p4  option handling in general (collab with Josh)

diff --git a/Engine/Source/Programs/AutomationTool/AutomationUtils/P4Utils.cs b/Engine/Source/Programs/AutomationTool/AutomationUtils/P4Utils.cs
--- a/Engine/Source/Programs/AutomationTool/AutomationUtils/P4Utils.cs
+++ b/Engine/Source/Programs/AutomationTool/AutomationUtils/P4Utils.cs
@@ -1283,11 +1283,11 @@ namespace AutomationTool
 		/// <param name="Input">Stdin input.</param>
 		/// <param name="AllowSpew">Whether the command should spew.</param>
 		/// <returns>True if succeeded, otherwise false.</returns>
-        public bool P4Output(out string Output, string CommandLine, string Input = null, bool AllowSpew = true, bool WithClient = true)
+        public bool P4Output(out string Output, string ExtraGlobalOptions, string CommandLine, string Input = null, bool AllowSpew = true, bool WithClient = true)
 		{
 			Output = "";
 
-            var Result = P4(CommandLine, Input, AllowSpew, WithClient);
+            var Result = P4(ExtraGlobalOptions, CommandLine, Input, AllowSpew, WithClient);
 
 			Output = Result.Output;
 			return Result.ExitCode == 0;
@@ -1301,10 +1301,10 @@ namespace AutomationTool
 		/// <param name="Input">Stdin input.</param>
 		/// <param name="AllowSpew">Whether the command should spew.</param>
 		/// <returns>True if succeeded, otherwise false.</returns>
-        public bool P4Output(out string[] OutputLines, string CommandLine, string Input = null, bool AllowSpew = true, bool WithClient = true)
+        public bool P4Output(out string[] OutputLines, string ExtraGlobalOptions, string CommandLine, string Input = null, bool AllowSpew = true, bool WithClient = true)
 		{
 			string Output;
-			bool bResult = P4Output(out Output, CommandLine, Input, AllowSpew, WithClient);
+			bool bResult = P4Output(out Output, ExtraGlobalOptions, CommandLine, Input, AllowSpew, WithClient);
 
 			List<string> Lines = new List<string>();
 			for(int Idx = 0; Idx < Output.Length; )
@@ -1338,10 +1338,10 @@ namespace AutomationTool
 		/// <param name="CommandLine">Commandline to pass to p4.</param>
 		/// <param name="Input">Stdin input.</param>
 		/// <param name="AllowSpew">Whether the command is allowed to spew.</param>
-        public void LogP4(string CommandLine, string Input = null, bool AllowSpew = true, bool WithClient = true, bool SpewIsVerbose = false)
+        public void LogP4(string ExtraGlobalOptions, string CommandLine, string Input = null, bool AllowSpew = true, bool WithClient = true, bool SpewIsVerbose = false)
 		{
 			string Output;
-            if (!LogP4Output(out Output, CommandLine, Input, AllowSpew, WithClient, SpewIsVerbose:SpewIsVerbose))
+            if (!LogP4Output(out Output, ExtraGlobalOptions, CommandLine, Input, AllowSpew, WithClient, SpewIsVerbose:SpewIsVerbose))
 			{
 				throw new P4Exception("p4.exe {0} failed. {1}", CommandLine, Output);
 			}
@@ -1355,7 +1355,7 @@ namespace AutomationTool
 		/// <param name="Input">Stdin input.</param>
 		/// <param name="AllowSpew">Whether the command should spew.</param>
 		/// <returns>True if succeeded, otherwise false.</returns>
-        public bool LogP4Output(out string Output, string CommandLine, string Input = null, bool AllowSpew = true, bool WithClient = true, bool SpewIsVerbose = false)
+        public bool LogP4Output(out string Output, string ExtraGlobalOptions, string CommandLine, string Input = null, bool AllowSpew = true, bool WithClient = true, bool SpewIsVerbose = false)
 		{
 			Output = "";
 
@@ -1365,7 +1365,7 @@ namespace AutomationTool
 				return false;
 			}
 
-            var Result = P4(CommandLine, Input, AllowSpew, WithClient, SpewIsVerbose:SpewIsVerbose);
+            var Result = P4(ExtraGlobalOptions, CommandLine, Input, AllowSpew, WithClient, SpewIsVerbose:SpewIsVerbose);
 
 			CommandUtils.WriteToFile(LogPath, CommandLine + "\n");
 			CommandUtils.WriteToFile(LogPath, Result.Output);
@@ -1517,7 +1517,7 @@ namespace AutomationTool
 
 			string Output;
             string P4Passwd = InternalUtils.GetEnvironmentVariable("uebp_PASS", "", true) + '\n';
-            P4Output(out Output, "login -a -p", P4Passwd);
+            P4Output(out Output, "", "login -a -p", P4Passwd);
 
 			// Validate output.
 			const string PasswordPromptString = "Enter password: \r\n";
@@ -1605,7 +1605,7 @@ namespace AutomationTool
                 // Change 1999345 on 2014/02/16 by buildmachine@BuildFarm_BUILD-23_buildmachine_++depot+UE4 'GUBP Node Shadow_LabelPromotabl'
 
                 string Output;
-                if (!LogP4Output(out Output, "changes " + CommandLine, null, AllowSpew, WithClient: WithClient))
+                if (!LogP4Output(out Output, "", "changes " + CommandLine, null, AllowSpew, WithClient: WithClient))
                 {
                     throw new AutomationException("P4 returned failure.");
                 }
@@ -1822,7 +1822,7 @@ namespace AutomationTool
 					CommandLine += " " + Changelist.ToString();
 				}
 
-                if (!LogP4Output(out Output, "describe " + CommandLine, null, AllowSpew))
+                if (!LogP4Output(out Output, "", "describe " + CommandLine, null, AllowSpew))
                 {
                     return false;
                 }
@@ -1989,17 +1989,18 @@ namespace AutomationTool
 		{
 			string SyncCommandLine = "sync " + CommandLine;
 
+			string ExtraGlobalOptions = "";
 			if (MaxWait > 0)
 			{
-				SyncCommandLine = string.Format("-vnet.maxwait={0} {1}", MaxWait, SyncCommandLine);
+				ExtraGlobalOptions = string.Format("-vnet.maxwait={0} {1}", MaxWait, ExtraGlobalOptions);
 			}
 
 			if (Retries > 0)
 			{
-				SyncCommandLine = string.Format("-r{0} {1}", Retries, SyncCommandLine);
+				ExtraGlobalOptions = string.Format("-r{0} {1}", Retries, ExtraGlobalOptions);
 			}
 
-			LogP4(SyncCommandLine, null, AllowSpew, SpewIsVerbose:SpewIsVerbose);
+			LogP4(ExtraGlobalOptions, SyncCommandLine, null, AllowSpew, SpewIsVerbose:SpewIsVerbose);
 		}
 
 		/// <summary>
@@ -2015,7 +2016,7 @@ namespace AutomationTool
 			try
 			{
 				string Output;
-				LogP4Output(out Output, "sync -n " + CommandLine, null, AllowSpew, SpewIsVerbose: SpewIsVerbose);
+				LogP4Output(out Output, "", "sync -n " + CommandLine, null, AllowSpew, SpewIsVerbose: SpewIsVerbose);
 
 				string UpToDateOutput = String.Format("{0} - file(s) up-to-date.\r\n", CommandLine);
 				if (Output == UpToDateOutput)
@@ -2052,7 +2053,7 @@ namespace AutomationTool
 		/// <param name="CommandLine">Commandline for the command.</param>
 		public void Unshelve(int FromCL, int ToCL, string CommandLine = "", bool SpewIsVerbose = false)
 		{
-			LogP4("unshelve " + String.Format("-s {0} ", FromCL) + String.Format("-c {0} ", ToCL) + CommandLine, SpewIsVerbose: SpewIsVerbose);
+			LogP4("", "unshelve " + String.Format("-s {0} ", FromCL) + String.Format("-c {0} ", ToCL) + CommandLine, SpewIsVerbose: SpewIsVerbose);
 		}
 
         /// <summary>
@@ -2063,7 +2064,7 @@ namespace AutomationTool
         /// <param name="CommandLine">Commandline for the command.</param>
         public void Shelve(int FromCL, string CommandLine = "", bool AllowSpew = true)
         {
-            LogP4("shelve " + String.Format("-r -c {0} ", FromCL) + CommandLine, AllowSpew: AllowSpew);
+            LogP4("", "shelve " + String.Format("-r -c {0} ", FromCL) + CommandLine, AllowSpew: AllowSpew);
         }
 
 		/// <summary>
@@ -2074,7 +2075,7 @@ namespace AutomationTool
         public void DeleteShelvedFiles(int FromCL, bool AllowSpew = true)
         {
 			string Output;
-            if (!LogP4Output(out Output, String.Format("shelve -d -c {0}", FromCL), AllowSpew: AllowSpew) && !Output.StartsWith("No shelved files in changelist to delete."))
+            if (!LogP4Output(out Output, "", String.Format("shelve -d -c {0}", FromCL), AllowSpew: AllowSpew) && !Output.StartsWith("No shelved files in changelist to delete."))
 			{
 				throw new P4Exception("Couldn't unshelve files: {0}", Output);
 			}
@@ -2087,7 +2088,7 @@ namespace AutomationTool
 		/// <param name="CommandLine">Commandline for the command.</param>
 		public void Edit(int CL, string CommandLine, bool AllowSpew = true)
 		{
-			LogP4("edit " + String.Format("-c {0} ", CL) + CommandLine, AllowSpew: AllowSpew);
+			LogP4("", "edit " + String.Format("-c {0} ", CL) + CommandLine, AllowSpew: AllowSpew);
 		}
 
 		/// <summary>
@@ -2104,7 +2105,7 @@ namespace AutomationTool
 			{
 				if(CommandLine.Length + Files[Idx].Length + 3 > MaxCommandLineLength)
 				{
-					LogP4(String.Format("edit -c {0} {1}", CL, CommandLine.ToString()), AllowSpew: AllowSpew);
+					LogP4("", String.Format("edit -c {0} {1}", CL, CommandLine.ToString()), AllowSpew: AllowSpew);
 					CommandLine.Clear();
 				}
 				if(CommandLine.Length > 0)
@@ -2115,7 +2116,7 @@ namespace AutomationTool
 			}
 			if(CommandLine.Length > 0)
 			{
-				LogP4(String.Format("edit -c {0} {1}", CL, CommandLine.ToString()), AllowSpew: AllowSpew);
+				LogP4("", String.Format("edit -c {0} {1}", CL, CommandLine.ToString()), AllowSpew: AllowSpew);
 			}
 		}
 
@@ -2129,7 +2130,7 @@ namespace AutomationTool
 			try
 			{
 				string Output;
-				if (!LogP4Output(out Output, "edit " + String.Format("-c {0} ", CL) + CommandLine, null, true))
+				if (!LogP4Output(out Output, "", "edit " + String.Format("-c {0} ", CL) + CommandLine, null, true))
 				{
 					return false;
 				}
@@ -2152,7 +2153,7 @@ namespace AutomationTool
 		/// <param name="CommandLine">Commandline for the command.</param>
 		public void Add(int CL, string CommandLine)
 		{
-			LogP4("add " + String.Format("-c {0} ", CL) + CommandLine);
+			LogP4("", "add " + String.Format("-c {0} ", CL) + CommandLine);
 		}
 
 		/// <summary>
@@ -2162,7 +2163,7 @@ namespace AutomationTool
 		/// <param name="CommandLine">Commandline for the command.</param>
 		public void Delete(int CL, string CommandLine)
 		{
-			LogP4("delete " + String.Format("-c {0} ", CL) + CommandLine);
+			LogP4("", "delete " + String.Format("-c {0} ", CL) + CommandLine);
 		}
 
 		/// <summary>
@@ -2172,7 +2173,7 @@ namespace AutomationTool
 		/// <param name="CommandLine">Commandline for the command.</param>
 		public void Reconcile(int CL, string CommandLine, bool AllowSpew = true)
 		{
-			LogP4("reconcile " + String.Format("-c {0} -ead -f ", CL) + CommandLine, AllowSpew: AllowSpew);
+			LogP4("", "reconcile " + String.Format("-c {0} -ead -f ", CL) + CommandLine, AllowSpew: AllowSpew);
 		}
 
         /// <summary>
@@ -2182,7 +2183,7 @@ namespace AutomationTool
         /// <param name="CommandLine">Commandline for the command.</param>
         public void ReconcilePreview(string CommandLine)
         {
-            LogP4("reconcile " + String.Format("-ead -n ") + CommandLine);
+            LogP4("", "reconcile " + String.Format("-ead -n ") + CommandLine);
         }
 
 		/// <summary>
@@ -2193,7 +2194,7 @@ namespace AutomationTool
 		/// <param name="CommandLine">Commandline for the command.</param>
 		public void ReconcileNoDeletes(int CL, string CommandLine, bool AllowSpew = true)
 		{
-			LogP4("reconcile " + String.Format("-c {0} -ea ", CL) + CommandLine, AllowSpew: AllowSpew);
+			LogP4("", "reconcile " + String.Format("-c {0} -ea ", CL) + CommandLine, AllowSpew: AllowSpew);
 		}
 
 		/// <summary>
@@ -2204,7 +2205,7 @@ namespace AutomationTool
 		/// <param name="CommandLine">Commandline for the command.</param>
 		public void Resolve(int CL, string CommandLine)
 		{
-			LogP4("resolve -ay " + String.Format("-c {0} ", CL) + CommandLine);
+			LogP4("", "resolve -ay " + String.Format("-c {0} ", CL) + CommandLine);
 		}
 
 		/// <summary>
@@ -2213,7 +2214,7 @@ namespace AutomationTool
 		/// <param name="CommandLine">Commandline for the command.</param>
 		public void Revert(string CommandLine, bool AllowSpew = true)
 		{
-			LogP4("revert " + CommandLine, AllowSpew: AllowSpew);
+			LogP4("", "revert " + CommandLine, AllowSpew: AllowSpew);
 		}
 
 		/// <summary>
@@ -2223,7 +2224,7 @@ namespace AutomationTool
 		/// <param name="CommandLine">Commandline for the command.</param>
 		public void Revert(int CL, string CommandLine = "", bool AllowSpew = true)
 		{
-			LogP4("revert " + String.Format("-c {0} ", CL) + CommandLine, AllowSpew: AllowSpew);
+			LogP4("", "revert " + String.Format("-c {0} ", CL) + CommandLine, AllowSpew: AllowSpew);
 		}
 
 		/// <summary>
@@ -2233,7 +2234,7 @@ namespace AutomationTool
 		public void RevertUnchanged(int CL)
 		{
 			// caution this is a really bad idea if you hope to force submit!!!
-			LogP4("revert -a " + String.Format("-c {0} ", CL));
+			LogP4("", "revert -a " + String.Format("-c {0} ", CL));
 		}
 
 		/// <summary>
@@ -2242,7 +2243,7 @@ namespace AutomationTool
 		/// <param name="CL">Changelist to revert.</param>
 		public void RevertAll(int CL, bool SpewIsVerbose = false)
 		{
-			LogP4("revert " + String.Format("-c {0} //...", CL), SpewIsVerbose: SpewIsVerbose);
+			LogP4("", "revert " + String.Format("-c {0} //...", CL), SpewIsVerbose: SpewIsVerbose);
 		}
 
 		/// <summary>
@@ -2294,7 +2295,7 @@ namespace AutomationTool
                     return;
                 }
 				string CmdOutput;
-				if (!LogP4Output(out CmdOutput, String.Format("submit -c {0} -f submitunchanged", CL)))
+				if (!LogP4Output(out CmdOutput, "", String.Format("submit -c {0} -f submitunchanged", CL)))
 				{
 					if (!Force)
 					{
@@ -2452,7 +2453,7 @@ namespace AutomationTool
 			{
 				CommandUtils.LogInformation("Creating Change\n {0}\n", ChangeSpec);
 			}
-			if (LogP4Output(out CmdOutput, "change -i", Input: ChangeSpec, AllowSpew: AllowSpew))
+			if (LogP4Output(out CmdOutput, "", "change -i", Input: ChangeSpec, AllowSpew: AllowSpew))
 			{
 				string EndStr = " created.";
 				string ChangeStr = "Change ";
@@ -2485,7 +2486,7 @@ namespace AutomationTool
 		public void UpdateChange(int CL, string NewOwner, string NewDescription, bool SpewIsVerbose = false)
 		{
 			string CmdOutput;
-			if(!LogP4Output(out CmdOutput, String.Format("change -o {0}", CL), SpewIsVerbose: SpewIsVerbose))
+			if(!LogP4Output(out CmdOutput, "", String.Format("change -o {0}", CL), SpewIsVerbose: SpewIsVerbose))
 			{
 				throw new P4Exception("Couldn't describe changelist {0}", CL);
 			}
@@ -2500,7 +2501,7 @@ namespace AutomationTool
 				Spec.SetField("Description", NewDescription);
 			}
 
-			if(!LogP4Output(out CmdOutput, "change -i", Input: Spec.ToString(), SpewIsVerbose: SpewIsVerbose))
+			if(!LogP4Output(out CmdOutput, "", "change -i", Input: Spec.ToString(), SpewIsVerbose: SpewIsVerbose))
 			{
 				throw new P4Exception("Failed to update spec for changelist {0}", CL);
 			}
@@ -2523,7 +2524,7 @@ namespace AutomationTool
 			}
 
 			string CmdOutput;
-			if (LogP4Output(out CmdOutput, String.Format("change -d {0}", CL), SpewIsVerbose: SpewIsVerbose, AllowSpew: AllowSpew))
+			if (LogP4Output(out CmdOutput, "", String.Format("change -d {0}", CL), SpewIsVerbose: SpewIsVerbose, AllowSpew: AllowSpew))
 			{
 				string EndStr = " deleted.";
 				string ChangeStr = "Change ";
@@ -2545,7 +2546,7 @@ namespace AutomationTool
 		public bool TryDeleteEmptyChange(int CL)
 		{
 			string CmdOutput;
-			if (LogP4Output(out CmdOutput, String.Format("change -d {0}", CL)))
+			if (LogP4Output(out CmdOutput, "", String.Format("change -d {0}", CL)))
 			{
 				string EndStr = " deleted.";
 				string ChangeStr = "Change ";
@@ -2568,7 +2569,7 @@ namespace AutomationTool
 		public string ChangeOutput(int CL, bool AllowSpew = true)
 		{
 			string CmdOutput;
-			if (LogP4Output(out CmdOutput, String.Format("change -o {0}", CL), AllowSpew: AllowSpew))
+			if (LogP4Output(out CmdOutput, "", String.Format("change -o {0}", CL), AllowSpew: AllowSpew))
 			{
 				return CmdOutput;
 			}
@@ -2678,7 +2679,7 @@ namespace AutomationTool
         public string OpenedOutput()
         {
             string CmdOutput;
-            if (LogP4Output(out CmdOutput, "opened"))
+            if (LogP4Output(out CmdOutput, "", "opened"))
             {
                 return CmdOutput;
             }
@@ -2694,7 +2695,7 @@ namespace AutomationTool
 
 			// NOTE: We don't throw exceptions when trying to delete a label
 			string Output;
-			if (!LogP4Output(out Output, CommandLine, null, AllowSpew))
+			if (!LogP4Output(out Output, "", CommandLine, null, AllowSpew))
 			{
 				CommandUtils.LogInformation("Couldn't delete label '{0}'.  It may not have existed in the first place.", LabelName);
 			}
@@ -2728,7 +2729,7 @@ namespace AutomationTool
 			LabelSpec += " " + View;
 
 			CommandUtils.LogInformation("Creating Label\n {0}\n", LabelSpec);
-			LogP4("label -i", Input: LabelSpec);
+			LogP4("", "label -i", Input: LabelSpec);
 		}
 
 		/// <summary>
@@ -2740,7 +2741,7 @@ namespace AutomationTool
 		/// <param name="AllowSpew">Whether the command is allowed to spew.</param>
 		public void Tag(string LabelName, string FilePath, bool AllowSpew = true)
 		{
-			LogP4("tag -l " + LabelName + " " + FilePath, null, AllowSpew);
+			LogP4("", "tag -l " + LabelName + " " + FilePath, null, AllowSpew);
 		}
 
 		/// <summary>
@@ -2757,11 +2758,11 @@ namespace AutomationTool
 			}
 			if (FileToLabel == "")
 			{
-				LogP4("labelsync " + Quiet + "-l " + LabelName);
+				LogP4("", "labelsync " + Quiet + "-l " + LabelName);
 			}
 			else
 			{
-				LogP4("labelsync " + Quiet + "-l" + LabelName + " " + FileToLabel);
+				LogP4("", "labelsync " + Quiet + "-l" + LabelName + " " + FileToLabel);
 			}
 		}
 
@@ -2778,7 +2779,7 @@ namespace AutomationTool
 			{
 				Quiet = "-q ";
 			}
-			LogP4("labelsync -a " + Quiet + "-l " + ToLabelName + " //...@" + FromLabelName);
+			LogP4("", "labelsync -a " + Quiet + "-l " + ToLabelName + " //...@" + FromLabelName);
 		}
 
 		/// <summary>
@@ -2789,7 +2790,7 @@ namespace AutomationTool
 		public bool LabelExistsAndHasFiles(string Name)
 		{
 			string Output;
-			return LogP4Output(out Output, "files -m 1 //...@" + Name);
+			return LogP4Output(out Output, "", "files -m 1 //...@" + Name);
 		}
 
 		/// <summary>
@@ -2803,7 +2804,7 @@ namespace AutomationTool
 		{
 			string Output;
 			Description = "";
-			if (LogP4Output(out Output, "label -o " + Name, AllowSpew: AllowSpew))
+			if (LogP4Output(out Output, "", "label -o " + Name, AllowSpew: AllowSpew))
 			{
 				string Desc = "Description:";
 				int Start = Output.LastIndexOf(Desc);
@@ -2830,7 +2831,7 @@ namespace AutomationTool
 		public P4Spec ReadLabelSpec(string Name, bool AllowSpew = true)
 		{
 			string LabelSpec;
-			if(!LogP4Output(out LabelSpec, "label -o " + Name, AllowSpew: AllowSpew))
+			if(!LogP4Output(out LabelSpec, "", "label -o " + Name, AllowSpew: AllowSpew))
 			{
 				throw new P4Exception("Couldn't describe existing label '{0}', output was:\n", Name, LabelSpec);
 			}
@@ -2844,7 +2845,7 @@ namespace AutomationTool
 		/// <param name="AllowSpew">Whether to allow log spew</param>
 		public void UpdateLabelSpec(P4Spec Spec, bool AllowSpew = true)
 		{
-			LogP4("label -i", Input: Spec.ToString(), AllowSpew: AllowSpew);
+			LogP4("", "label -i", Input: Spec.ToString(), AllowSpew: AllowSpew);
 		}
 
 		/// <summary>
@@ -2856,7 +2857,7 @@ namespace AutomationTool
 		public void UpdateLabelDescription(string Name, string NewDescription, bool AllowSpew = true)
 		{
 			string LabelSpec;
-			if(!LogP4Output(out LabelSpec, "label -o " + Name, AllowSpew: AllowSpew))
+			if(!LogP4Output(out LabelSpec, "", "label -o " + Name, AllowSpew: AllowSpew))
 			{
 				throw new P4Exception("Couldn't describe existing label '{0}', output was:\n", Name, LabelSpec);
 			}
@@ -2883,7 +2884,7 @@ namespace AutomationTool
 			LabelSpec = String.Join("\n", Lines);
 
 			// Update the label
-			LogP4("label -i", Input: LabelSpec, AllowSpew: AllowSpew);
+			LogP4("", "label -i", Input: LabelSpec, AllowSpew: AllowSpew);
 		}
 
 		/* Pattern to parse P4 changes command output. */
@@ -2896,7 +2897,7 @@ namespace AutomationTool
 		public int GetLatestCLNumber()
 		{
 			string Output;
-			if (!LogP4Output(out Output, "changes -s submitted -m1") || string.IsNullOrWhiteSpace(Output))
+			if (!LogP4Output(out Output, "", "changes -s submitted -m1") || string.IsNullOrWhiteSpace(Output))
 			{
 				throw new InvalidOperationException("The depot should have at least one submitted changelist. Brand new depot?");
 			}
@@ -2925,7 +2926,7 @@ namespace AutomationTool
 			var LabelList = new List<P4Label>();
 
 			string Output;
-			if (P4Output(out Output, "labels -t " + (bCaseSensitive ? "-e" : "-E") + Filter, null, false))
+			if (P4Output(out Output, "", "labels -t " + (bCaseSensitive ? "-e" : "-E") + Filter, null, false))
 			{
 				foreach (Match LabelMatch in LabelsListOutputPattern.Matches(Output))
 				{
@@ -2947,7 +2948,7 @@ namespace AutomationTool
 		public bool ValidateLabelContent(string LabelName)
 		{
 			string Output;
-			if (P4Output(out Output, "files -m 1 @" + LabelName, null, false))
+			if (P4Output(out Output, "", "files -m 1 @" + LabelName, null, false))
 			{
 				if (Output.StartsWith("//depot"))
 				{
@@ -2971,11 +2972,11 @@ namespace AutomationTool
         public string DepotToLocalPath(string DepotFile, bool AllowSpew = true)
         {
 			//  P4 where outputs missing entries 
-			string Command = String.Format("-z tag fstat \"{0}\"", DepotFile);
+			string Command = String.Format("fstat \"{0}\"", DepotFile);
 
 			// Run the command.
 			string[] Lines;
-			if (!P4Output(out Lines, Command, AllowSpew: AllowSpew))
+			if (!P4Output(out Lines, "-z tag", Command, AllowSpew: AllowSpew))
 			{
 				throw new P4Exception("p4.exe {0} failed.", Command);
 			}
@@ -3006,7 +3007,7 @@ namespace AutomationTool
 			for(int Idx = 0; Idx < DepotFiles.Length; Idx += BatchSize)
 			{
 				// Build the argument list
-				StringBuilder Command = new StringBuilder("-z tag fstat ");
+				StringBuilder Command = new StringBuilder("fstat ");
 				for(int ArgIdx = Idx; ArgIdx < Idx + BatchSize && ArgIdx < DepotFiles.Length; ArgIdx++)
 				{
 					Command.AppendFormat(" {0}", CommandUtils.MakePathSafeToUseWithCommandLine(DepotFiles[ArgIdx]));
@@ -3014,7 +3015,7 @@ namespace AutomationTool
 
 				// Run the command.
 				string[] Output;
-				if (!P4Output(out Output, Command.ToString(), AllowSpew: AllowSpew))
+				if (!P4Output(out Output, "-z tag", Command.ToString(), AllowSpew: AllowSpew))
 				{
 					throw new P4Exception("p4.exe {0} failed.", Command);
 				}
@@ -3080,11 +3081,11 @@ namespace AutomationTool
 		public P4WhereRecord[] Where(string DepotFile, bool AllowSpew = true)
 		{
 			//  P4 where outputs missing entries 
-			string Command = String.Format("-z tag where \"{0}\"", DepotFile);
+			string Command = String.Format("where \"{0}\"", DepotFile);
 
 			// Run the command.
 			string Output;
-			if (!LogP4Output(out Output, Command, AllowSpew: AllowSpew))
+			if (!LogP4Output(out Output, "-z tag", Command, AllowSpew: AllowSpew))
 			{
 				throw new P4Exception("p4.exe {0} failed.", Command);
 			}
@@ -3148,10 +3149,10 @@ namespace AutomationTool
 		/// <returns>List of records describing the file's mapping. Usually just one, but may be more.</returns>
 		public bool FileExistsInDepot(string DepotFile, bool AllowSpew = true)
 		{
-			string CommandLine = String.Format("-z tag fstat {0}", CommandUtils.MakePathSafeToUseWithCommandLine(DepotFile));
+			string CommandLine = String.Format("fstat {0}", CommandUtils.MakePathSafeToUseWithCommandLine(DepotFile));
 
 			string Output;
-			if(!LogP4Output(out Output, CommandLine, AllowSpew: false) || !Output.Contains("headRev"))
+			if(!LogP4Output(out Output, "-z tag", CommandLine, AllowSpew: false) || !Output.Contains("headRev"))
 
 			{
 				return false;
@@ -3169,7 +3170,7 @@ namespace AutomationTool
 		{
 			string Output;
 			string Command = "fstat " + CommandUtils.MakePathSafeToUseWithCommandLine(Filename);
-			if (!LogP4Output(out Output, Command))
+			if (!LogP4Output(out Output, "", Command))
 			{
 				throw new P4Exception("p4.exe {0} failed.", Command);
 			}
@@ -3235,7 +3236,7 @@ namespace AutomationTool
 				var CmdLine = String.Format("{0} -c {1} -t {2} {3}",
 					(Stat.Action != P4Action.None) ? "reopen" : "open",
 					Changelist, FileAttributesToString(Attributes | Stat.Attributes), CommandUtils.MakePathSafeToUseWithCommandLine(Filename));
-				LogP4(CmdLine);
+				LogP4("", CmdLine);
 			}
 		}
 
@@ -3518,7 +3519,7 @@ namespace AutomationTool
 		/// <param name="Force">Forces the operation (-f)</param>
 		public void DeleteClient(string Name, bool Force = false, bool AllowSpew = true)
         {
-            LogP4(String.Format("client -d {0} {1}", (Force ? "-f" : ""), Name), WithClient: false, AllowSpew: AllowSpew);
+            LogP4("", String.Format("client -d {0} {1}", (Force ? "-f" : ""), Name), WithClient: false, AllowSpew: AllowSpew);
         }
 
         /// <summary>
@@ -3622,7 +3623,7 @@ namespace AutomationTool
 		public string Print(string DepotPath, bool AllowSpew = true)
 		{
 			string Output;
-			if(!P4Output(out Output, "print -q " + DepotPath, AllowSpew: AllowSpew, WithClient: false))
+			if(!P4Output(out Output, "", "print -q " + DepotPath, AllowSpew: AllowSpew, WithClient: false))
 			{
 				throw new AutomationException("p4 print {0} failed", DepotPath);
 			}
@@ -3641,7 +3642,7 @@ namespace AutomationTool
 		public void PrintToFile(string DepotPath, string FileName, bool AllowSpew = true)
 		{
 			string Output;
-			if(!P4Output(out Output, "print -q -o \"" + FileName + "\" " + DepotPath, AllowSpew: AllowSpew, WithClient: false))
+			if(!P4Output(out Output, "", "print -q -o \"" + FileName + "\" " + DepotPath, AllowSpew: AllowSpew, WithClient: false))
 			{
 				throw new AutomationException("p4 print {0} failed", DepotPath);
 			}
@@ -3660,7 +3661,7 @@ namespace AutomationTool
 		public List<int> StreamInterchanges(string StreamName, bool bReverse)
 		{
 			string Output;
-			if(!P4Output(out Output, String.Format("interchanges {0}-S {1} -F", bReverse? "-r " : "", StreamName), Input:null, AllowSpew:false))
+			if(!P4Output(out Output, "", String.Format("interchanges {0}-S {1} -F", bReverse? "-r " : "", StreamName), Input:null, AllowSpew:false))
 			{
 				throw new AutomationException("Couldn't get unintegrated stream changes from {0}", StreamName);
 			}
@@ -4049,7 +4050,7 @@ namespace AutomationTool
 		public string GetIntegrationSource(string DepotPath)
 		{
 			string Output;
-			if(P4Output(out Output, "filelog -m 1 \"" + DepotPath + "\"", Input:null, AllowSpew:false))
+			if(P4Output(out Output, "", "filelog -m 1 \"" + DepotPath + "\"", Input:null, AllowSpew:false))
 			{
 				foreach(string Line in Output.Split('\n').Select(x => x.Trim()))
 				{
diff --git a/Engine/Source/Programs/AutomationTool/Scripts/FixPerforceCase.cs b/Engine/Source/Programs/AutomationTool/Scripts/FixPerforceCase.cs
--- a/Engine/Source/Programs/AutomationTool/Scripts/FixPerforceCase.cs
+++ b/Engine/Source/Programs/AutomationTool/Scripts/FixPerforceCase.cs
@@ -97,14 +97,14 @@ namespace AutomationTool
 			// Force sync all the old files 
 			foreach (string OldFile in SourceFiles)
 			{
-				P4.LogP4(String.Format("sync -f {0}", OldFile));
+				P4.LogP4("", String.Format("sync -f {0}", OldFile));
 			}
 
 			// Delete all the old files 
 			int DeleteChangeNumber = P4.CreateChange(Description: String.Format("Fixing case of {0} (1/2){1}", FileSpec, BoilerplateText));
 			foreach (string OldFile in SourceFiles)
 			{
-				P4.LogP4(String.Format("delete -k -c {0} {1}", DeleteChangeNumber, OldFile));
+				P4.LogP4("", String.Format("delete -k -c {0} {1}", DeleteChangeNumber, OldFile));
 			}
 			P4.Submit(DeleteChangeNumber);
 
@@ -112,7 +112,7 @@ namespace AutomationTool
 			int AddChangeNumber = P4.CreateChange(Description: String.Format("Fixing case of {0} (2/2){1}", FileSpec, BoilerplateText));
 			foreach (string NewFile in TargetFiles)
 			{
-				P4.LogP4(String.Format("add -c {0} {1}", AddChangeNumber, NewFile));
+				P4.LogP4("", String.Format("add -c {0} {1}", AddChangeNumber, NewFile));
 			}
 			P4.Submit(AddChangeNumber);
 		}
diff --git a/Engine/Source/Programs/AutomationTool/Scripts/Localisation.Automation.cs b/Engine/Source/Programs/AutomationTool/Scripts/Localisation.Automation.cs
--- a/Engine/Source/Programs/AutomationTool/Scripts/Localisation.Automation.cs
+++ b/Engine/Source/Programs/AutomationTool/Scripts/Localisation.Automation.cs
@@ -428,7 +428,7 @@ class Localize : BuildCommand
 						}
 					}
 
-					P4.LogP4(String.Format("-x{0} revert", P4RevertArgsFilename));
+					P4.LogP4("", String.Format("-x{0} revert", P4RevertArgsFilename));
 					DeleteFile_NoExceptions(P4RevertArgsFilename);
 				}
 			}
diff --git a/Engine/Source/Programs/AutomationTool/Scripts/Tests.Automation.cs b/Engine/Source/Programs/AutomationTool/Scripts/Tests.Automation.cs
--- a/Engine/Source/Programs/AutomationTool/Scripts/Tests.Automation.cs
+++ b/Engine/Source/Programs/AutomationTool/Scripts/Tests.Automation.cs
@@ -593,7 +593,7 @@ class TestP4_ClientOps : BuildCommand
 
 		//Sync(CombinePaths(PathSeparator.Slash, P4Env.BuildRootP4, "UE4Games.uprojectdirs"));
 		string Output;
-		P4.P4Output(out Output, "files -m 10 " + CombinePaths(PathSeparator.Slash, P4Env.Branch, "..."));
+		P4.P4Output(out Output, "", "files -m 10 " + CombinePaths(PathSeparator.Slash, P4Env.Branch, "..."));
 
 		var Lines = Output.Split(new string[] { Environment.NewLine }, StringSplitOptions.RemoveEmptyEntries);
 		string SlashRoot = CombinePaths(PathSeparator.Slash, P4Env.Branch, "/");
@@ -614,7 +614,7 @@ class TestP4_ClientOps : BuildCommand
 		}
 
 		// should be used as a sanity check! make sure there are no files!
-		P4.LogP4Output(out Output, "files -m 10 " + CombinePaths(PathSeparator.Slash, P4Env.Branch, "...", "NoRedist", "..."));
+		P4.LogP4Output(out Output, "", "files -m 10 " + CombinePaths(PathSeparator.Slash, P4Env.Branch, "...", "NoRedist", "..."));
 
 		// caution this doesn't actually use the client _view_ from the template at all!
 		// if you want that sync -f -k /...
