Fixedup path to UAT for netcore projects

diff --git a/Engine/Source/Programs/AutomationTool/AutomationUtils/CommandEnvironment.cs b/Engine/Source/Programs/AutomationTool/AutomationUtils/CommandEnvironment.cs
--- a/Engine/Source/Programs/AutomationTool/AutomationUtils/CommandEnvironment.cs
+++ b/Engine/Source/Programs/AutomationTool/AutomationUtils/CommandEnvironment.cs
@@ -71,7 +71,12 @@ namespace AutomationTool
 		internal CommandEnvironment()
 		{
 			// Get the path to the UAT executable
+#if NET_CORE
+			// the entry assembly is the .dll but it is easier to use apphost so change extension to the executable
+			UATExe = Path.ChangeExtension(Assembly.GetEntryAssembly().GetOriginalLocation(), Utils.IsRunningOnWindows ? "exe" : null);
+#else
 			UATExe = Assembly.GetEntryAssembly().GetOriginalLocation();
+#endif
 			if (!CommandUtils.FileExists(UATExe))
 			{
 				throw new AutomationException("Could not find AutomationTool.exe. Reflection indicated it was here: {0}", UATExe);
