- Missed a p4 command ExtraGlobalOptions

diff --git a/Engine/Source/Programs/AutomationTool/AutomationUtils/P4Utils.cs b/Engine/Source/Programs/AutomationTool/AutomationUtils/P4Utils.cs
--- a/Engine/Source/Programs/AutomationTool/AutomationUtils/P4Utils.cs
+++ b/Engine/Source/Programs/AutomationTool/AutomationUtils/P4Utils.cs
@@ -3552,7 +3552,7 @@ namespace AutomationTool
 				}
 			}
 			if (AllowSpew) CommandUtils.LogLog(SpecInput);
-            LogP4("client -i", SpecInput, AllowSpew: AllowSpew, WithClient: false);
+            LogP4("", "client -i", SpecInput, AllowSpew: AllowSpew, WithClient: false);
             return ClientSpec;
         }
 
