Moved UAT plugin building to a seperate step controlled by AutomationTool.proj, this makes the process of building addins identical for installed builds and engine builds (but we still apply seperate logic to them).

Added CopyLocalLockFileAssemblies to some addins that added new dependencies (nuget packages) to make sure these gets added to the output folder.

UAT now have a dependency on WindowsForms (on Windows) to enable the WinForms usage in Gauntlet and Turnkey, but is still a console application so just making sure the right assemblies are referenced.

Fixed up shootergame csproj were I had accidentally removed EngineDir props which would have prevented it from working for installed builds as a foreign project.

diff --git a/Engine/Build/BatchFiles/RunUAT.bat b/Engine/Build/BatchFiles/RunUAT.bat
--- a/Engine/Build/BatchFiles/RunUAT.bat
+++ b/Engine/Build/BatchFiles/RunUAT.bat
@@ -26,9 +26,13 @@ rem ## Change the CWD to /Engine.
 pushd "%~dp0..\..\"
 if not exist Build\BatchFiles\RunUAT.bat goto Error_BatchFileInWrongLocation
 
+set MSBUILD_LOGLEVEL=quiet
+for %%P in (%*) do if /I "%%P" == "-msbuild-verbose" set MSBUILD_LOGLEVEL=normal
+
 rem ## Use the pre-compiled UAT scripts if -nocompile is specified in the command line
 for %%P in (%*) do if /I "%%P" == "-nocompile" goto RunPrecompiled
 
+
 rem ## If we're running in an installed build, default to precompiled
 if exist Build\InstalledBuild.txt goto RunPrecompiled
 
@@ -59,7 +63,10 @@ call "%~dp0GetDotnetPath.bat"
 if errorlevel 1 goto Error_NoDotnetSDK
 
 echo Building AutomationTool...
-dotnet build Source\Programs\AutomationTool\AutomationToolCore.csproj -c Development -v quiet 
+dotnet msbuild /restore /property:Configuration=Development /property:AutomationToolProjectOnly=true /verbosity:%MSBUILD_LOGLEVEL% Source\Programs\AutomationTool\AutomationToolCore.csproj
+if errorlevel 1 goto Error_UATCompileFailed
+echo Building AutomationTool Plugins...
+dotnet msbuild /restore /property:Configuration=Development /verbosity:%MSBUILD_LOGLEVEL% Source\Programs\AutomationTool\AutomationTool.proj
 if errorlevel 1 goto Error_UATCompileFailed
 goto DoRunUAT
 
diff --git a/Engine/Build/BatchFiles/RunUAT.sh b/Engine/Build/BatchFiles/RunUAT.sh
--- a/Engine/Build/BatchFiles/RunUAT.sh
+++ b/Engine/Build/BatchFiles/RunUAT.sh
@@ -79,11 +79,19 @@ if [ "$UATCompileArg" = "-compile" ]; then
 		UATCompileArg=
 	else
 		if [ ${UE_USE_DOTNET:=0} -ne 0 ]; then
-			dotnet build Source/Programs/AutomationTool/AutomationToolCore.csproj -c Development
+			echo Building AutomationTool...
+			dotnet msbuild -restore Source/Programs/AutomationTool/AutomationToolCore.csproj /property:Configuration=Development /property:AutomationToolProjectOnly=true /verbosity:quiet
 			if [ $? -ne 0 ]; then
 				echo RunUAT ERROR: AutomationTool failed to compile.
 				exit 1
 			fi
+
+			echo Building AutomationTool Plugins...
+			dotnet msbuild -restore Source/Programs/AutomationTool/AutomationTool.proj /property:Configuration=Development /verbosity:quiet
+			if [ $? -ne 0 ]; then
+				echo RunUAT ERROR: AutomationTool plugins failed to compile.
+				exit 1
+			fi
 		else
 			# mono 5.0 and up include msbuild
 			if [ "$IS_MS_BUILD_AVAILABLE" == "1" ]; then
diff --git a/Engine/Source/Programs/AutomationTool/AutomationToolCore.csproj b/Engine/Source/Programs/AutomationTool/AutomationToolCore.csproj
--- a/Engine/Source/Programs/AutomationTool/AutomationToolCore.csproj
+++ b/Engine/Source/Programs/AutomationTool/AutomationToolCore.csproj
@@ -4,7 +4,7 @@
     <BaseIntermediateOutputPath>$(MSBuildThisFileDirectory)obj-core\$(MSBuildProjectName)\</BaseIntermediateOutputPath>
     <BaseOutputPath>$(MSBuildThisFileDirectory)bin-core\$(MSBuildProjectName)\</BaseOutputPath>
   </PropertyGroup>
-  <Import Project="Sdk.props" Sdk="Microsoft.NET.Sdk" />
+
   <PropertyGroup>
     <TargetFramework>netcoreapp3.1</TargetFramework>
     <Configuration Condition=" '$(Configuration)' == '' ">Development</Configuration>
@@ -44,6 +44,11 @@
         <RuntimeIdentifier>osx-x64</RuntimeIdentifier>
       </PropertyGroup>
     </When>
+    <When Condition="'$(IsWindows)'=='true'">
+      <PropertyGroup>
+        <UseWindowsForms>true</UseWindowsForms>
+      </PropertyGroup>
+    </When>
   </Choose>
   <ItemGroup>
     <Reference Include="Ionic.Zip.Reduced">
@@ -51,8 +56,21 @@
       <Private>true</Private>
     </Reference>
   </ItemGroup>
-  <Import Project="Sdk.targets" Sdk="Microsoft.NET.Sdk" />
 
+  <Import Project="Sdk.props" Sdk="Microsoft.NET.Sdk.WindowsDesktop" Condition="'$(IsWindows)'=='true'"/>
+  <Import Project="Sdk.targets" Sdk="Microsoft.NET.Sdk.WindowsDesktop" Condition="'$(IsWindows)'=='true'"/>
+  <Import Project="Sdk.props" Sdk="Microsoft.NET.Sdk" Condition="'$(IsWindows)'!='true'"/>
+  <Import Project="Sdk.targets" Sdk="Microsoft.NET.Sdk" Condition="'$(IsWindows)'!='true'"/>
+
+  <Choose>
+    <When Condition="'$(IsWindows)'=='true'">
+      <PropertyGroup>
+        <!-- Reset output type back to a console application as the import of WindowsDesktop sets it to a windows executable -->
+        <OutputType>Exe</OutputType>
+      </PropertyGroup>
+    </When>
+  </Choose>
+  
   <ItemGroup>
     <Compile Remove="obj\**" />
     <EmbeddedResource Remove="obj\**" />
@@ -76,8 +94,6 @@
     <ProjectReference Include="AutomationUtils\AutomationUtils.Automation.Core.csproj" />
   </ItemGroup>
 
-  <!-- If building inside visual studio it means we generate projects, thus UBT will have output this file and we do not need to add the relativley slow task of globbing this in msbuild -->
-  <Import Project="AutomationTool.proj" Condition="'$(AutomationToolProjectOnly)' == '' AND '$(BuildingInsideVisualStudio)' != 'true'" />
   <Import Project="AutomationToolCore.csproj.References" Condition="Exists('AutomationToolCore.csproj.References') And '$(AutomationToolProjectOnly)' == ''" />
   
 </Project>
\ No newline at end of file
diff --git a/Engine/Source/Programs/AutomationTool/AutomationUtils/AutomationUtils.Automation.Core.csproj b/Engine/Source/Programs/AutomationTool/AutomationUtils/AutomationUtils.Automation.Core.csproj
--- a/Engine/Source/Programs/AutomationTool/AutomationUtils/AutomationUtils.Automation.Core.csproj
+++ b/Engine/Source/Programs/AutomationTool/AutomationUtils/AutomationUtils.Automation.Core.csproj
@@ -22,6 +22,7 @@
     <OutputPath>..\..\..\..\Binaries\DotNET\AutomationTool</OutputPath>
     <Configurations>Debug;Release;Development</Configurations>
     <TreatWarningsAsErrors>true</TreatWarningsAsErrors>
+    <AppendTargetFrameworkToOutputPath>false</AppendTargetFrameworkToOutputPath>
   </PropertyGroup>
     <ItemGroup>
     <Reference Include="Ionic.Zip.Reduced">
diff --git a/Engine/Source/Programs/AutomationTool/Turnkey/Turnkey.Automation.Core.csproj b/Engine/Source/Programs/AutomationTool/Turnkey/Turnkey.Automation.Core.csproj
--- a/Engine/Source/Programs/AutomationTool/Turnkey/Turnkey.Automation.Core.csproj
+++ b/Engine/Source/Programs/AutomationTool/Turnkey/Turnkey.Automation.Core.csproj
@@ -80,6 +80,5 @@
     <PackageReference Include="System.Data.DataSetExtensions" Version="4.5.0" />
     <PackageReference Include="System.Net.Http" Version="4.3.4" />
     <PackageReference Include="System.Drawing.Common" Version="4.7.0" />
-    <PackageReference Include="Microsoft.VisualBasic" Version="10.3.0" />
   </ItemGroup>
 </Project>
\ No newline at end of file
diff --git a/Engine/Source/Programs/AutomationTool/XLocLocalization/XLocLocalization.Automation.Core.csproj b/Engine/Source/Programs/AutomationTool/XLocLocalization/XLocLocalization.Automation.Core.csproj
--- a/Engine/Source/Programs/AutomationTool/XLocLocalization/XLocLocalization.Automation.Core.csproj
+++ b/Engine/Source/Programs/AutomationTool/XLocLocalization/XLocLocalization.Automation.Core.csproj
@@ -19,6 +19,7 @@
     <OutputPath>..\..\..\..\Binaries\DotNET\AutomationTool\AutomationScripts</OutputPath>
     <DefineConstants>NET_CORE</DefineConstants>
     <TreatWarningsAsErrors>true</TreatWarningsAsErrors>
+    <CopyLocalLockFileAssemblies>true</CopyLocalLockFileAssemblies>
   </PropertyGroup>
   <PropertyGroup Condition=" '$(Configuration)|$(Platform)' == 'Development|AnyCPU' ">
     <DefineConstants>NET_CORE;TRACE</DefineConstants>
