Changed output directory for UAT projects under netcore, now outputting to a seperate directory for each addin. This avoids the risk of output stomping each other breaking incremental builds. Furthermore it makes it easier to understand who adds what dependency to the output.

Also cleaned up some output assemblies, preventing all transative references from being outputed (as we assume they build directly to the output directory instead). Unfortunatley the UBT references still cause tranastive dependencies to be copied as this is also a executable that needs to work outside of UAT as well, we should eventually move all UBT referenced code into BuildUtilities.

diff --git a/Engine/Source/Programs/AutomationTool/Android/Android.Automation.Core.csproj b/Engine/Source/Programs/AutomationTool/Android/Android.Automation.Core.csproj
--- a/Engine/Source/Programs/AutomationTool/Android/Android.Automation.Core.csproj
+++ b/Engine/Source/Programs/AutomationTool/Android/Android.Automation.Core.csproj
@@ -17,7 +17,7 @@
     <RootNamespace>Android.Automation</RootNamespace>
     <AssemblyName>Android.Automation</AssemblyName>
     <WarningsNotAsErrors>612,618</WarningsNotAsErrors>
-    <OutputPath>..\..\..\..\Binaries\DotNET\AutomationTool\AutomationScripts\Android</OutputPath>
+    <OutputPath>..\..\..\..\Binaries\DotNET\AutomationTool\AutomationScripts\Platforms\Android</OutputPath>
     <DefineConstants>NET_CORE</DefineConstants>
     <TreatWarningsAsErrors>true</TreatWarningsAsErrors>
   </PropertyGroup>
@@ -53,10 +53,10 @@
     </Compile>
   </ItemGroup>
   <ItemGroup>
-    <ProjectReference Include="..\..\DotNETCommon\DotNETUtilities\DotNETUtilitiesCore.csproj"><Private>false</Private></ProjectReference>
-    <ProjectReference Include="..\..\UnrealBuildTool\UnrealBuildToolCore.csproj"><Private>false</Private></ProjectReference>
-    <ProjectReference Include="..\AutomationUtils\AutomationUtils.Automation.Core.csproj"><Private>false</Private></ProjectReference>
-    <ProjectReference Include="..\Scripts\AutomationScripts.Automation.Core.csproj"><Private>false</Private></ProjectReference>
+    <ProjectReference Include="..\..\DotNETCommon\DotNETUtilities\DotNETUtilitiesCore.csproj" PrivateAssets="All"><Private>false</Private></ProjectReference>
+    <ProjectReference Include="..\..\UnrealBuildTool\UnrealBuildToolCore.csproj" PrivateAssets="All"><Private>false</Private></ProjectReference>
+    <ProjectReference Include="..\AutomationUtils\AutomationUtils.Automation.Core.csproj" PrivateAssets="All"><Private>false</Private></ProjectReference>
+    <ProjectReference Include="..\Scripts\AutomationScripts.Automation.Core.csproj" PrivateAssets="All"><Private>false</Private></ProjectReference>
   </ItemGroup>
   <ItemGroup>
     <PackageReference Include="Microsoft.CSharp" Version="4.7.0" />
diff --git a/Engine/Source/Programs/AutomationTool/AutomationUtils/AutomationUtils.Automation.Core.csproj b/Engine/Source/Programs/AutomationTool/AutomationUtils/AutomationUtils.Automation.Core.csproj
--- a/Engine/Source/Programs/AutomationTool/AutomationUtils/AutomationUtils.Automation.Core.csproj
+++ b/Engine/Source/Programs/AutomationTool/AutomationUtils/AutomationUtils.Automation.Core.csproj
@@ -50,8 +50,8 @@
   </ItemGroup>
 
   <ItemGroup>
-    <ProjectReference Include="..\..\DotNETCommon\DotNETUtilities\DotNETUtilitiesCore.csproj"><Private>false</Private></ProjectReference>
-    <ProjectReference Include="..\..\UnrealBuildTool\UnrealBuildToolCore.csproj"><Private>false</Private></ProjectReference>
+    <ProjectReference Include="..\..\DotNETCommon\DotNETUtilities\DotNETUtilitiesCore.csproj" PrivateAssets="All"><Private>false</Private></ProjectReference>
+    <ProjectReference Include="..\..\UnrealBuildTool\UnrealBuildToolCore.csproj" PrivateAssets="All"><Private>false</Private></ProjectReference>
   </ItemGroup>
   <ItemGroup>
     <PackageReference Include="Microsoft.CSharp" Version="4.7.0" />
diff --git a/Engine/Source/Programs/AutomationTool/BuildGraph/BuildGraph.Automation.Core.csproj b/Engine/Source/Programs/AutomationTool/BuildGraph/BuildGraph.Automation.Core.csproj
--- a/Engine/Source/Programs/AutomationTool/BuildGraph/BuildGraph.Automation.Core.csproj
+++ b/Engine/Source/Programs/AutomationTool/BuildGraph/BuildGraph.Automation.Core.csproj
@@ -16,7 +16,7 @@
     <RootNamespace>AutomationTool</RootNamespace>
     <AssemblyName>BuildGraph.Automation</AssemblyName>
     <WarningsNotAsErrors>612,618</WarningsNotAsErrors>
-    <OutputPath>..\..\..\..\Binaries\DotNET\AutomationTool\AutomationScripts</OutputPath>
+    <OutputPath>..\..\..\..\Binaries\DotNET\AutomationTool\AutomationScripts\BuildGraph</OutputPath>
     <DefineConstants>NET_CORE</DefineConstants>
     <TreatWarningsAsErrors>true</TreatWarningsAsErrors>
   </PropertyGroup>
@@ -52,9 +52,9 @@
     </Compile>
   </ItemGroup>
   <ItemGroup>
-    <ProjectReference Include="..\..\DotNETCommon\DotNETUtilities\DotNETUtilitiesCore.csproj"><Private>false</Private></ProjectReference>
-    <ProjectReference Include="..\..\UnrealBuildTool\UnrealBuildToolCore.csproj"><Private>false</Private></ProjectReference>
-    <ProjectReference Include="..\AutomationUtils\AutomationUtils.Automation.Core.csproj"><Private>false</Private></ProjectReference>
+    <ProjectReference Include="..\..\DotNETCommon\DotNETUtilities\DotNETUtilitiesCore.csproj" PrivateAssets="All"><Private>false</Private></ProjectReference>
+    <ProjectReference Include="..\..\UnrealBuildTool\UnrealBuildToolCore.csproj" PrivateAssets="All"><Private>false</Private></ProjectReference>
+    <ProjectReference Include="..\AutomationUtils\AutomationUtils.Automation.Core.csproj" PrivateAssets="All"><Private>false</Private></ProjectReference>
   </ItemGroup>
   <ItemGroup>
     <PackageReference Include="Microsoft.CSharp" Version="4.7.0" />
diff --git a/Engine/Source/Programs/AutomationTool/Gauntlet/Gauntlet.Automation.Core.csproj b/Engine/Source/Programs/AutomationTool/Gauntlet/Gauntlet.Automation.Core.csproj
--- a/Engine/Source/Programs/AutomationTool/Gauntlet/Gauntlet.Automation.Core.csproj
+++ b/Engine/Source/Programs/AutomationTool/Gauntlet/Gauntlet.Automation.Core.csproj
@@ -16,7 +16,7 @@
     <RootNamespace>Gauntlet.Automation</RootNamespace>
     <AssemblyName>Gauntlet.Automation</AssemblyName>
     <WarningsNotAsErrors>612,618</WarningsNotAsErrors>
-    <OutputPath>..\..\..\..\Binaries\DotNET\AutomationTool\AutomationScripts</OutputPath>
+    <OutputPath>..\..\..\..\Binaries\DotNET\AutomationTool\AutomationScripts\Gauntlet</OutputPath>
     <DefineConstants>NET_CORE</DefineConstants>
     <TreatWarningsAsErrors>true</TreatWarningsAsErrors>
   </PropertyGroup>
@@ -60,15 +60,15 @@
   <Import Project="../../../../Platforms/*/Source/Programs/AutomationTool/Gauntlet/*.Gauntlet.targets" Condition="'$(OS)' == 'Windows_NT'" />
 
   <ItemGroup>
-    <ProjectReference Include="..\..\DotNETCommon\DotNETUtilities\DotNETUtilitiesCore.csproj"><Private>false</Private></ProjectReference>
-    <ProjectReference Include="..\..\UnrealBuildTool\UnrealBuildToolCore.csproj"><Private>false</Private></ProjectReference>
-    <ProjectReference Include="..\AutomationUtils\AutomationUtils.Automation.Core.csproj"><Private>false</Private></ProjectReference>
-    <ProjectReference Include="..\Scripts\AutomationScripts.Automation.Core.csproj"><Private>false</Private></ProjectReference>
-    <ProjectReference Include="..\OneSkyLocalization\OneSkyLocalization.Automation.Core.csproj"><Private>false</Private></ProjectReference>
-    <ProjectReference Include="..\XLocLocalization\XLocLocalization.Automation.Core.csproj"><Private>false</Private></ProjectReference>
-    <ProjectReference Include="..\Localization\Localization.Automation.Core.csproj"><Private>false</Private></ProjectReference>
-    <ProjectReference Include="..\Lumin\Lumin.Automation.Core.csproj"><Private>false</Private></ProjectReference>
-    <ProjectReference Include="..\Android\Android.Automation.Core.csproj"><Private>false</Private></ProjectReference>
+    <ProjectReference Include="..\..\DotNETCommon\DotNETUtilities\DotNETUtilitiesCore.csproj" PrivateAssets="All"><Private>false</Private></ProjectReference>
+    <ProjectReference Include="..\..\UnrealBuildTool\UnrealBuildToolCore.csproj" PrivateAssets="All"><Private>false</Private></ProjectReference>
+    <ProjectReference Include="..\AutomationUtils\AutomationUtils.Automation.Core.csproj" PrivateAssets="All"><Private>false</Private></ProjectReference>
+    <ProjectReference Include="..\Scripts\AutomationScripts.Automation.Core.csproj" PrivateAssets="All"><Private>false</Private></ProjectReference>
+    <ProjectReference Include="..\OneSkyLocalization\OneSkyLocalization.Automation.Core.csproj" PrivateAssets="All"><Private>false</Private></ProjectReference>
+    <ProjectReference Include="..\XLocLocalization\XLocLocalization.Automation.Core.csproj" PrivateAssets="All"><Private>false</Private></ProjectReference>
+    <ProjectReference Include="..\Localization\Localization.Automation.Core.csproj" PrivateAssets="All"><Private>false</Private></ProjectReference>
+    <ProjectReference Include="..\Lumin\Lumin.Automation.Core.csproj" PrivateAssets="All"><Private>false</Private></ProjectReference>
+    <ProjectReference Include="..\Android\Android.Automation.Core.csproj" PrivateAssets="All"><Private>false</Private></ProjectReference>
   </ItemGroup>
   <ItemGroup>
     <PackageReference Include="Microsoft.CSharp" Version="4.7.0" />
diff --git a/Engine/Source/Programs/AutomationTool/HoloLens/HoloLens.Automation.Core.csproj b/Engine/Source/Programs/AutomationTool/HoloLens/HoloLens.Automation.Core.csproj
--- a/Engine/Source/Programs/AutomationTool/HoloLens/HoloLens.Automation.Core.csproj
+++ b/Engine/Source/Programs/AutomationTool/HoloLens/HoloLens.Automation.Core.csproj
@@ -16,7 +16,7 @@
     <RootNamespace>HoloLens.Automation</RootNamespace>
     <AssemblyName>HoloLens.Automation</AssemblyName>
     <WarningsNotAsErrors>612,618</WarningsNotAsErrors>
-    <OutputPath>..\..\..\..\Binaries\DotNET\AutomationTool\AutomationScripts\HoloLens\</OutputPath>
+    <OutputPath>..\..\..\..\Binaries\DotNET\AutomationTool\AutomationScripts\Platforms\HoloLens\</OutputPath>
     <DefineConstants>NET_CORE</DefineConstants>
     <TreatWarningsAsErrors>true</TreatWarningsAsErrors>
     <IsWindows Condition="'$([System.Runtime.InteropServices.RuntimeInformation]::IsOSPlatform($([System.Runtime.InteropServices.OSPlatform]::Windows)))' == 'true'">true</IsWindows> 
@@ -71,10 +71,10 @@
     </Compile>
   </ItemGroup>
   <ItemGroup>
-    <ProjectReference Include="..\..\DotNETCommon\DotNETUtilities\DotNETUtilitiesCore.csproj"><Private>false</Private></ProjectReference>
-    <ProjectReference Include="..\..\UnrealBuildTool\UnrealBuildToolCore.csproj"><Private>false</Private></ProjectReference>
-    <ProjectReference Include="..\AutomationUtils\AutomationUtils.Automation.Core.csproj"><Private>false</Private></ProjectReference>
-    <ProjectReference Include="..\Scripts\AutomationScripts.Automation.Core.csproj"><Private>false</Private></ProjectReference>
+    <ProjectReference Include="..\..\DotNETCommon\DotNETUtilities\DotNETUtilitiesCore.csproj" PrivateAssets="All"><Private>false</Private></ProjectReference>
+    <ProjectReference Include="..\..\UnrealBuildTool\UnrealBuildToolCore.csproj" PrivateAssets="All"><Private>false</Private></ProjectReference>
+    <ProjectReference Include="..\AutomationUtils\AutomationUtils.Automation.Core.csproj" PrivateAssets="All"><Private>false</Private></ProjectReference>
+    <ProjectReference Include="..\Scripts\AutomationScripts.Automation.Core.csproj" PrivateAssets="All"><Private>false</Private></ProjectReference>
   </ItemGroup>
   <ItemGroup>
     <PackageReference Include="Microsoft.CSharp" Version="4.7.0" />
diff --git a/Engine/Source/Programs/AutomationTool/IOS/IOS.Automation.Core.csproj b/Engine/Source/Programs/AutomationTool/IOS/IOS.Automation.Core.csproj
--- a/Engine/Source/Programs/AutomationTool/IOS/IOS.Automation.Core.csproj
+++ b/Engine/Source/Programs/AutomationTool/IOS/IOS.Automation.Core.csproj
@@ -16,7 +16,7 @@
     <RootNamespace>IOS.Automation</RootNamespace>
     <AssemblyName>IOS.Automation</AssemblyName>
     <WarningsNotAsErrors>612,618</WarningsNotAsErrors>
-    <OutputPath>..\..\..\..\Binaries\DotNET\AutomationTool\AutomationScripts\IOS</OutputPath>
+    <OutputPath>..\..\..\..\Binaries\DotNET\AutomationTool\AutomationScripts\Platforms\IOS</OutputPath>
     <DefineConstants>NET_CORE</DefineConstants>
     <TreatWarningsAsErrors>true</TreatWarningsAsErrors>
   </PropertyGroup>
@@ -57,9 +57,9 @@
     </Compile>
   </ItemGroup>
   <ItemGroup>
-    <ProjectReference Include="..\..\DotNETCommon\DotNETUtilities\DotNETUtilitiesCore.csproj"><Private>false</Private></ProjectReference>
-    <ProjectReference Include="..\..\UnrealBuildTool\UnrealBuildToolCore.csproj"><Private>false</Private></ProjectReference>
-    <ProjectReference Include="..\AutomationUtils\AutomationUtils.Automation.Core.csproj"><Private>false</Private></ProjectReference>
+    <ProjectReference Include="..\..\DotNETCommon\DotNETUtilities\DotNETUtilitiesCore.csproj" PrivateAssets="All"><Private>false</Private></ProjectReference>
+    <ProjectReference Include="..\..\UnrealBuildTool\UnrealBuildToolCore.csproj" PrivateAssets="All"><Private>false</Private></ProjectReference>
+    <ProjectReference Include="..\AutomationUtils\AutomationUtils.Automation.Core.csproj" PrivateAssets="All"><Private>false</Private></ProjectReference>
   </ItemGroup>
   <ItemGroup>
     <PackageReference Include="Microsoft.CSharp" Version="4.7.0" />
diff --git a/Engine/Source/Programs/AutomationTool/Linux/Linux.Automation.Core.csproj b/Engine/Source/Programs/AutomationTool/Linux/Linux.Automation.Core.csproj
--- a/Engine/Source/Programs/AutomationTool/Linux/Linux.Automation.Core.csproj
+++ b/Engine/Source/Programs/AutomationTool/Linux/Linux.Automation.Core.csproj
@@ -16,7 +16,7 @@
     <RootNamespace>Linux.Automation</RootNamespace>
     <AssemblyName>Linux.Automation</AssemblyName>
     <WarningsNotAsErrors>612,618</WarningsNotAsErrors>
-    <OutputPath>..\..\..\..\Binaries\DotNET\AutomationTool\AutomationScripts\Linux\</OutputPath>
+    <OutputPath>..\..\..\..\Binaries\DotNET\AutomationTool\AutomationScripts\Platforms\Linux\</OutputPath>
     <DefineConstants>NET_CORE</DefineConstants>
     <TreatWarningsAsErrors>true</TreatWarningsAsErrors>
   </PropertyGroup>
@@ -46,9 +46,9 @@
     </Compile>
   </ItemGroup>
   <ItemGroup>
-    <ProjectReference Include="..\..\DotNETCommon\DotNETUtilities\DotNETUtilitiesCore.csproj"><Private>false</Private></ProjectReference>
-    <ProjectReference Include="..\..\UnrealBuildTool\UnrealBuildToolCore.csproj"><Private>false</Private></ProjectReference>
-    <ProjectReference Include="..\AutomationUtils\AutomationUtils.Automation.Core.csproj"><Private>false</Private></ProjectReference>
+    <ProjectReference Include="..\..\DotNETCommon\DotNETUtilities\DotNETUtilitiesCore.csproj" PrivateAssets="All"><Private>false</Private></ProjectReference>
+    <ProjectReference Include="..\..\UnrealBuildTool\UnrealBuildToolCore.csproj" PrivateAssets="All"><Private>false</Private></ProjectReference>
+    <ProjectReference Include="..\AutomationUtils\AutomationUtils.Automation.Core.csproj" PrivateAssets="All"><Private>false</Private></ProjectReference>
   </ItemGroup>
   <ItemGroup>
     <PackageReference Include="Microsoft.CSharp" Version="4.7.0" />
diff --git a/Engine/Source/Programs/AutomationTool/Localization/Localization.Automation.Core.csproj b/Engine/Source/Programs/AutomationTool/Localization/Localization.Automation.Core.csproj
--- a/Engine/Source/Programs/AutomationTool/Localization/Localization.Automation.Core.csproj
+++ b/Engine/Source/Programs/AutomationTool/Localization/Localization.Automation.Core.csproj
@@ -16,7 +16,7 @@
     <RootNamespace>AutomationTool</RootNamespace>
     <AssemblyName>Localization.Automation</AssemblyName>
     <WarningsNotAsErrors>612,618</WarningsNotAsErrors>
-    <OutputPath>..\..\..\..\Binaries\DotNET\AutomationTool\AutomationScripts</OutputPath>
+    <OutputPath>..\..\..\..\Binaries\DotNET\AutomationTool\AutomationScripts\Localization</OutputPath>
     <DefineConstants>NET_CORE</DefineConstants>
     <TreatWarningsAsErrors>true</TreatWarningsAsErrors>
   </PropertyGroup>
@@ -46,9 +46,9 @@
     </Compile>
   </ItemGroup>
   <ItemGroup>
-    <ProjectReference Include="..\..\DotNETCommon\DotNETUtilities\DotNETUtilitiesCore.csproj"><Private>false</Private></ProjectReference>
-    <ProjectReference Include="..\..\UnrealBuildTool\UnrealBuildToolCore.csproj"><Private>false</Private></ProjectReference>
-    <ProjectReference Include="..\AutomationUtils\AutomationUtils.Automation.Core.csproj"><Private>false</Private></ProjectReference>
+    <ProjectReference Include="..\..\DotNETCommon\DotNETUtilities\DotNETUtilitiesCore.csproj" PrivateAssets="All"><Private>false</Private></ProjectReference>
+    <ProjectReference Include="..\..\UnrealBuildTool\UnrealBuildToolCore.csproj" PrivateAssets="All"><Private>false</Private></ProjectReference>
+    <ProjectReference Include="..\AutomationUtils\AutomationUtils.Automation.Core.csproj" PrivateAssets="All"><Private>false</Private></ProjectReference>
   </ItemGroup>
   <ItemGroup>
     <PackageReference Include="System.Data.DataSetExtensions" Version="4.5.0" />
diff --git a/Engine/Source/Programs/AutomationTool/Lumin/Lumin.Automation.Core.csproj b/Engine/Source/Programs/AutomationTool/Lumin/Lumin.Automation.Core.csproj
--- a/Engine/Source/Programs/AutomationTool/Lumin/Lumin.Automation.Core.csproj
+++ b/Engine/Source/Programs/AutomationTool/Lumin/Lumin.Automation.Core.csproj
@@ -16,7 +16,7 @@
     <RootNamespace>Lumin.Automation</RootNamespace>
     <AssemblyName>Lumin.Automation</AssemblyName>
     <WarningsNotAsErrors>612,618</WarningsNotAsErrors>
-    <OutputPath>..\..\..\..\Binaries\DotNET\AutomationTool\AutomationScripts\Lumin\</OutputPath>
+    <OutputPath>..\..\..\..\Binaries\DotNET\AutomationTool\AutomationScripts\Platforms\Lumin\</OutputPath>
     <DefineConstants>NET_CORE</DefineConstants>
     <TreatWarningsAsErrors>true</TreatWarningsAsErrors>
   </PropertyGroup>
@@ -52,10 +52,10 @@
     </Compile>
   </ItemGroup>
   <ItemGroup>
-    <ProjectReference Include="..\..\DotNETCommon\DotNETUtilities\DotNETUtilitiesCore.csproj"><Private>false</Private></ProjectReference>
-    <ProjectReference Include="..\..\UnrealBuildTool\UnrealBuildToolCore.csproj"><Private>false</Private></ProjectReference>
-    <ProjectReference Include="..\AutomationUtils\AutomationUtils.Automation.Core.csproj"><Private>false</Private></ProjectReference>
-    <ProjectReference Include="..\Scripts\AutomationScripts.Automation.Core.csproj"><Private>false</Private></ProjectReference>
+    <ProjectReference Include="..\..\DotNETCommon\DotNETUtilities\DotNETUtilitiesCore.csproj" PrivateAssets="All"><Private>false</Private></ProjectReference>
+    <ProjectReference Include="..\..\UnrealBuildTool\UnrealBuildToolCore.csproj" PrivateAssets="All"><Private>false</Private></ProjectReference>
+    <ProjectReference Include="..\AutomationUtils\AutomationUtils.Automation.Core.csproj" PrivateAssets="All"><Private>false</Private></ProjectReference>
+    <ProjectReference Include="..\Scripts\AutomationScripts.Automation.Core.csproj" PrivateAssets="All"><Private>false</Private></ProjectReference>
   </ItemGroup>
   <ItemGroup>
     <PackageReference Include="Microsoft.CSharp" Version="4.7.0" />
diff --git a/Engine/Source/Programs/AutomationTool/Mac/Mac.Automation.Core.csproj b/Engine/Source/Programs/AutomationTool/Mac/Mac.Automation.Core.csproj
--- a/Engine/Source/Programs/AutomationTool/Mac/Mac.Automation.Core.csproj
+++ b/Engine/Source/Programs/AutomationTool/Mac/Mac.Automation.Core.csproj
@@ -16,7 +16,7 @@
     <RootNamespace>Mac.Automation</RootNamespace>
     <AssemblyName>Mac.Automation</AssemblyName>
     <WarningsNotAsErrors>612,618</WarningsNotAsErrors>
-    <OutputPath>..\..\..\..\Binaries\DotNET\AutomationTool\AutomationScripts\Mac\</OutputPath>
+    <OutputPath>..\..\..\..\Binaries\DotNET\AutomationTool\AutomationScripts\Platforms\Mac\</OutputPath>
     <DefineConstants>NET_CORE</DefineConstants>
     <TreatWarningsAsErrors>true</TreatWarningsAsErrors>
   </PropertyGroup>
@@ -46,9 +46,9 @@
     </Compile>
   </ItemGroup>
   <ItemGroup>
-    <ProjectReference Include="..\..\DotNETCommon\DotNETUtilities\DotNETUtilitiesCore.csproj"><Private>false</Private></ProjectReference>
-    <ProjectReference Include="..\..\UnrealBuildTool\UnrealBuildToolCore.csproj"><Private>false</Private></ProjectReference>
-    <ProjectReference Include="..\AutomationUtils\AutomationUtils.Automation.Core.csproj"><Private>false</Private></ProjectReference>
+    <ProjectReference Include="..\..\DotNETCommon\DotNETUtilities\DotNETUtilitiesCore.csproj" PrivateAssets="All"><Private>false</Private></ProjectReference>
+    <ProjectReference Include="..\..\UnrealBuildTool\UnrealBuildToolCore.csproj" PrivateAssets="All"><Private>false</Private></ProjectReference>
+    <ProjectReference Include="..\AutomationUtils\AutomationUtils.Automation.Core.csproj" PrivateAssets="All"><Private>false</Private></ProjectReference>
   </ItemGroup>
   <ItemGroup>
     <PackageReference Include="Microsoft.CSharp" Version="4.7.0" />
diff --git a/Engine/Source/Programs/AutomationTool/OneSkyLocalization/OneSkyLocalization.Automation.Core.csproj b/Engine/Source/Programs/AutomationTool/OneSkyLocalization/OneSkyLocalization.Automation.Core.csproj
--- a/Engine/Source/Programs/AutomationTool/OneSkyLocalization/OneSkyLocalization.Automation.Core.csproj
+++ b/Engine/Source/Programs/AutomationTool/OneSkyLocalization/OneSkyLocalization.Automation.Core.csproj
@@ -16,7 +16,7 @@
     <RootNamespace>AutomationTool</RootNamespace>
     <AssemblyName>OneSkyLocalization.Automation</AssemblyName>
     <WarningsNotAsErrors>612,618</WarningsNotAsErrors>
-    <OutputPath>..\..\..\..\Binaries\DotNET\AutomationTool\AutomationScripts</OutputPath>
+    <OutputPath>..\..\..\..\Binaries\DotNET\AutomationTool\AutomationScripts\OneSkyLocalization</OutputPath>
     <DefineConstants>NET_CORE</DefineConstants>
     <TreatWarningsAsErrors>true</TreatWarningsAsErrors>
   </PropertyGroup>
@@ -53,10 +53,10 @@
     </Compile>
   </ItemGroup>
   <ItemGroup>
-    <ProjectReference Include="..\..\DotNETCommon\DotNETUtilities\DotNETUtilitiesCore.csproj"><Private>false</Private></ProjectReference>
-    <ProjectReference Include="..\..\UnrealBuildTool\UnrealBuildToolCore.csproj"><Private>false</Private></ProjectReference>
-    <ProjectReference Include="..\AutomationUtils\AutomationUtils.Automation.Core.csproj"><Private>false</Private></ProjectReference>
-    <ProjectReference Include="..\Localization\Localization.Automation.Core.csproj"><Private>false</Private></ProjectReference>
+    <ProjectReference Include="..\..\DotNETCommon\DotNETUtilities\DotNETUtilitiesCore.csproj" PrivateAssets="All"><Private>false</Private></ProjectReference>
+    <ProjectReference Include="..\..\UnrealBuildTool\UnrealBuildToolCore.csproj" PrivateAssets="All"><Private>false</Private></ProjectReference>
+    <ProjectReference Include="..\AutomationUtils\AutomationUtils.Automation.Core.csproj" PrivateAssets="All"><Private>false</Private></ProjectReference>
+    <ProjectReference Include="..\Localization\Localization.Automation.Core.csproj" PrivateAssets="All"><Private>false</Private></ProjectReference>
   </ItemGroup>
   <ItemGroup>
     <PackageReference Include="Microsoft.CSharp" Version="4.7.0" />
diff --git a/Engine/Source/Programs/AutomationTool/Scripts/AutomationScripts.Automation.Core.csproj b/Engine/Source/Programs/AutomationTool/Scripts/AutomationScripts.Automation.Core.csproj
--- a/Engine/Source/Programs/AutomationTool/Scripts/AutomationScripts.Automation.Core.csproj
+++ b/Engine/Source/Programs/AutomationTool/Scripts/AutomationScripts.Automation.Core.csproj
@@ -16,7 +16,7 @@
     <RootNamespace>AutomationTool</RootNamespace>
     <AssemblyName>AutomationScripts.Automation</AssemblyName>
     <WarningsNotAsErrors>612,618</WarningsNotAsErrors>
-    <OutputPath>..\..\..\..\Binaries\DotNET\AutomationTool\AutomationScripts</OutputPath>
+    <OutputPath>..\..\..\..\Binaries\DotNET\AutomationTool\AutomationScripts\Scripts</OutputPath>
     <DefineConstants>NET_CORE</DefineConstants>
     <TreatWarningsAsErrors>true</TreatWarningsAsErrors>
   </PropertyGroup>
@@ -57,10 +57,10 @@
     </Compile>
   </ItemGroup>
   <ItemGroup>
-    <ProjectReference Include="..\..\DotNETCommon\DotNETUtilities\DotNETUtilitiesCore.csproj"><Private>false</Private></ProjectReference>
-    <ProjectReference Include="..\..\UnrealBuildTool\UnrealBuildToolCore.csproj"><Private>false</Private></ProjectReference>
-    <ProjectReference Include="..\AutomationUtils\AutomationUtils.Automation.Core.csproj"><Private>false</Private></ProjectReference>
-    <ProjectReference Include="..\Localization\Localization.Automation.Core.csproj"><Private>false</Private></ProjectReference>
+    <ProjectReference Include="..\..\DotNETCommon\DotNETUtilities\DotNETUtilitiesCore.csproj" PrivateAssets="All"><Private>false</Private></ProjectReference>
+    <ProjectReference Include="..\..\UnrealBuildTool\UnrealBuildToolCore.csproj" PrivateAssets="All"><Private>false</Private></ProjectReference>
+    <ProjectReference Include="..\AutomationUtils\AutomationUtils.Automation.Core.csproj" PrivateAssets="All"><Private>false</Private></ProjectReference>
+    <ProjectReference Include="..\Localization\Localization.Automation.Core.csproj" PrivateAssets="All"><Private>false</Private></ProjectReference>
   </ItemGroup>
   <ItemGroup>
     <PackageReference Include="Microsoft.CSharp" Version="4.7.0" />
diff --git a/Engine/Source/Programs/AutomationTool/TVOS/TVOS.Automation.Core.csproj b/Engine/Source/Programs/AutomationTool/TVOS/TVOS.Automation.Core.csproj
--- a/Engine/Source/Programs/AutomationTool/TVOS/TVOS.Automation.Core.csproj
+++ b/Engine/Source/Programs/AutomationTool/TVOS/TVOS.Automation.Core.csproj
@@ -16,7 +16,7 @@
     <RootNamespace>TVOS.Automation</RootNamespace>
     <AssemblyName>TVOS.Automation</AssemblyName>
     <WarningsNotAsErrors>612,618</WarningsNotAsErrors>
-    <OutputPath>..\..\..\..\Binaries\DotNET\AutomationTool\AutomationScripts\TVOS\</OutputPath>
+    <OutputPath>..\..\..\..\Binaries\DotNET\AutomationTool\AutomationScripts\Platforms\TVOS\</OutputPath>
     <DefineConstants>NET_CORE</DefineConstants>
     <TreatWarningsAsErrors>true</TreatWarningsAsErrors>
   </PropertyGroup>
@@ -52,10 +52,10 @@
     </Compile>
   </ItemGroup>
   <ItemGroup>
-    <ProjectReference Include="..\..\DotNETCommon\DotNETUtilities\DotNETUtilitiesCore.csproj"><Private>false</Private></ProjectReference>
-    <ProjectReference Include="..\..\UnrealBuildTool\UnrealBuildToolCore.csproj"><Private>false</Private></ProjectReference>
-    <ProjectReference Include="..\AutomationUtils\AutomationUtils.Automation.Core.csproj"><Private>false</Private></ProjectReference>
-    <ProjectReference Include="..\IOS\IOS.Automation.Core.csproj"><Private>false</Private></ProjectReference>
+    <ProjectReference Include="..\..\DotNETCommon\DotNETUtilities\DotNETUtilitiesCore.csproj" PrivateAssets="All"><Private>false</Private></ProjectReference>
+    <ProjectReference Include="..\..\UnrealBuildTool\UnrealBuildToolCore.csproj" PrivateAssets="All"><Private>false</Private></ProjectReference>
+    <ProjectReference Include="..\AutomationUtils\AutomationUtils.Automation.Core.csproj" PrivateAssets="All"><Private>false</Private></ProjectReference>
+    <ProjectReference Include="..\IOS\IOS.Automation.Core.csproj" PrivateAssets="All"><Private>false</Private></ProjectReference>
   </ItemGroup>
   <ItemGroup>
     <PackageReference Include="Microsoft.CSharp" Version="4.7.0" />
diff --git a/Engine/Source/Programs/AutomationTool/Turnkey/Turnkey.Automation.Core.csproj b/Engine/Source/Programs/AutomationTool/Turnkey/Turnkey.Automation.Core.csproj
--- a/Engine/Source/Programs/AutomationTool/Turnkey/Turnkey.Automation.Core.csproj
+++ b/Engine/Source/Programs/AutomationTool/Turnkey/Turnkey.Automation.Core.csproj
@@ -69,11 +69,11 @@
     </Compile>
   </ItemGroup>
   <ItemGroup>
-    <ProjectReference Include="..\..\DotNETCommon\BuildUtilities\BuildUtilitiesCore.csproj"><Private>false</Private></ProjectReference>
-    <ProjectReference Include="..\..\DotNETCommon\DotNETUtilities\DotNETUtilitiesCore.csproj"><Private>false</Private></ProjectReference>
-    <ProjectReference Include="..\..\UnrealBuildTool\UnrealBuildToolCore.csproj"><Private>false</Private></ProjectReference>
-    <ProjectReference Include="..\AutomationUtils\AutomationUtils.Automation.Core.csproj"><Private>false</Private></ProjectReference>
-    <ProjectReference Include="..\Gauntlet\Gauntlet.Automation.Core.csproj"><Private>false</Private></ProjectReference>
+    <ProjectReference Include="..\..\DotNETCommon\BuildUtilities\BuildUtilitiesCore.csproj" PrivateAssets="All"><Private>false</Private></ProjectReference>
+    <ProjectReference Include="..\..\DotNETCommon\DotNETUtilities\DotNETUtilitiesCore.csproj" PrivateAssets="All"><Private>false</Private></ProjectReference>
+    <ProjectReference Include="..\..\UnrealBuildTool\UnrealBuildToolCore.csproj" PrivateAssets="All"><Private>false</Private></ProjectReference>
+    <ProjectReference Include="..\AutomationUtils\AutomationUtils.Automation.Core.csproj" PrivateAssets="All"><Private>false</Private></ProjectReference>
+    <ProjectReference Include="..\Gauntlet\Gauntlet.Automation.Core.csproj" PrivateAssets="All"><Private>false</Private></ProjectReference>
   </ItemGroup>
   <ItemGroup>
     <PackageReference Include="Microsoft.CSharp" Version="4.7.0" />
diff --git a/Engine/Source/Programs/AutomationTool/Win/Win.Automation.Core.csproj b/Engine/Source/Programs/AutomationTool/Win/Win.Automation.Core.csproj
--- a/Engine/Source/Programs/AutomationTool/Win/Win.Automation.Core.csproj
+++ b/Engine/Source/Programs/AutomationTool/Win/Win.Automation.Core.csproj
@@ -16,7 +16,7 @@
     <RootNamespace>Win.Automation</RootNamespace>
     <AssemblyName>Win.Automation</AssemblyName>
     <WarningsNotAsErrors>612,618</WarningsNotAsErrors>
-    <OutputPath>..\..\..\..\Binaries\DotNET\AutomationTool\AutomationScripts</OutputPath>
+    <OutputPath>..\..\..\..\Binaries\DotNET\AutomationTool\AutomationScripts\Platforms\Windows</OutputPath>
     <DefineConstants>NET_CORE</DefineConstants>
     <TreatWarningsAsErrors>true</TreatWarningsAsErrors>
   </PropertyGroup>
@@ -36,11 +36,6 @@
       <HintPath>..\..\..\..\Binaries\DotNET\Ionic.Zip.Reduced.dll</HintPath>
       <Private>False</Private>
     </Reference>
-    <!--<Reference Include="OneSky, Version=1.0.0.0, Culture=neutral, processorArchitecture=MSIL">
-      <HintPath>..\..\..\..\Binaries\DotNET\OneSky.dll</HintPath>
-      <SpecificVersion>False</SpecificVersion>
-      <Private>False</Private>
-    </Reference>-->
   </ItemGroup>
   <Import Project="Sdk.props" Sdk="Microsoft.NET.Sdk" />
   <Import Project="Sdk.targets" Sdk="Microsoft.NET.Sdk" />
@@ -57,10 +52,10 @@
     </Compile>
   </ItemGroup>
   <ItemGroup>
-    <ProjectReference Include="..\..\DotNETCommon\DotNETUtilities\DotNETUtilitiesCore.csproj"><Private>false</Private></ProjectReference>
-    <ProjectReference Include="..\..\UnrealBuildTool\UnrealBuildToolCore.csproj"><Private>false</Private></ProjectReference>
-    <ProjectReference Include="..\AutomationUtils\AutomationUtils.Automation.Core.csproj"><Private>false</Private></ProjectReference>
-    <ProjectReference Include="..\BuildGraph\BuildGraph.Automation.Core.csproj"><Private>false</Private></ProjectReference>
+    <ProjectReference Include="..\..\DotNETCommon\DotNETUtilities\DotNETUtilitiesCore.csproj" PrivateAssets="All"><Private>false</Private></ProjectReference>
+    <ProjectReference Include="..\..\UnrealBuildTool\UnrealBuildToolCore.csproj" PrivateAssets="All"><Private>false</Private></ProjectReference>
+    <ProjectReference Include="..\AutomationUtils\AutomationUtils.Automation.Core.csproj" PrivateAssets="All"><Private>false</Private></ProjectReference>
+    <ProjectReference Include="..\BuildGraph\BuildGraph.Automation.Core.csproj" PrivateAssets="All"><Private>false</Private></ProjectReference>
   </ItemGroup>
   <ItemGroup>
     <PackageReference Include="Microsoft.CSharp" Version="4.7.0" />
diff --git a/Engine/Source/Programs/AutomationTool/XLocLocalization/XLocLocalization.Automation.Core.csproj b/Engine/Source/Programs/AutomationTool/XLocLocalization/XLocLocalization.Automation.Core.csproj
--- a/Engine/Source/Programs/AutomationTool/XLocLocalization/XLocLocalization.Automation.Core.csproj
+++ b/Engine/Source/Programs/AutomationTool/XLocLocalization/XLocLocalization.Automation.Core.csproj
@@ -16,10 +16,11 @@
     <RootNamespace>AutomationTool</RootNamespace>
     <AssemblyName>XLocLocalization.Automation</AssemblyName>
     <WarningsNotAsErrors>612,618</WarningsNotAsErrors>
-    <OutputPath>..\..\..\..\Binaries\DotNET\AutomationTool\AutomationScripts</OutputPath>
+    <OutputPath>..\..\..\..\Binaries\DotNET\AutomationTool\AutomationScripts\XLocLocalization</OutputPath>
     <DefineConstants>NET_CORE</DefineConstants>
     <TreatWarningsAsErrors>true</TreatWarningsAsErrors>
     <CopyLocalLockFileAssemblies>true</CopyLocalLockFileAssemblies>
+    <SatelliteResourceLanguages>en</SatelliteResourceLanguages> <!-- remove non english resource languages -->
   </PropertyGroup>
   <PropertyGroup Condition=" '$(Configuration)|$(Platform)' == 'Development|AnyCPU' ">
     <DefineConstants>NET_CORE;TRACE</DefineConstants>
@@ -43,10 +44,10 @@
   </ItemGroup>
 
   <ItemGroup>
-    <ProjectReference Include="..\..\DotNETCommon\DotNETUtilities\DotNETUtilitiesCore.csproj"><Private>false</Private></ProjectReference>
-    <ProjectReference Include="..\..\UnrealBuildTool\UnrealBuildToolCore.csproj"><Private>false</Private></ProjectReference>
-    <ProjectReference Include="..\AutomationUtils\AutomationUtils.Automation.Core.csproj"><Private>false</Private></ProjectReference>
-    <ProjectReference Include="..\Localization\Localization.Automation.Core.csproj"><Private>false</Private></ProjectReference>
+    <ProjectReference Include="..\..\DotNETCommon\DotNETUtilities\DotNETUtilitiesCore.csproj" PrivateAssets="All"><Private>false</Private></ProjectReference>
+    <ProjectReference Include="..\..\UnrealBuildTool\UnrealBuildToolCore.csproj" PrivateAssets="All"><Private>false</Private></ProjectReference>
+    <ProjectReference Include="..\AutomationUtils\AutomationUtils.Automation.Core.csproj" PrivateAssets="All"><Private>false</Private></ProjectReference>
+    <ProjectReference Include="..\Localization\Localization.Automation.Core.csproj" PrivateAssets="All"><Private>false</Private></ProjectReference>
   </ItemGroup>
   <ItemGroup>
     <PackageReference Include="Microsoft.CSharp" Version="4.7.0" />
