Fix BundleId for iOS/tvOS packaged products.

#rnx

diff --git a/Engine/Source/Programs/AutomationTool/IOS/IOSPlatform.Automation.cs b/Engine/Source/Programs/AutomationTool/IOS/IOSPlatform.Automation.cs
--- a/Engine/Source/Programs/AutomationTool/IOS/IOSPlatform.Automation.cs
+++ b/Engine/Source/Programs/AutomationTool/IOS/IOSPlatform.Automation.cs
@@ -255,7 +255,7 @@ public class IOSPlatform : Platform
 			if (string.IsNullOrEmpty(Password)) EngineConfig.GetString("/Script/IOSRuntimeSettings.IOSRuntimeSettings", "DevCenterPassword", out Password);
 			if (string.IsNullOrEmpty(Provision)) EngineConfig.GetString("/Script/IOSRuntimeSettings.IOSRuntimeSettings", "MobileProvision", out Provision);
 			
-			BundleId.Replace("[PROJECT_NAME]", ProjectName);
+			BundleId = BundleId.Replace("[PROJECT_NAME]", ProjectName);
 
 			// some are required
 			if (string.IsNullOrEmpty(BundleId))
