Editgrate 14600524
Avoid redundant error report in the log by implementing a new property LogWarningsAndErrorsAfterSummary to ITestNode

#rnx

diff --git a/Engine/Source/Programs/AutomationTool/Gauntlet/Framework/Base/Gauntlet.BaseTestNode.cs b/Engine/Source/Programs/AutomationTool/Gauntlet/Framework/Base/Gauntlet.BaseTestNode.cs
--- a/Engine/Source/Programs/AutomationTool/Gauntlet/Framework/Base/Gauntlet.BaseTestNode.cs
+++ b/Engine/Source/Programs/AutomationTool/Gauntlet/Framework/Base/Gauntlet.BaseTestNode.cs
@@ -49,6 +49,11 @@ namespace Gauntlet
 		/// </summary>
 		private TestStatus InnerStatus;
 
+		/// <summary>
+		/// Return true if the warnings and errors needs to log after summary
+		/// </summary>
+		public virtual bool LogWarningsAndErrorsAfterSummary { get; protected set; } = true;
+
 		/// <summary>
 		/// 
 		/// </summary>
diff --git a/Engine/Source/Programs/AutomationTool/Gauntlet/Framework/Base/Gauntlet.TestNode.cs b/Engine/Source/Programs/AutomationTool/Gauntlet/Framework/Base/Gauntlet.TestNode.cs
--- a/Engine/Source/Programs/AutomationTool/Gauntlet/Framework/Base/Gauntlet.TestNode.cs
+++ b/Engine/Source/Programs/AutomationTool/Gauntlet/Framework/Base/Gauntlet.TestNode.cs
@@ -67,6 +67,11 @@ namespace Gauntlet
 		/// </summary>
 		float MaxDuration { get; }
 
+		/// <summary>
+		/// Return true if the warnings and errors needs to log after summary
+		/// </summary>
+		bool LogWarningsAndErrorsAfterSummary { get; }
+
 		/// <summary>
 		/// What the test result should be treated as if we reach max duration.
 		/// </summary>
diff --git a/Engine/Source/Programs/AutomationTool/Gauntlet/Framework/Gauntlet.TestExecutor.cs b/Engine/Source/Programs/AutomationTool/Gauntlet/Framework/Gauntlet.TestExecutor.cs
--- a/Engine/Source/Programs/AutomationTool/Gauntlet/Framework/Gauntlet.TestExecutor.cs
+++ b/Engine/Source/Programs/AutomationTool/Gauntlet/Framework/Gauntlet.TestExecutor.cs
@@ -534,9 +534,12 @@ namespace Gauntlet
 			}
 			Summary.Split('\n').ToList().ForEach(L => Log.Info("  " + L));
 
-			TestInfo.TestNode.GetErrors().ToList().ForEach(E => Log.Error("{0}", E));
+			if (TestInfo.TestNode.LogWarningsAndErrorsAfterSummary)
+			{
+				TestInfo.TestNode.GetErrors().ToList().ForEach(E => Log.Error("{0}", E));
 
-			TestInfo.TestNode.GetWarnings().ToList().ForEach(E => Log.Warning("{0}", E));
+				TestInfo.TestNode.GetWarnings().ToList().ForEach(E => Log.Warning("{0}", E));
+			}
 
 		}
 
diff --git a/Engine/Source/Programs/AutomationTool/Gauntlet/SelfTest/Framework/Gauntlet.SelfTest.OrderOfOpsTest.cs b/Engine/Source/Programs/AutomationTool/Gauntlet/SelfTest/Framework/Gauntlet.SelfTest.OrderOfOpsTest.cs
--- a/Engine/Source/Programs/AutomationTool/Gauntlet/SelfTest/Framework/Gauntlet.SelfTest.OrderOfOpsTest.cs
+++ b/Engine/Source/Programs/AutomationTool/Gauntlet/SelfTest/Framework/Gauntlet.SelfTest.OrderOfOpsTest.cs
@@ -58,6 +58,8 @@ namespace Gauntlet.SelfTest
 
 		public bool HasWarnings { get { return false; } }
 
+		public virtual bool LogWarningsAndErrorsAfterSummary { get { return true; } }
+
 		protected void SetNewStage(TestStages NewStage)
 		{
 			int CurrentValue = (int)CurrentStage;
diff --git a/Engine/Source/Programs/AutomationTool/Gauntlet/SelfTest/Gauntlet.SelfTest.BaseNode.cs b/Engine/Source/Programs/AutomationTool/Gauntlet/SelfTest/Gauntlet.SelfTest.BaseNode.cs
--- a/Engine/Source/Programs/AutomationTool/Gauntlet/SelfTest/Gauntlet.SelfTest.BaseNode.cs
+++ b/Engine/Source/Programs/AutomationTool/Gauntlet/SelfTest/Gauntlet.SelfTest.BaseNode.cs
@@ -45,6 +45,8 @@ namespace Gauntlet.SelfTest
 
 		public bool HasWarnings { get { return false; } }
 
+		public virtual bool LogWarningsAndErrorsAfterSummary { get { return true; } }
+
 		public TestResult GetTestResult()
 		{
 			return InnerResult;
