Disable ionic zip for build graph archives under net core, this was causing crashes on Mac builds in Horde when netcore was enabled.

Also slightly refactored how we start these parallel zip threads to avoid potential fallout due to variable captures.

diff --git a/Engine/Source/Programs/AutomationTool/BuildGraph/TempStorage.cs b/Engine/Source/Programs/AutomationTool/BuildGraph/TempStorage.cs
--- a/Engine/Source/Programs/AutomationTool/BuildGraph/TempStorage.cs
+++ b/Engine/Source/Programs/AutomationTool/BuildGraph/TempStorage.cs
@@ -913,13 +913,17 @@ namespace AutomationTool
 			ConcurrentBag<FileInfo> ZipFiles = new ConcurrentBag<FileInfo>();
 
 			DirectoryReference ZipDir = StagingDir ?? OutputDir;
+			// Disabled for net core as using ionic seems to crash when run on build machines (causes a process panic)
+#if !NET_CORE
 			if (Utils.IsRunningOnMono)
 			{
 				ZipThreads = (
 					from CoreNum in Enumerable.Range(0, bZipInParallel ? Environment.ProcessorCount : 1)
-					let ZipFileName = FileReference.Combine(ZipDir, string.Format("{0}{1}.zip", ZipBaseName, bZipInParallel ? "-" + CoreNum.ToString("00") : ""))
-					select new Thread(() =>
+					select new Thread((object indexObject) =>
 					{
+						int index = (int) indexObject;
+						var ZipFileName = FileReference.Combine(ZipDir, string.Format("{0}{1}.zip", ZipBaseName, bZipInParallel ? "-" + index.ToString("00") : ""));
+
 						// don't create the zip unless we have at least one file to add
 						FileReference File;
 						if (FilesToZip.TryDequeue(out File))
@@ -952,12 +956,14 @@ namespace AutomationTool
 					})).ToList();
 			}
 			else
+#endif
 			{
 				ZipThreads = (
 					from CoreNum in Enumerable.Range(0, bZipInParallel ? Environment.ProcessorCount : 1)
-					let ZipFileName = FileReference.Combine(ZipDir, string.Format("{0}{1}.zip", ZipBaseName, bZipInParallel ? "-" + CoreNum.ToString("00") : ""))
-					select new Thread(() =>
+					select new Thread((object indexObject) =>
 					{
+						int index = (int) indexObject;
+						var ZipFileName = FileReference.Combine(ZipDir, string.Format("{0}{1}.zip", ZipBaseName, bZipInParallel ? "-" + index.ToString("00") : ""));
 						// don't create the zip unless we have at least one file to add
 						FileReference File;
 						if (FilesToZip.TryDequeue(out File))
@@ -989,7 +995,13 @@ namespace AutomationTool
 						}
 					})).ToList();
 			}
-			ZipThreads.ForEach(thread => thread.Start());
+
+			for (int index = 0; index < ZipThreads.Count; index++)
+			{
+				Thread thread = ZipThreads[index];
+				thread.Start(index);
+			}
+
 			ZipThreads.ForEach(thread => thread.Join());
 			
 			return ZipFiles.OrderBy(x => x.Name).ToArray();
