- Added support for a platform to use the Staging type and also filename in some circumstances the filename (could change to ONLY have the filename version of the function)

diff --git a/Engine/Source/Programs/AutomationTool/Android/AndroidPlatform.Automation.cs b/Engine/Source/Programs/AutomationTool/Android/AndroidPlatform.Automation.cs
--- a/Engine/Source/Programs/AutomationTool/Android/AndroidPlatform.Automation.cs
+++ b/Engine/Source/Programs/AutomationTool/Android/AndroidPlatform.Automation.cs
@@ -2506,7 +2506,7 @@ public class AndroidPlatform : Platform
 		return bIsClientOnly ? "AndroidClient" : "Android";
 	}
 
-	public override bool DeployLowerCaseFilenames()
+	public override bool DeployLowerCaseFilenames(StagedFileType FileType)
 	{
 		return false;
 	}
diff --git a/Engine/Source/Programs/AutomationTool/AutomationUtils/Platform.cs b/Engine/Source/Programs/AutomationTool/AutomationUtils/Platform.cs
--- a/Engine/Source/Programs/AutomationTool/AutomationUtils/Platform.cs
+++ b/Engine/Source/Programs/AutomationTool/AutomationUtils/Platform.cs
@@ -449,12 +449,24 @@ namespace AutomationTool
 		/// <summary>
 		/// return true if we need to change the case of filenames outside of pak files
 		/// </summary>
-		/// <returns></returns>
-		public virtual bool DeployLowerCaseFilenames()
+		/// <param name="FileType" The staged file type to check (UFS vs SsytemNonUFS, etc)
+		/// <returns>true if files should be lower-cased during staging, for the given filetype</returns>
+		public virtual bool DeployLowerCaseFilenames(StagedFileType FileType)
 		{
 			return false;
 		}
 
+		/// <summary>
+		/// return true if we need to change the case of a particular file
+		/// </summary>
+		/// <param name="FileType" The staged file type to check (UFS vs SsytemNonUFS, etc)
+		/// <returns>true if files should be lower-cased during staging, for the given filetype</returns>
+		public virtual bool DeployLowerCaseFile(FileReference File, StagedFileType FileType)
+		{
+			return DeployLowerCaseFilenames(FileType);
+		}
+
+
 		/// <summary>
 		/// Converts local path to target platform path.
 		/// </summary>
@@ -606,7 +618,8 @@ namespace AutomationTool
 		public virtual HashSet<StagedFileReference> GetFilesForCRCCheck()
 		{
 			string CmdLine = "UE4CommandLine.txt";
-			if (DeployLowerCaseFilenames())
+			// using SystemNonUFS because that is how it's staged in CreateStagingManifest
+			if (DeployLowerCaseFilenames(StagedFileType.SystemNonUFS))
 			{
 				CmdLine = CmdLine.ToLowerInvariant();
 			}
diff --git a/Engine/Source/Programs/AutomationTool/IOS/IOSPlatform.Automation.cs b/Engine/Source/Programs/AutomationTool/IOS/IOSPlatform.Automation.cs
--- a/Engine/Source/Programs/AutomationTool/IOS/IOSPlatform.Automation.cs
+++ b/Engine/Source/Programs/AutomationTool/IOS/IOSPlatform.Automation.cs
@@ -1759,7 +1759,7 @@ public class IOSPlatform : Platform
 		return bIsClientOnly ? "IOSClient" : "IOS";
 	}
 
-	public override bool DeployLowerCaseFilenames()
+	public override bool DeployLowerCaseFilenames(StagedFileType FileType)
 	{
 		// we shouldn't modify the case on files like Info.plist or the icons
 		return true;
diff --git a/Engine/Source/Programs/AutomationTool/Linux/LinuxPlatform.Automation.cs b/Engine/Source/Programs/AutomationTool/Linux/LinuxPlatform.Automation.cs
--- a/Engine/Source/Programs/AutomationTool/Linux/LinuxPlatform.Automation.cs
+++ b/Engine/Source/Programs/AutomationTool/Linux/LinuxPlatform.Automation.cs
@@ -178,7 +178,7 @@ public abstract class BaseLinuxPlatform : Platform
 	/// return true if we need to change the case of filenames outside of pak files
 	/// </summary>
 	/// <returns></returns>
-	public override bool DeployLowerCaseFilenames()
+	public override bool DeployLowerCaseFilenames(StagedFileType FileType)
 	{
 		return false;
 	}
diff --git a/Engine/Source/Programs/AutomationTool/Lumin/LuminPlatform.Automation.cs b/Engine/Source/Programs/AutomationTool/Lumin/LuminPlatform.Automation.cs
--- a/Engine/Source/Programs/AutomationTool/Lumin/LuminPlatform.Automation.cs
+++ b/Engine/Source/Programs/AutomationTool/Lumin/LuminPlatform.Automation.cs
@@ -279,7 +279,7 @@ public class LuminPlatform : Platform
 	/// return true if we need to change the case of filenames outside of pak files
 	/// </summary>
 	/// <returns></returns>
-	public override bool DeployLowerCaseFilenames()
+	public override bool DeployLowerCaseFilenames(StagedFileType FileType)
 	{
 		// @todo quail hack: we don't actually want to lower case some of the OS files, but I HACKED the
 		// CopyBuildToStagingDirectory code to not apply lowercase to root files 
diff --git a/Engine/Source/Programs/AutomationTool/Scripts/CopyBuildToStagingDirectory.Automation.cs b/Engine/Source/Programs/AutomationTool/Scripts/CopyBuildToStagingDirectory.Automation.cs
--- a/Engine/Source/Programs/AutomationTool/Scripts/CopyBuildToStagingDirectory.Automation.cs
+++ b/Engine/Source/Programs/AutomationTool/Scripts/CopyBuildToStagingDirectory.Automation.cs
@@ -703,12 +703,12 @@ public partial class Project : CommandUtils
 
 			// move the UE4Commandline.txt file to the root of the stage
 			// this file needs to be treated as a UFS file for casing, but NonUFS for being put into the .pak file
-			// @todo: Maybe there should be a new category - UFSNotForPak
+			// @todo: Maybe there should be a new category - UFSNotForPak - and reevaluate all SystemNonUFS references
 			FileReference CommandLineFile = FileReference.Combine(GetIntermediateCommandlineDir(SC), "UE4CommandLine.txt");
 			if (FileReference.Exists(CommandLineFile))
 			{
 				StagedFileReference StagedCommandLineFile = new StagedFileReference("UE4CommandLine.txt");
-				if (SC.StageTargetPlatform.DeployLowerCaseFilenames())
+				if (SC.StageTargetPlatform.DeployLowerCaseFilenames(StagedFileType.SystemNonUFS))
 				{
 					StagedCommandLineFile = StagedCommandLineFile.ToLowerInvariant();
 				}
@@ -1089,14 +1089,15 @@ public partial class Project : CommandUtils
 		}
 
 		// Make all the filenames lowercase
-		if (SC.StageTargetPlatform.DeployLowerCaseFilenames())
+		if (SC.StageTargetPlatform.DeployLowerCaseFilenames(StagedFileType.NonUFS))
 		{
 			SC.FilesToStage.NonUFSFiles = SC.FilesToStage.NonUFSFiles.ToDictionary(x => x.Key.ToLowerInvariant(), x => x.Value);
-			SC.FilesToStage.NonUFSDebugFiles = SC.FilesToStage.NonUFSDebugFiles.ToDictionary(x => x.Key.ToLowerInvariant(), x => x.Value);
-			if (!Params.UsePak(SC.StageTargetPlatform))
-			{
-				SC.FilesToStage.UFSFiles = SC.FilesToStage.UFSFiles.ToDictionary(x => x.Key.ToLowerInvariant(), x => x.Value);
-			}
+		}
+		// DebugNonUFS is weird, so ask per-file
+		SC.FilesToStage.NonUFSDebugFiles = SC.FilesToStage.NonUFSDebugFiles.ToDictionary(x => SC.StageTargetPlatform.DeployLowerCaseFile(x.Value, StagedFileType.DebugNonUFS) ? x.Key.ToLowerInvariant() : x.Key, x => x.Value);
+		if (!Params.UsePak(SC.StageTargetPlatform) && SC.StageTargetPlatform.DeployLowerCaseFilenames(StagedFileType.UFS))
+		{
+			SC.FilesToStage.UFSFiles = SC.FilesToStage.UFSFiles.ToDictionary(x => x.Key.ToLowerInvariant(), x => x.Value);
 		}
 
 		// Remap all the non-ufs files if not using a PAK file
@@ -2223,7 +2224,7 @@ public partial class Project : CommandUtils
 			{
 				OutputRelativeLocation = StagedFileReference.Combine(SC.RelativeProjectRootForStage, "Content", "Paks", OutputFilename + OutputFilenameExtension);
 			}
-			if (SC.StageTargetPlatform.DeployLowerCaseFilenames())
+			if (SC.StageTargetPlatform.DeployLowerCaseFilenames(StagedFileType.UFS))
 			{
 				OutputRelativeLocation = OutputRelativeLocation.ToLowerInvariant();
 			}
@@ -2280,7 +2281,7 @@ public partial class Project : CommandUtils
 				// Check to see if we have an existing pak file we can use
 
 				StagedFileReference SourceOutputRelativeLocation = StagedFileReference.Combine(SC.RelativeProjectRootForStage, "Content/Paks/", PakParams.PakName + "-" + SC.CookPlatform + PostFix + ".pak");
-				if (SC.CookSourcePlatform.DeployLowerCaseFilenames())
+				if (SC.CookSourcePlatform.DeployLowerCaseFilenames(StagedFileType.UFS))
 				{
 					SourceOutputRelativeLocation = SourceOutputRelativeLocation.ToLowerInvariant();
 				}
@@ -2674,7 +2675,7 @@ public partial class Project : CommandUtils
 							{
 								InternalUtils.SafeCopyFile(PakFilePath, OutputDestinationPath.FullName);
 								StagedFileReference OutputDestinationRelativeLocation = StagedFileReference.Combine(SC.RelativeProjectRootForStage, "Content/Paks/", Path.GetFileName(PakFilePath));
-								if (SC.StageTargetPlatform.DeployLowerCaseFilenames())
+								if (SC.StageTargetPlatform.DeployLowerCaseFilenames(StagedFileType.UFS))
 								{
 									OutputDestinationRelativeLocation = OutputDestinationRelativeLocation.ToLowerInvariant();
 								}
@@ -2692,7 +2693,7 @@ public partial class Project : CommandUtils
 	{
 		StagedFileReference GlobalContainerOutputRelativeLocation;
 		GlobalContainerOutputRelativeLocation = StagedFileReference.Combine(SC.RelativeProjectRootForStage, "Content", "Paks", "global.utoc");
-		if (SC.StageTargetPlatform.DeployLowerCaseFilenames())
+		if (SC.StageTargetPlatform.DeployLowerCaseFilenames(StagedFileType.UFS))
 		{
 			GlobalContainerOutputRelativeLocation = GlobalContainerOutputRelativeLocation.ToLowerInvariant();
 		}
@@ -3460,7 +3461,7 @@ public partial class Project : CommandUtils
 	{
 		// this file needs to be treated as a UFS file for casing, but NonUFS for being put into the .pak file. 
 		// @todo: Maybe there should be a new category - UFSNotForPak
-		if (SC.StageTargetPlatform.DeployLowerCaseFilenames())
+		if (SC.StageTargetPlatform.DeployLowerCaseFilenames(StagedFileType.SystemNonUFS))
 		{
 			IntermediateCmdLineFile = new FileReference(CombinePaths(Path.GetDirectoryName(IntermediateCmdLineFile.FullName), Path.GetFileName(IntermediateCmdLineFile.FullName).ToLowerInvariant()));
 		}
