- Fixed linux crash in Turnkey in GDI mono

diff --git a/Engine/Source/Programs/AutomationTool/Turnkey/Providers/HybridIOProvider.cs b/Engine/Source/Programs/AutomationTool/Turnkey/Providers/HybridIOProvider.cs
--- a/Engine/Source/Programs/AutomationTool/Turnkey/Providers/HybridIOProvider.cs
+++ b/Engine/Source/Programs/AutomationTool/Turnkey/Providers/HybridIOProvider.cs
@@ -9,9 +9,7 @@ using System.Text.RegularExpressions;
 using System.Drawing;
 using Tools.DotNETCommon;
 
-#if NET_CORE && WINDOWS
-using System.Windows.Forms;
-#elif !NET_CORE
+#if WINDOWS
 using System.Windows.Forms;
 #endif
 
@@ -30,15 +28,10 @@ namespace Turnkey
 				SetProcessDPIAware();
 			}
 
-#if NET_CORE
 #if WINDOWS
 			Application.EnableVisualStyles();
 			Application.SetCompatibleTextRenderingDefault(false);
 #endif
-#else
-			Application.EnableVisualStyles();
-			Application.SetCompatibleTextRenderingDefault(false);
-#endif
 
 		}
 
@@ -122,7 +115,7 @@ namespace Turnkey
 #if NET_CORE
 #if WINDOWS
 			Result = ShowFormsDialogs(Prompt, Default);
-#elif  OSX
+#elif OSX
 			Result = ShowMacDialog(Prompt, Default);
 #else
 			throw new NotImplementedException("Linux dialog not implemented");
@@ -133,10 +126,14 @@ namespace Turnkey
 			{
 				Result = ShowMacDialog(Prompt, Default);
 			}
-			else
+			else if (UnrealBuildTool.BuildHostPlatform.Current.Platform == UnrealBuildTool.UnrealTargetPlatform.Win64)
 			{
 				Result = Microsoft.VisualBasic.Interaction.InputBox(Prompt, "Turnkey Input", Default);
 			}
+			else
+			{
+				throw new NotImplementedException("Linux dialog not implemented");
+			}
 #endif
 
 			if (string.IsNullOrEmpty(Result) && bIsList)
@@ -146,110 +143,6 @@ namespace Turnkey
 
 			return Result;
 #pragma warning restore
-/*
-			// Create a new instance of the form.
-			Form Form1 = new Form();
-
-			System.Windows.Forms.TextBox Label;
-			System.Windows.Forms.Button CancelBtn;
-			System.Windows.Forms.Button OkBtn;
-			System.Windows.Forms.TextBox TextBox;
-
-			TextBox = new System.Windows.Forms.TextBox();
-			Label = new System.Windows.Forms.TextBox();
-			CancelBtn = new System.Windows.Forms.Button();
-			OkBtn = new System.Windows.Forms.Button();
-			Form1.SuspendLayout();
-			// 
-			// Label
-			// 
-			Label.Anchor = ((System.Windows.Forms.AnchorStyles)(((System.Windows.Forms.AnchorStyles.Top | System.Windows.Forms.AnchorStyles.Left)
-			| System.Windows.Forms.AnchorStyles.Right | System.Windows.Forms.AnchorStyles.Bottom)));
-			Label.BackColor = System.Drawing.SystemColors.Control;
-			Label.BorderStyle = System.Windows.Forms.BorderStyle.None;
-			Label.Location = new System.Drawing.Point(13, 13);
-			Label.Multiline = true;
-			Label.Name = "Label";
-			Label.ScrollBars = System.Windows.Forms.ScrollBars.Vertical;
-			Label.ShortcutsEnabled = false;
-			Label.Size = new System.Drawing.Size(775, 199);
-			Label.TabIndex = 0;
-			Label.TabStop = false;
-			Label.Text = Prompt;
-
-			// 
-			// CancelBtn
-			// 
-			CancelBtn.Anchor = ((System.Windows.Forms.AnchorStyles)((System.Windows.Forms.AnchorStyles.Right | System.Windows.Forms.AnchorStyles.Bottom)));
-			CancelBtn.DialogResult = System.Windows.Forms.DialogResult.Cancel;
-			CancelBtn.Location = new System.Drawing.Point(679, 262);
-			CancelBtn.Margin = new System.Windows.Forms.Padding(4);
-			CancelBtn.Name = "CancelBtn";
-			CancelBtn.Size = new System.Drawing.Size(109, 34);
-			CancelBtn.TabIndex = 6;
-			CancelBtn.Text = "Cancel";
-			CancelBtn.UseVisualStyleBackColor = true;
-			// 
-			// TextBox
-			// 
-			TextBox.Anchor = ((System.Windows.Forms.AnchorStyles)(((System.Windows.Forms.AnchorStyles.Left | System.Windows.Forms.AnchorStyles.Bottom)
-			| System.Windows.Forms.AnchorStyles.Right)));
-			TextBox.Location = new System.Drawing.Point(13, 219);
-			TextBox.Margin = new System.Windows.Forms.Padding(4);
-			TextBox.Name = "TextBox";
-			TextBox.Size = new System.Drawing.Size(772, 22);
-			TextBox.TabIndex = 4;
-			TextBox.Text = Default;
-			// 
-			// OkBtn
-			// 
-			OkBtn.Anchor = ((System.Windows.Forms.AnchorStyles)((System.Windows.Forms.AnchorStyles.Right | System.Windows.Forms.AnchorStyles.Bottom)));
-			OkBtn.DialogResult = System.Windows.Forms.DialogResult.OK;
-			OkBtn.Location = new System.Drawing.Point(562, 262);
-			OkBtn.Margin = new System.Windows.Forms.Padding(4);
-			OkBtn.Name = "OkBtn";
-			OkBtn.Size = new System.Drawing.Size(109, 34);
-			OkBtn.TabIndex = 5;
-			OkBtn.Text = "Ok";
-			OkBtn.UseVisualStyleBackColor = true;
-			// 
-			// Form1
-			// 
-			Form1.AcceptButton = OkBtn;
-			Form1.StartPosition = FormStartPosition.CenterScreen;
-			Form1.AutoScaleDimensions = new System.Drawing.SizeF(8F, 16F);
-			Form1.AutoScaleMode = System.Windows.Forms.AutoScaleMode.Font;
-			Form1.CancelButton = CancelBtn;
-			Form1.ClientSize = new System.Drawing.Size(800, 315);
-			Form1.Controls.Add(CancelBtn);
-			Form1.Controls.Add(OkBtn);
-			Form1.Controls.Add(TextBox);
-			Form1.Controls.Add(Label);
-			Form1.Name = "Form1";
-			Form1.Text = "Make selection";
-			Form1.ResumeLayout(false);
-			Form1.PerformLayout();
-
-			if (!bIsList)
-			{
-				Form1.Height -= 150;
-				Form1.PerformLayout();
-			}
-
-			// Display the form as a modal dialog box.
-			Form1.ShowDialog();
-
-			string Result = null;
-
-			// Determine if the OK button was clicked on the dialog box.
-			if (Form1.DialogResult == DialogResult.OK)
-			{
-				Result = TextBox.Text;
-			}
-
-			Form1.Dispose();
-			return Result;
-*/
 		}
 
 		public override string ReadInput(string Prompt, string Default, bool bAppendNewLine)
