Fix Linux for turnkey.
Assume p4 is in the $PATH somewhere

diff --git a/Engine/Source/Programs/AutomationTool/AutomationUtils/LinuxHostPlatform.cs b/Engine/Source/Programs/AutomationTool/AutomationUtils/LinuxHostPlatform.cs
--- a/Engine/Source/Programs/AutomationTool/AutomationUtils/LinuxHostPlatform.cs
+++ b/Engine/Source/Programs/AutomationTool/AutomationUtils/LinuxHostPlatform.cs
@@ -74,7 +74,7 @@ namespace AutomationTool
 
 		public override string P4Exe
 		{
-			get { return "/usr/bin/p4"; }
+			get { return "p4"; }
 		}
 
 		public override Process CreateProcess(string AppName)
diff --git a/Engine/Source/Programs/AutomationTool/Turnkey/Providers/NullCopyProvider.cs b/Engine/Source/Programs/AutomationTool/Turnkey/Providers/NullCopyProvider.cs
--- a/Engine/Source/Programs/AutomationTool/Turnkey/Providers/NullCopyProvider.cs
+++ b/Engine/Source/Programs/AutomationTool/Turnkey/Providers/NullCopyProvider.cs
@@ -59,7 +59,20 @@ namespace Turnkey
 			}
 
 			// accessing large remote files direcly can be crazy slow, so we cache large-ish files locally before using
-			bool bIsRemotePath = (new System.Uri(OutputPath).IsUnc) || (new DriveInfo(OutputPath).DriveType == DriveType.Network);
+			bool bIsRemotePath = false;
+
+			if (UnrealBuildTool.BuildHostPlatform.Current.Platform == UnrealBuildTool.UnrealTargetPlatform.Mac ||
+				UnrealBuildTool.BuildHostPlatform.Current.Platform == UnrealBuildTool.UnrealTargetPlatform.Linux)
+			{
+				// [REDACTED_PROFILE_URL]/mono/issues/12157
+				// resolved in newer versions of mono but the current bundled version is crashing when trying to create a new DriveInfo
+				bIsRemotePath = false;
+			}
+			else
+			{
+				bIsRemotePath = (new System.Uri(OutputPath).IsUnc) || (new DriveInfo(OutputPath).DriveType == DriveType.Network);
+			}
+
 			bool bIsLargeFile = File.Exists(OutputPath) && new FileInfo(OutputPath).Length > 5 * 1024 * 1024;
 			bool bIsDirectory = Directory.Exists(OutputPath);
 
