Changed pdb format to not use portable pdbs (which is the default in net core) and instead back to our old pdb configuration.
This resolves a issue with PCBs were pdbcopy failed to run because it does not support portable pdbs.

diff --git a/Engine/Source/Programs/AutomationTool/Android/Android.Automation.csproj b/Engine/Source/Programs/AutomationTool/Android/Android.Automation.csproj
--- a/Engine/Source/Programs/AutomationTool/Android/Android.Automation.csproj
+++ b/Engine/Source/Programs/AutomationTool/Android/Android.Automation.csproj
@@ -1,4 +1,4 @@
-<Project Sdk="Microsoft.NET.Sdk">
+﻿<Project Sdk="Microsoft.NET.Sdk">
   <PropertyGroup>
     <TargetFramework>netcoreapp3.1</TargetFramework>
     <Configuration Condition=" '$(Configuration)' == '' ">Development</Configuration>
@@ -12,6 +12,7 @@
     <WarningsNotAsErrors>612,618</WarningsNotAsErrors>
     <OutputPath>..\..\..\..\Binaries\DotNET\AutomationTool\AutomationScripts\Platforms\Android</OutputPath>
     <TreatWarningsAsErrors>true</TreatWarningsAsErrors>
+    <DebugType>pdbonly</DebugType>
   </PropertyGroup>
   <PropertyGroup Condition=" '$(Configuration)|$(Platform)' == 'Development|AnyCPU' ">
     <DefineConstants>$(DefineConstants);TRACE</DefineConstants>
@@ -23,6 +24,8 @@
   </PropertyGroup>
   <PropertyGroup Condition=" '$(Configuration)|$(Platform)' == 'Debug|AnyCPU' ">
     <DefineConstants>$(DefineConstants);DEBUG;TRACE</DefineConstants>
+    <DebugSymbols>true</DebugSymbols>
+    <DebugType>full</DebugType>
   </PropertyGroup>
   <ItemGroup>
     <Reference Include="Ionic.Zip.Reduced">
diff --git a/Engine/Source/Programs/AutomationTool/AutomationUtils/AutomationUtils.Automation.csproj b/Engine/Source/Programs/AutomationTool/AutomationUtils/AutomationUtils.Automation.csproj
--- a/Engine/Source/Programs/AutomationTool/AutomationUtils/AutomationUtils.Automation.csproj
+++ b/Engine/Source/Programs/AutomationTool/AutomationUtils/AutomationUtils.Automation.csproj
@@ -1,4 +1,4 @@
-<Project Sdk="Microsoft.NET.Sdk">
+﻿<Project Sdk="Microsoft.NET.Sdk">
   <PropertyGroup>
     <TargetFramework>netcoreapp3.1</TargetFramework>
     <Configuration Condition=" '$(Configuration)' == '' ">Development</Configuration>
@@ -11,6 +11,7 @@
     <Configurations>Debug;Release;Development</Configurations>
     <TreatWarningsAsErrors>true</TreatWarningsAsErrors>
     <AppendTargetFrameworkToOutputPath>false</AppendTargetFrameworkToOutputPath>
+    <DebugType>pdbonly</DebugType>
   </PropertyGroup>
     <ItemGroup>
     <Reference Include="Ionic.Zip.Reduced">
@@ -18,9 +19,10 @@
       <Private>False</Private>
     </Reference>
   </ItemGroup>
-  <!-- We do not output these binaries to Binaries as it targets multiple frameworks. They would have overloaded names for each framework, let the application copy it to its output directory instead -->
   <PropertyGroup Condition=" '$(Configuration)|$(Platform)' == 'Debug|AnyCPU' ">
     <DefineConstants>$(DefineConstants);TRACE</DefineConstants>
+    <DebugSymbols>true</DebugSymbols>
+    <DebugType>full</DebugType>
   </PropertyGroup>
   <PropertyGroup Condition="'$(Configuration)|$(Platform)' == 'Development|AnyCPU'">
     <Optimize>true</Optimize>
diff --git a/Engine/Source/Programs/AutomationTool/BuildGraph/BuildGraph.Automation.csproj b/Engine/Source/Programs/AutomationTool/BuildGraph/BuildGraph.Automation.csproj
--- a/Engine/Source/Programs/AutomationTool/BuildGraph/BuildGraph.Automation.csproj
+++ b/Engine/Source/Programs/AutomationTool/BuildGraph/BuildGraph.Automation.csproj
@@ -12,6 +12,7 @@
     <WarningsNotAsErrors>612,618</WarningsNotAsErrors>
     <OutputPath>..\..\..\..\Binaries\DotNET\AutomationTool\AutomationScripts\BuildGraph</OutputPath>
     <TreatWarningsAsErrors>true</TreatWarningsAsErrors>
+    <DebugType>pdbonly</DebugType>
   </PropertyGroup>
   <PropertyGroup Condition=" '$(Configuration)|$(Platform)' == 'Development|AnyCPU' ">
     <DefineConstants>$(DefineConstants);TRACE</DefineConstants>
@@ -23,6 +24,8 @@
   </PropertyGroup>
   <PropertyGroup Condition=" '$(Configuration)|$(Platform)' == 'Debug|AnyCPU' ">
     <DefineConstants>$(DefineConstants);DEBUG;TRACE</DefineConstants>
+    <DebugSymbols>true</DebugSymbols>
+    <DebugType>full</DebugType>
   </PropertyGroup>
   <ItemGroup>
     <Reference Include="Ionic.Zip.Reduced">
diff --git a/Engine/Source/Programs/AutomationTool/Gauntlet/Gauntlet.Automation.csproj b/Engine/Source/Programs/AutomationTool/Gauntlet/Gauntlet.Automation.csproj
--- a/Engine/Source/Programs/AutomationTool/Gauntlet/Gauntlet.Automation.csproj
+++ b/Engine/Source/Programs/AutomationTool/Gauntlet/Gauntlet.Automation.csproj
@@ -12,6 +12,7 @@
     <WarningsNotAsErrors>612,618</WarningsNotAsErrors>
     <OutputPath>..\..\..\..\Binaries\DotNET\AutomationTool\AutomationScripts\Gauntlet</OutputPath>
     <TreatWarningsAsErrors>true</TreatWarningsAsErrors>
+    <DebugType>pdbonly</DebugType>
   </PropertyGroup>
   <PropertyGroup Condition=" '$(Configuration)|$(Platform)' == 'Development|AnyCPU' ">
     <DefineConstants>$(DefineConstants);TRACE</DefineConstants>
@@ -23,6 +24,8 @@
   </PropertyGroup>
   <PropertyGroup Condition=" '$(Configuration)|$(Platform)' == 'Debug|AnyCPU' ">
     <DefineConstants>$(DefineConstants);DEBUG;TRACE</DefineConstants>
+    <DebugSymbols>true</DebugSymbols>
+    <DebugType>full</DebugType>
   </PropertyGroup>
   <ItemGroup>
     <Reference Include="Ionic.Zip.Reduced">
diff --git a/Engine/Source/Programs/AutomationTool/HoloLens/HoloLens.Automation.csproj b/Engine/Source/Programs/AutomationTool/HoloLens/HoloLens.Automation.csproj
--- a/Engine/Source/Programs/AutomationTool/HoloLens/HoloLens.Automation.csproj
+++ b/Engine/Source/Programs/AutomationTool/HoloLens/HoloLens.Automation.csproj
@@ -15,7 +15,7 @@
     <IsWindows Condition="'$([System.Runtime.InteropServices.RuntimeInformation]::IsOSPlatform($([System.Runtime.InteropServices.OSPlatform]::Windows)))' == 'true'">true</IsWindows> 
     <IsOSX Condition="'$([System.Runtime.InteropServices.RuntimeInformation]::IsOSPlatform($([System.Runtime.InteropServices.OSPlatform]::OSX)))' == 'true'">true</IsOSX> 
     <IsLinux Condition="'$([System.Runtime.InteropServices.RuntimeInformation]::IsOSPlatform($([System.Runtime.InteropServices.OSPlatform]::Linux)))' == 'true'">true</IsLinux> 
-
+    <DebugType>pdbonly</DebugType>
   </PropertyGroup>
   <PropertyGroup Condition=" '$(Configuration)|$(Platform)' == 'Development|AnyCPU' ">
     <DefineConstants>$(DefineConstants);TRACE</DefineConstants>
@@ -27,6 +27,8 @@
   </PropertyGroup>
   <PropertyGroup Condition=" '$(Configuration)|$(Platform)' == 'Debug|AnyCPU' ">
     <DefineConstants>$(DefineConstants);DEBUG;TRACE</DefineConstants>
+    <DebugSymbols>true</DebugSymbols>
+    <DebugType>full</DebugType>
   </PropertyGroup>
   <ItemGroup>
     <Reference Include="WindowsDevicePortalWrapper">
diff --git a/Engine/Source/Programs/AutomationTool/IOS/IOS.Automation.csproj b/Engine/Source/Programs/AutomationTool/IOS/IOS.Automation.csproj
--- a/Engine/Source/Programs/AutomationTool/IOS/IOS.Automation.csproj
+++ b/Engine/Source/Programs/AutomationTool/IOS/IOS.Automation.csproj
@@ -12,6 +12,7 @@
     <WarningsNotAsErrors>612,618</WarningsNotAsErrors>
     <OutputPath>..\..\..\..\Binaries\DotNET\AutomationTool\AutomationScripts\Platforms\IOS</OutputPath>
     <TreatWarningsAsErrors>true</TreatWarningsAsErrors>
+    <DebugType>pdbonly</DebugType>
   </PropertyGroup>
   <PropertyGroup Condition=" '$(Configuration)|$(Platform)' == 'Development|AnyCPU' ">
     <DefineConstants>$(DefineConstants);TRACE</DefineConstants>
@@ -23,6 +24,8 @@
   </PropertyGroup>
   <PropertyGroup Condition=" '$(Configuration)|$(Platform)' == 'Debug|AnyCPU' ">
     <DefineConstants>$(DefineConstants);DEBUG;TRACE</DefineConstants>
+    <DebugSymbols>true</DebugSymbols>
+    <DebugType>full</DebugType>
   </PropertyGroup>
   <ItemGroup>
     <!-- BouncyCastle.Crypto is not actually used but is required so that IPhonePackager has it available after we have copied files to a local folder -->
diff --git a/Engine/Source/Programs/AutomationTool/Linux/Linux.Automation.csproj b/Engine/Source/Programs/AutomationTool/Linux/Linux.Automation.csproj
--- a/Engine/Source/Programs/AutomationTool/Linux/Linux.Automation.csproj
+++ b/Engine/Source/Programs/AutomationTool/Linux/Linux.Automation.csproj
@@ -12,6 +12,7 @@
     <WarningsNotAsErrors>612,618</WarningsNotAsErrors>
     <OutputPath>..\..\..\..\Binaries\DotNET\AutomationTool\AutomationScripts\Platforms\Linux\</OutputPath>
     <TreatWarningsAsErrors>true</TreatWarningsAsErrors>
+    <DebugType>pdbonly</DebugType>
   </PropertyGroup>
   <PropertyGroup Condition=" '$(Configuration)|$(Platform)' == 'Development|AnyCPU' ">
     <DefineConstants>$(DefineConstants);TRACE</DefineConstants>
@@ -23,6 +24,8 @@
   </PropertyGroup>
   <PropertyGroup Condition=" '$(Configuration)|$(Platform)' == 'Debug|AnyCPU' ">
     <DefineConstants>$(DefineConstants);DEBUG;TRACE</DefineConstants>
+    <DebugSymbols>true</DebugSymbols>
+    <DebugType>full</DebugType>
   </PropertyGroup>
 
   <ItemGroup>
diff --git a/Engine/Source/Programs/AutomationTool/Localization/Localization.Automation.csproj b/Engine/Source/Programs/AutomationTool/Localization/Localization.Automation.csproj
--- a/Engine/Source/Programs/AutomationTool/Localization/Localization.Automation.csproj
+++ b/Engine/Source/Programs/AutomationTool/Localization/Localization.Automation.csproj
@@ -12,6 +12,7 @@
     <WarningsNotAsErrors>612,618</WarningsNotAsErrors>
     <OutputPath>..\..\..\..\Binaries\DotNET\AutomationTool\AutomationScripts\Localization</OutputPath>
     <TreatWarningsAsErrors>true</TreatWarningsAsErrors>
+    <DebugType>pdbonly</DebugType>
   </PropertyGroup>
   <PropertyGroup Condition=" '$(Configuration)|$(Platform)' == 'Development|AnyCPU' ">
     <DefineConstants>$(DefineConstants);TRACE</DefineConstants>
@@ -23,6 +24,8 @@
   </PropertyGroup>
   <PropertyGroup Condition=" '$(Configuration)|$(Platform)' == 'Debug|AnyCPU' ">
     <DefineConstants>$(DefineConstants);DEBUG;TRACE</DefineConstants>
+    <DebugSymbols>true</DebugSymbols>
+    <DebugType>full</DebugType>
   </PropertyGroup>
 
   <ItemGroup>
diff --git a/Engine/Source/Programs/AutomationTool/Lumin/Lumin.Automation.csproj b/Engine/Source/Programs/AutomationTool/Lumin/Lumin.Automation.csproj
--- a/Engine/Source/Programs/AutomationTool/Lumin/Lumin.Automation.csproj
+++ b/Engine/Source/Programs/AutomationTool/Lumin/Lumin.Automation.csproj
@@ -12,6 +12,7 @@
     <WarningsNotAsErrors>612,618</WarningsNotAsErrors>
     <OutputPath>..\..\..\..\Binaries\DotNET\AutomationTool\AutomationScripts\Platforms\Lumin\</OutputPath>
     <TreatWarningsAsErrors>true</TreatWarningsAsErrors>
+    <DebugType>pdbonly</DebugType>
   </PropertyGroup>
   <PropertyGroup Condition=" '$(Configuration)|$(Platform)' == 'Development|AnyCPU' ">
     <DefineConstants>$(DefineConstants);TRACE</DefineConstants>
@@ -23,6 +24,8 @@
   </PropertyGroup>
   <PropertyGroup Condition=" '$(Configuration)|$(Platform)' == 'Debug|AnyCPU' ">
     <DefineConstants>$(DefineConstants);DEBUG;TRACE</DefineConstants>
+    <DebugSymbols>true</DebugSymbols>
+    <DebugType>full</DebugType>
   </PropertyGroup>
   <ItemGroup>
     <Reference Include="Ionic.Zip.Reduced">
diff --git a/Engine/Source/Programs/AutomationTool/Mac/Mac.Automation.csproj b/Engine/Source/Programs/AutomationTool/Mac/Mac.Automation.csproj
--- a/Engine/Source/Programs/AutomationTool/Mac/Mac.Automation.csproj
+++ b/Engine/Source/Programs/AutomationTool/Mac/Mac.Automation.csproj
@@ -12,6 +12,7 @@
     <WarningsNotAsErrors>612,618</WarningsNotAsErrors>
     <OutputPath>..\..\..\..\Binaries\DotNET\AutomationTool\AutomationScripts\Platforms\Mac\</OutputPath>
     <TreatWarningsAsErrors>true</TreatWarningsAsErrors>
+    <DebugType>pdbonly</DebugType>
   </PropertyGroup>
   <PropertyGroup Condition=" '$(Configuration)|$(Platform)' == 'Development|AnyCPU' ">
     <DefineConstants>$(DefineConstants);TRACE</DefineConstants>
@@ -23,6 +24,8 @@
   </PropertyGroup>
   <PropertyGroup Condition=" '$(Configuration)|$(Platform)' == 'Debug|AnyCPU' ">
     <DefineConstants>$(DefineConstants);DEBUG;TRACE</DefineConstants>
+    <DebugSymbols>true</DebugSymbols>
+    <DebugType>full</DebugType>
   </PropertyGroup>
   <ItemGroup>
     <Compile Include="..\..\DotNETCommon\MetaData.cs">
diff --git a/Engine/Source/Programs/AutomationTool/OneSkyLocalization/OneSkyLocalization.Automation.csproj b/Engine/Source/Programs/AutomationTool/OneSkyLocalization/OneSkyLocalization.Automation.csproj
--- a/Engine/Source/Programs/AutomationTool/OneSkyLocalization/OneSkyLocalization.Automation.csproj
+++ b/Engine/Source/Programs/AutomationTool/OneSkyLocalization/OneSkyLocalization.Automation.csproj
@@ -12,6 +12,7 @@
     <WarningsNotAsErrors>612,618</WarningsNotAsErrors>
     <OutputPath>..\..\..\..\Binaries\DotNET\AutomationTool\AutomationScripts\OneSkyLocalization</OutputPath>
     <TreatWarningsAsErrors>true</TreatWarningsAsErrors>
+    <DebugType>pdbonly</DebugType>
   </PropertyGroup>
   <PropertyGroup Condition=" '$(Configuration)|$(Platform)' == 'Development|AnyCPU' ">
     <DefineConstants>$(DefineConstants);TRACE</DefineConstants>
@@ -23,6 +24,8 @@
   </PropertyGroup>
   <PropertyGroup Condition=" '$(Configuration)|$(Platform)' == 'Debug|AnyCPU' ">
     <DefineConstants>$(DefineConstants);DEBUG;TRACE</DefineConstants>
+    <DebugSymbols>true</DebugSymbols>
+    <DebugType>full</DebugType>
   </PropertyGroup>
   <ItemGroup>
     <Reference Include="OneSky, Version=1.0.0.0, Culture=neutral, processorArchitecture=MSIL">
diff --git a/Engine/Source/Programs/AutomationTool/Scripts/AutomationScripts.Automation.csproj b/Engine/Source/Programs/AutomationTool/Scripts/AutomationScripts.Automation.csproj
--- a/Engine/Source/Programs/AutomationTool/Scripts/AutomationScripts.Automation.csproj
+++ b/Engine/Source/Programs/AutomationTool/Scripts/AutomationScripts.Automation.csproj
@@ -12,6 +12,7 @@
     <WarningsNotAsErrors>612,618</WarningsNotAsErrors>
     <OutputPath>..\..\..\..\Binaries\DotNET\AutomationTool\AutomationScripts\Scripts</OutputPath>
     <TreatWarningsAsErrors>true</TreatWarningsAsErrors>
+    <DebugType>pdbonly</DebugType>
   </PropertyGroup>
   <PropertyGroup Condition=" '$(Configuration)|$(Platform)' == 'Development|AnyCPU' ">
     <DefineConstants>$(DefineConstants);TRACE</DefineConstants>
@@ -23,6 +24,8 @@
   </PropertyGroup>
   <PropertyGroup Condition=" '$(Configuration)|$(Platform)' == 'Debug|AnyCPU' ">
     <DefineConstants>$(DefineConstants);DEBUG;TRACE</DefineConstants>
+    <DebugSymbols>true</DebugSymbols>
+    <DebugType>full</DebugType>
   </PropertyGroup>
   <ItemGroup>
     <Reference Include="Ionic.Zip.Reduced">
diff --git a/Engine/Source/Programs/AutomationTool/TVOS/TVOS.Automation.csproj b/Engine/Source/Programs/AutomationTool/TVOS/TVOS.Automation.csproj
--- a/Engine/Source/Programs/AutomationTool/TVOS/TVOS.Automation.csproj
+++ b/Engine/Source/Programs/AutomationTool/TVOS/TVOS.Automation.csproj
@@ -12,6 +12,7 @@
     <WarningsNotAsErrors>612,618</WarningsNotAsErrors>
     <OutputPath>..\..\..\..\Binaries\DotNET\AutomationTool\AutomationScripts\Platforms\TVOS\</OutputPath>
     <TreatWarningsAsErrors>true</TreatWarningsAsErrors>
+    <DebugType>pdbonly</DebugType>
   </PropertyGroup>
   <PropertyGroup Condition=" '$(Configuration)|$(Platform)' == 'Development|AnyCPU' ">
     <DefineConstants>$(DefineConstants);TRACE</DefineConstants>
@@ -23,6 +24,8 @@
   </PropertyGroup>
   <PropertyGroup Condition=" '$(Configuration)|$(Platform)' == 'Debug|AnyCPU' ">
     <DefineConstants>$(DefineConstants);DEBUG;TRACE</DefineConstants>
+    <DebugSymbols>true</DebugSymbols>
+    <DebugType>full</DebugType>
   </PropertyGroup>
   <ItemGroup>
     <Reference Include="Ionic.Zip.Reduced">
diff --git a/Engine/Source/Programs/AutomationTool/Turnkey/Turnkey.Automation.csproj b/Engine/Source/Programs/AutomationTool/Turnkey/Turnkey.Automation.csproj
--- a/Engine/Source/Programs/AutomationTool/Turnkey/Turnkey.Automation.csproj
+++ b/Engine/Source/Programs/AutomationTool/Turnkey/Turnkey.Automation.csproj
@@ -15,6 +15,7 @@
     <IsWindows Condition="'$([System.Runtime.InteropServices.RuntimeInformation]::IsOSPlatform($([System.Runtime.InteropServices.OSPlatform]::Windows)))' == 'true'">true</IsWindows> 
     <IsOSX Condition="'$([System.Runtime.InteropServices.RuntimeInformation]::IsOSPlatform($([System.Runtime.InteropServices.OSPlatform]::OSX)))' == 'true'">true</IsOSX> 
     <IsLinux Condition="'$([System.Runtime.InteropServices.RuntimeInformation]::IsOSPlatform($([System.Runtime.InteropServices.OSPlatform]::Linux)))' == 'true'">true</IsLinux> 
+    <DebugType>pdbonly</DebugType>
   </PropertyGroup>
   <PropertyGroup Condition="'$(Configuration)|$(Platform)' == 'Development|AnyCPU' ">
     <DefineConstants>$(DefineConstants);TRACE</DefineConstants>
@@ -26,6 +27,8 @@
   </PropertyGroup>
   <PropertyGroup Condition="'$(Configuration)|$(Platform)' == 'Debug|AnyCPU' ">
     <DefineConstants>$(DefineConstants);DEBUG;TRACE</DefineConstants>
+    <DebugSymbols>true</DebugSymbols>
+    <DebugType>full</DebugType>
   </PropertyGroup>
   <PropertyGroup Condition="'$(IsWindows)'=='true'">
     <DefineConstants>$(DefineConstants);WINDOWS</DefineConstants>
diff --git a/Engine/Source/Programs/AutomationTool/Win/Win.Automation.csproj b/Engine/Source/Programs/AutomationTool/Win/Win.Automation.csproj
--- a/Engine/Source/Programs/AutomationTool/Win/Win.Automation.csproj
+++ b/Engine/Source/Programs/AutomationTool/Win/Win.Automation.csproj
@@ -12,6 +12,7 @@
     <WarningsNotAsErrors>612,618</WarningsNotAsErrors>
     <OutputPath>..\..\..\..\Binaries\DotNET\AutomationTool\AutomationScripts\Platforms\Windows</OutputPath>
     <TreatWarningsAsErrors>true</TreatWarningsAsErrors>
+    <DebugType>pdbonly</DebugType>
   </PropertyGroup>
   <PropertyGroup Condition=" '$(Configuration)|$(Platform)' == 'Development|AnyCPU' ">
     <DefineConstants>$(DefineConstants);TRACE</DefineConstants>
@@ -23,6 +24,8 @@
   </PropertyGroup>
   <PropertyGroup Condition=" '$(Configuration)|$(Platform)' == 'Debug|AnyCPU' ">
     <DefineConstants>$(DefineConstants);DEBUG;TRACE</DefineConstants>
+    <DebugSymbols>true</DebugSymbols>
+    <DebugType>full</DebugType>
   </PropertyGroup>
   <ItemGroup>
     <Reference Include="Ionic.Zip.Reduced">
diff --git a/Engine/Source/Programs/AutomationTool/XLocLocalization/XLocLocalization.Automation.csproj b/Engine/Source/Programs/AutomationTool/XLocLocalization/XLocLocalization.Automation.csproj
--- a/Engine/Source/Programs/AutomationTool/XLocLocalization/XLocLocalization.Automation.csproj
+++ b/Engine/Source/Programs/AutomationTool/XLocLocalization/XLocLocalization.Automation.csproj
@@ -14,6 +14,7 @@
     <TreatWarningsAsErrors>true</TreatWarningsAsErrors>
     <CopyLocalLockFileAssemblies>true</CopyLocalLockFileAssemblies>
     <SatelliteResourceLanguages>en</SatelliteResourceLanguages> <!-- remove non english resource languages -->
+    <DebugType>pdbonly</DebugType>
   </PropertyGroup>
   <PropertyGroup Condition=" '$(Configuration)|$(Platform)' == 'Development|AnyCPU' ">
     <DefineConstants>$(DefineConstants);TRACE</DefineConstants>
@@ -25,6 +26,8 @@
   </PropertyGroup>
   <PropertyGroup Condition=" '$(Configuration)|$(Platform)' == 'Debug|AnyCPU' ">
     <DefineConstants>$(DefineConstants);DEBUG;TRACE</DefineConstants>
+    <DebugSymbols>true</DebugSymbols>
+    <DebugType>full</DebugType>
   </PropertyGroup>
 
   <ItemGroup>
