Changed pdb format to not use portable pdbs (which is the default in net core) and instead back to our old pdb configuration.
This resolves a issue with PCBs were pdbcopy failed to run because it does not support portable pdbs.

diff --git a/Engine/Source/Programs/AutomationTool/AutomationTool.csproj b/Engine/Source/Programs/AutomationTool/AutomationTool.csproj
--- a/Engine/Source/Programs/AutomationTool/AutomationTool.csproj
+++ b/Engine/Source/Programs/AutomationTool/AutomationTool.csproj
@@ -18,6 +18,7 @@
     <IsWindows Condition="'$([System.Runtime.InteropServices.RuntimeInformation]::IsOSPlatform($([System.Runtime.InteropServices.OSPlatform]::Windows)))' == 'true'">true</IsWindows> 
     <IsOSX Condition="'$([System.Runtime.InteropServices.RuntimeInformation]::IsOSPlatform($([System.Runtime.InteropServices.OSPlatform]::OSX)))' == 'true'">true</IsOSX> 
     <IsLinux Condition="'$([System.Runtime.InteropServices.RuntimeInformation]::IsOSPlatform($([System.Runtime.InteropServices.OSPlatform]::Linux)))' == 'true'">true</IsLinux> 
+    <DebugType>pdbonly</DebugType>
   </PropertyGroup>
   <PropertyGroup Condition=" '$(Configuration)|$(Platform)' == 'Development|AnyCPU' ">
     <DefineConstants>$(DefineConstants);TRACE</DefineConstants>
@@ -29,6 +30,8 @@
   </PropertyGroup>
   <PropertyGroup Condition=" '$(Configuration)|$(Platform)' == 'Debug|AnyCPU' ">
     <DefineConstants>$(DefineConstants);DEBUG;TRACE</DefineConstants>
+    <DebugSymbols>true</DebugSymbols>
+    <DebugType>full</DebugType>
   </PropertyGroup>
   <Choose>
     <When Condition="'$(IsOSX)'=='true'">
diff --git a/Engine/Source/Programs/UnrealBuildTool/UnrealBuildTool.csproj b/Engine/Source/Programs/UnrealBuildTool/UnrealBuildTool.csproj
--- a/Engine/Source/Programs/UnrealBuildTool/UnrealBuildTool.csproj
+++ b/Engine/Source/Programs/UnrealBuildTool/UnrealBuildTool.csproj
@@ -15,6 +15,7 @@
     <IsWindows Condition="'$([System.Runtime.InteropServices.RuntimeInformation]::IsOSPlatform($([System.Runtime.InteropServices.OSPlatform]::Windows)))' == 'true'">true</IsWindows> 
     <IsOSX Condition="'$([System.Runtime.InteropServices.RuntimeInformation]::IsOSPlatform($([System.Runtime.InteropServices.OSPlatform]::OSX)))' == 'true'">true</IsOSX> 
     <IsLinux Condition="'$([System.Runtime.InteropServices.RuntimeInformation]::IsOSPlatform($([System.Runtime.InteropServices.OSPlatform]::Linux)))' == 'true'">true</IsLinux> 
+    <DebugType>pdbonly</DebugType>
   </PropertyGroup>
   <Choose>
     <When Condition="'$(IsOSX)'=='true'">
@@ -37,6 +38,8 @@
   <PropertyGroup Condition=" '$(Configuration)|$(Platform)' == 'Debug|AnyCPU' ">
     <OutputPath>..\..\..\Binaries\DotNET\UnrealBuildTool</OutputPath>
     <DocumentationFile>..\..\..\Binaries\DotNET\UnrealBuildTool\UnrealBuildTool.XML</DocumentationFile>
+    <DebugSymbols>true</DebugSymbols>
+    <DebugType>full</DebugType>
   </PropertyGroup>
   <ItemGroup>
     <Reference Include="Ionic.Zip.Reduced">
