Brought back -nocompile and -compile options in UAT even if they are not used, to avoid erroring out when these gets forward by the batch file which will not remove already handled flags.

diff --git a/Engine/Source/Programs/AutomationTool/AutomationUtils/Automation.cs b/Engine/Source/Programs/AutomationTool/AutomationUtils/Automation.cs
--- a/Engine/Source/Programs/AutomationTool/AutomationUtils/Automation.cs
+++ b/Engine/Source/Programs/AutomationTool/AutomationUtils/Automation.cs
@@ -103,6 +103,15 @@ namespace AutomationTool
 		public static CommandLineArg UTF8Output = new CommandLineArg("-UTF8Output");
         public static CommandLineArg AllowStdOutLogVerbosity = new CommandLineArg("-AllowStdOutLogVerbosity");
         public static CommandLineArg NoAutoSDK = new CommandLineArg("-NoAutoSDK");
+
+        /// <summary>
+        /// These command is LEGACY and are all handled by RunUAT shell files.
+        /// However, the shell files simply passes on all arguments, so UAT will choke when encountering them.
+        /// Keep this CommandLineArg around so that doesn't happen.
+        /// </summary>
+        public static CommandLineArg NoCompileLegacyDontUse = new CommandLineArg("-NoCompile");
+        public static CommandLineArg CompileLegacyDontUse = new CommandLineArg("-Compile");
+
         /// <summary>
         /// Allows you to use local storage for your root build storage dir (default of P:\Builds (on PC) is changed to Engine\Saved\LocalBuilds). Used for local testing.
         /// </summary>
