# Simplify gauntlet report for Horde

Introduce Gauntlet.ITestReport interface
 - enable attaching artifact to report through ITestReport.AttachArtifact(string)
 - enable reporting errors and warnings to ITestReport

Deprecate CreateReport(TestResult, Context, Build, Artifacts, ArtifactPath) and favor to CreateReport(ITestReport)

Generate a simplified report for Horde Test Data

Better support copying artifacts with long path on Windows

diff --git a/Engine/Source/Programs/AutomationTool/Gauntlet/Framework/Base/Gauntlet.BaseTestNode.cs b/Engine/Source/Programs/AutomationTool/Gauntlet/Framework/Base/Gauntlet.BaseTestNode.cs
--- a/Engine/Source/Programs/AutomationTool/Gauntlet/Framework/Base/Gauntlet.BaseTestNode.cs
+++ b/Engine/Source/Programs/AutomationTool/Gauntlet/Framework/Base/Gauntlet.BaseTestNode.cs
@@ -39,6 +39,11 @@ namespace Gauntlet
 		/// </summary>
 		public virtual bool HasWarnings { get; protected set; }
 
+		/// <summary>
+		/// Returns reason for the test cancellation
+		/// </summary>
+		public virtual string CancellationReason { get; protected set; }
+
 		/// <summary>
 		/// Returns true if the test was cancelled
 		/// </summary>
@@ -159,6 +164,16 @@ namespace Gauntlet
 		{
 		}
 
+		/// <summary>
+		/// Sets Cancellation Reason.
+		/// </summary>
+		/// <param name="InReason"></param>
+		/// <returns></returns>
+		public virtual void SetCancellationReason(string InReason)
+		{
+			CancellationReason = InReason;
+		}
+
 		/// <summary>
 		/// Called after the test is completed and shutdown
 		/// </summary>
diff --git a/Engine/Source/Programs/AutomationTool/Gauntlet/Framework/Base/Gauntlet.TestNode.cs b/Engine/Source/Programs/AutomationTool/Gauntlet/Framework/Base/Gauntlet.TestNode.cs
--- a/Engine/Source/Programs/AutomationTool/Gauntlet/Framework/Base/Gauntlet.TestNode.cs
+++ b/Engine/Source/Programs/AutomationTool/Gauntlet/Framework/Base/Gauntlet.TestNode.cs
@@ -109,6 +109,13 @@ namespace Gauntlet
 		/// </summary>
 		void TickTest();
 
+		/// <summary>
+		/// Sets Cancellation Reason.
+		/// </summary>
+		/// <param name="Reason"></param>
+		/// <returns></returns>
+		void SetCancellationReason(string Reason);
+
 		/// <summary>
 		/// Called to request that the test stop, either because it's reported that its complete or due to external factors.
 		/// Tests should consider whether they passed or succeeded (even a terminated test may have gotten all the data it needs) 
diff --git a/Engine/Source/Programs/AutomationTool/Gauntlet/Framework/Base/Gauntlet.TestReport.cs b/Engine/Source/Programs/AutomationTool/Gauntlet/Framework/Base/Gauntlet.TestReport.cs
new file mode 100644
--- /dev/null
+++ b/Engine/Source/Programs/AutomationTool/Gauntlet/Framework/Base/Gauntlet.TestReport.cs
@@ -0,0 +1,102 @@
+// Copyright Epic Games, Inc. All Rights Reserved.
+
+using System;
+using System.Reflection;
+
+
+namespace Gauntlet
+{
+	/// <summary>
+	/// Interface for test report
+	/// </summary>
+	public interface ITestReport
+	{
+		/// <summary>
+		/// Return report type
+		/// </summary>
+		string Type { get; }
+
+		/// <summary>
+		/// Set a property of the test report type
+		/// </summary>
+		/// <param name="Property"></param>
+		/// <param name="Value"></param>
+		/// <returns></returns>
+		void SetProperty(string Property, object Value);
+
+		/// <summary>
+		/// Add event to the test report
+		/// </summary>
+		/// <param name="Type"></param>
+		/// <param name="Message"></param>
+		/// <param name="Context"></param>
+		/// <returns></returns>
+		void AddEvent(string Type, string Message, object Context = null);
+
+		/// <summary>
+		/// Attach an Artifact to the ITestReport
+		/// </summary>
+		/// <param name="ArtifactPath"></param>
+		/// <returns></returns>
+		bool AttachArtifact(string ArtifactPath, string Name = null);
+	}
+
+	/// <summary>
+	/// Test Report Base class
+	/// </summary>
+	public class BaseTestReport : ITestReport
+	{
+		/// <summary>
+		/// Return report type
+		/// </summary>
+		public virtual string Type { get { return "Base Test Report"; } }
+
+		/// <summary>
+		/// Set a property of the test report type
+		/// </summary>
+		/// <param name="Property"></param>
+		/// <param name="Value"></param>
+		/// <returns></returns>
+		public virtual void SetProperty(string Property, object Value)
+		{
+			PropertyInfo PropertyInstance = GetType().GetProperty(Property);
+			PropertyInstance.SetValue(this, Convert.ChangeType(Value, PropertyInstance.PropertyType));
+		}
+
+		/// <summary>
+		/// Add event to the test report
+		/// </summary>
+		/// <param name="Type"></param>
+		/// <param name="Message"></param>
+		/// <param name="Context"></param>
+		/// <returns></returns>
+		public virtual void AddEvent(string Type, string Message, object Context = null)
+		{
+			throw new System.NotImplementedException("AddEvent not implemented");
+		}
+
+		public void AddError(string Message, object Context = null)
+		{
+			AddEvent("Error", Message, Context);
+		}
+		public void AddWarning(string Message, object Context = null)
+		{
+			AddEvent("Warning", Message, Context);
+		}
+		public void AddInfo(string Message, object Context = null)
+		{
+			AddEvent("Info", Message, Context);
+		}
+
+		/// <summary>
+		/// Attach an Artifact to the ITestReport
+		/// </summary>
+		/// <param name="ArtifactPath"></param>
+		/// <param name="Name"></param>
+		/// <returns>bool</returns>
+		public virtual bool AttachArtifact(string ArtifactPath, string Name = null)
+		{
+			throw new System.NotImplementedException("AttachArtifact not implemented");
+		}
+	}
+}
\ No newline at end of file
diff --git a/Engine/Source/Programs/AutomationTool/Gauntlet/Framework/Gauntlet.TestExecutor.cs b/Engine/Source/Programs/AutomationTool/Gauntlet/Framework/Gauntlet.TestExecutor.cs
--- a/Engine/Source/Programs/AutomationTool/Gauntlet/Framework/Gauntlet.TestExecutor.cs
+++ b/Engine/Source/Programs/AutomationTool/Gauntlet/Framework/Gauntlet.TestExecutor.cs
@@ -650,6 +650,11 @@ namespace Gauntlet
 				Log.Info("{0}", TestInfo.CancellationReason);
 			}
 
+			if (!string.IsNullOrEmpty(TestInfo.CancellationReason))
+			{
+				TestInfo.TestNode.SetCancellationReason(TestInfo.CancellationReason);
+			}
+
 			// if the test is not running. or we've determined a result for it..
 			if (TestIsRunning == false || TestInfo.FinalResult != TestResult.Invalid)
 			{
diff --git a/Engine/Source/Programs/AutomationTool/Gauntlet/Framework/Gauntlet.Utils.cs b/Engine/Source/Programs/AutomationTool/Gauntlet/Framework/Gauntlet.Utils.cs
--- a/Engine/Source/Programs/AutomationTool/Gauntlet/Framework/Gauntlet.Utils.cs
+++ b/Engine/Source/Programs/AutomationTool/Gauntlet/Framework/Gauntlet.Utils.cs
@@ -124,6 +124,15 @@ namespace Gauntlet
 
 		}
 
+		/// <summary>
+		/// Return @"\\?\"; adding the prefix to a path string tells the Windows APIs to disable all string parsing and to send the string that follows it straight to the file system.
+		/// Allowing for longer than 260 characters path.
+		/// </summary>
+		public static string LongPathPrefix
+		{
+			get { return Path.DirectorySeparatorChar == '\\'? @"\\?\" : ""; }
+		}
+
 
 		/// <summary>
 		/// Acquired and released during the main Tick of the Gauntlet systems. Use this before touchung anything global scope from a 
@@ -744,6 +753,24 @@ namespace Gauntlet
 				return ConstructedTypes;
 			}
 
+			/// <summary>
+			/// Check if the method signature is overridden.
+			/// </summary>
+			/// <param name="ClassType"></param>
+			/// <param name="MethodName"></param>
+			/// <param name="Signature"></param>
+			/// <returns></returns>
+			public static bool HasOverriddenMethod(Type ClassType, string MethodName, Type[] Signature)
+			{
+				MethodInfo Method = ClassType.GetMethod(MethodName, Signature);
+				if (Method == null)
+				{
+					throw new Exception(string.Format("Unknown method {0} from {1} with signature {2}.", MethodName, ClassType.Name, Signature.ToString()));
+				}
+				bool IsOverridden = Method.GetBaseDefinition().DeclaringType != Method.DeclaringType;
+				return Method.GetBaseDefinition().DeclaringType != Method.DeclaringType;
+			}
+
 		}
 
 
@@ -1530,4 +1557,17 @@ namespace Gauntlet
 		}
 	}
 
+	public static class FileUtils
+	{
+		/// <summary>
+		/// Sanitize filename
+		/// </summary>
+		/// <param name="Name"></param>
+		/// <returns></returns>
+		static public string SanitizeFilename(string Name)
+		{
+			return Regex.Replace(Name, @"[^a-z0-9_\-.]+", "_", RegexOptions.IgnoreCase);
+		}
+	}
+
 }
diff --git a/Engine/Source/Programs/AutomationTool/Gauntlet/Framework/Utils/Gauntlet.HordeReport.cs b/Engine/Source/Programs/AutomationTool/Gauntlet/Framework/Utils/Gauntlet.HordeReport.cs
--- a/Engine/Source/Programs/AutomationTool/Gauntlet/Framework/Utils/Gauntlet.HordeReport.cs
+++ b/Engine/Source/Programs/AutomationTool/Gauntlet/Framework/Utils/Gauntlet.HordeReport.cs
@@ -17,7 +17,7 @@ namespace Gauntlet
 		{
 			get
 			{
-				return Environment.GetEnvironmentVariable("UE_TESTDATA_DIR") ?? Path.Combine(CommandUtils.CmdEnv.EngineSavedFolder, "TestData");
+				return Path.GetFullPath(Environment.GetEnvironmentVariable("UE_TESTDATA_DIR") ?? Path.Combine(CommandUtils.CmdEnv.EngineSavedFolder, "TestData"));
 			}
 		}
 		/// <summary>
@@ -27,7 +27,68 @@ namespace Gauntlet
 		{
 			get
 			{
-				return Environment.GetEnvironmentVariable("UE_ARTIFACTS_DIR") ?? CommandUtils.CmdEnv.LogFolder;
+				return Path.GetFullPath(Environment.GetEnvironmentVariable("UE_ARTIFACTS_DIR") ?? CommandUtils.CmdEnv.LogFolder);
+			}
+		}
+		/// <summary>
+		/// Is Environement set by Horde Agent 
+		/// </summary>
+		public static bool IsUnderHordeAgent
+		{
+			get
+			{
+				return !string.IsNullOrEmpty(Environment.GetEnvironmentVariable("UE_TESTDATA_DIR"));
+			}
+		}
+
+		public class BaseHordeReport : BaseTestReport
+		{
+			protected string OutputArtifactPath;
+
+			/// <summary>
+			/// Attach Artifact to the Test Report
+			/// </summary>
+			/// <param name="ArtifactPath"></param>
+			/// <param name="Name"></param>
+			/// <returns>bool</returns>
+			public override bool AttachArtifact(string ArtifactPath, string Name = null)
+			{
+				if (string.IsNullOrEmpty(OutputArtifactPath))
+				{
+					throw new InvalidOperationException("OutputArtifactPath must be set before attaching any artifact");
+				}
+
+				if (File.Exists(ArtifactPath))
+				{
+					try
+					{
+						string TargetPath = Path.Combine(OutputArtifactPath, Name ?? Path.GetFileName(ArtifactPath));
+						string TargetDirectry = Path.GetDirectoryName(TargetPath);
+						if (!Directory.Exists(TargetDirectry)) { Directory.CreateDirectory(TargetDirectry); }
+						File.Copy(Globals.LongPathPrefix + ArtifactPath, TargetPath);
+						return true;
+					}
+					catch (Exception Ex)
+					{
+						Log.Error("Failed to copy artifact '{0}'. {1}", Path.GetFileName(ArtifactPath), Ex);
+					}
+				}
+				return false;
+			}
+
+			/// <summary>
+			/// Set Output Artifact Path and create the directory if missing
+			/// </summary>
+			/// <param name="ArtifactPath"></param>
+			/// <returns></returns>
+			public void SetOutputArtifactPath(string Path)
+			{
+				OutputArtifactPath = Path;
+
+				if (!Directory.Exists(OutputArtifactPath))
+				{
+					Directory.CreateDirectory(OutputArtifactPath);
+				}
 			}
 		}
 
@@ -169,9 +230,14 @@ namespace Gauntlet
 		/// <summary>
 		/// Contains information about an entire test pass 
 		/// </summary>
-		public class TestPassResults
+		public class UnrealEngineTestPassResults : BaseHordeReport
 		{
-			public TestPassResults()
+			public override string Type
+			{
+				get { return "Unreal Automated Tests"; }
+			}
+
+			public UnrealEngineTestPassResults()
 			{
 				Tests = new List<TestResult>();
 			}
@@ -203,9 +269,9 @@ namespace Gauntlet
 			/// <param name="InTestPassResults"></param>
 			/// <param name="ReportPath"></param>
 			/// <param name="ReportURL"></param>
-			public static TestPassResults FromUnrealAutomatedTests(UnrealAutomatedTestPassResults InTestPassResults, string ReportPath, string ReportURL)
+			public static UnrealEngineTestPassResults FromUnrealAutomatedTests(UnrealAutomatedTestPassResults InTestPassResults, string ReportPath, string ReportURL)
 			{
-				TestPassResults OutTestPassResults = new TestPassResults();
+				UnrealEngineTestPassResults OutTestPassResults = new UnrealEngineTestPassResults();
 				OutTestPassResults.ClientDescriptor = InTestPassResults.clientDescriptor;
 				OutTestPassResults.ReportCreatedOn = InTestPassResults.reportCreatedOn;
 				OutTestPassResults.ReportURL = ReportURL;
@@ -261,13 +327,9 @@ namespace Gauntlet
 			/// Copy Test Results Artifacts
 			/// </summary>
 			/// <param name="OutputArtifactPath"></param>
-			public void CopyTestResultsArtifacts(string OutputArtifactPath)
+			public void CopyTestResultsArtifacts(string InOutputArtifactPath)
 			{
-
-				if (!Directory.Exists(OutputArtifactPath))
-				{
-					Directory.CreateDirectory(OutputArtifactPath);
-				}
+				SetOutputArtifactPath(InOutputArtifactPath);
 				int ArtifactsCount = 0;
 				foreach (TestResult OutputTestResult in Tests)
 				{
@@ -278,18 +340,7 @@ namespace Gauntlet
 						string[] ArtifactPaths= { TestArtifact.Files.Difference, TestArtifact.Files.Approved, TestArtifact.Files.Unapproved };
 						foreach (string ArtifactPath in ArtifactPaths)
 						{
-							if (File.Exists(ArtifactPath))
-							{
-								try
-								{
-									File.Copy(ArtifactPath, Path.Combine(OutputArtifactPath, Path.GetFileName(ArtifactPath)));
-									ArtifactsCount++;
-								}
-								catch (Exception Ex)
-								{
-									Log.Error("Failed to copy artifact '{0}'. {1}", Path.GetFileName(ArtifactPath), Ex);
-								}
-							}
+							if (AttachArtifact(ArtifactPath)) { ArtifactsCount++; }
 						}
 						TestArtifact.Files.Difference = Path.GetFileName(TestArtifact.Files.Difference);
 						TestArtifact.Files.Approved = Path.GetFileName(TestArtifact.Files.Approved);
@@ -311,6 +362,56 @@ namespace Gauntlet
 			}
 		}
 
+		/// <summary>
+		/// Contains Test success/failed status and a list of errors and warnings
+		/// </summary>
+		public class SimpleTestReport : BaseHordeReport
+		{
+			public override string Type
+			{
+				get { return "Simple Report"; }
+			}
+			public SimpleTestReport()
+			{
+				Logs = new List<String>();
+				Errors = new List<String>();
+				Warnings = new List<String>();
+			}
+
+			public string Description;
+			public string ReportCreatedOn;
+			public float TotalDurationSeconds;
+			public bool HasSucceeded;
+			public string Status;
+			public string URLLink;
+			public List<String> Logs;
+			public List<String> Errors;
+			public List<String> Warnings;
+
+			public override void AddEvent(string Type, string Message, object Context = null)
+			{
+				if (Type.ToLower() == "error")
+				{
+					Errors.Add(Message);
+				}
+				else if (Type.ToLower() == "warning")
+				{
+					Warnings.Add(Message);
+				}
+				// The rest is ignored with this report type.
+			}
+
+			public override bool AttachArtifact(string ArtifactPath, string Name = null)
+			{
+				if (base.AttachArtifact(ArtifactPath, Name))
+				{
+					Logs.Add(Name ?? Path.GetFileName(ArtifactPath));
+					return true;
+				}
+				return false;
+			}
+		}
+
 		/// <summary>
 		/// Container for Test Data items 
 		/// </summary>
@@ -326,10 +427,14 @@ namespace Gauntlet
 				Items = new List<DataItem>();
 			}
 
-			public DataItem AddNewItem(string InType, string InKey, object InData)
+			public DataItem AddNewTestReport(string InKey, ITestReport InData)
 			{
+				if (string.IsNullOrEmpty(InKey))
+				{
+					throw new System.ArgumentException("Test Data key can't be an empty string.");
+				}
 				DataItem NewDataItem = new DataItem();
-				NewDataItem.Key = InType +"::"+ InKey;
+				NewDataItem.Key = InData.Type +"::"+ InKey;
 				NewDataItem.Data = InData;
 				Items.Add(NewDataItem);
 				return NewDataItem;
@@ -341,13 +446,25 @@ namespace Gauntlet
 			/// Write Test Data Collection to json
 			/// </summary>
 			/// <param name="OutputTestDataFilePath"></param>
-			public void WriteToJson(string OutputTestDataFilePath)
+			public void WriteToJson(string OutputTestDataFilePath, bool bIncrementNameIfFileExists = false)
 			{
 				string OutputTestDataDir = Path.GetDirectoryName(OutputTestDataFilePath);
 				if (!Directory.Exists(OutputTestDataDir))
 				{
 					Directory.CreateDirectory(OutputTestDataDir);
 				}
+				if (File.Exists(OutputTestDataFilePath) && bIncrementNameIfFileExists)
+				{
+					// increment filename if file exists
+					string Ext = Path.GetExtension(OutputTestDataFilePath);
+					string Filename = OutputTestDataFilePath.Replace(Ext, "");
+					int Incr = 0;
+					do
+					{
+						Incr++;
+						OutputTestDataFilePath = string.Format("{0}{1}{2}", Filename, Incr, Ext);
+					} while (File.Exists(OutputTestDataFilePath));
+				}
 				// write test pass summary
 				Log.Verbose("Writing Test Data Collection for Horde at {0}", OutputTestDataFilePath);
 				try
diff --git a/Engine/Source/Programs/AutomationTool/Gauntlet/SelfTest/Framework/Gauntlet.SelfTest.OrderOfOpsTest.cs b/Engine/Source/Programs/AutomationTool/Gauntlet/SelfTest/Framework/Gauntlet.SelfTest.OrderOfOpsTest.cs
--- a/Engine/Source/Programs/AutomationTool/Gauntlet/SelfTest/Framework/Gauntlet.SelfTest.OrderOfOpsTest.cs
+++ b/Engine/Source/Programs/AutomationTool/Gauntlet/SelfTest/Framework/Gauntlet.SelfTest.OrderOfOpsTest.cs
@@ -140,6 +140,10 @@ namespace Gauntlet.SelfTest
 			return true;
 		}
 
+		public virtual void SetCancellationReason(string Reason)
+		{
+		}
+
 		public void StopTest(bool WasCancelled)
 		{
 			SetNewStage(TestStages.StopTest);
diff --git a/Engine/Source/Programs/AutomationTool/Gauntlet/SelfTest/Gauntlet.SelfTest.BaseNode.cs b/Engine/Source/Programs/AutomationTool/Gauntlet/SelfTest/Gauntlet.SelfTest.BaseNode.cs
--- a/Engine/Source/Programs/AutomationTool/Gauntlet/SelfTest/Gauntlet.SelfTest.BaseNode.cs
+++ b/Engine/Source/Programs/AutomationTool/Gauntlet/SelfTest/Gauntlet.SelfTest.BaseNode.cs
@@ -113,6 +113,10 @@ namespace Gauntlet.SelfTest
 		{
 		}
 
+		public virtual void SetCancellationReason(string Reason)
+		{
+		}
+
 		public virtual void StopTest(bool WasCancelled)
 		{
 		}
diff --git a/Engine/Source/Programs/AutomationTool/Gauntlet/Unreal/Base/Gauntlet.UnrealSession.cs b/Engine/Source/Programs/AutomationTool/Gauntlet/Unreal/Base/Gauntlet.UnrealSession.cs
--- a/Engine/Source/Programs/AutomationTool/Gauntlet/Unreal/Base/Gauntlet.UnrealSession.cs
+++ b/Engine/Source/Programs/AutomationTool/Gauntlet/Unreal/Base/Gauntlet.UnrealSession.cs
@@ -1046,8 +1046,8 @@ namespace Gauntlet
 
 				if (Directory.Exists(SourceSavedDir))
 				{
-					Utils.SystemHelpers.CopyDirectory(SourceSavedDir, DestSavedDir);
-					Log.Info("Archived artifacts to to {0}", DestSavedDir);
+					Utils.SystemHelpers.CopyDirectory(SourceSavedDir, Globals.LongPathPrefix + DestSavedDir);
+					Log.Info("Archived artifacts to {0}", DestSavedDir);
 				}
 				else
 				{
@@ -1078,7 +1078,7 @@ namespace Gauntlet
 					{
 						// Grab the final dir name to copy everything into so everything's not just going into root artifact dir.
 						string IntendedCopyLocation = Path.Combine(InDestArtifactPath, DirToCopy.Name);
-						Utils.SystemHelpers.CopyDirectory(SourcePath, IntendedCopyLocation);
+						Utils.SystemHelpers.CopyDirectory(SourcePath, Globals.LongPathPrefix + IntendedCopyLocation);
 					}
 				}
 			}
diff --git a/Engine/Source/Programs/AutomationTool/Gauntlet/Unreal/Base/Gauntlet.UnrealTestConfiguration.cs b/Engine/Source/Programs/AutomationTool/Gauntlet/Unreal/Base/Gauntlet.UnrealTestConfiguration.cs
--- a/Engine/Source/Programs/AutomationTool/Gauntlet/Unreal/Base/Gauntlet.UnrealTestConfiguration.cs
+++ b/Engine/Source/Programs/AutomationTool/Gauntlet/Unreal/Base/Gauntlet.UnrealTestConfiguration.cs
@@ -506,6 +506,30 @@ namespace Gauntlet
 		[AutoParam(600.0f)]
 		public float MaxDuration { get; set; }
 
+		/// <summary>
+		/// Produce test artifacts for Horde build system
+		/// </summary>
+		[AutoParam]
+		public bool WriteTestResultsForHorde = false;
+
+		/// <summary>
+		/// Path to store test data for Horde build system
+		/// </summary>
+		[AutoParam]
+		public string HordeTestDataPath = "";
+
+		/// <summary>
+		/// Key to store Horde Test Data
+		/// </summary>
+		[AutoParam]
+		public string HordeTestDataKey = "";
+
+		/// <summary>
+		/// Path to store test artifacts for Horde build system
+		/// </summary>
+		[AutoParam]
+		public string HordeArtifactPath = "";
+
 		/// <summary>
 		/// What the test result should be treated as if we reach max duration.
 		/// </summary>
@@ -723,6 +747,12 @@ namespace Gauntlet
 			{
 				AppConfig.CommandLine += string.Format(" -gauntlet.heartbeatperiod={0}", HeartbeatPeriod);
 			}
+
+			// we write results to Horde test data if we run under Horde agent
+			if (HordeReport.IsUnderHordeAgent)
+			{
+				WriteTestResultsForHorde = true;
+			}
 		}
 	}
 
diff --git a/Engine/Source/Programs/AutomationTool/Gauntlet/Unreal/Base/Gauntlet.UnrealTestContext.cs b/Engine/Source/Programs/AutomationTool/Gauntlet/Unreal/Base/Gauntlet.UnrealTestContext.cs
--- a/Engine/Source/Programs/AutomationTool/Gauntlet/Unreal/Base/Gauntlet.UnrealTestContext.cs
+++ b/Engine/Source/Programs/AutomationTool/Gauntlet/Unreal/Base/Gauntlet.UnrealTestContext.cs
@@ -223,7 +223,9 @@ namespace Gauntlet
 			{
 				Globals.LogDir = LogDir;
 			}
-
+			// Normalize paths
+			LogDir = Path.GetFullPath(LogDir);
+			TempDir = Path.GetFullPath(TempDir);
 
 			// normalize the project name and get the path
 			if (File.Exists(Project))
diff --git a/Engine/Source/Programs/AutomationTool/Gauntlet/Unreal/Base/Gauntlet.UnrealTestNode.cs b/Engine/Source/Programs/AutomationTool/Gauntlet/Unreal/Base/Gauntlet.UnrealTestNode.cs
--- a/Engine/Source/Programs/AutomationTool/Gauntlet/Unreal/Base/Gauntlet.UnrealTestNode.cs
+++ b/Engine/Source/Programs/AutomationTool/Gauntlet/Unreal/Base/Gauntlet.UnrealTestNode.cs
@@ -108,13 +108,13 @@ namespace Gauntlet
 		{
 			if (SessionArtifacts == null)
 			{
-				return new string[0];
+				return Warnings;
 			}
 
 			return SessionArtifacts.SelectMany(A =>
 			{
 				return A.LogSummary.Ensures.Select(E => E.Message);
-			}); 
+			}).Union(Warnings); 
 		}
 
 		/// <summary>
@@ -124,12 +124,77 @@ namespace Gauntlet
 		{
 			if (SessionArtifacts == null)
 			{
-				return new string[0];
+				return Errors;
 			}
 
 			var FailedArtifacts = GetArtifactsWithFailures();
 
-			return FailedArtifacts.Where(A => A.LogSummary.FatalError != null).Select(A => A.LogSummary.FatalError.Message);
+			return FailedArtifacts.Where(A => A.LogSummary.FatalError != null).Select(A => A.LogSummary.FatalError.Message).Union(Errors);
+		}
+
+		/// <summary>
+		/// Returns Errors found during tests. Including Abnornal Exit reasons
+		/// </summary>
+		public virtual IEnumerable<string> GetErrorsAndAbnornalExits()
+		{
+			IEnumerable<string>  Errors = GetErrors();
+
+			Dictionary<UnrealRoleArtifacts, Tuple<int, string>> ErrorCodesAndReasons = new Dictionary<UnrealRoleArtifacts, Tuple<int, string>>();
+
+			var FailedArtifacts = SessionArtifacts.Where(
+				A => {
+					if (A.AppInstance.WasKilled)
+					{
+						return false;
+					}
+					string ExitReason;
+					int ExitCode = GetExitCodeAndReason(A, out ExitReason);
+					ErrorCodesAndReasons.Add(A, new Tuple<int, string>(ExitCode, ExitReason));
+					return ExitCode != 0;
+				}
+			);
+
+			return Errors.Union(FailedArtifacts.Select(
+				A => {
+					int ExitCode = ErrorCodesAndReasons[A].Item1;
+					string ExitReason = ErrorCodesAndReasons[A].Item2;
+					return string.Format("Abnormal Exit: Reason={0}, ExitCode={1}, Log={2}", ExitReason, ExitCode, Path.GetFileName(A.LogPath));
+				}
+			));
+		}
+
+		/// <summary>
+		/// Returns the test URL Link
+		/// </summary>
+		public virtual string GetURLLink()
+		{
+			return "";
+		}
+
+		/// <summary>
+		/// Report an error
+		/// </summary>
+		/// <param name="Message"></param>
+		public virtual void ReportError(string Message, params object[] Args)
+		{
+			Message = string.Format(Message, Args);
+			Errors.Add(Message);
+			if (!LogWarningsAndErrorsAfterSummary) { Log.Error(Message); }
+			if (GetTestStatus() == TestStatus.Complete && GetTestResult() == TestResult.Passed)
+			{
+				SetUnrealTestResult(TestResult.Failed);
+			}
+		}
+
+		/// <summary>
+		/// Report a warning
+		/// </summary>
+		/// <param name="Message"></param>
+		public virtual void ReportWarning(string Message, params object[] Args)
+		{
+			Message = string.Format(Message, Args);
+			if (!LogWarningsAndErrorsAfterSummary) { Warnings.Add(Message); }
+			Log.Warning(Message);
 		}
 
 		// Begin UnrealTestNode properties and members
@@ -149,6 +214,16 @@ namespace Gauntlet
 		/// </summary>
 		public IEnumerable<UnrealRoleArtifacts> SessionArtifacts { get; private set; }
 
+		/// <summary>
+		/// Error collection.
+		/// </summary>
+		protected List<string> Errors { get; private set; } = new List<string>();
+
+		/// <summary>
+		/// Warning collection.
+		/// </summary>
+		protected List<string> Warnings { get; private set; } = new List<string>();
+
 		/// <summary>
 		/// Whether we submit to the dashboard
 		/// </summary>
@@ -795,10 +870,20 @@ namespace Gauntlet
 			}
 
 			string Message = string.Empty;
+			ITestReport Report = null;
+
+			var DeprecatedSignature = new[] { typeof(TestResult), typeof(UnrealTestContext), typeof(UnrealBuildSource), typeof(IEnumerable<UnrealRoleArtifacts>), typeof(string) };
 
 			try
 			{
-				CreateReport(GetTestResult(), Context, Context.BuildInfo, SessionArtifacts, ArtifactPath);
+				// Check if the deprecated signature is overriden, call it anyway if that the case, and trigger a warning.
+				if (Utils.InterfaceHelpers.HasOverriddenMethod(this.GetType(), "CreateReport", DeprecatedSignature))
+				{
+					Log.Warning("This CreateReport signature is deprecated, please use the overloaded version.");
+					CreateReport(GetTestResult(), Context, Context.BuildInfo, SessionArtifacts, ArtifactPath);
+				}
+
+				Report = CreateReport(GetTestResult());
 			}
 			catch (Exception Ex)
 			{
@@ -821,10 +906,17 @@ namespace Gauntlet
 
 			try
 			{
-				SubmitToDashboard(GetTestResult(), Context, Context.BuildInfo, SessionArtifacts, ArtifactPath);
+				// Check if the deprecated signature is overriden, call it anyway if that the case, and trigger a warning.
+				if (Utils.InterfaceHelpers.HasOverriddenMethod(this.GetType(), "SubmitToDashboard", DeprecatedSignature))
+				{
+					Log.Warning("This SubmitToDashboard signature is deprecated, please use the overloaded version.");
+					SubmitToDashboard(GetTestResult(), Context, Context.BuildInfo, SessionArtifacts, ArtifactPath);
+				}
+
+				if (Report != null) { SubmitToDashboard(Report); }
 			}
 			catch (Exception Ex)
-			{				
+			{
 				Log.Warning("Failed to submit results to dashboard. {0}", Ex);
 			}
 		}
@@ -904,11 +996,14 @@ namespace Gauntlet
 		}
 
 		/// <summary>
-		/// Optional function that is called on test completion and gives an opportunity to create a report
+		/// DEPRECATED Optional function that is called on test completion and gives an opportunity to create a report
 		/// </summary>
 		/// <param name="Result"></param>
 		/// <param name="Context"></param>
 		/// <param name="Build"></param>
+		/// <param name="Artifacts"></param>
+		/// <param name="ArtifactPath"></param>
+		/// <returns></returns>
 		public virtual void CreateReport(TestResult Result, UnrealTestContext Context, UnrealBuildSource Build, IEnumerable<UnrealRoleArtifacts> Artifacts, string ArtifactPath)
 		{
 		}
@@ -917,12 +1012,95 @@ namespace Gauntlet
 		/// Optional function that is called on test completion and gives an opportunity to create a report
 		/// </summary>
 		/// <param name="Result"></param>
-		/// <param name="Contex"></param>
+		/// <returns>ITestReport</returns>
+		public virtual ITestReport CreateReport(TestResult Result)
+		{
+			if (GetConfiguration().WriteTestResultsForHorde)
+			{
+				// write test report for Horde
+				HordeReport.SimpleTestReport HordeTestReport = CreateSimpleReportForHorde(Result);
+				return HordeTestReport;
+			}
+
+			return null;
+		}
+
+		/// <summary>
+		/// Generate a Simple Test Report from the results of this test
+		/// </summary>
+		/// <param name="Result"></param>
+		protected virtual HordeReport.SimpleTestReport CreateSimpleReportForHorde(TestResult Result)
+		{
+			HordeReport.SimpleTestReport HordeTestReport = new HordeReport.SimpleTestReport();
+			HordeTestReport.ReportCreatedOn = DateTime.Now.ToString();
+			HordeTestReport.TotalDurationSeconds = (float) (DateTime.Now - SessionStartTime).TotalSeconds;
+			HordeTestReport.Description = Context.ToString();
+			HordeTestReport.Status = Result.ToString();
+			HordeTestReport.URLLink = GetURLLink();
+			HordeTestReport.Errors.AddRange(GetErrorsAndAbnornalExits());
+			if (!string.IsNullOrEmpty(CancellationReason))
+			{
+				HordeTestReport.Errors.Add(CancellationReason);
+			}
+			HordeTestReport.Warnings.AddRange(GetWarnings());
+			HordeTestReport.HasSucceeded = !(Result == TestResult.Failed || Result == TestResult.TimedOut || HordeTestReport.Errors.Count > 0);
+			string HordeArtifactPath = string.IsNullOrEmpty(GetConfiguration().HordeArtifactPath) ? HordeReport.DefaultArtifactsDir : GetConfiguration().HordeArtifactPath;
+			HordeTestReport.SetOutputArtifactPath(HordeArtifactPath);
+			if (SessionArtifacts != null)
+			{
+				foreach (UnrealRoleArtifacts Artifact in SessionArtifacts)
+				{
+					string LogName = Artifact.LogPath.Replace(Path.GetFullPath(Path.Combine(ArtifactPath, "..")), "").TrimStart(Path.DirectorySeparatorChar);
+					HordeTestReport.AttachArtifact(Artifact.LogPath, LogName);
+
+					UnrealLogParser.LogSummary LogSummary = Artifact.LogSummary;
+					if (LogSummary.Errors.Count() > 0)
+					{
+						HordeTestReport.Warnings.Add(
+							string.Format(
+								"Log Parsing: FatalErrors={0}, Ensures={1}, Errors={2}, Warnings={3}, Log={4}",
+								(LogSummary.FatalError != null ? 1 : 0), LogSummary.Ensures.Count(), LogSummary.Errors.Count(), LogSummary.Warnings.Count(), LogName
+							)
+						);
+					}
+				}
+			}
+			return HordeTestReport;
+		}
+
+		/// <summary>
+		/// DEPRECATED Optional function that is called on test completion and gives an opportunity to submit a report to a Dashboard
+		/// </summary>
+		/// <param name="Result"></param>
+		/// <param name="Context"></param>
 		/// <param name="Build"></param>
-		public virtual void SubmitToDashboard(TestResult Result, UnrealTestContext Contex, UnrealBuildSource Build, IEnumerable<UnrealRoleArtifacts> Artifacts, string ArtifactPath)
+		/// <param name="Artifacts"></param>
+		/// <param name="ArtifactPath"></param>
+		/// <returns></returns>
+		public virtual void SubmitToDashboard(TestResult Result, UnrealTestContext Context, UnrealBuildSource Build, IEnumerable<UnrealRoleArtifacts> Artifacts, string ArtifactPath)
 		{
 		}
 
+		/// <summary>
+		/// Optional function that is called on test completion and gives an opportunity to submit a report to a Dashboard
+		/// </summary>
+		/// <param name="Report"></param>
+		public virtual void SubmitToDashboard(ITestReport Report)
+		{
+			if (GetConfiguration().WriteTestResultsForHorde)
+			{
+				// write test data collection for Horde
+				string HordeTestDataKey = string.IsNullOrEmpty(GetConfiguration().HordeTestDataKey) ? Name + " " + Context.ToString() : GetConfiguration().HordeTestDataKey;
+				string HordeTestDataFilePath = Path.Combine(
+					string.IsNullOrEmpty(GetConfiguration().HordeTestDataPath) ? HordeReport.DefaultTestDataDir : GetConfiguration().HordeTestDataPath,
+					FileUtils.SanitizeFilename(Name) + ".TestData.json"
+				);
+				HordeReport.TestDataCollection HordeTestDataCollection = new HordeReport.TestDataCollection();
+				HordeTestDataCollection.AddNewTestReport(HordeTestDataKey, Report);
+				HordeTestDataCollection.WriteToJson(HordeTestDataFilePath, true);
+			}
+		}
+
 		/// <summary>
 		/// Called to request that the test save all artifacts from the completed test to the specified 
 		/// output path. By default the app will save all logs and crash dumps
@@ -1288,8 +1466,6 @@ namespace Gauntlet
 				return new UnrealRoleArtifacts[0] { };
 			}
 
-			bool DidKillClients = SessionArtifacts.Any(A => A.SessionRole.RoleType.IsClient() && A.AppInstance.WasKilled);
-
 			Dictionary<UnrealRoleArtifacts, int> ErrorCodes = new Dictionary<UnrealRoleArtifacts, int>();
 
 			var FailureList = SessionArtifacts.Where(A =>
diff --git a/Engine/Source/Programs/AutomationTool/Gauntlet/Unreal/Game/Samples/ElementalDemoTest.cs b/Engine/Source/Programs/AutomationTool/Gauntlet/Unreal/Game/Samples/ElementalDemoTest.cs
--- a/Engine/Source/Programs/AutomationTool/Gauntlet/Unreal/Game/Samples/ElementalDemoTest.cs
+++ b/Engine/Source/Programs/AutomationTool/Gauntlet/Unreal/Game/Samples/ElementalDemoTest.cs
@@ -33,16 +33,16 @@ namespace UnrealGame
 			return Config;
 		}
 
-		public override void CreateReport(TestResult Result, UnrealTestContext Contex, UnrealBuildSource Build, IEnumerable<UnrealRoleArtifacts> Artifacts, string ArtifactPath)
+		public override ITestReport CreateReport(TestResult Result)
 		{
-			UnrealRoleArtifacts ClientArtifacts = Artifacts.Where(A => A.SessionRole.RoleType == UnrealTargetRole.Client).FirstOrDefault();
+			UnrealRoleArtifacts ClientArtifacts = SessionArtifacts.Where(A => A.SessionRole.RoleType == UnrealTargetRole.Client).FirstOrDefault();
 
 			var SnapshotSummary = new UnrealSnapshotSummary<UnrealHealthSnapshot>(ClientArtifacts.AppInstance.StdOut);
 
 			Log.Info("Elemental Performance Report");
 			Log.Info(SnapshotSummary.ToString());
 
-			base.CreateReport(Result, Contex, Build, Artifacts, ArtifactPath);
+			return base.CreateReport(Result);
 		}
 
 		public override void SaveArtifacts_DEPRECATED(string OutputPath)
diff --git a/Engine/Source/Programs/AutomationTool/Gauntlet/Unreal/Game/Tests/Gauntlet.UnrealTest.BootTest.cs b/Engine/Source/Programs/AutomationTool/Gauntlet/Unreal/Game/Tests/Gauntlet.UnrealTest.BootTest.cs
--- a/Engine/Source/Programs/AutomationTool/Gauntlet/Unreal/Game/Tests/Gauntlet.UnrealTest.BootTest.cs
+++ b/Engine/Source/Programs/AutomationTool/Gauntlet/Unreal/Game/Tests/Gauntlet.UnrealTest.BootTest.cs
@@ -124,58 +124,54 @@ namespace Gauntlet.UnrealTest
 		/// Called after a test finishes to create an overall summary based on looking at the artifacts
 		/// </summary>
 		/// <param name="Result"></param>
-		/// <param name="Context"></param>
-		/// <param name="Build"></param>
-		/// <param name="Artifacts"></param>
-		/// <param name="InArtifactPath"></param>
-		public override void CreateReport(TestResult Result, UnrealTestContext Context, UnrealBuildSource Build, IEnumerable<UnrealRoleArtifacts> Artifacts, string InArtifactPath)
+		public override ITestReport CreateReport(TestResult Result)
 		{
 			// only check for artifacts if the test passed
-			if (Result !=TestResult.Passed)
+			if (Result == TestResult.Passed)
 			{
-				return;
-			}
+				if (!DidDetectLaunch)
+				{
+					SetUnrealTestResult(TestResult.Failed);
+					ReportError("Failed to detect completion of launch");
+				}
+				else
+				{
+					bool MissingFiles = false;
 
-			if (!DidDetectLaunch)
-			{
-				SetUnrealTestResult(TestResult.Failed);
-				Log.Error("Failed to detect completion of launch");
-				return;
-			}
+					foreach (var RoleArtifact in SessionArtifacts)
+					{
+						DirectoryInfo RoleDir = new DirectoryInfo(RoleArtifact.ArtifactPath);
 
-			bool MissingFiles = false;
+						IEnumerable<FileInfo> ArtifactFiles = RoleDir.EnumerateFiles("*.*", SearchOption.AllDirectories);
 
-			foreach(var RoleArtifact in Artifacts)
-			{
-				DirectoryInfo RoleDir = new DirectoryInfo(RoleArtifact.ArtifactPath);
+						// user may not have cleared paths between runs, so throw away anything that's older than 2m
+						ArtifactFiles = ArtifactFiles.Where(F => (DateTime.Now - F.LastWriteTime).TotalMinutes < 2);
 
-				IEnumerable<FileInfo> ArtifactFiles = RoleDir.EnumerateFiles("*.*", SearchOption.AllDirectories);
+						if (ArtifactFiles.Any() == false)
+						{
+							MissingFiles = true;
+							ReportError("No artifact files found for {0}. Were they not retrieved from the device?", RoleArtifact.SessionRole);
+						}
 
-				// user may not have cleared paths between runs, so throw away anything that's older than 2m
-				ArtifactFiles = ArtifactFiles.Where(F => (DateTime.Now - F.LastWriteTime).TotalMinutes < 2);
+						IEnumerable<FileInfo> LogFiles = ArtifactFiles.Where(F => F.Extension.Equals(".log", StringComparison.OrdinalIgnoreCase));
 
-				if (ArtifactFiles.Any() == false)
-				{
-					MissingFiles = true;
-					Log.Error("No artifact files found for {0}. Were they not retrieved from the device?", RoleArtifact.SessionRole);
-				}
+						if (LogFiles.Any() == false)
+						{
+							MissingFiles = true;
+							ReportError("No log files found for {0}. Were they not retrieved from the device?", RoleArtifact.SessionRole);
+						}
+					}
 
-				IEnumerable<FileInfo> LogFiles = ArtifactFiles.Where(F => F.Extension.Equals(".log", StringComparison.OrdinalIgnoreCase));
+					if (MissingFiles)
+					{
+						SetUnrealTestResult(TestResult.Failed);
+						ReportError("One or more roles did not generated any artifacts");
+					}
 
-				if (LogFiles.Any() == false)
-				{
-					MissingFiles = true;
-					Log.Error("No log files found for {0}. Were they not retrieved from the device?", RoleArtifact.SessionRole);
+					Log.Info("Found valid artifacts for test");
 				}
 			}
-
-			if (MissingFiles)
-			{
-				SetUnrealTestResult(TestResult.Failed);
-				Log.Error("One or more roles did not generated any artifacts");
-			}
-
-			Log.Info("Found valid artifacts for test");
+			return base.CreateReport(Result);
 		}
 	}
 }
