Make sure path are fully qualified when coping artifact files with Gauntlet

#rnx

diff --git a/Engine/Source/Programs/AutomationTool/Gauntlet/Framework/Base/Gauntlet.TestReport.cs b/Engine/Source/Programs/AutomationTool/Gauntlet/Framework/Base/Gauntlet.TestReport.cs
--- a/Engine/Source/Programs/AutomationTool/Gauntlet/Framework/Base/Gauntlet.TestReport.cs
+++ b/Engine/Source/Programs/AutomationTool/Gauntlet/Framework/Base/Gauntlet.TestReport.cs
@@ -93,7 +93,7 @@ namespace Gauntlet
 		/// </summary>
 		/// <param name="ArtifactPath"></param>
 		/// <param name="Name"></param>
-		/// <returns>bool</returns>
+		/// <returns>return true if the file was successfully attached</returns>
 		public virtual bool AttachArtifact(string ArtifactPath, string Name = null)
 		{
 			throw new System.NotImplementedException("AttachArtifact not implemented");
diff --git a/Engine/Source/Programs/AutomationTool/Gauntlet/Framework/Gauntlet.Utils.cs b/Engine/Source/Programs/AutomationTool/Gauntlet/Framework/Gauntlet.Utils.cs
--- a/Engine/Source/Programs/AutomationTool/Gauntlet/Framework/Gauntlet.Utils.cs
+++ b/Engine/Source/Programs/AutomationTool/Gauntlet/Framework/Gauntlet.Utils.cs
@@ -1341,6 +1341,21 @@ namespace Gauntlet
 					}
 				}				
 			}
+
+			/// <summary>
+			/// Return the fully qualified location of file path, including the Long path prefix if necessary
+			/// </summary>
+			/// <param name="FilePath"></param>
+			/// <returns></returns>
+			public static string GetFullyQualifiedPath(string FilePath)
+			{
+				FilePath = Path.GetFullPath(FilePath);
+				if (FilePath.Length > 260)
+				{
+					return Globals.LongPathPrefix + FilePath;
+				}
+				return FilePath;
+			}
 		}
 
 		public class Image
diff --git a/Engine/Source/Programs/AutomationTool/Gauntlet/Framework/Utils/Gauntlet.HordeReport.cs b/Engine/Source/Programs/AutomationTool/Gauntlet/Framework/Utils/Gauntlet.HordeReport.cs
--- a/Engine/Source/Programs/AutomationTool/Gauntlet/Framework/Utils/Gauntlet.HordeReport.cs
+++ b/Engine/Source/Programs/AutomationTool/Gauntlet/Framework/Utils/Gauntlet.HordeReport.cs
@@ -50,7 +50,7 @@ namespace Gauntlet
 			/// </summary>
 			/// <param name="ArtifactPath"></param>
 			/// <param name="Name"></param>
-			/// <returns>bool</returns>
+			/// <returns>return true if the file was successfully attached</returns>
 			public override bool AttachArtifact(string ArtifactPath, string Name = null)
 			{
 				if (string.IsNullOrEmpty(OutputArtifactPath))
@@ -65,7 +65,7 @@ namespace Gauntlet
 						string TargetPath = Path.Combine(OutputArtifactPath, Name ?? Path.GetFileName(ArtifactPath));
 						string TargetDirectry = Path.GetDirectoryName(TargetPath);
 						if (!Directory.Exists(TargetDirectry)) { Directory.CreateDirectory(TargetDirectry); }
-						File.Copy(Globals.LongPathPrefix + ArtifactPath, TargetPath);
+						File.Copy(Utils.SystemHelpers.GetFullyQualifiedPath(ArtifactPath), Utils.SystemHelpers.GetFullyQualifiedPath(TargetPath));
 						return true;
 					}
 					catch (Exception Ex)
@@ -79,16 +79,17 @@ namespace Gauntlet
 			/// <summary>
 			/// Set Output Artifact Path and create the directory if missing
 			/// </summary>
-			/// <param name="ArtifactPath"></param>
+			/// <param name="InPath"></param>
 			/// <returns></returns>
-			public void SetOutputArtifactPath(string Path)
+			public void SetOutputArtifactPath(string InPath)
 			{
-				OutputArtifactPath = Path;
+				OutputArtifactPath = Path.GetFullPath(InPath);
 
 				if (!Directory.Exists(OutputArtifactPath))
 				{
 					Directory.CreateDirectory(OutputArtifactPath);
 				}
+				Log.Verbose(string.Format("Test Report output artifact path is set to: {0}", OutputArtifactPath));
 			}
 		}
 
