Create IterationProfile automated tests to gather raw trace data from Editor

UAT:
Add a new EditorTests runner, existing runner is derelict (will be deleted in upcoming merge from UE4)
BuildGraph is project-agnostic, forked and modified from EngineTest's build graph. Follow up change will try to generalize this logic for future benefit.
Requires specification on the cmdline of project content to use and platform to run on. Project content is specified in DefaultEngine.ini.
EditorTest plugin and log error/warning handling is all handled on the command line to be unobtrusive to production uprojects

EngineTest Plugin:
Add new EditorIterationTests.cpp and a simple test object. This is copied from another test for now for the sake of iteration profiling having its own filter and automation behavior.

AutomationController:
bTreatLogErrorsAsTestErrors was specified in some ini files, but was bugged. Restore this likely deleted config variable and check it when scanning log results

Reverb:
Add Iteration test group

EngineTest:
Add Iteration test group

diff --git a/Engine/Plugins/Tests/EditorTests/Source/EditorTests/Private/UnrealEd/EditorIterationTests.cpp b/Engine/Plugins/Tests/EditorTests/Source/EditorTests/Private/UnrealEd/EditorIterationTests.cpp
new file mode 100644
--- /dev/null
+++ b/Engine/Plugins/Tests/EditorTests/Source/EditorTests/Private/UnrealEd/EditorIterationTests.cpp
@@ -0,0 +1,59 @@
+// Copyright Epic Games, Inc. All Rights Reserved.
+
+#include "CoreMinimal.h"
+#include "Misc/Paths.h"
+#include "Misc/AutomationTest.h"
+#include "Misc/App.h"
+#include "UObject/Object.h"
+
+#include "Tests/AutomationTestSettings.h"
+#include "Tests/AutomationEditorCommon.h"
+#include "Tests/AutomationCommon.h"
+
+#include "Subsystems/AssetEditorSubsystem.h"
+#include "Editor.h"
+
+/**
+ * Test to open the sub editor windows for a specified list of assets.
+ * This list can be setup in the Editor Preferences window within the editor or the DefaultEngine.ini file for that particular project.
+*/
+IMPLEMENT_COMPLEX_AUTOMATION_TEST(FIterationOpenAssets, "Project.Iteration.OpenAssets", EAutomationTestFlags::EditorContext | EAutomationTestFlags::ProductFilter);
+
+void FIterationOpenAssets::GetTests(TArray<FString>& OutBeautifiedNames, TArray<FString>& OutTestCommands) const
+{
+	const UAutomationTestSettings* AutomationTestSettings = GetDefault<UAutomationTestSettings>();
+	for ( FSoftObjectPath AssetRef : AutomationTestSettings->AssetsToOpen )
+	{
+		OutBeautifiedNames.Add(AssetRef.GetAssetName());
+		OutTestCommands.Add(AssetRef.GetLongPackageName());
+	}
+}
+
+bool FIterationOpenAssets::RunTest(const FString& LongAssetPath)
+{
+	//start with all editors closed
+	GEditor->GetEditorSubsystem<UAssetEditorSubsystem>()->CloseAllAssetEditors();
+
+	// below is all latent action, so before sending there, verify the asset exists
+	UObject* Object = FSoftObjectPath(LongAssetPath).TryLoad();
+	if ( Object == nullptr )
+	{
+		UE_LOG(LogEditorAutomationTests, Error, TEXT("Failed to Open Asset '%s'."), *LongAssetPath);
+		return false;
+	}
+
+	AddCommand(new FOpenEditorForAssetCommand(LongAssetPath));
+	AddCommand(new FWaitLatentCommand(0.5f));
+	AddCommand(new FDelayedFunctionLatentCommand([=] {
+		if ( Object->GetOutermost()->IsDirty() )
+		{
+			UE_LOG(LogEditorAutomationTests, Error, TEXT("Asset '%s' was dirty after opening it."), *LongAssetPath);
+		}
+	}));
+	AddCommand(new FCloseAllAssetEditorsCommand());
+	AddCommand(new FDelayedFunctionLatentCommand([] {
+		CollectGarbage(GARBAGE_COLLECTION_KEEPFLAGS);
+	}));
+
+	return true;
+}
diff --git a/Engine/Source/Developer/AutomationController/Private/AutomationCommandline.cpp b/Engine/Source/Developer/AutomationController/Private/AutomationCommandline.cpp
--- a/Engine/Source/Developer/AutomationController/Private/AutomationCommandline.cpp
+++ b/Engine/Source/Developer/AutomationController/Private/AutomationCommandline.cpp
@@ -513,7 +513,8 @@ public:
 				{
 					if (!GIsCriticalError)
 					{
-						if (AutomationController->ReportsHaveErrors())
+						const UAutomationControllerSettings* Settings = GetDefault<UAutomationControllerSettings>();
+						if (AutomationController->ReportsHaveErrors() && Settings->bTreatLogErrorsAsTestErrors)
 						{
 							UE_LOG(LogAutomationCommandLine, Display, TEXT("Setting GIsCriticalError due to test failures (will cause non-zero exit code)."));
 							GIsCriticalError = AutomationController->ReportsHaveErrors();
diff --git a/Engine/Source/Developer/AutomationController/Private/AutomationControllerSettings.cpp b/Engine/Source/Developer/AutomationController/Private/AutomationControllerSettings.cpp
--- a/Engine/Source/Developer/AutomationController/Private/AutomationControllerSettings.cpp
+++ b/Engine/Source/Developer/AutomationController/Private/AutomationControllerSettings.cpp
@@ -8,6 +8,7 @@ UAutomationControllerSettings::UAutomationControllerSettings(const FObjectInitia
 	: Super(ObjectInitializer)
 	, bSuppressLogErrors(false)
 	, bSuppressLogWarnings(false)
+	, bTreatLogErrorsAsTestErrors(true)
 	, bTreatLogWarningsAsTestErrors(true)
 {
 }
diff --git a/Engine/Source/Developer/AutomationController/Public/AutomationControllerSettings.h b/Engine/Source/Developer/AutomationController/Public/AutomationControllerSettings.h
--- a/Engine/Source/Developer/AutomationController/Public/AutomationControllerSettings.h
+++ b/Engine/Source/Developer/AutomationController/Public/AutomationControllerSettings.h
@@ -75,6 +75,10 @@ public:
 	UPROPERTY(Config)
 	bool bSuppressLogWarnings;
 
+	/** Whether to treat log errors as test errors (default=true) */
+	UPROPERTY(Config)
+	bool bTreatLogErrorsAsTestErrors;
+
 	/** Whether to treat log warnings as test errors (default=true) */
 	UPROPERTY(Config)
 	bool bTreatLogWarningsAsTestErrors;
diff --git a/Engine/Source/Programs/AutomationTool/Gauntlet/Unreal/Editor/EditorTests.cs b/Engine/Source/Programs/AutomationTool/Gauntlet/Unreal/Editor/EditorTests.cs
new file mode 100644
--- /dev/null
+++ b/Engine/Source/Programs/AutomationTool/Gauntlet/Unreal/Editor/EditorTests.cs
@@ -0,0 +1,111 @@
+using System;
+using System.Collections.Generic;
+using System.Linq;
+using System.Text;
+using System.Threading.Tasks;
+using Gauntlet;
+
+namespace UnrealEditor
+{
+	/// <summary>
+	/// Default set of options for testing Editor. Options that tests can configure
+	/// should be public, external command-line driven options should be protected/private
+	/// </summary>
+	public class EditorTestConfig : UnrealTestConfiguration
+	{
+		/// <summary>
+		/// Filter or groups of tests to apply
+		/// </summary>
+		[AutoParam]
+		public string TestFilter = "Editor";
+
+		/// <summary>
+		/// Force some specific plugins to load, comma delimited (at time of writing)
+		/// </summary>
+		[AutoParam]
+		public string EnablePlugins = string.Empty;
+
+		/// <summary>
+		/// The file to trace profile data to
+		/// </summary>
+		[AutoParam]
+		public string TraceFile = string.Empty;
+
+		/// <summary>
+		/// Control for interpretation of log warnings as test failures
+		/// </summary>
+		[AutoParam]
+		public string TreatLogWarningsAsTestErrors = "true";
+
+		/// <summary>
+		/// Control for interpretation of log errors as test failures
+		/// </summary>
+		[AutoParam]
+		public string TreatLogErrorsAsTestErrors = "true";
+
+		/// <summary>
+		/// Applies these options to the provided app config
+		/// </summary>
+		/// <param name="AppConfig"></param>
+		public override void ApplyToConfig(UnrealAppConfig AppConfig, UnrealSessionRole ConfigRole, IEnumerable<UnrealSessionRole> OtherRoles)
+		{
+			base.ApplyToConfig(AppConfig, ConfigRole, OtherRoles);
+
+			if (AppConfig.ProcessType.IsEditor())
+			{
+				AppConfig.CommandLineParams.Add("ExecCmds", string.Format("\"Automation RunTests {0}; Quit;\"", TestFilter));
+			}
+
+			if (EnablePlugins != string.Empty)
+			{
+				AppConfig.CommandLineParams.Add(string.Format("EnablePlugins={0}", EnablePlugins));
+			}
+
+			if (TraceFile != string.Empty)
+			{
+				AppConfig.CommandLineParams.Add(string.Format("tracefile={0}", TraceFile));
+				AppConfig.CommandLineParams.Add("tracefiletrunc"); // replace existing
+				AppConfig.CommandLineParams.Add("trace=cpu");
+				AppConfig.CommandLineParams.Add("statnamedevents");
+			}
+
+			if (TreatLogWarningsAsTestErrors != "true" || TreatLogWarningsAsTestErrors != "1")
+			{
+				AppConfig.CommandLineParams.Add("ini:Engine:[/Script/AutomationController.AutomationControllerSettings]:bTreatLogWarningsAsTestErrors=False");
+			}
+
+			if (TreatLogErrorsAsTestErrors != "true" || TreatLogErrorsAsTestErrors != "1")
+			{
+				AppConfig.CommandLineParams.Add("ini:Engine:[/Script/AutomationController.AutomationControllerSettings]:bTreatLogErrorsAsTestErrors=False");
+			}
+		}
+	}
+
+	public class EditorTests : UnrealTestNode<EditorTestConfig>
+	{
+		public EditorTests(UnrealTestContext InContext) : base(InContext)
+		{
+		}
+
+		public override EditorTestConfig GetConfiguration()
+		{
+			// just need a single client
+			EditorTestConfig Config = base.GetConfiguration();
+			Config.RequireRole(UnrealTargetRole.Editor);	
+			return Config;
+		}
+
+		public override ITestReport CreateReport(TestResult Result)
+		{
+			ITestReport Report = base.CreateReport(Result);
+
+			if (Report != null && CachedConfig.TraceFile != string.Empty)
+			{
+				string TraceFilePath = System.IO.Path.GetFullPath( CachedConfig.TraceFile );
+				Log.Info("Attaching trace {0} to report", TraceFilePath);
+				Report.AttachArtifact(TraceFilePath);
+			}
+
+			return Report;
+		}
+	}}
diff --git a/Engine/Source/Programs/AutomationTool/Gauntlet/Unreal/RunUnreal.cs b/Engine/Source/Programs/AutomationTool/Gauntlet/Unreal/RunUnreal.cs
--- a/Engine/Source/Programs/AutomationTool/Gauntlet/Unreal/RunUnreal.cs
+++ b/Engine/Source/Programs/AutomationTool/Gauntlet/Unreal/RunUnreal.cs
@@ -56,7 +56,7 @@ namespace Gauntlet
 				throw new AutomationException("No project specified. Use -project=ShooterGame etc");
 			}
 
-			ContextOptions.Namespaces = "Gauntlet.UnrealTest,UnrealGame";
+			ContextOptions.Namespaces = "Gauntlet.UnrealTest,UnrealGame,UnrealEditor";
 			ContextOptions.UsesSharedBuildType = true;
 
 			return RunTests(ContextOptions);
