- Updated default AutoSDK install location to work across host platforms better

diff --git a/Engine/Source/Programs/AutomationTool/Turnkey/Providers/PerforceCopyProvider.cs b/Engine/Source/Programs/AutomationTool/Turnkey/Providers/PerforceCopyProvider.cs
--- a/Engine/Source/Programs/AutomationTool/Turnkey/Providers/PerforceCopyProvider.cs
+++ b/Engine/Source/Programs/AutomationTool/Turnkey/Providers/PerforceCopyProvider.cs
@@ -243,7 +243,7 @@ namespace Turnkey
 					// get local pathname
 					// if we had UE_SDKS_ROOT already set, we use it as a default
 					string AutoSdksDir = Environment.GetEnvironmentVariable("UE_SDKS_ROOT");
-					string DefaultDir = @"D:\Sdks";
+					string DefaultDir = Path.Combine(Environment.GetFolderPath(Environment.SpecialFolder.Personal), "AutoSdks");
 					if (!string.IsNullOrEmpty(AutoSdksDir))
 					{
 						// @todo turnkey - for a deeper directory structure, we should match directory names to the perforce Operation, to figure it out properly
