Move support for UnrealEngine automated tests results from EngineTest to Gauntlet.UnrealTest

#rnx

diff --git a/Engine/Plugins/Tests/EditorTests/Source/EditorTests/EditorTests.Build.cs b/Engine/Plugins/Tests/EditorTests/Source/EditorTests/EditorTests.Build.cs
--- a/Engine/Plugins/Tests/EditorTests/Source/EditorTests/EditorTests.Build.cs
+++ b/Engine/Plugins/Tests/EditorTests/Source/EditorTests/EditorTests.Build.cs
@@ -69,6 +69,7 @@ public class EditorTests : ModuleRules
                 "MeshBuilder",
                 "RawMesh",
 				"FunctionalTesting",
+				"AutomationController",
 			}
 		);
 
diff --git a/Engine/Plugins/Tests/EditorTests/Source/EditorTests/Private/UnrealEd/EditorIterationTests.cpp b/Engine/Plugins/Tests/EditorTests/Source/EditorTests/Private/UnrealEd/EditorIterationTests.cpp
--- a/Engine/Plugins/Tests/EditorTests/Source/EditorTests/Private/UnrealEd/EditorIterationTests.cpp
+++ b/Engine/Plugins/Tests/EditorTests/Source/EditorTests/Private/UnrealEd/EditorIterationTests.cpp
@@ -9,15 +9,37 @@
 #include "Tests/AutomationTestSettings.h"
 #include "Tests/AutomationEditorCommon.h"
 #include "Tests/AutomationCommon.h"
+#include "AutomationControllerSettings.h"
 
 #include "Subsystems/AssetEditorSubsystem.h"
 #include "Editor.h"
 
+class FFIterationOpenAssetsBase : public FAutomationTestBase
+{
+public:
+	FFIterationOpenAssetsBase(const FString& InName, const bool bInComplexTask)
+		: FAutomationTestBase(InName, bInComplexTask)
+	{
+		const UAutomationControllerSettings* Settings = GetDefault<UAutomationControllerSettings>();
+		bSuppressLogErrors = Settings->bSuppressLogErrors || !Settings->bTreatLogErrorsAsTestErrors;
+		bElevateLogWarningsToErrors = Settings->bTreatLogWarningsAsTestErrors;
+	}
+
+	virtual bool SuppressLogErrors() override { return bSuppressLogErrors; }
+
+	virtual bool ElevateLogWarningsToErrors() override { return bElevateLogWarningsToErrors; }
+
+private:
+	bool bSuppressLogErrors = false;
+	bool bElevateLogWarningsToErrors = true;
+
+};
+
 /**
  * Test to open the sub editor windows for a specified list of assets.
  * This list can be setup in the Editor Preferences window within the editor or the DefaultEngine.ini file for that particular project.
 */
-IMPLEMENT_COMPLEX_AUTOMATION_TEST(FIterationOpenAssets, "Project.Iteration.OpenAssets", EAutomationTestFlags::EditorContext | EAutomationTestFlags::ProductFilter);
+IMPLEMENT_CUSTOM_COMPLEX_AUTOMATION_TEST(FIterationOpenAssets, FFIterationOpenAssetsBase, "Project.Iteration.OpenAssets", EAutomationTestFlags::EditorContext | EAutomationTestFlags::ProductFilter);
 
 void FIterationOpenAssets::GetTests(TArray<FString>& OutBeautifiedNames, TArray<FString>& OutTestCommands) const
 {
diff --git a/Engine/Source/Programs/AutomationTool/Gauntlet/Unreal/Base/Gauntlet.UnrealTestNode.cs b/Engine/Source/Programs/AutomationTool/Gauntlet/Unreal/Base/Gauntlet.UnrealTestNode.cs
--- a/Engine/Source/Programs/AutomationTool/Gauntlet/Unreal/Base/Gauntlet.UnrealTestNode.cs
+++ b/Engine/Source/Programs/AutomationTool/Gauntlet/Unreal/Base/Gauntlet.UnrealTestNode.cs
@@ -94,8 +94,6 @@ namespace Gauntlet
 		/// </summary>
 		public override float MaxDuration { get; protected set; }
 
-
-
 		/// <summary>
 		/// Priority of this test
 		/// </summary>
@@ -129,7 +127,7 @@ namespace Gauntlet
 
 			var FailedArtifacts = GetArtifactsWithFailures();
 
-			return FailedArtifacts.Where(A => A.LogSummary.FatalError != null).Select(A => A.LogSummary.FatalError.Message).Union(Errors);
+			 return FailedArtifacts.Where(A => A.LogSummary.FatalError != null).Select(A => A.LogSummary.FatalError.Message).Union(Errors);
 		}
 
 		/// <summary>
@@ -1050,7 +1048,7 @@ namespace Gauntlet
 			{
 				foreach (UnrealRoleArtifacts Artifact in SessionArtifacts)
 				{
-					string LogName = Artifact.LogPath.Replace(Path.GetFullPath(Path.Combine(ArtifactPath, "..")), "").TrimStart(Path.DirectorySeparatorChar);
+					string LogName = Path.GetFullPath(Artifact.LogPath).Replace(Path.GetFullPath(Path.Combine(ArtifactPath, "..")), "").TrimStart(Path.DirectorySeparatorChar);
 					HordeTestReport.AttachArtifact(Artifact.LogPath, LogName);
 
 					UnrealLogParser.LogSummary LogSummary = Artifact.LogSummary;
@@ -1068,6 +1066,32 @@ namespace Gauntlet
 			return HordeTestReport;
 		}
 
+		/// <summary>
+		/// Generate report from Unreal Automated Test Results
+		/// </summary>
+		/// <param name="UnrealAutomatedTestReportPath"></param>
+		/// <param name="ReportURL"></param>
+		/// <returns>ITestReport</returns>
+		public virtual ITestReport CreateUnrealEngineTestPassReport(string UnrealAutomatedTestReportPath, string ReportURL)
+		{
+			string JsonReportPath = Path.Combine(UnrealAutomatedTestReportPath, "index.json");
+			if (File.Exists(JsonReportPath))
+			{
+				Log.Verbose("Reading json Unreal Automated test report from {0}", JsonReportPath);
+				UnrealAutomatedTestPassResults JsonTestPassResults = UnrealAutomatedTestPassResults.LoadFromJson(JsonReportPath);
+				// Convert test results for Horde
+				string HordeArtifactPath = GetConfiguration().HordeArtifactPath;
+				HordeReport.UnrealEngineTestPassResults HordeTestPassResults = HordeReport.UnrealEngineTestPassResults.FromUnrealAutomatedTests(JsonTestPassResults, UnrealAutomatedTestReportPath, ReportURL);
+				HordeTestPassResults.CopyTestResultsArtifacts(HordeArtifactPath);
+				return HordeTestPassResults;
+			}
+			else
+			{
+				Log.Warning("Could not find Unreal Automated test report at {0}. Reverting to base report.", JsonReportPath);
+				return CreateSimpleReportForHorde(GetTestResult());
+			}
+		}
+
 		/// <summary>
 		/// DEPRECATED Optional function that is called on test completion and gives an opportunity to submit a report to a Dashboard
 		/// </summary>
diff --git a/Engine/Source/Programs/AutomationTool/Gauntlet/Unreal/Editor/EditorTests.cs b/Engine/Source/Programs/AutomationTool/Gauntlet/Unreal/Editor/EditorTests.cs
--- a/Engine/Source/Programs/AutomationTool/Gauntlet/Unreal/Editor/EditorTests.cs
+++ b/Engine/Source/Programs/AutomationTool/Gauntlet/Unreal/Editor/EditorTests.cs
@@ -3,9 +3,9 @@
 using System;
 using System.Collections.Generic;
 using System.Linq;
-using System.Text;
-using System.Threading.Tasks;
 using Gauntlet;
+using Gauntlet.UnrealTest;
+
 
 namespace UnrealEditor
 {
@@ -13,13 +13,17 @@ namespace UnrealEditor
 	/// Default set of options for testing Editor. Options that tests can configure
 	/// should be public, external command-line driven options should be protected/private
 	/// </summary>
-	public class EditorTestConfig : UnrealTestConfiguration
+	public class EditorTestConfig : AutomatedTestsConfig
 	{
 		/// <summary>
 		/// Filter or groups of tests to apply
 		/// </summary>
-		[AutoParam]
-		public string TestFilter = "Editor";
+		public override bool UseEditor { get; set; } = true;
+
+		/// <summary>
+		/// Use Simple Horde Report instead of Unreal Automated Tests
+		/// </summary>
+		public override bool SimpleHordeReport { get; set; } = true;
 
 		/// <summary>
 		/// Force some specific plugins to load, comma delimited (at time of writing)
@@ -53,11 +57,6 @@ namespace UnrealEditor
 		{
 			base.ApplyToConfig(AppConfig, ConfigRole, OtherRoles);
 
-			if (AppConfig.ProcessType.IsEditor())
-			{
-				AppConfig.CommandLineParams.Add("ExecCmds", string.Format("\"Automation RunTests {0}; Quit;\"", TestFilter));
-			}
-
 			if (EnablePlugins != string.Empty)
 			{
 				AppConfig.CommandLineParams.Add(string.Format("EnablePlugins={0}", EnablePlugins));
@@ -71,19 +70,19 @@ namespace UnrealEditor
 				AppConfig.CommandLineParams.Add("statnamedevents");
 			}
 
-			if (TreatLogWarningsAsTestErrors != "true" || TreatLogWarningsAsTestErrors != "1")
+			if (TreatLogWarningsAsTestErrors.ToLower() == "false" || TreatLogWarningsAsTestErrors == "0")
 			{
-				AppConfig.CommandLineParams.Add("ini:Engine:[/Script/AutomationController.AutomationControllerSettings]:bTreatLogWarningsAsTestErrors=False");
+				AppConfig.CommandLineParams.Add("ini:Engine:[/Script/AutomationController.AutomationControllerSettings]:bTreatLogWarningsAsTestErrors=false");
 			}
 
-			if (TreatLogErrorsAsTestErrors != "true" || TreatLogErrorsAsTestErrors != "1")
+			if (TreatLogErrorsAsTestErrors.ToLower() == "false" || TreatLogErrorsAsTestErrors == "0")
 			{
-				AppConfig.CommandLineParams.Add("ini:Engine:[/Script/AutomationController.AutomationControllerSettings]:bTreatLogErrorsAsTestErrors=False");
+				AppConfig.CommandLineParams.Add("ini:Engine:[/Script/AutomationController.AutomationControllerSettings]:bTreatLogErrorsAsTestErrors=false");
 			}
 		}
 	}
 
-	public class EditorTests : UnrealTestNode<EditorTestConfig>
+	public class EditorTests : AutomatedTestsBase<EditorTestConfig>
 	{
 		public EditorTests(UnrealTestContext InContext) : base(InContext)
 		{
@@ -110,4 +109,5 @@ namespace UnrealEditor
 
 			return Report;
 		}
-	}}
+	}
+}
diff --git a/Engine/Source/Programs/AutomationTool/Gauntlet/Unreal/Engine/UnrealTests.cs b/Engine/Source/Programs/AutomationTool/Gauntlet/Unreal/Engine/UnrealTests.cs
new file mode 100644
--- /dev/null
+++ b/Engine/Source/Programs/AutomationTool/Gauntlet/Unreal/Engine/UnrealTests.cs
@@ -0,0 +1,521 @@
+// Copyright Epic Games, Inc. All Rights Reserved.
+
+using System;
+using System.IO;
+using System.Collections.Generic;
+using System.Linq;
+using AutomationTool;
+using UnrealBuildTool;
+
+
+namespace Gauntlet
+{
+
+	namespace UnrealTest
+	{
+		/// <summary>
+		/// Test runner for excuting Unreal Tests in Gauntlet
+		/// </summary>
+		public class RunUnrealTests : RunUnreal
+		{
+			public override string DefaultTestName { get { return "AutomatedTests"; } }
+
+			public override ExitCode Execute()
+			{
+				Globals.Params = new Gauntlet.Params(this.Params);
+
+				UnrealTestOptions ContextOptions = new UnrealTestOptions();
+
+				AutoParam.ApplyParamsAndDefaults(ContextOptions, Globals.Params.AllArguments);
+
+				if (string.IsNullOrEmpty(ContextOptions.Project))
+				{
+					throw new AutomationException("No project specified. Use -project=EngineTest etc");
+				}
+
+				ContextOptions.Namespaces = "Gauntlet.UnrealTest,UnrealGame,UnrealEditor";
+				ContextOptions.UsesSharedBuildType = true;
+
+				return RunTests(ContextOptions);
+			}
+		}
+
+		public class AutomatedTests : AutomatedTestsBase<AutomatedTestsConfig>
+		{
+			public AutomatedTests(UnrealTestContext InContext) : base(InContext)
+			{
+			}
+		}
+
+		/// <summary>
+		/// Runs automated tests on a platform
+		/// </summary>
+		public abstract class AutomatedTestsBase<TConfigClass> : UnrealTestNode<TConfigClass>
+		where TConfigClass : AutomatedTestsConfig, new()
+		{
+			private int LastAutomationEntryCount = 0;
+
+			private DateTime LastAutomationEntryTime = DateTime.MinValue;
+
+			public AutomatedTestsBase(Gauntlet.UnrealTestContext InContext)
+				: base(InContext)
+			{
+				LogWarningsAndErrorsAfterSummary = false;
+			}
+
+			/// <summary>
+			/// Returns Errors found during tests. By default only fatal errors are considered
+			/// </summary>
+			public override IEnumerable<string> GetErrors()
+			{
+				List<string> AllErrors = base.GetErrors().ToList();
+
+				foreach (var Artifact in GetArtifactsWithFailures())
+				{
+					if (Artifact.SessionRole.RoleType == UnrealTargetRole.Editor)
+					{
+						AutomationLogParser Parser = new AutomationLogParser(Artifact.LogParser);
+						AllErrors.AddRange(
+							Parser.GetResults().Where(R => !R.Passed)
+							.SelectMany(R => R.Events
+								.Where(E => E.ToLower().Contains("error")).Select(E => string.Format("[test={0}] {1}", R.DisplayName, E))
+								)
+							);
+					}
+				}
+
+				return AllErrors;
+			}
+
+			/// <summary>
+			/// Define the configuration for this test. Most options are applied by the test config above
+			/// </summary>
+			/// <returns></returns>
+			public override TConfigClass GetConfiguration()
+			{
+				TConfigClass Config = base.GetConfiguration();
+
+				if (Config.UseEditor)
+				{
+					Config.RequireRole(UnrealTargetRole.Editor);
+				}
+				else
+				{
+					Config.RequireRole(UnrealTargetRole.Editor);
+					Config.RequireRole(UnrealTargetRole.Client);
+				}
+
+				Config.MaxDuration = Context.TestParams.ParseValue("MaxDuration", 3600);
+
+				return Config;
+			}
+
+			/// <summary>
+			/// Override our name to include the filter we're testing
+			/// </summary>
+			public override string Name
+			{
+				get
+				{
+					string BaseName = base.Name;
+
+					var Config = GetConfiguration();
+
+					if (!string.IsNullOrEmpty(Config.TestFilter))
+					{
+						BaseName += " " + Config.TestFilter;
+					}
+
+					return BaseName;
+				}
+			}
+
+			/// <summary>
+			/// Override TickTest to log interesting things and make sure nothing has stalled
+			/// </summary>
+			public override void TickTest()
+			{
+				const float IdleTimeout = 30 * 60;
+
+				List<string> ChannelEntries = new List<string>();
+
+				var AppInstance = TestInstance.EditorApp;
+
+				UnrealLogParser Parser = new UnrealLogParser(AppInstance.StdOut);
+				ChannelEntries.AddRange(Parser.GetLogChannels(new string[] { "Automation", "FunctionalTest", "Material", "DerivedDataCache", "ShaderCompilers" }, false));
+
+				if (ChannelEntries.Count > LastAutomationEntryCount)
+				{
+					// log new entries so people have something to look at
+					ChannelEntries.Skip(LastAutomationEntryCount).ToList().ForEach(S => Log.Info("{0}", S));
+					LastAutomationEntryTime = DateTime.Now;
+					LastAutomationEntryCount = ChannelEntries.Count;
+				}
+				else
+				{
+					if (LastAutomationEntryTime == DateTime.MinValue)
+					{
+						LastAutomationEntryTime = DateTime.Now;
+					}
+
+					double ElapsedTime = (DateTime.Now - LastAutomationEntryTime).TotalSeconds;
+
+					// Check for timeout
+					if (ElapsedTime > IdleTimeout)
+					{
+						Log.Error("No activity observed in last {0:0.00} minutes. Aborting test", IdleTimeout / 60);
+						MarkTestComplete();
+						SetUnrealTestResult(TestResult.TimedOut);
+					}
+				}
+
+				base.TickTest();
+			}
+
+			/// <summary>
+			/// Optional function that is called on test completion and gives an opportunity to create a report
+			/// </summary>
+			/// <param name="Result"></param>
+			/// <returns>ITestReport</returns>
+			public override ITestReport CreateReport(TestResult Result)
+			{
+				// Save test result data for Horde build system
+				bool WriteTestResultsForHorde = GetConfiguration().WriteTestResultsForHorde;
+				if (WriteTestResultsForHorde)
+				{
+					if (GetConfiguration().SimpleHordeReport)
+					{
+						return base.CreateReport(Result);
+					}
+					else
+					{
+						string ReportPath = GetConfiguration().ReportExportPath;
+						if (!string.IsNullOrEmpty(ReportPath))
+						{
+							return CreateUnrealEngineTestPassReport(ReportPath, GetConfiguration().ReportURL);
+						}
+					}
+				}
+
+				return null;
+			}
+
+			/// <summary>
+			/// Parses the provided artifacts to determine the cause of an exit and whether it was abnormal
+			/// </summary>
+			/// <param name="InArtifacts"></param>
+			/// <param name="Reason"></param>
+			/// <param name="WasAbnormal"></param>
+			/// <returns></returns>
+			protected override int GetExitCodeAndReason(UnrealRoleArtifacts InArtifacts, out string ExitReason)
+			{
+				int ExitCode = base.GetExitCodeAndReason(InArtifacts, out ExitReason);
+
+				if (InArtifacts.SessionRole.RoleType == UnrealTargetRole.Editor)
+				{
+					// if no fatal errors, check test results
+					if (InArtifacts.LogParser.GetFatalError() == null)
+					{
+						AutomationLogParser Parser = new AutomationLogParser(InArtifacts.LogParser);
+
+						IEnumerable<AutomationTestResult> TotalTests = Parser.GetResults();
+						IEnumerable<AutomationTestResult> FailedTests = TotalTests.Where(R => !R.Passed);
+
+						if (FailedTests.Any())
+						{
+							ExitReason = string.Format("{0}/{1} tests failed", FailedTests.Count(), TotalTests.Count());
+							ExitCode = -1;
+						}
+
+						// Warn if no tests were run
+						if (!TotalTests.Any())
+						{
+							ExitReason = "No tests were executed!";
+							ExitCode = -1;
+						}
+					}
+				}
+
+				return ExitCode;
+			}
+
+			/// <summary>
+			/// Override the summary report so we can insert a link to our report and the failed tests
+			/// </summary>
+			/// <returns></returns>
+			protected override string GetTestSummaryHeader()
+			{
+				MarkdownBuilder MB = new MarkdownBuilder(base.GetTestSummaryHeader());
+
+				var EditorArtifacts = SessionArtifacts.Where(A => A.SessionRole.RoleType == UnrealTargetRole.Editor).FirstOrDefault();
+
+				if (EditorArtifacts != null)
+				{
+					AutomationLogParser Parser = new AutomationLogParser(EditorArtifacts.LogParser);
+
+					IEnumerable<AutomationTestResult> AllTests = Parser.GetResults();
+					IEnumerable<AutomationTestResult> FailedTests = AllTests.Where(R => R.Completed && !R.Passed);
+					IEnumerable<AutomationTestResult> IncompleteTests = AllTests.Where(R => !R.Completed);
+
+					if (AllTests.Count() == 0)
+					{
+						MB.Paragraph("No tests were executed!");
+					}
+					else
+					{
+						int PassedTests = AllTests.Count() - (FailedTests.Count() + IncompleteTests.Count());
+						MB.Paragraph(string.Format("{0} of {1} tests passed", PassedTests, AllTests.Count()));
+
+						if (FailedTests.Count() > 0)
+						{
+							MB.H3("The following tests failed:");
+
+							foreach (AutomationTestResult Result in FailedTests)
+							{
+								MB.H4(string.Format("{0}: {1}", Result.DisplayName, Result.TestName));
+								MB.UnorderedList(Result.Events);
+							}
+						}
+
+						if (IncompleteTests.Count() > 0)
+						{
+							MB.H3("The following tests timed out:");
+
+							foreach (AutomationTestResult Result in IncompleteTests)
+							{
+								MB.H4(string.Format("{0}: {1}", Result.DisplayName, Result.TestName));
+								MB.UnorderedList(Result.Events);
+							}
+						}
+
+						string ReportLink = GetConfiguration().ReportURL;
+						string ReportPath = GetConfiguration().ReportExportPath;
+
+						if (!string.IsNullOrEmpty(ReportLink) || !string.IsNullOrEmpty(ReportPath))
+						{
+							MB.H3("Links");
+
+							if (string.IsNullOrEmpty(ReportLink) == false)
+							{
+								MB.Paragraph(string.Format("View results here: {0}", ReportLink));
+							}
+
+							if (string.IsNullOrEmpty(ReportPath) == false)
+							{
+								MB.Paragraph(string.Format("Open results in UnrealEd from {0}", ReportPath));
+							}
+						}
+
+
+					}
+				}
+
+				return MB.ToString();
+			}
+		}
+
+		/// <summary>
+		/// Define a Config class for this test that contains the available options
+		/// </summary>
+		public class AutomatedTestsConfig : UnrealTestConfiguration
+		{
+			/// <summary>
+			/// Set to true if the editor executes all these tests in is own process and PIE
+			/// </summary>
+			[AutoParam]
+			public virtual bool UseEditor { get; set; } = false;
+
+			/// <summary>
+			/// Run with d3d12 RHI
+			/// </summary>
+			[AutoParam]
+			public virtual bool D3D12 { get; set; } = false;
+
+			/// <summary>
+			/// Run with Nvidia cards for raytracing support
+			/// </summary>
+			[AutoParam]
+			public virtual bool PreferNvidia { get; set; } = false;
+
+			/// <summary>
+			/// Run forcing raytracing 
+			/// </summary>
+			[AutoParam]
+			public virtual bool RayTracing { get; set; } = false;
+
+			/// <summary>
+			/// Use the StompMalloc memory allocator.
+			/// </summary>
+			[AutoParam]
+			public virtual bool StompMalloc { get; set; } = false;
+
+			/// <summary>
+			/// Enable the D3D debug layer.
+			/// </summary>
+			[AutoParam]
+			public virtual bool D3DDebug { get; set; } = false;
+
+			/// <summary>
+			/// Enable the RHI validation layer.
+			/// </summary>
+			[AutoParam]
+			public virtual bool RHIValidation { get; set; } = false;
+
+			/// <summary>
+			/// Disable capturing frame trace for image based tests
+			/// </summary>
+			[AutoParam]
+			public virtual bool DisableFrameTraceCapture { get; set; } = true;
+
+			/// <summary>
+			/// Filter or groups of tests to apply
+			/// </summary>
+			[AutoParam]
+			public virtual string TestFilter { get; set; } = "Filter:Project+Filter:System";
+
+			/// <summary>
+			/// Path to write the automation report to
+			/// </summary>
+			[AutoParam]
+			public virtual string ReportExportPath { get; set; } = "";
+
+			/// <summary>
+			/// Path the report can be found at
+			/// </summary>
+			[AutoParam]
+			public virtual string ReportURL { get; set; } = "";
+
+			/// <summary>
+			/// Use Simple Horde Report instead of Unreal Automated Tests
+			/// </summary>
+			[AutoParam]
+			public virtual bool SimpleHordeReport { get; set; } = false;
+
+			/// <summary>
+			/// Validate DDC during tests
+			/// </summary>
+			[AutoParam]
+			public virtual bool VerifyDDC { get; set; } = false;
+
+			/// <summary>
+			/// Validate DDC during tests
+			/// </summary>
+			[AutoParam]
+			public virtual string DDC { get; set; } = "";
+
+			/// <summary>
+			/// Used for having the editor and any client communicate
+			/// </summary>
+			public string SessionID = Guid.NewGuid().ToString();
+
+			/// <summary>
+			/// Apply this config to the role that is passed in
+			/// </summary>
+			/// <param name="AppConfig"></param>
+			/// <param name="ConfigRole"></param>
+			/// <param name="OtherRoles"></param>
+			public override void ApplyToConfig(UnrealAppConfig AppConfig, UnrealSessionRole ConfigRole, IEnumerable<UnrealSessionRole> OtherRoles)
+			{
+				base.ApplyToConfig(AppConfig, ConfigRole, OtherRoles);
+
+				if (WriteTestResultsForHorde)
+				{
+					if (string.IsNullOrEmpty(ReportExportPath))
+					{
+						ReportExportPath = Path.Combine(Globals.TempDir, "TestReport");
+					}
+					if (string.IsNullOrEmpty(HordeTestDataPath))
+					{
+						HordeTestDataPath = HordeReport.DefaultTestDataDir;
+					}
+					if (string.IsNullOrEmpty(HordeArtifactPath))
+					{
+						HordeArtifactPath = HordeReport.DefaultArtifactsDir;
+					}
+				}
+
+				string ReportOutputArg = "";
+				if (!string.IsNullOrEmpty(ReportExportPath))
+				{
+					if (Directory.Exists(ReportExportPath))
+					{
+						// clean up previous run if any
+						DirectoryInfo ReportDirInfo = new DirectoryInfo(ReportExportPath);
+						ReportDirInfo.Delete(true);
+					}
+					ReportOutputArg = string.Format("-ReportExportPath=\"{0}\"", ReportExportPath);
+				}
+
+				if (UseEditor)
+				{
+					if (ConfigRole.RoleType == UnrealTargetRole.Editor)
+					{
+						AppConfig.CommandLine += string.Format(" -NoWatchdog -stdout -FORCELOGFLUSH -CrashForUAT -unattended -buildmachine {0} -ExecCmds=\"Automation RunTests {1}; Quit;\"", ReportOutputArg, TestFilter);
+					}
+
+					if (VerifyDDC)
+					{
+						AppConfig.CommandLine += " -VerifyDDC";
+					}
+
+					if (!string.IsNullOrEmpty(DDC))
+					{
+						AppConfig.CommandLine += string.Format(" -ddc={0}", DDC);
+					}
+				}
+				else
+				{
+					// If the test isnt using the editor for both roles then pass the IP of the editor (us) to the client
+					string HostIP = UnrealHelpers.GetHostIpAddress();
+
+					if (ConfigRole.RoleType.IsClient())
+					{
+						AppConfig.CommandLine += string.Format(" -sessionid={0} -messaging -log -TcpMessagingConnect={1}:6666", SessionID, HostIP);
+					}
+					else if (ConfigRole.RoleType.IsEditor())
+					{
+						AppConfig.CommandLine += string.Format(" -nullrhi -ExecCmds=\"Automation list;StartRemoteSession {0};RunTests {1}; Quit\" -TcpMessagingListen={2}:6666 -multihome={3} -log {4}", SessionID, TestFilter, HostIP, HostIP, ReportOutputArg);
+					}
+				}
+
+				if (DisableFrameTraceCapture || RayTracing)
+				{
+					AppConfig.CommandLine += " -DisableFrameTraceCapture";
+				}
+
+				if (RayTracing)
+				{
+					AppConfig.CommandLine += " -dpcvars=r.RayTracing=1,r.SkinCache.CompileShaders=1,AutomationAllowFrameTraceCapture=0 -noxgeshadercompile";
+				}
+
+				if (ConfigRole.Platform == UnrealTargetPlatform.Win64)
+				{
+					if (PreferNvidia)
+					{
+						AppConfig.CommandLine += " -preferNvidia";
+					}
+
+					if (D3D12)
+					{
+						AppConfig.CommandLine += " -d3d12";
+					}
+
+					if (D3DDebug)
+					{
+						AppConfig.CommandLine += " -d3ddebug";
+					}
+
+					if (StompMalloc)
+					{
+						AppConfig.CommandLine += " -stompmalloc";
+					}
+				}
+
+				if (RHIValidation)
+				{
+					AppConfig.CommandLine += " -rhivalidation";
+				}
+			}
+		}
+	}
+}
