BuildGraph: Downgrade invalid argument errors to a warning.

diff --git a/Engine/Source/Programs/AutomationTool/BuildGraph/Script.cs b/Engine/Source/Programs/AutomationTool/BuildGraph/Script.cs
--- a/Engine/Source/Programs/AutomationTool/BuildGraph/Script.cs
+++ b/Engine/Source/Programs/AutomationTool/BuildGraph/Script.cs
@@ -296,16 +296,9 @@ namespace AutomationTool
 			}
 
 			// Make sure all the arguments were valid
-			bool bInvalidArgument = false;
 			foreach(string InvalidArgumentName in Arguments.Keys.Except(Reader.Graph.Options.Select(x => x.Name), StringComparer.InvariantCultureIgnoreCase))
 			{
-				CommandUtils.LogError("Unknown argument '{0}' for '{1}'", InvalidArgumentName, File.FullName);
-				bInvalidArgument = true;
-			}
-			if(bInvalidArgument)
-			{
-				Graph = null;
-				return false;
+				CommandUtils.LogWarning("Unknown argument '{0}' for '{1}'", InvalidArgumentName, File.FullName);
 			}
 
 			// Return the constructed graph
