Fix Gauntlet Engine Test name references

#rnx

diff --git a/Engine/Source/Programs/AutomationTool/Gauntlet/Unreal/Editor/EditorTests.cs b/Engine/Source/Programs/AutomationTool/Gauntlet/Unreal/Editor/EditorTests.cs
--- a/Engine/Source/Programs/AutomationTool/Gauntlet/Unreal/Editor/EditorTests.cs
+++ b/Engine/Source/Programs/AutomationTool/Gauntlet/Unreal/Editor/EditorTests.cs
@@ -13,7 +13,7 @@ namespace UnrealEditor
 	/// Default set of options for testing Editor. Options that tests can configure
 	/// should be public, external command-line driven options should be protected/private
 	/// </summary>
-	public class EditorTestConfig : AutomatedTestsConfig
+	public class EditorTestConfig : EngineTestConfig
 	{
 		/// <summary>
 		/// Filter or groups of tests to apply
@@ -82,7 +82,7 @@ namespace UnrealEditor
 		}
 	}
 
-	public class EditorTests : AutomatedTestsBase<EditorTestConfig>
+	public class EditorTests : EngineTestBase<EditorTestConfig>
 	{
 		public EditorTests(UnrealTestContext InContext) : base(InContext)
 		{
diff --git a/Engine/Source/Programs/AutomationTool/Gauntlet/Unreal/Engine/UnrealTests.cs b/Engine/Source/Programs/AutomationTool/Gauntlet/Unreal/Engine/UnrealTests.cs
--- a/Engine/Source/Programs/AutomationTool/Gauntlet/Unreal/Engine/UnrealTests.cs
+++ b/Engine/Source/Programs/AutomationTool/Gauntlet/Unreal/Engine/UnrealTests.cs
@@ -18,7 +18,7 @@ namespace Gauntlet
 		/// </summary>
 		public class RunUnrealTests : RunUnreal
 		{
-			public override string DefaultTestName { get { return "AutomatedTests"; } }
+			public override string DefaultTestName { get { return "EngineTest"; } }
 
 			public override ExitCode Execute()
 			{
@@ -40,9 +40,9 @@ namespace Gauntlet
 			}
 		}
 
-		public class AutomatedTests : AutomatedTestsBase<AutomatedTestsConfig>
+		public class EngineTest : EngineTestBase<EngineTestConfig>
 		{
-			public AutomatedTests(UnrealTestContext InContext) : base(InContext)
+			public EngineTest(UnrealTestContext InContext) : base(InContext)
 			{
 			}
 		}
@@ -50,14 +50,14 @@ namespace Gauntlet
 		/// <summary>
 		/// Runs automated tests on a platform
 		/// </summary>
-		public abstract class AutomatedTestsBase<TConfigClass> : UnrealTestNode<TConfigClass>
-		where TConfigClass : AutomatedTestsConfig, new()
+		public abstract class EngineTestBase<TConfigClass> : UnrealTestNode<TConfigClass>
+		where TConfigClass : EngineTestConfig, new()
 		{
 			private int LastAutomationEntryCount = 0;
 
 			private DateTime LastAutomationEntryTime = DateTime.MinValue;
 
-			public AutomatedTestsBase(Gauntlet.UnrealTestContext InContext)
+			public EngineTestBase(Gauntlet.UnrealTestContext InContext)
 				: base(InContext)
 			{
 				LogWarningsAndErrorsAfterSummary = false;
@@ -317,7 +317,7 @@ namespace Gauntlet
 		/// <summary>
 		/// Define a Config class for this test that contains the available options
 		/// </summary>
-		public class AutomatedTestsConfig : UnrealTestConfiguration
+		public class EngineTestConfig : UnrealTestConfiguration
 		{
 			/// <summary>
 			/// Set to true if the editor executes all these tests in is own process and PIE
