- Updated Turnkey studio settings (split into separate .xml file so it can be read without trying to hunt down server files, added idea of what platforms are supported by the studio). Set Epic settings
- Return more info from Turnkey so callers can know if Turnkey can install full or autosdks (based on studio settings)
- Added idea of project settings having custom builds that can be run from Turnkey (and with a later checkin, Platforms menu in the editor)

diff --git a/Engine/Build/Turnkey/TurnkeyManifest.xml b/Engine/Build/Turnkey/TurnkeyManifest.xml
--- a/Engine/Build/Turnkey/TurnkeyManifest.xml
+++ b/Engine/Build/Turnkey/TurnkeyManifest.xml
@@ -1,18 +1,7 @@
 <?xml version="1.0" encoding="utf-8" ?>
 <TurnkeyManifest>
-	<!-- Some default values for Studio-wide settings. It's recommended you set these appropriately in one of the .xml files referenced below. "RunuAT Turnkey -command=ManageSettings -listvars" for more information -->
-	<SavedSettings>
-		<SavedSetting>
-			<Variable>Studio_AppleSigningCertPassword</Variable>
-			<Value></Value>
-		</SavedSetting>
-		<SavedSetting>
-			<Variable>Studio_GoogleDriveCredentials</Variable>
-			<Value>file:$(EngineDir)/Build/Turnkey/GoogleDriveCredentials.json</Value>
-		</SavedSetting>
-	</SavedSettings>
 	<AdditionalManifests>
-		<Manifest>$(ThisManifestDir)/StudioTurnkeyManifest.xml</Manifest>
+		<Manifest>$(ThisManifestDir)/Licensee_TurnkeyManifest.xml</Manifest>
 		<Manifest>file:$(EngineDir)/Restricted/NotForLicensees/Build/Turnkey/TurnkeyManifest.xml</Manifest>
 		<Manifest>file:$(EngineDir)/Restricted/NoRedist/Build/Turnkey/TurnkeyManifest.xml</Manifest>
 		<Manifest>$(UE_STUDIO_TURNKEY_LOCATION)</Manifest>
diff --git a/Engine/Build/Turnkey/TurnkeyStudioSettings.xml b/Engine/Build/Turnkey/TurnkeyStudioSettings.xml
new file mode 100644
--- /dev/null
+++ b/Engine/Build/Turnkey/TurnkeyStudioSettings.xml
@@ -0,0 +1,33 @@
+<?xml version="1.0" encoding="utf-8" ?>
+<TurnkeyManifest>
+	<!-- Some default values for Studio-wide settings. It's recommended you set these appropriately in one of the .xml files referenced below. "RunuAT Turnkey -command=ManageSettings -listvars" for more information.  -->
+	<SavedSettings>
+		<SavedSetting>
+			<Variable>Studio_AppleSigningCertPassword</Variable>
+			<Value></Value>
+		</SavedSetting>
+		<SavedSetting>
+			<Variable>Studio_GoogleDriveCredentials</Variable>
+			<Value>file:$(EngineDir)/Build/Turnkey/GoogleDriveCredentials.json</Value>
+		</SavedSetting>
+		<SavedSetting>
+			<Variable>Studio_GoogleDriveAppName</Variable>
+			<Value></Value>
+		</SavedSetting>
+		<SavedSetting>
+			<!-- By default, indicate that Turnkey cannot install any platforms. This can be a comma separated list, or All to indicate all platforms should try to install via Turnkey. -->
+			<Variable>Studio_FullInstallPlatforms</Variable>
+			<Value></Value>
+		</SavedSetting>
+		<SavedSetting>
+			<Variable>Studio_AutoSdkPlatforms</Variable>
+			<Value></Value>
+		</SavedSetting>
+	</SavedSettings>
+	<AdditionalManifests>
+		<Manifest>$(ThisManifestDir)/Licensee_TurnkeyStudioSettings.xml</Manifest>
+		<Manifest>file:$(EngineDir)/Restricted/NotForLicensees/Build/Turnkey/TurnkeyStudioSettings.xml</Manifest>
+		<Manifest>file:$(EngineDir)/Restricted/NoRedist/Build/Turnkey/TurnkeyStudioSettings.xml</Manifest>
+		<Manifest>$(UE_STUDIO_TURNKEY_SETTINGS_LOCATION)</Manifest>
+	</AdditionalManifests>
+</TurnkeyManifest>
diff --git a/Engine/Config/BaseGame.ini b/Engine/Config/BaseGame.ini
--- a/Engine/Config/BaseGame.ini
+++ b/Engine/Config/BaseGame.ini
@@ -80,6 +80,7 @@ MinWindowWidth=16
 MinWindowHeight=16
 
 [/Script/UnrealEd.ProjectPackagingSettings]
+@ExtraProjectBuilds=Name
 BuildConfiguration=PPBC_Development
 FullRebuild=False
 UsePakFile=True
diff --git a/Engine/Source/Programs/AutomationTool/AutomationUtils/Automation.cs b/Engine/Source/Programs/AutomationTool/AutomationUtils/Automation.cs
--- a/Engine/Source/Programs/AutomationTool/AutomationUtils/Automation.cs
+++ b/Engine/Source/Programs/AutomationTool/AutomationUtils/Automation.cs
@@ -14,7 +14,7 @@ namespace AutomationTool
 	/// <summary>
 	/// Command to execute info.
 	/// </summary>
-	class CommandInfo
+	public class CommandInfo
 	{
 		public string CommandName;
 		public List<string> Arguments = new List<string>();
@@ -521,7 +521,7 @@ AutomationTool.exe [-verbose] [-compileonly] [-p4] Command0 [-Arg0 -Arg1 -Arg2 .
 		/// </summary>
 		/// <param name="CommandsToExecute"></param>
 		/// <param name="Commands"></param>
-		private static ExitCode Execute(List<CommandInfo> CommandsToExecute, Dictionary<string, Type> Commands)
+		public static ExitCode Execute(List<CommandInfo> CommandsToExecute, Dictionary<string, Type> Commands)
 		{
 			for (int CommandIndex = 0; CommandIndex < CommandsToExecute.Count; ++CommandIndex)
 			{
diff --git a/Engine/Source/Programs/AutomationTool/Turnkey/Commands/ExecuteBuild.cs b/Engine/Source/Programs/AutomationTool/Turnkey/Commands/ExecuteBuild.cs
new file mode 100644
--- /dev/null
+++ b/Engine/Source/Programs/AutomationTool/Turnkey/Commands/ExecuteBuild.cs
@@ -0,0 +1,216 @@
+// Copyright Epic Games, Inc. All Rights Reserved.
+
+using AutomationTool;
+using System;
+using System.Collections.Generic;
+using System.Linq;
+using System.Text;
+using System.Text.RegularExpressions;
+using Tools.DotNETCommon;
+using UnrealBuildTool;
+
+namespace Turnkey.Commands
+{
+	class ExecuteBuild : TurnkeyCommand
+	{
+		protected override void Execute(string[] CommandOptions)
+		{
+			// we need a platform to execute
+			List<UnrealTargetPlatform> Platforms = TurnkeyUtils.GetPlatformsFromCommandLineOrUser(CommandOptions, null);
+
+			string Project = TurnkeyUtils.GetVariableValue("Project");
+			if (string.IsNullOrEmpty(Project))
+			{
+				Project = TurnkeyUtils.ReadInput("Enter a project to build");
+			}
+			
+			FileReference ProjectFile = ProjectUtils.FindProjectFileFromName(Project);
+
+			string DesiredBuild = TurnkeyUtils.ParseParamValue("Build", null, CommandOptions);
+
+			// get a list of builds from config
+			foreach (UnrealTargetPlatform Platform in Platforms)
+			{
+ 				ConfigHierarchy GameConfig = ConfigCache.ReadHierarchy(ConfigHierarchyType.Game, ProjectFile.Directory, Platform);
+
+				List<string> Builds;
+				GameConfig.GetArray("/Script/UnrealEd.ProjectPackagingSettings", "ExtraProjectBuilds", out Builds);
+
+				Dictionary<string, string> BuildCommands = new Dictionary<string, string>(StringComparer.InvariantCultureIgnoreCase);
+				foreach (string Build in Builds)
+				{
+					Match NameResult = Regex.Match(Build, "Name\\s*=\\s*\"(.*?)\"");
+					Match PlatformsResult = Regex.Match(Build, "SpecificPlatforms\\s*=\\s*(.*?)");
+					Match ParamsResult = Regex.Match(Build, "BuildCookRunParams=\\s*\"(.*?)\"");
+
+					// make sure required entries are there
+					if (!NameResult.Success || !ParamsResult.Success)
+					{
+						continue;
+					}
+
+					// if platforms are specified, and this platform isn't one of them, skip it
+					if (PlatformsResult.Success)
+					{
+						string[] SpecificPlatforms = PlatformsResult.Groups[1].Value.Split(",\"".ToCharArray());
+						if (SpecificPlatforms.Length > 0 && !SpecificPlatforms.Contains(Platform.ToString()))
+						{
+							continue;
+						}
+					}
+
+					// add to list of commands
+					BuildCommands.Add(NameResult.Groups[1].Value, ParamsResult.Groups[1].Value);
+				}
+
+				if (BuildCommands.Count == 0)
+				{
+					TurnkeyUtils.Log("Unable to find a build for platform {0} and project {1}", Platform, Project);
+					continue;
+				}
+
+				string FinalParams;
+				if (!string.IsNullOrEmpty(DesiredBuild) && BuildCommands.ContainsKey(DesiredBuild))
+				{
+					FinalParams = BuildCommands[DesiredBuild];
+				}
+				else
+				{
+					List<string> BuildNames = BuildCommands.Keys.ToList();
+					int Choice = TurnkeyUtils.ReadInputInt("Choose a build to execute", BuildNames, true);
+					if (Choice == 0)
+					{
+						continue;
+					}
+
+					FinalParams = BuildCommands[BuildNames[Choice - 1]];
+				}
+
+				FinalParams = FinalParams.Replace("{Project}", "\"" + ProjectFile.FullName + "\"");
+				FinalParams = FinalParams.Replace("{Platform}", Platform.ToString());
+				FinalParams = PerformIniReplacements(FinalParams, ProjectFile.Directory, Platform);
+
+				TurnkeyUtils.Log("Executing '{0}'...", FinalParams);
+
+				ExecuteBuildCookRun(FinalParams);
+			}
+		}
+
+		private static void ExecuteBuildCookRun(string Params)
+		{
+			// split the params on whitespace not inside quotes (see https://stackoverflow.com/questions/4780728/regex-split-string-preserving-quotes/4780801#4780801 to explain the regex)
+			Regex Matcher = new Regex("(?<=^[^\"]*(?:\"[^\"]*\"[^\"]*)*)\\s(?=(?:[^\"]*\"[^\"]*\")*[^\"]*$)");
+			// split the string, removing empty results
+			List<string> Arguments = Matcher.Split(Params).Where(x => x != "").ToList();
+
+			AutomationTool.CommandInfo BCRCommand = new AutomationTool.CommandInfo();
+			BCRCommand.CommandName = "BuildCookRun";
+			// chop off the first - character in all the commands (see Automation.ParseParam)
+			BCRCommand.Arguments = Arguments.Select(x => x.Substring(1)).ToList();
+
+			// use the BCR's exitcode as Turnkey's exitcode
+			TurnkeyUtils.ExitCode = Automation.Execute(new List<AutomationTool.CommandInfo>() { BCRCommand }, ScriptCompiler.Commands);
+		}
+
+		string GetIniSetting(string Spec, DirectoryReference ProjectDir, UnrealTargetPlatform Platform)
+		{
+			// handle these cases:
+			//  iniif:-option:Engine:/Script/Module.Class:bUseOption
+			//  iniif:-option:bUseOption   [convenience for ProjectPackagingSettings setting]
+			//  inivalue:Engine:/Script/Module.Class:SomeSetting
+			//  inivalue:SomeSetting       [convenience for ProjectPackagingSettings setting]
+
+			string[] Tokens = Spec.Split(":".ToCharArray());
+
+			string ConfigName;
+			string SectionName;
+			string Key;
+			string IniIfValue = "";
+			if (Tokens[0].Equals("iniif", StringComparison.InvariantCultureIgnoreCase))
+			{
+				if (Tokens.Length == 3)
+				{
+					ConfigName = "Game";
+					SectionName = "/Script/UnrealEd.ProjectPackagingSettings";
+					Key = Tokens[2];
+					IniIfValue = Tokens[1];
+				}
+				else if (Tokens.Length == 5)
+				{
+					ConfigName = Tokens[2];
+					SectionName = Tokens[3];
+					Key = Tokens[4];
+					IniIfValue = Tokens[1];
+				}
+				else
+				{
+					TurnkeyUtils.Log("Found a bad iniif spec: {0}", Spec);
+					return "";
+				}
+			}
+			else if (Tokens[0].Equals("inivalue", StringComparison.InvariantCultureIgnoreCase))
+			{
+				if (Tokens.Length == 2)
+				{
+					ConfigName = "Game";
+					SectionName = "/Script/UnrealEd.ProjectPackagingSettings";
+					Key = Tokens[1];
+				}
+				else if (Tokens.Length == 4)
+				{
+					ConfigName = Tokens[1];
+					SectionName = Tokens[2];
+					Key = Tokens[3];
+				}
+				else
+				{
+					TurnkeyUtils.Log("Found a bad inivalue spec: {0}", Spec);
+					return "";
+				}
+			}
+			else
+			{
+				TurnkeyUtils.Log("Found a bad ini spec: {0}", Spec);
+				return "";
+			}
+
+			// get the value, if it exists (or empty string if not)
+			ConfigHierarchyType ConfigType;
+			if (!ConfigHierarchyType.TryParse(ConfigName, out ConfigType))
+			{
+				TurnkeyUtils.Log("Found a bad config name {0} in spec {1}", ConfigName, Spec);
+				return "";
+			}
+			ConfigHierarchy Config = ConfigCache.ReadHierarchy(ConfigType, ProjectDir, Platform);
+			string FoundValue;
+			Config.GetString(SectionName, Key, out FoundValue);
+
+			if (Tokens[0].Equals("iniif", StringComparison.InvariantCultureIgnoreCase))
+			{
+				bool bIsTrue;
+				if (bool.TryParse(FoundValue, out bIsTrue) && bIsTrue)
+				{
+					return IniIfValue;
+				}
+				return "";
+			}
+
+			return FoundValue;
+		}
+
+		private string PerformIniReplacements(string Params, DirectoryReference ProjectDir, UnrealTargetPlatform Platform)
+		{
+			Regex IniMatch = new Regex("({(ini.*?)})+");
+			foreach (Match Match in IniMatch.Matches(Params))
+			{
+				if (Match.Success)
+				{
+					// group[1] is {ini.....}, groups[2] is the same but without the {}
+					Params = Params.Replace(Match.Groups[1].Value, GetIniSetting(Match.Groups[2].Value, ProjectDir, Platform));
+				}
+			}
+
+			return Params;
+		}
+	}
+}
diff --git a/Engine/Source/Programs/AutomationTool/Turnkey/Commands/VerifySdk.cs b/Engine/Source/Programs/AutomationTool/Turnkey/Commands/VerifySdk.cs
--- a/Engine/Source/Programs/AutomationTool/Turnkey/Commands/VerifySdk.cs
+++ b/Engine/Source/Programs/AutomationTool/Turnkey/Commands/VerifySdk.cs
@@ -130,7 +130,7 @@ namespace Turnkey.Commands
 				else
 				{
 					StatusString = "Valid";
-					ReportedState &= (SdkUtils.LocalAvailability.AutoSdk_ValidVersionExists | SdkUtils.LocalAvailability.InstalledSdk_ValidVersionExists);
+					ReportedState &= (SdkUtils.LocalAvailability.AutoSdk_ValidVersionExists | SdkUtils.LocalAvailability.InstalledSdk_ValidVersionExists | SdkUtils.LocalAvailability.Support_FullSdk| SdkUtils.LocalAvailability.Support_AutoSdk);
 					//					TurnkeyUtils.Log("{0}: Valid [Installed: '{1}', Required: '{2}']", Platform, PlatformObject.GetInstalledSdk(), PlatformObject.GetAllowedSdks());
 				}
 
diff --git a/Engine/Source/Programs/AutomationTool/Turnkey/SdkInfo.cs b/Engine/Source/Programs/AutomationTool/Turnkey/SdkInfo.cs
--- a/Engine/Source/Programs/AutomationTool/Turnkey/SdkInfo.cs
+++ b/Engine/Source/Programs/AutomationTool/Turnkey/SdkInfo.cs
@@ -28,9 +28,12 @@ namespace Turnkey
 			Platform_InvalidHostPrerequisites = 128,
 
 			Device_InvalidPrerequisites = 256,
-			Device_InstallSoftwareValid,
-			Device_InstallSoftwareInvalid,
-			Device_CannotConnect,
+			Device_InstallSoftwareValid = 512,
+			Device_InstallSoftwareInvalid = 1024,
+			Device_CannotConnect = 2048,
+
+			Support_FullSdk = 4096,
+			Support_AutoSdk = 8192,
 		}
 
 		static public LocalAvailability GetLocalAvailability(AutomationTool.Platform AutomationPlatform, bool bAllowUpdatingPrerequisites)
@@ -102,6 +105,23 @@ namespace Turnkey
 			}
 
 
+			string FullSupportedPlatforms = TurnkeyUtils.GetVariableValue("Studio_FullInstallPlatforms");
+			string AutoSdkSupportedPlatforms = TurnkeyUtils.GetVariableValue("Studio_AutoSdkPlatforms");
+			if (!string.IsNullOrEmpty(FullSupportedPlatforms))
+			{
+				if (FullSupportedPlatforms.ToLower() == "all" || FullSupportedPlatforms.Split(",".ToCharArray()).Contains(AutomationPlatform.PlatformType.ToString(), StringComparer.InvariantCultureIgnoreCase))
+				{
+					Result |= LocalAvailability.Support_FullSdk;
+				}
+			}
+			if (!string.IsNullOrEmpty(AutoSdkSupportedPlatforms))
+			{
+				if (AutoSdkSupportedPlatforms.ToLower() == "all" || AutoSdkSupportedPlatforms.Split(",".ToCharArray()).Contains(AutomationPlatform.PlatformType.ToString(), StringComparer.InvariantCultureIgnoreCase))
+				{
+					Result |= LocalAvailability.Support_AutoSdk;
+				}
+			}
+
 			return Result;
 		}
 
diff --git a/Engine/Source/Programs/AutomationTool/Turnkey/Turnkey.Automation.cs b/Engine/Source/Programs/AutomationTool/Turnkey/Turnkey.Automation.cs
--- a/Engine/Source/Programs/AutomationTool/Turnkey/Turnkey.Automation.cs
+++ b/Engine/Source/Programs/AutomationTool/Turnkey/Turnkey.Automation.cs
@@ -56,11 +56,12 @@ namespace Turnkey
 
 		public static AutomationTool.ExitCode Execute(IOProvider IOProvider, BuildCommand CommandUtilHelper)
 		{
-			TurnkeyUtils.Initialize(IOProvider, CommandUtilHelper);
-			TurnkeySettings.Initialize();
-
 			// cache some settings for other classes
 			TurnkeyUtils.SetVariable("EngineDir", EngineDirectory.FullName);
+			TurnkeyUtils.SetVariable("Project", CommandUtilHelper.ParseParamValue("Project="));
+			
+			TurnkeyUtils.Initialize(IOProvider, CommandUtilHelper);
+			TurnkeySettings.Initialize();
 
 			//			Microsoft.Win32.SystemEvents.UserPreferenceChanged += new Microsoft.Win32.UserPreferenceChangedEventHandler((sender, args) => TurnkeyUtils.Log("Got a change! {0}", args.Category));
 
diff --git a/Engine/Source/Programs/AutomationTool/Turnkey/TurnkeyCommand.cs b/Engine/Source/Programs/AutomationTool/Turnkey/TurnkeyCommand.cs
--- a/Engine/Source/Programs/AutomationTool/Turnkey/TurnkeyCommand.cs
+++ b/Engine/Source/Programs/AutomationTool/Turnkey/TurnkeyCommand.cs
@@ -71,6 +71,12 @@ namespace Turnkey
 			}
 		}
 
+		static bool bExitAfterCommand = false;
+		public static void ExitTurnkeyLoopAfterCommand()
+		{
+			bExitAfterCommand = true;
+		}
+
 		public static bool ExecuteCommand(string CommandString=null)
 		{
 			TurnkeyCommand Command;
@@ -115,7 +121,7 @@ namespace Turnkey
 			}
 
 			// true if we should keep going
-			return true;
+			return !bExitAfterCommand;
 		}
 
 		private static int PromptForCommand()
diff --git a/Engine/Source/Programs/AutomationTool/Turnkey/TurnkeySettings.cs b/Engine/Source/Programs/AutomationTool/Turnkey/TurnkeySettings.cs
--- a/Engine/Source/Programs/AutomationTool/Turnkey/TurnkeySettings.cs
+++ b/Engine/Source/Programs/AutomationTool/Turnkey/TurnkeySettings.cs
@@ -48,6 +48,8 @@ namespace Turnkey
 			new UserSetting("Studio_AppleSigningCertPassword", null, "A shared password that is used across Apple Signing Certificates"),
 			new UserSetting("Studio_GoogleDriveCredentials", null, "The location of credentials file needed for GoogleDrive integration. The contents should start with: '{\"installed\":{\"client_id\":'"),
 			new UserSetting("Studio_GoogleDriveAppName", null, "The name of the application your studio needs to create (see documentation for help)"),
+			new UserSetting("Studio_FullInstallPlatforms", null, "The list of platforms your studio has support for Full Sdk installation. Note that this generally affects UI exposing of options, you can still use Turnkey commandline interface to install Sdks for unlisted platforms. This can be a comma separated list, or All."),
+			new UserSetting("Studio_AutoSdkPlatforms", null, "The list of platforms your studio has support for AutoSdk. Note that this generally affects UI exposing of options, you can still use Turnkey commandline interface to install Sdks for unlisted platforms. This can be a comma separated list, or All."),
 		};
 
 		// basically same as Turnkey variables, but this only contains ones that were loaded so we can write them back out
@@ -58,8 +60,10 @@ namespace Turnkey
 		{
 			// then load any saved settings
 			TurnkeyManifest.LoadManifestsFromProvider("file:" + TurnkeySettings.UserSettingManifestLocation);
+			TurnkeyManifest.LoadManifestsFromProvider("file:$(EngineDir)/Build/Turnkey/TurnkeyStudioSettings.xml");
 
 			// now pull out anything that was set and unset it from Turnkey
+			// Studio settings aren't special like UserSettings (read-only) so we don't have to unset them
 			foreach (UserSetting PossibleSetting in AllUserSettings)
 			{
 				string Value = TurnkeyUtils.GetVariableValue(PossibleSetting.VariableName);
@@ -70,6 +74,7 @@ namespace Turnkey
 				}
 			}
 		}
+
 		public static bool HasSetUserSetting(string VariableName)
 		{
 			return SetUserSettings.ContainsKey(VariableName);
diff --git a/Engine/Source/Programs/AutomationTool/Turnkey/TurnkeyUtils.cs b/Engine/Source/Programs/AutomationTool/Turnkey/TurnkeyUtils.cs
--- a/Engine/Source/Programs/AutomationTool/Turnkey/TurnkeyUtils.cs
+++ b/Engine/Source/Programs/AutomationTool/Turnkey/TurnkeyUtils.cs
@@ -109,6 +109,12 @@ namespace Turnkey
 			string PlatformString = TurnkeyUtils.ParseParamValue("Platform", null, CommandOptions);
 			bool bUnattended = TurnkeyUtils.ParseParam("Unattended", CommandOptions);
 
+			// use default set of platforms if none specified
+			if (PossiblePlatforms == null)
+			{
+				PossiblePlatforms = UnrealTargetPlatform.GetValidPlatforms().ToList();
+			}
+
 			// HACK UNTIL WIN32 IS GONE
 			PossiblePlatforms = PossiblePlatforms.Where(x => x != UnrealTargetPlatform.Win32 && x != UnrealTargetPlatform.XboxOne).ToList();
 
