- Fixed a Turnkey bug that would delete other platforms' SDKs when installing an SDK

diff --git a/Engine/Source/Programs/AutomationTool/Turnkey/TurnkeyManifest.cs b/Engine/Source/Programs/AutomationTool/Turnkey/TurnkeyManifest.cs
--- a/Engine/Source/Programs/AutomationTool/Turnkey/TurnkeyManifest.cs
+++ b/Engine/Source/Programs/AutomationTool/Turnkey/TurnkeyManifest.cs
@@ -100,8 +100,8 @@ namespace Turnkey
 			List<FileSource> NotNeedsExpansion = Sources.FindAll(x => !x.NeedsFileExpansion());
 
 
-			// remove them, then we add expansions below
-			DiscoveredFileSources = DiscoveredFileSources.FindAll(x => !x.NeedsFileExpansion());
+			// remove the ones that we are going to expand below
+			DiscoveredFileSources = DiscoveredFileSources.FindAll(x => !NeedsExpansion.Contains(x));
 
 			foreach (FileSource Source in NeedsExpansion)
 			{
