Move HashCollection out of DotNETCommon.

#rnx

diff --git a/Engine/Source/Programs/AutomationTool/AutomationUtils/HashCollection.cs b/Engine/Source/Programs/AutomationTool/AutomationUtils/HashCollection.cs
new file mode 100644
--- /dev/null
+++ b/Engine/Source/Programs/AutomationTool/AutomationUtils/HashCollection.cs
@@ -0,0 +1,287 @@
+// Copyright Epic Games, Inc. All Rights Reserved.
+
+using System;
+using System.Collections.Generic;
+using System.IO;
+using System.Security.Cryptography;
+using System.Text;
+using System.Xml;
+using System.Xml.Serialization;
+using System.Linq;
+using Tools.DotNETCommon;
+
+namespace AutomationTool
+{
+
+	/// <summary>
+	/// Class that holds a single hash entry. Only the hash is used for comparisons but the
+	/// name and meta fields can hold related data for debugging etc.
+	/// </summary>
+	public class HashEntry : IEquatable<HashEntry>
+	{
+		public string Name { get; set; }
+		public string MetaData { get; set; }
+
+		[XmlIgnore]
+		public ContentHash Hash { get; set; }
+		
+		[XmlElement(ElementName = "Hash", DataType = "hexBinary")]
+		public byte[] HashBytes
+		{
+			get { return Hash.Bytes; }
+			set { Hash = new ContentHash(value); }
+		}
+
+		public HashEntry(string InName, ContentHash InHash, string InMeta)
+		{
+			Name = InName;
+			Hash = InHash;
+			MetaData = InMeta;
+		}
+
+		// parameterless constructor for serialization
+		public HashEntry()
+		{
+		}
+
+		// Equality check using the underlying content hash
+		public bool Equals(HashEntry RHS)
+		{
+			return Hash == RHS.Hash;
+		}
+
+		public override int GetHashCode()
+		{
+			return Hash.GetHashCode();
+		}
+
+		public override string ToString()
+		{
+			return Hash.ToString();
+		}
+	}
+
+	/// <summary>
+	/// Class that holds a collection of hashes with helpers for adding files, comparisions, and
+	/// serialization to and from files
+	/// </summary>
+	public class HashCollection
+	{
+		public enum HashType
+		{
+			MetaData,
+			Content
+		}
+
+		// underlying hashset that holds our data
+		public HashSet<HashEntry> Hashes { get; protected set; }
+
+		public HashCollection()
+		{
+			Hashes = new HashSet<HashEntry>();
+		}
+
+		/// <summary>
+		/// Adds a file to our hash collection. If the specified path does not exist an exception
+		/// will be thrown
+		/// </summary>
+		/// <param name="InPath"></param>
+		/// <param name="HashType"></param>
+		public bool AddFile(FileReference InPath, HashCollection.HashType HashType)
+		{
+			var Fi = new FileInfo(InPath.FullName);
+
+			if (!Fi.Exists)
+			{
+				return false;
+			}
+
+			string MetaData = string.Format("File={0} Size={1} LastWrite={2}", InPath, Fi.Length, Fi.LastWriteTimeUtc.ToString());
+
+			// Hash metadata or content
+			ContentHash Hash = HashType == HashCollection.HashType.MetaData ? ContentHash.SHA1(MetaData) : ContentHash.SHA1(InPath);
+
+			HashEntry Entry = new HashEntry(InPath.FullName, Hash, MetaData);
+
+			Hashes.Add(Entry);
+
+			return true;
+		}
+
+		/// <summary>
+		/// Convenience function that adds a range of files
+		/// </summary>
+		/// <param name="Files"></param>
+		/// <param name="HashType"></param>
+		public bool AddFiles(IEnumerable<FileReference> Files, HashCollection.HashType HashType)
+		{
+			foreach (var File in Files)
+			{
+				if (!AddFile(File, HashType))
+				{
+					return false;
+				}
+			}
+
+			return true;
+		}
+
+		/// <summary>
+		/// Compares two collections by comparing the underlying hashsets
+		/// </summary>
+		/// <param name="obj"></param>
+		/// <returns></returns>
+		public override bool Equals(object obj)
+		{
+			HashCollection RHS = obj as HashCollection;
+
+			if (RHS == null)
+			{
+				return false;
+
+			}
+			return Hashes.SetEquals(RHS.Hashes);
+		}
+
+		/// <summary>
+		/// Debugging helper that returns true if there's a hash with the specified name
+		/// </summary>
+		/// <param name="Name"></param>
+		/// <returns></returns>
+		public bool HasEntryForName(string Name)
+		{
+			return Hashes.Where(E => E.Name == Name).Any();
+		}
+
+		/// <summary>
+		/// Debugging helper that returns the hash entry for the specified name
+		/// </summary>
+		/// <param name="Input"></param>
+		/// <returns></returns>
+		public HashEntry GetEntryForName(string Input)
+		{
+			return Hashes.Where(E => E.Name == Input).FirstOrDefault();
+		}
+
+		/// <summary>
+		/// Logs differences between two collections for debugging
+		/// </summary>
+		/// <param name="RHS"></param>
+		public void LogDifferences(HashCollection RHS)
+		{
+			foreach (var Entry in Hashes)
+			{
+				if (!RHS.HasEntryForName(Entry.Name))
+				{
+					Log.TraceInformation("RHS does not have an entry for {0}", Entry.Name);
+				}
+				else if (!RHS.Hashes.Contains(Entry))
+				{
+					HashEntry RHSEntry = RHS.GetEntryForName(Entry.Name);
+					Log.TraceInformation("RHS hash mismatch for {0}", Entry.Name);
+					Log.TraceInformation("LHS:\n\t{0}\n\t{1}\n\t{2}", Entry.Hash, Entry.Name, Entry.MetaData);
+					Log.TraceInformation("RHS:\n\t{0}\n\t{1}\n\t{2}", RHSEntry.Hash, RHSEntry.Name, RHSEntry.MetaData);
+				}
+			}
+		}
+
+		/// <summary>
+		/// 
+		/// </summary>
+		/// <returns></returns>
+		public override int GetHashCode()
+		{
+			return EqualityComparer<HashSet<HashEntry>>.Default.GetHashCode(Hashes);
+		}
+
+		/// <summary>
+		/// Creates a hash collection from a previously serialized file
+		/// </summary>
+		/// <param name="InPath"></param>
+		/// <returns></returns>
+		public static HashCollection CreateFromFile(string InPath)
+		{
+			try
+			{
+				using (var stream = System.IO.File.OpenRead(InPath))
+				{
+					var Serializer = new XmlSerializer(typeof(HashCollection));
+					return Serializer.Deserialize(stream) as HashCollection;
+				}
+			}
+			catch (Exception Ex)
+			{
+				Log.TraceWarning("Failed to load HashCollection from {0}. {1}", InPath, Ex.Message);
+			}
+
+			return null;
+		}
+
+		/// <summary>
+		/// Saves a hash collection to the specified file path as XML.
+		/// </summary>
+		/// <param name="OutPath"></param>
+		public void SaveToFile(string OutPath)
+		{
+			XmlSerializer XS = new XmlSerializer(this.GetType());
+
+			//first serialize the object to memory stream,
+			//in case of exception, the original file is not corrupted
+			using (MemoryStream MS = new MemoryStream())
+			{
+				XmlWriterSettings Settings = new XmlWriterSettings();
+				Settings.Indent = true;
+				Settings.IndentChars = ("\t");
+				Settings.Encoding = Encoding.UTF8;
+
+				using (XmlWriter Writer = XmlWriter.Create(MS, Settings))
+				{
+					XS.Serialize(Writer, this);
+				}
+
+				MS.Flush();
+
+				string String = Encoding.UTF8.GetString(MS.ToArray());
+
+				Directory.CreateDirectory(Path.GetDirectoryName(OutPath));
+
+				File.WriteAllText(OutPath, String);
+			}
+		}
+	}
+
+	/// <summary>
+	/// Extension methods for CsProject support
+	/// </summary>
+	public static class CsProjectInfoExtensionMethods
+	{
+		/// <summary>
+		/// Adds aall input/output properties of a CSProject to a hash collection
+		/// </summary>
+		/// <param name="Hasher"></param>
+		/// <param name="Project"></param>
+		/// <returns></returns>
+		public static bool AddCsProjectInfo(this HashCollection Hasher, CsProjectInfo Project, HashCollection.HashType HashType)
+		{
+			// Get the output assembly and pdb file
+			FileReference OutputFile;
+			if (!Project.TryGetOutputFile(out OutputFile))
+			{
+				throw new Exception(String.Format("Unable to get output file for {0}", Project.ProjectPath));
+			}
+			FileReference DebugFile = OutputFile.ChangeExtension("pdb");
+
+			// build a list of all input and output files from this module
+			List<FileReference> DependentFiles = new List<FileReference> { Project.ProjectPath, OutputFile, DebugFile };
+
+			DependentFiles.AddRange(Project.CompileReferences);
+
+			if (!Hasher.AddFiles(DependentFiles, HashType))
+			{
+				return false;
+			}
+
+			return true;
+		}
+	}
+}
diff --git a/Engine/Source/Programs/DotNETCommon/DotNETUtilities/CsProjectInfo.cs b/Engine/Source/Programs/DotNETCommon/DotNETUtilities/CsProjectInfo.cs
--- a/Engine/Source/Programs/DotNETCommon/DotNETUtilities/CsProjectInfo.cs
+++ b/Engine/Source/Programs/DotNETCommon/DotNETUtilities/CsProjectInfo.cs
@@ -1020,39 +1020,4 @@ namespace Tools.DotNETCommon
 			return NewText;
 		}
 	}
-
-	/// <summary>
-	/// Extension methods for CsProject support
-	/// </summary>
-	public static class CsProjectInfoExtensionMethods
-	{
-		/// <summary>
-		/// Adds aall input/output properties of a CSProject to a hash collection
-		/// </summary>
-		/// <param name="Hasher"></param>
-		/// <param name="Project"></param>
-		/// <returns></returns>
-		public static bool AddCsProjectInfo(this HashCollection Hasher, CsProjectInfo Project, HashCollection.HashType HashType)
-		{
-			// Get the output assembly and pdb file
-			FileReference OutputFile;
-			if(!Project.TryGetOutputFile(out OutputFile))
-			{
-				throw new Exception(String.Format("Unable to get output file for {0}", Project.ProjectPath));
-			}
-			FileReference DebugFile = OutputFile.ChangeExtension("pdb");
-
-			// build a list of all input and output files from this module
-			List<FileReference> DependentFiles = new List<FileReference> { Project.ProjectPath, OutputFile, DebugFile };
-
-			DependentFiles.AddRange(Project.CompileReferences);
-	
-			if (!Hasher.AddFiles(DependentFiles, HashType))
-			{
-				return false;
-			}
-
-			return true;
-		}
-	}
 }
diff --git a/Engine/Source/Programs/DotNETCommon/DotNETUtilities/HashCollection.cs b/Engine/Source/Programs/DotNETCommon/DotNETUtilities/HashCollection.cs
deleted file mode 100644
diff --git a/Engine/Source/Programs/DotNet/EpicGames.Common/HelpUtils.cs b/Engine/Source/Programs/DotNet/EpicGames.Common/HelpUtils.cs
new file mode 100644
--- /dev/null
+++ b/Engine/Source/Programs/DotNet/EpicGames.Common/HelpUtils.cs
@@ -0,0 +1,198 @@
+// Copyright Epic Games, Inc. All Rights Reserved.
+
+using System;
+using System.Collections.Generic;
+using System.ComponentModel;
+using System.Linq;
+using System.Text;
+using System.Threading.Tasks;
+
+namespace Tools.DotNETCommon
+{
+	/// <summary>
+	/// Utility functions for showing help for objects
+	/// </summary>
+	public static class HelpUtils
+	{
+		/// <summary>
+		/// Gets the width of the window for formatting purposes
+		/// </summary>
+		public static int WindowWidth
+		{
+			get
+			{
+				// Get the window width, using a default value if there's no console attached to this process.
+				int NewWindowWidth;
+				try
+				{
+					NewWindowWidth = Console.WindowWidth;
+				}
+				catch
+				{
+					NewWindowWidth = 240;
+				}
+
+				if (NewWindowWidth <= 0)
+				{
+					NewWindowWidth = 240;
+				}
+
+				return NewWindowWidth;
+			}
+		}
+
+		/// <summary>
+		/// Prints help for the given object type
+		/// </summary>
+		/// <param name="Type">Type to print help for</param>
+		public static void PrintHelp(string Title, Type Type)
+		{
+			PrintHelp(Title, GetDescription(Type), CommandLineArguments.GetParameters(Type));
+		}
+
+		/// <summary>
+		/// Prints help for a command
+		/// </summary>
+		/// <param name="Title">Title for the help text</param>
+		/// <param name="Description">Description for the command</param>
+		/// <param name="Parameters">List of parameters</param>
+		public static void PrintHelp(string Title, string Description, List<KeyValuePair<string, string>> Parameters)
+		{
+			int MaxWidth = WindowWidth;
+
+			bool bFirstLine = true;
+			if (!String.IsNullOrEmpty(Title))
+			{
+				PrintParagraph(Title, MaxWidth);
+				bFirstLine = false;
+			}
+
+			if (!String.IsNullOrEmpty(Description))
+			{
+				if (!bFirstLine)
+				{
+					Log.TraceInformation("");
+				}
+				PrintParagraph(Description, MaxWidth);
+				bFirstLine = false;
+			}
+
+			if (Parameters.Count > 0)
+			{
+				if (!bFirstLine)
+				{
+					Log.TraceInformation("");
+				}
+
+				Log.TraceInformation("Parameters:");
+				PrintTable(Parameters, 4, 24, MaxWidth);
+			}
+		}
+
+		/// <summary>
+		/// Gets the description from a type
+		/// </summary>
+		/// <param name="Type">The type to get a description for</param>
+		/// <returns>The description text</returns>
+		public static string GetDescription(Type Type)
+		{
+			StringBuilder DescriptionText = new StringBuilder();
+			foreach (DescriptionAttribute Attribute in Type.GetCustomAttributes(typeof(DescriptionAttribute), false))
+			{
+				if (DescriptionText.Length > 0)
+				{
+					DescriptionText.AppendLine();
+				}
+				DescriptionText.AppendLine(Attribute.Description);
+			}
+			return DescriptionText.ToString();
+		}
+
+		/// <summary>
+		/// Prints a paragraph of text using word wrapping
+		/// </summary>
+		/// <param name="Text">Text to print</param>
+		/// <param name="MaxWidth">Maximum width for each line</param>
+		public static void PrintParagraph(string Text)
+		{
+			PrintParagraph(Text, WindowWidth - 1);
+		}
+
+		/// <summary>
+		/// Prints a paragraph of text using word wrapping
+		/// </summary>
+		/// <param name="Text">Text to print</param>
+		/// <param name="MaxWidth">Maximum width for each line</param>
+		public static void PrintParagraph(string Text, int MaxWidth)
+		{
+			List<string> Lines = StringUtils.WordWrap(Text, MaxWidth);
+			foreach (string Line in Lines)
+			{
+				Log.TraceInformation("{0}", Line);
+			}
+		}
+
+		/// <summary>
+		/// Formats the given parameters as so:
+		///     -Param1     Param1 Description
+		///
+		///     -Param2      Param2 Description, this description is
+		///                  longer and splits onto a separate line. 
+		///
+		///     -Param3      Param3 Description continues as before. 
+		/// </summary>
+		/// <param name="Items">List of parameters arranged as "-ParamName Param Description"</param>
+		/// <param name="Indent">Indent from the left hand side</param>
+		/// <param name="DefaultRightPadding">The minimum padding from the start of the param name to the start of the description (resizes with larger param names)</param>
+		/// <returns></returns>
+		public static void PrintTable(List<KeyValuePair<string, string>> Items, int Indent, int MinFirstColumnWidth)
+		{
+			PrintTable(Items, Indent, MinFirstColumnWidth, WindowWidth - 1);
+		}
+
+		/// <summary>
+		/// Formats the given parameters as so:
+		///     -Param1     Param1 Description
+		///
+		///     -Param2      Param2 Description, this description is
+		///                  longer and splits onto a separate line. 
+		///
+		///     -Param3      Param3 Description continues as before. 
+		/// </summary>
+		/// <param name="Items">List of parameters arranged as "-ParamName Param Description"</param>
+		/// <param name="Indent">Indent from the left hand side</param>
+		/// <param name="DefaultRightPadding">The minimum padding from the start of the param name to the start of the description (resizes with larger param names)</param>
+		/// <returns></returns>
+		public static void PrintTable(List<KeyValuePair<string, string>> Items, int Indent, int MinFirstColumnWidth, int MaxWidth)
+		{
+			if(Items.Count > 0)
+			{
+				// string used to intent the param
+				string IndentString = new string(' ', Indent);
+
+				// default the padding value
+				int RightPadding = Math.Max(MinFirstColumnWidth, Items.Max(x => x.Key.Length + 1));
+
+				// Build the formatted params
+				foreach(KeyValuePair<string, string> Item in Items)
+				{
+					// build the param first, including intend and padding on the rights size
+					string ParamString = IndentString + Item.Key.PadRight(RightPadding);
+
+					// Build the description line by line, adding the same amount of intending each time. 
+					List<string> DescriptionLines = StringUtils.WordWrap(Item.Value, MaxWidth - ParamString.Length);
+
+					foreach(string DescriptionLine in DescriptionLines)
+					{
+						// Formatting as following:
+						// <Indent>-param<Right Padding>Description<New line>
+						Log.TraceInformation("{0}{1}", ParamString, DescriptionLine);
+
+						// we replace the param string on subsequent lines with white space of the same length
+						ParamString = string.Empty.PadRight(IndentString.Length + RightPadding);
+					}
+				}
+			}
+		}
+	}
+}
