Remove references to FastJSON functionality in DotNETCommon. Use a duplicate copy in AutomationScripts called Json_DEPRECATED instead.

#rnx

diff --git a/Engine/Source/Programs/AutomationTool/AutomationUtils/Json.cs b/Engine/Source/Programs/AutomationTool/AutomationUtils/Json.cs
new file mode 100644
--- /dev/null
+++ b/Engine/Source/Programs/AutomationTool/AutomationUtils/Json.cs
@@ -0,0 +1,264 @@
+using System;
+using System.Collections.Generic;
+using System.Text;
+using Tools.DotNETCommon;
+
+namespace AutomationTool
+{
+	/// <summary>
+	/// Options for serializing an object to JSON
+	/// </summary>
+	public enum JsonSerializeOptions_DEPRECATED
+	{
+		/// <summary>
+		/// No options
+		/// </summary>
+		None,
+
+		/// <summary>
+		/// Whether to format the output for readability
+		/// </summary>
+		PrettyPrint = 1,
+	}
+
+	/// <summary>
+	/// Wrapper class around fastJSON methods. These are deprecated. Prefer the native System.Text.Json framework instead.
+	/// </summary>
+	public static class Json_DEPRECATED
+	{
+		/// <summary>
+		/// Static constructor. Registers known types to the serializer.
+		/// </summary>
+		static Json_DEPRECATED()
+		{
+			fastJSON.JSON.Instance.RegisterCustomType(typeof(FileReference), x => EscapeString(x.ToString()), x => new FileReference(x));
+			fastJSON.JSON.Instance.RegisterCustomType(typeof(DirectoryReference), x => EscapeString(x.ToString()), x => new DirectoryReference(x));
+		}
+
+		/// <summary>
+		/// Serializes an object to JSON
+		/// </summary>
+		/// <param name="Object">Object to be serialized</param>
+		/// <returns>JSON representation of the object</returns>
+		public static string Serialize(object Object)
+		{
+			return Serialize(Object, JsonSerializeOptions_DEPRECATED.PrettyPrint);
+		}
+
+		/// <summary>
+		/// Serializes an object to JSON
+		/// </summary>
+		/// <param name="Object">Object to be serialized</param>
+		/// <param name="Options">Options for formatting the object</param>
+		/// <returns>JSON representation of the object</returns>
+		public static string Serialize(object Object, JsonSerializeOptions_DEPRECATED Options)
+		{
+			fastJSON.JSONParameters Params = new fastJSON.JSONParameters();
+			Params.EnableAnonymousTypes = true;
+
+			string Text = fastJSON.JSON.Instance.ToJSON(Object, Params);
+			if ((Options & JsonSerializeOptions_DEPRECATED.PrettyPrint) != 0)
+			{
+				Text = Format(Text);
+			}
+			return Text;
+		}
+
+		/// <summary>
+		/// Deserialize an object from text
+		/// </summary>
+		/// <typeparam name="T">Type of the object to deserialize</typeparam>
+		/// <param name="Text">String to parse</param>
+		/// <returns>Constructed object</returns>
+		public static T Deserialize<T>(string Text)
+		{
+			return fastJSON.JSON.Instance.ToObject<T>(Text);
+		}
+
+		/// <summary>
+		/// Read a JSON object from a file on disk
+		/// </summary>
+		/// <typeparam name="T">The type of the object to load</typeparam>
+		/// <param name="Location">The location of the file to read</param>
+		/// <returns>The deserialized object</returns>
+		public static T Load<T>(FileReference Location)
+		{
+			string Text;
+			try
+			{
+				Text = FileReference.ReadAllText(Location);
+			}
+			catch (Exception Ex)
+			{
+				throw new Exception(String.Format("Unable to read '{0}'", Location), Ex);
+			}
+			return Deserialize<T>(Text);
+		}
+
+		/// <summary>
+		/// Save a JSON object to a file on disk
+		/// </summary>
+		/// <param name="Location">The location of the file to read</param>
+		/// <param name="Object">The object to save</param>
+		public static void Save(FileReference Location, object Object)
+		{
+			string Text = Serialize(Object);
+			try
+			{
+				DirectoryReference.CreateDirectory(Location.Directory);
+				FileReference.WriteAllText(Location, Text);
+			}
+			catch (Exception Ex)
+			{
+				throw new Exception(String.Format("Unable to write '{0}'", Location), Ex);
+			}
+		}
+
+		/// <summary>
+		/// Formats the given JSON string in 'Epic' format (parentheses/brackets on new lines, tabs instead of spaces, no space before colons in key/value pairs)
+		/// </summary>
+		/// <param name="Input">The input string</param>
+		/// <returns>The formatted string</returns>
+		public static string Format(string Input)
+		{
+			StringBuilder Result = new StringBuilder(Input.Length * 2);
+
+			bool bOnEmptyLine = true;
+
+			int Depth = 0;
+			for (int Idx = 0; Idx < Input.Length; Idx++)
+			{
+				switch (Input[Idx])
+				{
+					case '{':
+					case '[':
+						if (!bOnEmptyLine)
+						{
+							Result.AppendLine();
+							bOnEmptyLine = true;
+						}
+
+						AppendIndent(Result, Depth++);
+						Result.Append(Input[Idx]);
+						Result.AppendLine();
+						break;
+					case '}':
+					case ']':
+						if (!bOnEmptyLine)
+						{
+							Result.AppendLine();
+						}
+
+						AppendIndent(Result, --Depth);
+						Result.Append(Input[Idx]);
+						bOnEmptyLine = false;
+						break;
+					case ':':
+						Result.Append(Input[Idx]);
+						Result.Append(' ');
+						break;
+					case ',':
+						Result.Append(Input[Idx]);
+						Result.AppendLine();
+						bOnEmptyLine = true;
+						break;
+					default:
+						if (!Char.IsWhiteSpace(Input[Idx]))
+						{
+							// If this is an empty line, append the indent
+							if (bOnEmptyLine)
+							{
+								AppendIndent(Result, Depth);
+								bOnEmptyLine = false;
+							}
+
+							// Append the character
+							Result.Append(Input[Idx]);
+
+							// If this was a string, also append everything to the end of the string
+							if (Input[Idx] == '\"')
+							{
+								Idx++;
+								for (; Idx < Input.Length; Idx++)
+								{
+									Result.Append(Input[Idx]);
+									if (Input[Idx] == '\\' && Idx + 1 < Input.Length)
+									{
+										Result.Append(Input[++Idx]);
+									}
+									else if (Input[Idx] == '\"')
+									{
+										break;
+									}
+								}
+							}
+						}
+						break;
+				}
+			}
+
+			return Result.ToString();
+		}
+
+		/// <summary>
+		/// Appends an indent to the output string
+		/// </summary>
+		/// <param name="Builder">The string builder to append to</param>
+		/// <param name="Depth">The number of tabs to append</param>
+		private static void AppendIndent(StringBuilder Builder, int Depth)
+		{
+			for (int Idx = 0; Idx < Depth; Idx++)
+			{
+				Builder.Append('\t');
+			}
+		}
+
+		/// <summary>
+		/// Escapes a string for serializing to JSON
+		/// </summary>
+		/// <param name="Value">The string to escape</param>
+		/// <returns>The escaped string</returns>
+		public static string EscapeString(string Value)
+		{
+			StringBuilder Result = new StringBuilder();
+			for (int Idx = 0; Idx < Value.Length; Idx++)
+			{
+				switch (Value[Idx])
+				{
+					case '\"':
+						Result.Append("\\\"");
+						break;
+					case '\\':
+						Result.Append("\\\\");
+						break;
+					case '\b':
+						Result.Append("\\b");
+						break;
+					case '\f':
+						Result.Append("\\f");
+						break;
+					case '\n':
+						Result.Append("\\n");
+						break;
+					case '\r':
+						Result.Append("\\r");
+						break;
+					case '\t':
+						Result.Append("\\t");
+						break;
+					default:
+						if (Char.IsControl(Value[Idx]))
+						{
+							Result.AppendFormat("\\u{0:X4}", (int)Value[Idx]);
+						}
+						else
+						{
+							Result.Append(Value[Idx]);
+						}
+						break;
+				}
+			}
+			return Result.ToString();
+		}
+	}
+}
diff --git a/Engine/Source/Programs/AutomationTool/BuildGraph/Tasks/ModifyJsonValue.cs b/Engine/Source/Programs/AutomationTool/BuildGraph/Tasks/ModifyJsonValue.cs
--- a/Engine/Source/Programs/AutomationTool/BuildGraph/Tasks/ModifyJsonValue.cs
+++ b/Engine/Source/Programs/AutomationTool/BuildGraph/Tasks/ModifyJsonValue.cs
@@ -70,7 +70,7 @@ namespace BuildGraph.Tasks
 
 				CurrObj[Keys[Keys.Length - 1]] = Parameters.NewValue;
 
-				var NewContents = Json.Serialize(ParamObj, JsonSerializeOptions.PrettyPrint);
+				var NewContents = Json_DEPRECATED.Serialize(ParamObj, JsonSerializeOptions_DEPRECATED.PrettyPrint);
 				File.WriteAllText(JsonFile, NewContents, new UTF8Encoding(encoderShouldEmitUTF8Identifier: false));
 			}
 		}
diff --git a/Engine/Source/Programs/AutomationTool/Gauntlet/Framework/Utils/Gauntlet.HordeReport.cs b/Engine/Source/Programs/AutomationTool/Gauntlet/Framework/Utils/Gauntlet.HordeReport.cs
--- a/Engine/Source/Programs/AutomationTool/Gauntlet/Framework/Utils/Gauntlet.HordeReport.cs
+++ b/Engine/Source/Programs/AutomationTool/Gauntlet/Framework/Utils/Gauntlet.HordeReport.cs
@@ -351,7 +351,7 @@ namespace Gauntlet
 					string TestResultFilePath = Path.Combine(OutputArtifactPath, OutputTestResult.ArtifactName);
 					try
 					{
-						Json.Save(new FileReference(TestResultFilePath), OutputTestResultDetailed);
+						Json_DEPRECATED.Save(new FileReference(TestResultFilePath), OutputTestResultDetailed);
 						ArtifactsCount++;
 					}
 					catch (Exception Ex)
@@ -470,7 +470,7 @@ namespace Gauntlet
 				Log.Verbose("Writing Test Data Collection for Horde at {0}", OutputTestDataFilePath);
 				try
 				{
-					Json.Save(new FileReference(OutputTestDataFilePath), this);
+					Json_DEPRECATED.Save(new FileReference(OutputTestDataFilePath), this);
 				}
 				catch (Exception Ex)
 				{
diff --git a/Engine/Source/Programs/AutomationTool/Gauntlet/Unreal/Utils/Gauntlet.UnrealTestResults.cs b/Engine/Source/Programs/AutomationTool/Gauntlet/Unreal/Utils/Gauntlet.UnrealTestResults.cs
--- a/Engine/Source/Programs/AutomationTool/Gauntlet/Unreal/Utils/Gauntlet.UnrealTestResults.cs
+++ b/Engine/Source/Programs/AutomationTool/Gauntlet/Unreal/Utils/Gauntlet.UnrealTestResults.cs
@@ -1,5 +1,6 @@
 // Copyright Epic Games, Inc. All Rights Reserved.
 
+using AutomationTool;
 using System.Collections.Generic;
 using Tools.DotNETCommon;
 
@@ -66,7 +67,7 @@ namespace Gauntlet
 		/// <param name="FilePath"></param>
 		public static UnrealAutomatedTestPassResults LoadFromJson(string FilePath)
 		{
-			UnrealAutomatedTestPassResults JsonTestPassResults = Json.Load<UnrealAutomatedTestPassResults>(new FileReference(FilePath));
+			UnrealAutomatedTestPassResults JsonTestPassResults = Json_DEPRECATED.Load<UnrealAutomatedTestPassResults>(new FileReference(FilePath));
 			return JsonTestPassResults;
 		}
 	}
diff --git a/Engine/Source/Programs/AutomationTool/Turnkey/LocalCache.cs b/Engine/Source/Programs/AutomationTool/Turnkey/LocalCache.cs
--- a/Engine/Source/Programs/AutomationTool/Turnkey/LocalCache.cs
+++ b/Engine/Source/Programs/AutomationTool/Turnkey/LocalCache.cs
@@ -10,7 +10,7 @@ using System.IO;
 using System.Runtime.Serialization.Formatters.Binary;
 using System.Reflection;
 using Tools.DotNETCommon;
-
+using AutomationTool;
 
 namespace Turnkey
 {
@@ -151,7 +151,7 @@ namespace Turnkey
 
 		static void SerializeObject(SavedCache Object, string Filename)
 		{
-			string Str = Json.Serialize(Object);
+			string Str = Json_DEPRECATED.Serialize(Object);
 			File.WriteAllText(Filename, Str);
 
 // 			using (FileStream Stream = new FileStream(Filename, FileMode.Create))
@@ -172,7 +172,7 @@ namespace Turnkey
 				string Str = File.ReadAllText(Filename);
 				if (!string.IsNullOrEmpty(Str))
 				{
-					return Json.Deserialize<SavedCache>(Str);
+					return Json_DEPRECATED.Deserialize<SavedCache>(Str);
 				}
 
 // 				if (File.Exists(Filename))
diff --git a/Engine/Source/Programs/UnrealBuildTool/Modes/AggregateParsedTimingInfo.cs b/Engine/Source/Programs/UnrealBuildTool/Modes/AggregateParsedTimingInfo.cs
--- a/Engine/Source/Programs/UnrealBuildTool/Modes/AggregateParsedTimingInfo.cs
+++ b/Engine/Source/Programs/UnrealBuildTool/Modes/AggregateParsedTimingInfo.cs
@@ -6,6 +6,7 @@ using System.Collections.Generic;
 using System.IO;
 using System.IO.Compression;
 using System.Linq;
+using System.Text.Json;
 using System.Text.RegularExpressions;
 using System.Threading;
 using System.Threading.Tasks;
@@ -83,7 +84,7 @@ namespace UnrealBuildTool
 				{
 					CompileTimes.Add(Json.EscapeString(TimingData.Key), TimingData.Value.InclusiveDuration);
 				}
-				string JsonCompileTimes = Json.Serialize(CompileTimes);
+				string JsonCompileTimes = JsonSerializer.Serialize(CompileTimes);
 				FileReference.WriteAllText(CompileTimingFile, JsonCompileTimes);
 			}
 
