Remove reference to DotNETUtilities from AddUBTAndUATFilesToBuildProducts().

#rnx

diff --git a/Engine/Source/Programs/AutomationTool/AutomationUtils/UE4Build.cs b/Engine/Source/Programs/AutomationTool/AutomationUtils/UE4Build.cs
--- a/Engine/Source/Programs/AutomationTool/AutomationUtils/UE4Build.cs
+++ b/Engine/Source/Programs/AutomationTool/AutomationUtils/UE4Build.cs
@@ -1291,7 +1291,7 @@ namespace AutomationTool
 						"AutomationTool.exe",
 						"UnrealBuildTool.exe",
 						"AutomationUtils.Automation.dll",
-						"DotNETUtilities.dll",
+						"EpicGames.Core.dll",
 						"BuildUtilities.dll",
 					});
 
