Include entire Engine/Binaries/DotNET/AutomationTool folder in UAT build products.

#rnx

diff --git a/Engine/Source/Programs/AutomationTool/AutomationUtils/UE4Build.cs b/Engine/Source/Programs/AutomationTool/AutomationUtils/UE4Build.cs
--- a/Engine/Source/Programs/AutomationTool/AutomationUtils/UE4Build.cs
+++ b/Engine/Source/Programs/AutomationTool/AutomationUtils/UE4Build.cs
@@ -1284,25 +1284,10 @@ namespace AutomationTool
 		public void AddUATFilesToBuildProducts()
 		{
 			// Find all DLLs (scripts and their dependencies)
-            var DotNetOutputLocation = CommandUtils.CombinePaths(CommandUtils.CmdEnv.LocalRoot, "Engine", "Binaries", "DotNET", "AutomationTool");
-
-			var UATFiles = new List<string>(new string[] 
-					{
-						"AutomationTool.exe",
-						"UnrealBuildTool.exe",
-						"AutomationUtils.Automation.dll",
-						"EpicGames.Core.dll",
-						"BuildUtilities.dll",
-					});
-
-			foreach (var UATFile in UATFiles)
+            DirectoryReference OutputDir = DirectoryReference.Combine(CommandUtils.RootDirectory, "Engine", "Binaries", "DotNET", "AutomationTool");
+			foreach (FileReference OutputFile in DirectoryReference.EnumerateFiles(OutputDir, "*", SearchOption.AllDirectories))
 			{
-				var OutputFile = CommandUtils.CombinePaths(DotNetOutputLocation, UATFile);
-				if (!CommandUtils.FileExists_NoExceptions(OutputFile))
-				{
-					throw new UE4BuildException("Cannot add UAT to the build products because {0} does not exist.", OutputFile);
-				}
-				AddBuildProduct(OutputFile);
+				AddBuildProduct(OutputFile.FullName);
 			}
 
 			// All scripts are expected to exist in DotNET/AutomationScripts subfolder.
