Add missing copyright

diff --git a/Engine/Source/Programs/AutomationTool/AutomationUtils/Json.cs b/Engine/Source/Programs/AutomationTool/AutomationUtils/Json.cs
--- a/Engine/Source/Programs/AutomationTool/AutomationUtils/Json.cs
+++ b/Engine/Source/Programs/AutomationTool/AutomationUtils/Json.cs
@@ -1,3 +1,5 @@
+// Copyright Epic Games, Inc. All Rights Reserved.
+
 using System;
 using System.Collections.Generic;
 using System.Text;
