Restore packages before building projects. Should fix IncludeTool errors in CIS.

#rnx

diff --git a/Engine/Source/Programs/AutomationTool/BuildGraph/Tasks/CsCompileTask.cs b/Engine/Source/Programs/AutomationTool/BuildGraph/Tasks/CsCompileTask.cs
--- a/Engine/Source/Programs/AutomationTool/BuildGraph/Tasks/CsCompileTask.cs
+++ b/Engine/Source/Programs/AutomationTool/BuildGraph/Tasks/CsCompileTask.cs
@@ -159,6 +159,7 @@ namespace AutomationTool.Tasks
 				{
 					Arguments.Add(String.Format("/target:{0}", CommandUtils.MakePathSafeToUseWithCommandLine(Parameters.Target)));
 				}
+				Arguments.Add("/restore");
 				Arguments.Add("/verbosity:minimal");
 				Arguments.Add("/nologo");
 				foreach(FileReference ProjectFile in ProjectFiles)
