#autosubmit

Rename of Gauntlet files from UE4Game to UnrealGame

diff --git a/Engine/Source/Programs/AutomationTool/Gauntlet/Unreal/Game/UE4Game.DefaultNode.cs b/Engine/Source/Programs/AutomationTool/Gauntlet/Unreal/Game/UE4Game.DefaultNode.cs
deleted file mode 100644
diff --git a/Engine/Source/Programs/AutomationTool/Gauntlet/Unreal/Game/UE4Game.TestConfig.cs b/Engine/Source/Programs/AutomationTool/Gauntlet/Unreal/Game/UE4Game.TestConfig.cs
deleted file mode 100644
diff --git a/Engine/Source/Programs/AutomationTool/Gauntlet/Unreal/Game/UnrealGame.DefaultNode.cs b/Engine/Source/Programs/AutomationTool/Gauntlet/Unreal/Game/UnrealGame.DefaultNode.cs
new file mode 100644
--- /dev/null
+++ b/Engine/Source/Programs/AutomationTool/Gauntlet/Unreal/Game/UnrealGame.DefaultNode.cs
@@ -0,0 +1,43 @@
+﻿// Copyright Epic Games, Inc. All Rights Reserved.
+
+using System;
+using System.Collections.Generic;
+using System.Linq;
+using System.Text;
+using System.Threading.Tasks;
+using Gauntlet;
+
+namespace UnrealGame
+{
+	/// <summary>
+	/// Implements a default node that should boot most projects and supports n number of clients
+	/// via -numclients=<n> and a server via -server
+	/// </summary>
+	public class DefaultTest : UnrealTestNode<UE4TestConfig>
+	{
+		public DefaultTest(UnrealTestContext InContext) : base(InContext)
+		{
+		}
+
+		public override UE4TestConfig GetConfiguration()
+		{
+			// just need a single client
+			UE4TestConfig Config = base.GetConfiguration();
+
+			int ClientCount = Context.TestParams.ParseValue("numclients", 1);
+			bool WithServer = Context.TestParams.ParseParam("server");
+
+			if (ClientCount > 0)
+			{
+				Config.RequireRoles(UnrealTargetRole.Client, ClientCount);
+			}
+
+			if (WithServer)
+			{
+				Config.RequireRoles(UnrealTargetRole.Server, 1);
+			}
+
+			return Config;
+		}
+	}
+}
diff --git a/Engine/Source/Programs/AutomationTool/Gauntlet/Unreal/Game/UnrealGame.TestConfig.cs b/Engine/Source/Programs/AutomationTool/Gauntlet/Unreal/Game/UnrealGame.TestConfig.cs
new file mode 100644
--- /dev/null
+++ b/Engine/Source/Programs/AutomationTool/Gauntlet/Unreal/Game/UnrealGame.TestConfig.cs
@@ -0,0 +1,53 @@
+﻿// Copyright Epic Games, Inc. All Rights Reserved.
+
+using System;
+using System.Collections.Generic;
+using System.Linq;
+using System.Text;
+using System.Threading.Tasks;
+using Gauntlet;
+
+namespace UnrealGame
+{
+	/// <summary>
+	/// Default set of options for testing "DefaultGame". Options that tests can configure
+	/// should be public, external command-line driven options should be protected/private
+	/// </summary>
+	public class UE4TestConfig : UnrealTestConfiguration
+	{
+		/// <summary>
+		/// Map to use
+		/// </summary>
+		[AutoParam]
+		public string Map = "";
+
+		/// <summary>
+		/// Applies these options to the provided app config
+		/// </summary>
+		/// <param name="AppConfig"></param>
+		public override void ApplyToConfig(UnrealAppConfig AppConfig, UnrealSessionRole ConfigRole, IEnumerable<UnrealSessionRole> OtherRoles)
+		{
+			base.ApplyToConfig(AppConfig, ConfigRole, OtherRoles);
+
+			if (string.IsNullOrEmpty(ConfigRole.MapOverride) == false)
+			{
+				// If map is specified, pass it to the server, or just the client if there is no server
+				if (AppConfig.ProcessType.IsServer()
+				|| (AppConfig.ProcessType.IsClient() && RoleCount(UnrealTargetRole.Server) == 0))
+				{
+					// must be the first argument!
+					AppConfig.CommandLineParams.GameMap = ConfigRole.MapOverride;
+				}
+			}
+			else if (string.IsNullOrEmpty(Map) == false)
+			{
+				// If map is specified, pass it to the server, or just the client if there is no server
+				if (AppConfig.ProcessType.IsServer()
+				|| (AppConfig.ProcessType.IsClient() && RoleCount(UnrealTargetRole.Server) == 0))
+				{
+					AppConfig.CommandLineParams.GameMap = Map;
+				}
+			}			
+		}
+	}
+}
