Fixed issue were ubt was being spawned with the wrong path and via mono when creating a installed engine build

diff --git a/Engine/Build/InstalledEngineBuild.xml b/Engine/Build/InstalledEngineBuild.xml
--- a/Engine/Build/InstalledEngineBuild.xml
+++ b/Engine/Build/InstalledEngineBuild.xml
@@ -441,7 +441,10 @@
 	
 	<Agent Name="BuildRules" Type="Win64_Licensee">
 		<Node Name="Compile Rules Assemblies" Produces="#Rules Assemblies">
-			<Spawn Exe="Engine/Binaries/DotNET/UnrealBuildTool/UnrealBuildTool.exe" Arguments="-Mode=QueryTargets"/>
+			<Property Name="UBTPath" Value="Engine/Binaries/DotNET/UnrealBuildTool/UnrealBuildTool.exe"/>
+			<Property Name="UBTPath" Value="Engine/Binaries/DotNET/UnrealBuildTool/UnrealBuildTool" If="'$(HostPlatform)' != 'Win64'"/>
+
+			<Spawn Exe="$(UBTPath)" Arguments="-Mode=QueryTargets"/>
 			<Tag Files="Engine/Intermediate/Build/BuildRules/..." With="#Rules Assemblies"/>
 		</Node>
 	</Agent>
diff --git a/Engine/Source/Programs/AutomationTool/AutomationUtils/LinuxHostPlatform.cs b/Engine/Source/Programs/AutomationTool/AutomationUtils/LinuxHostPlatform.cs
--- a/Engine/Source/Programs/AutomationTool/AutomationUtils/LinuxHostPlatform.cs
+++ b/Engine/Source/Programs/AutomationTool/AutomationUtils/LinuxHostPlatform.cs
@@ -117,6 +117,12 @@ namespace AutomationTool
 					AppName = AppName.Replace("-Cmd.exe", "");
 					AppName = AppName.Replace(".exe", "");
 				}
+				// some of our C# applications are converted to dotnet core, do not run those via mono
+				// they are instead assumed to produce a host executable that can just be run
+				else if (AppName.Contains("UnrealBuildTool") || AppName.Contains("AutomationTool"))
+				{
+					Options &= ~CommandUtils.ERunOptions.AppMustExist;
+				}
 				else
 				{
 					// It's a C# app, so run it with Mono
diff --git a/Engine/Source/Programs/AutomationTool/AutomationUtils/MacHostPlatform.cs b/Engine/Source/Programs/AutomationTool/AutomationUtils/MacHostPlatform.cs
--- a/Engine/Source/Programs/AutomationTool/AutomationUtils/MacHostPlatform.cs
+++ b/Engine/Source/Programs/AutomationTool/AutomationUtils/MacHostPlatform.cs
@@ -135,6 +135,12 @@ namespace AutomationTool
 						AppName = AppName + ".app/Contents/MacOS/" + AppFilename;
 					}
 				}
+				// some of our C# applications are converted to dotnet core, do not run those via mono
+				// they are instead assumed to produce a host executable that can just be run
+				else if (AppName.Contains("UnrealBuildTool") || AppName.Contains("AutomationTool"))
+				{
+					Options &= ~CommandUtils.ERunOptions.AppMustExist;
+				}
 				else
 				{
 					// It's a C# app, so run it with Mono
