- Fixed crash with getting SDK info on legacy NDA platforms when the Platforms directory doesn't exist

diff --git a/Engine/Source/Programs/AutomationTool/Turnkey/Commands/ListPlatforms.cs b/Engine/Source/Programs/AutomationTool/Turnkey/Commands/ListPlatforms.cs
--- a/Engine/Source/Programs/AutomationTool/Turnkey/Commands/ListPlatforms.cs
+++ b/Engine/Source/Programs/AutomationTool/Turnkey/Commands/ListPlatforms.cs
@@ -1,6 +1,5 @@
 // Copyright Epic Games, Inc. All Rights Reserved.
 
-using System;
 using System.Collections.Generic;
 using UnrealBuildTool;
 using AutomationTool;
@@ -13,22 +12,11 @@ namespace Turnkey.Commands
 		protected override void Execute(string[] CommandOptions)
 		{
 			TurnkeyUtils.Log("");
-			TurnkeyUtils.Log("Valid Platforms:");
-			string PlatformString = TurnkeyUtils.ParseParamValue("Platform", null, CommandOptions);
+			TurnkeyUtils.Log("Platform Information:");
 
-			foreach (UnrealTargetPlatform TargetPlatform in UnrealTargetPlatform.GetValidPlatforms())
+			List<UnrealTargetPlatform> ChosenPlatforms = TurnkeyUtils.GetPlatformsFromCommandLineOrUser(CommandOptions, null);
+			foreach (UnrealTargetPlatform TargetPlatform in ChosenPlatforms)
 			{
-				// HACK UNTIUL WIN32 IS GONE COMPLETELY
-				if (TargetPlatform == UnrealTargetPlatform.Win32 || TargetPlatform == UnrealTargetPlatform.XboxOne)
-				{
-					continue;
-				}
-
-				if (PlatformString != null && UnrealTargetPlatform.Parse(PlatformString) != TargetPlatform)
-				{
-					continue;
-				}
-
 				Platform Platform = Platform.Platforms[new TargetPlatformDescriptor(TargetPlatform)];
 				UEBuildPlatformSDK SDK = UEBuildPlatformSDK.GetSDKForPlatform(TargetPlatform.ToString());
 				if (SDK == null || !SDK.bIsSdkAllowedOnHost)
diff --git a/Engine/Source/Programs/AutomationTool/Turnkey/Commands/VerifySdk.cs b/Engine/Source/Programs/AutomationTool/Turnkey/Commands/VerifySdk.cs
--- a/Engine/Source/Programs/AutomationTool/Turnkey/Commands/VerifySdk.cs
+++ b/Engine/Source/Programs/AutomationTool/Turnkey/Commands/VerifySdk.cs
@@ -108,12 +108,15 @@ namespace Turnkey.Commands
 
 				SdkUtils.LocalAvailability LocalState = SdkUtils.GetLocalAvailability(AutomationPlatform, bUpdateIfNeeded);
 
-				string ManualSDKVersion, AutoSDKVersion;
-				string MinAllowedVersion, MaxAllowedVersion;
-				string MinSoftwareAllowedVersion, MaxSoftwareAllowedVersion;
-				PlatformSDK.GetInstalledVersions(out ManualSDKVersion, out AutoSDKVersion);
-				PlatformSDK.GetValidVersionRange(out MinAllowedVersion, out MaxAllowedVersion);
-				PlatformSDK.GetValidSoftwareVersionRange(out MinSoftwareAllowedVersion, out MaxSoftwareAllowedVersion);
+				string ManualSDKVersion = "", AutoSDKVersion = "";
+				string MinAllowedVersion = "", MaxAllowedVersion = "";
+				string MinSoftwareAllowedVersion = "", MaxSoftwareAllowedVersion = "";
+				if (PlatformSDK != null)
+				{
+					PlatformSDK.GetInstalledVersions(out ManualSDKVersion, out AutoSDKVersion);
+					PlatformSDK.GetValidVersionRange(out MinAllowedVersion, out MaxAllowedVersion);
+					PlatformSDK.GetValidSoftwareVersionRange(out MinSoftwareAllowedVersion, out MaxSoftwareAllowedVersion);
+				}
 
 				SdkUtils.LocalAvailability ReportedState = LocalState;
 				string StatusString;
@@ -137,6 +140,11 @@ namespace Turnkey.Commands
 				TurnkeyUtils.Report("{0}: (Status={1}, Installed={2}, AutoSDK={3}, MinAllowed={4}, MaxAllowed={5}, Flags=\"{6}\")", Platform, StatusString, ManualSDKVersion, AutoSDKVersion,
 					MinAllowedVersion, MaxAllowedVersion, ReportedState.ToString());
 
+				if (PlatformSDK == null)
+				{
+					continue;
+				}
+
 				// install if out of date, or if forcing it
 				if (bForceSdkInstall || (bUpdateIfNeeded && TurnkeyUtils.ExitCode != AutomationTool.ExitCode.Success))
 				{
diff --git a/Engine/Source/Programs/AutomationTool/Turnkey/SdkInfo.cs b/Engine/Source/Programs/AutomationTool/Turnkey/SdkInfo.cs
--- a/Engine/Source/Programs/AutomationTool/Turnkey/SdkInfo.cs
+++ b/Engine/Source/Programs/AutomationTool/Turnkey/SdkInfo.cs
@@ -40,6 +40,13 @@ namespace Turnkey
 		{
 			LocalAvailability Result = LocalAvailability.None;
 
+			// for some legacy NDA platforms, we could have an UnrealTargetPlatform but no registered SDK
+			UEBuildPlatformSDK SDK = UEBuildPlatformSDK.GetSDKForPlatform(AutomationPlatform.PlatformType.ToString());
+			if (SDK == null)
+			{
+				return Result;
+			}
+
 			// no need to retrieve anything if we aren't allowing updates
 			CopyProviderRetriever Retriever = new CopyProviderRetriever();
 
@@ -52,8 +59,6 @@ namespace Turnkey
 				Result |= LocalAvailability.Platform_InvalidHostPrerequisites;
 			}
 
-			UEBuildPlatformSDK SDK = UEBuildPlatformSDK.GetSDKForPlatform(AutomationPlatform.PlatformType.ToString());
-
 			string ManualSDKVersion, AutoSDKVersion;
 			SDK.GetInstalledVersions(out ManualSDKVersion, out AutoSDKVersion);
 
diff --git a/Engine/Source/Programs/AutomationTool/Turnkey/TurnkeyUtils.cs b/Engine/Source/Programs/AutomationTool/Turnkey/TurnkeyUtils.cs
--- a/Engine/Source/Programs/AutomationTool/Turnkey/TurnkeyUtils.cs
+++ b/Engine/Source/Programs/AutomationTool/Turnkey/TurnkeyUtils.cs
@@ -115,8 +115,8 @@ namespace Turnkey
 				PossiblePlatforms = UnrealTargetPlatform.GetValidPlatforms().ToList();
 			}
 
-			// HACK UNTIL WIN32 IS GONE
-			PossiblePlatforms = PossiblePlatforms.Where(x => x != UnrealTargetPlatform.Win32 && x != UnrealTargetPlatform.XboxOne).ToList();
+			// Remove known bad platforms and platforms that don't have a UEPlatformSDK
+			PossiblePlatforms = PossiblePlatforms.Where(x => x != UnrealTargetPlatform.Win32 && x != UnrealTargetPlatform.XboxOne && UEBuildPlatformSDK.GetSDKForPlatform(x.ToString()) != null).ToList();
 
 			// sort by name
 			PossiblePlatforms.Sort((x, y) => string.Compare(x.ToString(), y.ToString()));
