- Stopped returning SDKNotFound for bad SDKs in VerifySDK, because it's not an actual error case, it's normal operation. Trying to install an SDK and that failing can still return errors

diff --git a/Engine/Source/Programs/AutomationTool/Turnkey/Commands/VerifySdk.cs b/Engine/Source/Programs/AutomationTool/Turnkey/Commands/VerifySdk.cs
--- a/Engine/Source/Programs/AutomationTool/Turnkey/Commands/VerifySdk.cs
+++ b/Engine/Source/Programs/AutomationTool/Turnkey/Commands/VerifySdk.cs
@@ -120,6 +120,7 @@ namespace Turnkey.Commands
 
 				SdkUtils.LocalAvailability ReportedState = LocalState;
 				string StatusString;
+				bool bHasOutOfDateSDK = false;
 
 				if ((LocalState & SdkUtils.LocalAvailability.Platform_ValidHostPrerequisites) == 0)
 				{
@@ -128,12 +129,12 @@ namespace Turnkey.Commands
 				else if ((LocalState & (SdkUtils.LocalAvailability.AutoSdk_ValidVersionExists | SdkUtils.LocalAvailability.InstalledSdk_ValidVersionExists)) == 0)
 				{
 					StatusString = "Invalid";
-					TurnkeyUtils.ExitCode = AutomationTool.ExitCode.Error_SDKNotFound;
+					bHasOutOfDateSDK = true;
 				}
 				else
 				{
 					StatusString = "Valid";
-					ReportedState &= (SdkUtils.LocalAvailability.AutoSdk_ValidVersionExists | SdkUtils.LocalAvailability.InstalledSdk_ValidVersionExists | SdkUtils.LocalAvailability.Support_FullSdk| SdkUtils.LocalAvailability.Support_AutoSdk);
+					ReportedState &= (SdkUtils.LocalAvailability.AutoSdk_ValidVersionExists | SdkUtils.LocalAvailability.InstalledSdk_ValidVersionExists | SdkUtils.LocalAvailability.Support_FullSdk | SdkUtils.LocalAvailability.Support_AutoSdk);
 					//					TurnkeyUtils.Log("{0}: Valid [Installed: '{1}', Required: '{2}']", Platform, PlatformObject.GetInstalledSdk(), PlatformObject.GetAllowedSdks());
 				}
 
@@ -146,10 +147,8 @@ namespace Turnkey.Commands
 				}
 
 				// install if out of date, or if forcing it
-				if (bForceSdkInstall || (bUpdateIfNeeded && TurnkeyUtils.ExitCode != AutomationTool.ExitCode.Success))
+				if (bForceSdkInstall || (bUpdateIfNeeded && bHasOutOfDateSDK))
 				{
-					TurnkeyUtils.ExitCode = AutomationTool.ExitCode.Success;
-
 // 						if (!bUnattended)
 // 						{
 // 							string Response = TurnkeyUtils.ReadInput("Your Sdk installation is not up to date. Would you like to install a valid Sdk? [Y/n]", "Y");
@@ -245,8 +244,6 @@ namespace Turnkey.Commands
 
 						if (bForceDeviceInstall || !bIsSoftwareValid)
 						{
-							TurnkeyUtils.ExitCode = AutomationTool.ExitCode.Error_SDKNotFound;
-
 							if (bUpdateIfNeeded)
 							{
 								if (Device.bCanConnect)
@@ -261,7 +258,6 @@ namespace Turnkey.Commands
 									else
 									{
 										MatchingInstallableSdk.DownloadOrInstall(Platform, Device, bUnattended);
-										TurnkeyUtils.ExitCode = AutomationTool.ExitCode.Success;
 									}
 								}
 								else
