- Adding idea of "quick switching" installed SDKs if the platform can detect multiple installed SDKs

diff --git a/Engine/Source/Programs/AutomationTool/Turnkey/Commands/ListPlatforms.cs b/Engine/Source/Programs/AutomationTool/Turnkey/Commands/ListPlatforms.cs
--- a/Engine/Source/Programs/AutomationTool/Turnkey/Commands/ListPlatforms.cs
+++ b/Engine/Source/Programs/AutomationTool/Turnkey/Commands/ListPlatforms.cs
@@ -40,6 +40,9 @@ namespace Turnkey.Commands
 				TurnkeyUtils.Log("  Valid Manual SDK Installed? {0}", bIsManualSdkValid);
 				TurnkeyUtils.Log("  Valid Auto SDK Installed? {0}", bIsAutoSdkValid);
 
+				string[] AllVersions = SDK.GetAllInstalledSDKVersions();
+				TurnkeyUtils.Log("  AllVersions Installed: {0}", string.Join(",", AllVersions));
+
 				// look for available sdks
 				List<FileSource> MatchingFullSdks = TurnkeyManifest.FilterDiscoveredFileSources(TargetPlatform, FileSource.SourceType.Full);
 				if (MatchingFullSdks == null || MatchingFullSdks.Count == 0)
diff --git a/Engine/Source/Programs/AutomationTool/Turnkey/Commands/VerifySdk.cs b/Engine/Source/Programs/AutomationTool/Turnkey/Commands/VerifySdk.cs
--- a/Engine/Source/Programs/AutomationTool/Turnkey/Commands/VerifySdk.cs
+++ b/Engine/Source/Programs/AutomationTool/Turnkey/Commands/VerifySdk.cs
@@ -97,6 +97,8 @@ namespace Turnkey.Commands
 
 			CopyProviderRetriever Retriever = new CopyProviderRetriever();
 
+			TurnkeyUtils.StartTrackingExternalEnvVarChanges();
+
 			// check all the platforms
 			foreach (var Pair in PlatformsAndDevices)
 			{
@@ -158,6 +160,24 @@ namespace Turnkey.Commands
 // 							}
 // 						}
 
+					// if the platform has a valid sdk but isn't the "installed one", then try to switch to it
+					if ((LocalState & SdkUtils.LocalAvailability.InstalledSdk_ValidInactiveVersionExists) != 0)
+					{
+
+						// find the highest number that is valid (because a valid version exists, we know there will be at least one valid version)
+						string BestAlternateVersion = PlatformSDK.GetAllInstalledSDKVersions().ToList().OrderByDescending(x => x).Where(x => PlatformSDK.IsVersionValid(x, false)).First();
+
+						bool bWasSwitched = PlatformSDK.SwitchToAlternateSDK(BestAlternateVersion, false);
+
+						if (bWasSwitched == true)
+						{
+							TurnkeyUtils.Log("Fast-switched to already-installed version {0}", BestAlternateVersion);
+
+							// if SwitchToAlternateSDK returns true, then we are good to go!
+							continue;
+						}
+					}
+
 					FileSource BestSdk = null;
 					// find the best Sdk, prioritizing as request
 					if (bPreferFullSdk)
@@ -269,6 +289,8 @@ namespace Turnkey.Commands
 					}
 				}
 			}
+
+			TurnkeyUtils.EndTrackingExternalEnvVarChanges();
 		}
 	}
 }
diff --git a/Engine/Source/Programs/AutomationTool/Turnkey/SdkInfo.cs b/Engine/Source/Programs/AutomationTool/Turnkey/SdkInfo.cs
--- a/Engine/Source/Programs/AutomationTool/Turnkey/SdkInfo.cs
+++ b/Engine/Source/Programs/AutomationTool/Turnkey/SdkInfo.cs
@@ -21,9 +21,11 @@ namespace Turnkey
 			AutoSdk_VariableExists = 1,
 			AutoSdk_ValidVersionExists = 2,
 			AutoSdk_InvalidVersionExists = 4,
-			// InstalledSdk_BuildOnlyWasInstalled = 8,
+			
+			InstalledSdk_ValidInactiveVersionExists = 8,
 			InstalledSdk_ValidVersionExists = 16,
 			InstalledSdk_InvalidVersionExists = 32,
+
 			Platform_ValidHostPrerequisites = 64,
 			Platform_InvalidHostPrerequisites = 128,
 
@@ -110,6 +112,16 @@ namespace Turnkey
 			}
 
 
+			// look for other, inactive, versions
+			foreach (string AlternateVersion in SDK.GetAllInstalledSDKVersions())
+			{
+				if (SDK.IsVersionValid(AlternateVersion, bForAutoSDK: false))
+				{
+					Result |= LocalAvailability.InstalledSdk_ValidInactiveVersionExists;
+				}
+			}
+
+
 			string FullSupportedPlatforms = TurnkeyUtils.GetVariableValue("Studio_FullInstallPlatforms");
 			string AutoSdkSupportedPlatforms = TurnkeyUtils.GetVariableValue("Studio_AutoSdkPlatforms");
 			if (!string.IsNullOrEmpty(FullSupportedPlatforms))
diff --git a/Engine/Source/Programs/DotNETCommon/BuildUtilities/UEBuildPlatformSDK.cs b/Engine/Source/Programs/DotNETCommon/BuildUtilities/UEBuildPlatformSDK.cs
--- a/Engine/Source/Programs/DotNETCommon/BuildUtilities/UEBuildPlatformSDK.cs
+++ b/Engine/Source/Programs/DotNETCommon/BuildUtilities/UEBuildPlatformSDK.cs
@@ -171,6 +171,26 @@ namespace EpicGames.Core
 		/// <returns></returns>
 		public abstract string GetMainVersion();
 
+		/// <summary>
+		/// Return a list of full SDK versions that are already installed and can be quickly switched to without running any installers.
+		/// </summary>
+		/// <returns></returns>
+		public virtual string[] GetAllInstalledSDKVersions()
+		{
+			return new string[] { };
+		}
+
+		/// <summary>
+		/// Switch to another version of the SDK than what GetInstalledSDKVersion() returns. This will be one of the versions returned from GetAllInstalledSDKVersions()
+		/// </summary>
+		/// <param name="Version">String name of the version to switch to</param>
+		/// <param name="bSwitchForThisProcessOnly">If true, only switch for this process (usually via process environment variable). If false, switch permanently (usually via system-wide environment variable)</param>
+		/// <returns>True if successful</returns>
+		public virtual bool SwitchToAlternateSDK(string Version, bool bSwitchForThisProcessOnly)
+		{
+			return false;
+		}
+
 		/// <summary>
 		/// Gets the valid string range of Sdk versions. TryConvertVersionToInt() will need to succeed to make this usable for range checks
 		/// </summary>
