[Switch][Turnkey] Add support for full Switch SDK installs and add installers on google drive

#rnx

diff --git a/Engine/Source/Programs/AutomationTool/Turnkey/LocalCache.cs b/Engine/Source/Programs/AutomationTool/Turnkey/LocalCache.cs
--- a/Engine/Source/Programs/AutomationTool/Turnkey/LocalCache.cs
+++ b/Engine/Source/Programs/AutomationTool/Turnkey/LocalCache.cs
@@ -204,7 +204,7 @@ namespace Turnkey
 		{
 			TurnkeyUtils.Log("Loading cache bin {0}", TagCacheFile);
 			TagToDownloadCache = DeserializeObject(TagCacheFile);
-			TurnkeyUtils.Log("Loadied: {0}", TagToDownloadCache);
+			TurnkeyUtils.Log("Loaded: {0}", TagToDownloadCache);
 
 			if (TagToDownloadCache == null)
 			{
