Remove VS2017 as a supported compiler by default. Can still be overridden with the -2017 argument, but will not be maintained from this point onwards.

diff --git a/Engine/Build/InstalledEngineBuild.xml b/Engine/Build/InstalledEngineBuild.xml
--- a/Engine/Build/InstalledEngineBuild.xml
+++ b/Engine/Build/InstalledEngineBuild.xml
@@ -71,9 +71,7 @@
 	<!-- Generate full debug info for binary editor and packaged application builds -->
 	<Option Name="WithFullDebugInfo" Restrict="true|false" DefaultValue="false" Description="Generate full debug info for binary editor and packaged application builds"/>
 
-	<!-- Whether to use VS2019 -->
-	<Option Name="VS2019" Restrict="true|false" DefaultValue="false" Description="Use Visual Studio 2019 to build Windows targets. By default, Visual Studio 2017 is used for maximum compatibility"/>
-	
+
 	<!-- The local output directory -->
 	<Property Name="LocalInstalledDir" Value="$(RootDir)/LocalBuilds/Engine/Windows"/>
 	<Property Name="LocalInstalledDirMac" Value="$(RootDir)/LocalBuilds/Engine/Mac"/>
@@ -102,10 +100,6 @@
 	<Property Name="TargetDebugInfoArg" Value="-nodebuginfo"/>
 	<Property Name="TargetDebugInfoArg" Value="" If="$(WithFullDebugInfo)"/>
 
-	<!-- Compile flags for Windows targets -->
-	<Property Name="VSCompilerArg" Value="-2017"/>
-	<Property Name="VSCompilerArg" Value="-2019" If="$(VS2019)"/>
-
 	<!-- Names of the client and server targets to build -->
 	<Property Name="OptionalClientTarget" Value=""/>
 	<Property Name="OptionalClientTarget" Value="UnrealClient" If="$(WithClient)"/>
@@ -256,8 +250,8 @@
 		<Node Name="Compile UnrealGame Win64" Requires="Compile UnrealHeaderTool Win64" Produces="#UnrealGame Win64;#UnrealGame Win64 Unstripped;#UnrealGame Win64 Stripped;#UnrealGame Win64 Unsigned;#UnrealGame Win64 Signed">
 			<ForEach Name="Target" Values="UnrealGame;$(OptionalClientTarget);$(OptionalServerTarget)">
 				<ForEach Name="Configuration" Values="$(GameConfigurations)">
-					<Compile Target="$(Target)" Platform="Win64" Configuration="$(Configuration)" Tag="#UnrealGame Win64" Arguments="-precompile -allmodules -nolink $(VSCompilerArg) $(TargetDebugInfoArg)"/>
-					<Compile Target="$(Target)" Platform="Win64" Configuration="$(Configuration)" Tag="#UnrealGame Win64" Arguments="-precompile $(VSCompilerArg) $(TargetDebugInfoArg)" Clean="false"/>
+					<Compile Target="$(Target)" Platform="Win64" Configuration="$(Configuration)" Tag="#UnrealGame Win64" Arguments="-precompile -allmodules -nolink $(TargetDebugInfoArg)"/>
+					<Compile Target="$(Target)" Platform="Win64" Configuration="$(Configuration)" Tag="#UnrealGame Win64" Arguments="-precompile $(TargetDebugInfoArg)" Clean="false"/>
 				</ForEach>
 			</ForEach>
 
diff --git a/Engine/Source/Programs/AutomationTool/Scripts/BuildPluginCommand.Automation.cs b/Engine/Source/Programs/AutomationTool/Scripts/BuildPluginCommand.Automation.cs
--- a/Engine/Source/Programs/AutomationTool/Scripts/BuildPluginCommand.Automation.cs
+++ b/Engine/Source/Programs/AutomationTool/Scripts/BuildPluginCommand.Automation.cs
@@ -48,9 +48,6 @@ class BuildPlugin : BuildCommand
 		// Option for verifying that all include directive s
 		bool bStrictIncludes = ParseParam("StrictIncludes");
 
-		// Whether to use VS2019 for compiling all targets. By default, we currently use 2017 for compiling static libraries for maximum compatibility.
-		bool bVS2019 = ParseParam("VS2019");
-
 		// Make sure the packaging directory is valid
 		DirectoryReference PackageDir = new DirectoryReference(PackageParam);
 		if (PluginFile.IsUnderDirectory(PackageDir))
@@ -108,7 +105,7 @@ class BuildPlugin : BuildCommand
 		// Compile the plugin for all the target platforms
 		List<UnrealTargetPlatform> HostPlatforms = ParseParam("NoHostPlatform")? new List<UnrealTargetPlatform>() : new List<UnrealTargetPlatform> { BuildHostPlatform.Current.Platform };
 		List<UnrealTargetPlatform> TargetPlatforms = GetTargetPlatforms(this, BuildHostPlatform.Current.Platform);
-		FileReference[] BuildProducts = CompilePlugin(HostProjectFile, HostProjectPluginFile, Plugin, HostPlatforms, TargetPlatforms, AdditionalArgs.ToString(), bVS2019);
+		FileReference[] BuildProducts = CompilePlugin(HostProjectFile, HostProjectPluginFile, Plugin, HostPlatforms, TargetPlatforms, AdditionalArgs.ToString());
 
 		// Package up the final plugin data
 		PackagePlugin(HostProjectPluginFile, BuildProducts, PackageDir, ParseParam("unversioned"));
@@ -137,7 +134,7 @@ class BuildPlugin : BuildCommand
 		return FileReference.Combine(HostProjectPluginDir, PluginFile.GetFileName());
 	}
 
-	FileReference[] CompilePlugin(FileReference HostProjectFile, FileReference HostProjectPluginFile, PluginDescriptor Plugin, List<UnrealTargetPlatform> HostPlatforms, List<UnrealTargetPlatform> TargetPlatforms, string AdditionalArgs, bool bVS2019)
+	FileReference[] CompilePlugin(FileReference HostProjectFile, FileReference HostProjectPluginFile, PluginDescriptor Plugin, List<UnrealTargetPlatform> HostPlatforms, List<UnrealTargetPlatform> TargetPlatforms, string AdditionalArgs)
 	{
 		List<FileReference> ManifestFileNames = new List<FileReference>();
 
@@ -163,7 +160,7 @@ class BuildPlugin : BuildCommand
 			{
 				if(Plugin.SupportsTargetPlatform(TargetPlatform))
 				{
-					string AdditionalTargetArgs = AdditionalArgs + (bVS2019 ? "" : " -2017");
+					string AdditionalTargetArgs = AdditionalArgs;
 					CompilePluginWithUBT(HostProjectFile, HostProjectPluginFile, Plugin, "UnrealGame", TargetType.Game, TargetPlatform, UnrealTargetConfiguration.Development, ManifestFileNames, AdditionalTargetArgs);
 					CompilePluginWithUBT(HostProjectFile, HostProjectPluginFile, Plugin, "UnrealGame", TargetType.Game, TargetPlatform, UnrealTargetConfiguration.Shipping, ManifestFileNames, AdditionalTargetArgs);
 				}
diff --git a/Engine/Source/Programs/UnrealBuildTool/Platform/Windows/UEBuildWindows.cs b/Engine/Source/Programs/UnrealBuildTool/Platform/Windows/UEBuildWindows.cs
--- a/Engine/Source/Programs/UnrealBuildTool/Platform/Windows/UEBuildWindows.cs
+++ b/Engine/Source/Programs/UnrealBuildTool/Platform/Windows/UEBuildWindows.cs
@@ -811,9 +811,9 @@ namespace UnrealBuildTool
 			// Set the compiler version if necessary
 			if (Target.WindowsPlatform.Compiler == WindowsCompiler.Default)
 			{
-				if (Target.WindowsPlatform.StaticAnalyzer == WindowsStaticAnalyzer.PVSStudio && HasCompiler(WindowsCompiler.VisualStudio2017))
+				if (Target.WindowsPlatform.StaticAnalyzer == WindowsStaticAnalyzer.PVSStudio && HasCompiler(WindowsCompiler.VisualStudio2019))
 				{
-					Target.WindowsPlatform.Compiler = WindowsCompiler.VisualStudio2017;
+					Target.WindowsPlatform.Compiler = WindowsCompiler.VisualStudio2019;
 				}
 				else
 				{
@@ -907,10 +907,6 @@ namespace UnrealBuildTool
 					{
 						return WindowsCompiler.VisualStudio2019;
 					}
-					else if (Format == ProjectFileFormat.VisualStudio2017)
-					{
-						return WindowsCompiler.VisualStudio2017;
-					}
 				} 
 			}
 
@@ -923,10 +919,6 @@ namespace UnrealBuildTool
 				{
 					return WindowsCompiler.VisualStudio2019;
 				}
-				else if (ProjectFormat == VCProjectFileFormat.VisualStudio2017)
-				{
-					return WindowsCompiler.VisualStudio2017;
-				}
 			}
 
 			// Check the editor settings too
@@ -937,10 +929,6 @@ namespace UnrealBuildTool
 			    {
 				    return WindowsCompiler.VisualStudio2019;
 			    }
-				else if(PreferredAccessor == ProjectFileFormat.VisualStudio2017)
-			    {
-				    return WindowsCompiler.VisualStudio2017;
-			    }
 			}
 
 			// Second, default based on what's installed, test for 2015 first
@@ -948,10 +936,6 @@ namespace UnrealBuildTool
 			{
 				return WindowsCompiler.VisualStudio2019;
 			}
-			if (HasCompiler(WindowsCompiler.VisualStudio2017))
-			{
-				return WindowsCompiler.VisualStudio2017;
-			}
 
 			// If we do have a Visual Studio installation, but we're missing just the C++ parts, warn about that.
 			DirectoryReference? VSInstallDir;
@@ -959,13 +943,9 @@ namespace UnrealBuildTool
 			{
 				Log.TraceWarning("Visual Studio 2019 is installed, but is missing the C++ toolchain. Please verify that the \"VC++ 2019 toolset\" component is selected in the Visual Studio 2019 installation options.");
 			}
-			else if (TryGetVSInstallDir(WindowsCompiler.VisualStudio2017, out VSInstallDir))
-			{
-				Log.TraceWarning("Visual Studio 2017 is installed, but is missing the C++ toolchain. Please verify that the \"VC++ 2017 toolset\" component is selected in the Visual Studio 2017 installation options.");
-			}
 			else
 			{
-				Log.TraceWarning("No Visual C++ installation was found. Please download and install Visual Studio 2017 with C++ components.");
+				Log.TraceWarning("No Visual C++ installation was found. Please download and install Visual Studio 2019 with C++ components.");
 			}
 
 			// Finally, default to VS2019 anyway
diff --git a/Engine/Source/Runtime/Core/Public/Windows/WindowsPlatformCompilerSetup.h b/Engine/Source/Runtime/Core/Public/Windows/WindowsPlatformCompilerSetup.h
--- a/Engine/Source/Runtime/Core/Public/Windows/WindowsPlatformCompilerSetup.h
+++ b/Engine/Source/Runtime/Core/Public/Windows/WindowsPlatformCompilerSetup.h
@@ -14,12 +14,9 @@
 #endif
 
 //
-// We require at least Visual Studio 2015 Update 3 to compile
+// We require at least Visual Studio 2019 to compile
 //
-static_assert(_MSC_VER != 1900 || _MSC_FULL_VER >= 190024210, "Visual Studio 2015 Update 3 is required to compile on Windows (http://go.microsoft.com/fwlink/?LinkId=691129)");
-static_assert(_MSC_VER >= 1900, "Visual Studio 2015 or later is required to compile on Windows");
-static_assert(_MSC_VER < 1910 || _MSC_VER >= 1913, "Visual Studio 2017 version 15.6 is required to compile on Windows. Please install updates through the Visual Studio installer.");
-static_assert(_MSC_VER != 1914 && _MSC_VER != 1915, "Visual Studio 2017 versions 15.7 and 15.8 are known to have code generation bugs that affect UE4. Please update to version 15.9.");
+static_assert(_MSC_VER >= 1920, "Visual Studio 2019 is required to build Unreal Engine.");
 
 //
 // Manually enable all warnings as errors, except ones that are explicitly skipped.
