- Started work on "Control" command that can be used to execute basic functionality that is currently in C++ (poweron, reboot, etc). Work in progress
- Some Turnkey utilities cleanup/refactor for getting platforms/devices from commandline, and generic method for getting options from commandline or user
- Allow multiple Turnkey executions on a single RunUAT commandline (RunUAT Turnkey -command=Control -operation=Reboot Turnkey -command=Control -operation=Disconnect)

diff --git a/Engine/Source/Programs/AutomationTool/AutomationUtils/Platform.cs b/Engine/Source/Programs/AutomationTool/AutomationUtils/Platform.cs
--- a/Engine/Source/Programs/AutomationTool/AutomationUtils/Platform.cs
+++ b/Engine/Source/Programs/AutomationTool/AutomationUtils/Platform.cs
@@ -34,6 +34,12 @@ namespace AutomationTool
 		// is the device able to be connected to (this is more about able to flash SDK or run, not about matching SDK version)
 		// if false, any of the above fields are suspect, especually SoftwareVersion
 		public bool bCanConnect = true;
+
+		public virtual void Connect() { }
+		public virtual void Disconnect() { }
+		public virtual void PowerOn() { }
+		public virtual void PowerOff() { }
+		public virtual void Reboot() { }
 	}
 
 	/// <summary>
diff --git a/Engine/Source/Programs/AutomationTool/Turnkey/Commands/InstallBuild.cs b/Engine/Source/Programs/AutomationTool/Turnkey/Commands/InstallBuild.cs
--- a/Engine/Source/Programs/AutomationTool/Turnkey/Commands/InstallBuild.cs
+++ b/Engine/Source/Programs/AutomationTool/Turnkey/Commands/InstallBuild.cs
@@ -107,30 +107,38 @@ namespace Turnkey.Commands
 			Config.Build = Build;
 			Config.ProjectName = Project;
 
-			DeviceInfo TurnkeyDevice = TurnkeyUtils.GetDeviceFromCommandLineOrUser(CommandOptions, Platform);
+			Dictionary<UnrealTargetPlatform, List<DeviceInfo>> TurnkeyDevices = TurnkeyUtils.GetDevicesFromCommandLineOrUser(CommandOptions, Platform);
+			List<ITargetDevice> GauntletDevices = new List<ITargetDevice>();
 
-			ITargetDevice GauntletDevice;
-			if (TurnkeyDevice == null)
+			if (TurnkeyDevices == null)
 			{
+				// @todo turnkey: probably don't want this, just cancel?
 				IEnumerable<IDefaultDeviceSource> DeviceSources = Gauntlet.Utils.InterfaceHelpers.FindImplementations<IDefaultDeviceSource>();
-				GauntletDevice = DeviceSources.Where(S => S.CanSupportPlatform(Platform)).SelectMany(S => S.GetDefaultDevices()).FirstOrDefault();
+				GauntletDevices.AddRange(DeviceSources.Where(S => S.CanSupportPlatform(Platform)).SelectMany(S => S.GetDefaultDevices()));
 			}
 			else
 			{
 				IDeviceFactory Factory = Gauntlet.Utils.InterfaceHelpers.FindImplementations<IDeviceFactory>()
 					.Where(F => F.CanSupportPlatform(Platform))
 					.FirstOrDefault();
-				GauntletDevice = Factory.CreateDevice(TurnkeyDevice.Name, null);
+
+				GauntletDevices.AddRange(TurnkeyDevices[Platform].Select(x => Factory.CreateDevice(x.Name, null)));
 			}
 
-			if (GauntletDevice == null)
+			if (GauntletDevices.Count == 0 || GauntletDevices[0] == null)
 			{
 				TurnkeyUtils.Log("Could not find a device to install on!");
 				return;
 			}
 
-			IAppInstall Install = GauntletDevice.InstallApplication(Config);
-			GauntletDevice.Run(Install);
+			foreach (ITargetDevice GauntletDevice in GauntletDevices)
+			{
+				if (GauntletDevice != null)
+				{
+					IAppInstall Install = GauntletDevice.InstallApplication(Config);
+					GauntletDevice.Run(Install);
+				}
+			}
 		}
 	}
 }
diff --git a/Engine/Source/Programs/AutomationTool/Turnkey/Commands/ListSdks.cs b/Engine/Source/Programs/AutomationTool/Turnkey/Commands/ListSdks.cs
--- a/Engine/Source/Programs/AutomationTool/Turnkey/Commands/ListSdks.cs
+++ b/Engine/Source/Programs/AutomationTool/Turnkey/Commands/ListSdks.cs
@@ -16,10 +16,9 @@ namespace Turnkey.Commands
 			TurnkeyUtils.Log("Available Installers:");
 
 			string TypeString = TurnkeyUtils.ParseParamValue("Type", null, CommandOptions);
-			string PlatformString = TurnkeyUtils.ParseParamValue("Platform", null, CommandOptions);
 
+			List<UnrealTargetPlatform> Platforms = TurnkeyUtils.GetPlatformsFromCommandLineOrUser(CommandOptions, null);
 
-			UnrealTargetPlatform? OptionalPlatform = null;
 			FileSource.SourceType? OptionalType = null;
 
 			if (TypeString != null)
@@ -31,16 +30,15 @@ namespace Turnkey.Commands
 				}
 			}
 
-			if (PlatformString != null)
+			List<FileSource> Sdks;
+			if (Platforms == null)
 			{
-				UnrealTargetPlatform Platform;
-				if (UnrealTargetPlatform.TryParse(PlatformString, out Platform))
-				{
-					OptionalPlatform = Platform;
-				}
+				Sdks = TurnkeyManifest.FilterDiscoveredFileSources(null, OptionalType);
+			}
+			else
+			{
+				Sdks = Platforms.SelectMany(x => TurnkeyManifest.FilterDiscoveredFileSources(x, OptionalType)).ToList();
 			}
-
-			List<FileSource> Sdks = TurnkeyManifest.FilterDiscoveredFileSources(OptionalPlatform, OptionalType);
 
 			foreach (FileSource Sdk in Sdks)
 			{
diff --git a/Engine/Source/Programs/AutomationTool/Turnkey/Commands/VerifySdk.cs b/Engine/Source/Programs/AutomationTool/Turnkey/Commands/VerifySdk.cs
--- a/Engine/Source/Programs/AutomationTool/Turnkey/Commands/VerifySdk.cs
+++ b/Engine/Source/Programs/AutomationTool/Turnkey/Commands/VerifySdk.cs
@@ -21,73 +21,29 @@ namespace Turnkey.Commands
 			bool bForceDeviceInstall = TurnkeyUtils.ParseParam("ForceDeviceInstall", CommandOptions);
 			bool bUpdateIfNeeded = bForceSdkInstall || bForceDeviceInstall || TurnkeyUtils.ParseParam("UpdateIfNeeded", CommandOptions);
 
-			// track each platform to check, and 
-			Dictionary<UnrealTargetPlatform, List<string>> PlatformsAndDevices = null;
-
-			// look at any devices on the commandline, and see if they have platforms or not
-			string DeviceList = TurnkeyUtils.ParseParamValue("Device", null, CommandOptions);
-			List<string> SplitDeviceList = null;
-			if (DeviceList != null)
+			// track each platform to check, and if -device was specified, get the devices
+			string DeviceString = TurnkeyUtils.ParseParamValue("Device", null, CommandOptions);
+			Dictionary<UnrealTargetPlatform, List<DeviceInfo>> PlatformsAndDevices = null;
+			if (string.IsNullOrEmpty(DeviceString))
 			{
-				SplitDeviceList = DeviceList.Split("+".ToCharArray()).ToList();
-
-				// look if they have platform@ tags
-				bool bAnyHavePlatform = SplitDeviceList.Any(x => x.Contains("@"));
-				if (bAnyHavePlatform)
+				List<UnrealTargetPlatform> Platforms = TurnkeyUtils.GetPlatformsFromCommandLineOrUser(CommandOptions, null);
+				if (Platforms != null)
 				{
-					if (!SplitDeviceList.All(x => x.Contains("@")))
-					{
-						throw new AutomationException("If any device in -device has a platform indicator ('Platform@Device'), they must all have a platform indicator");
-					}
-
-					// now split it up for devices for each platform
-					foreach (string DeviceToken in SplitDeviceList)
+					PlatformsAndDevices = new Dictionary<UnrealTargetPlatform, List<DeviceInfo>>();
+					foreach (UnrealTargetPlatform Platform in Platforms)
 					{
-						string[] Tokens = DeviceToken.Split("@".ToCharArray(), StringSplitOptions.RemoveEmptyEntries);
-						if (Tokens.Length != 2)
-						{
-							throw new AutomationException("{0} did not have the Platform@Device format", DeviceToken);
-						}
-						UnrealTargetPlatform Platform;
-						if (!UnrealTargetPlatform.TryParse(Tokens[0], out Platform))
-						{
-							TurnkeyUtils.Log("Platform indicator {0} is an invalid platform, skipping", Tokens[0]);
-							continue;
-						}
-
-						string DeviceName = Tokens[1];
-
-						// track it
-						if (PlatformsAndDevices == null)
-						{
-							PlatformsAndDevices = new Dictionary<UnrealTargetPlatform, List<string>>();
-						}
-						if (!PlatformsAndDevices.ContainsKey(Platform))
-						{
-							PlatformsAndDevices[Platform] = new List<string>();
-						}
-						PlatformsAndDevices[Platform].Add(DeviceName);
+						PlatformsAndDevices[Platform] = new List<DeviceInfo>();
 					}
-					SplitDeviceList = null;
 				}
 			}
-
-			// if we didn't get some platforms already from -device list, then get or ask the user for platforms
-			if (PlatformsAndDevices == null)
+			else
 			{
-				PlatformsAndDevices = new Dictionary<UnrealTargetPlatform, List<string>>();
-				List<UnrealTargetPlatform> ChosenPlatforms = TurnkeyUtils.GetPlatformsFromCommandLineOrUser(CommandOptions, UnrealTargetPlatform.GetValidPlatforms().ToList());
-				ChosenPlatforms.ForEach(x => PlatformsAndDevices.Add(x, null));
-
-				if (ChosenPlatforms.Count > 1 && SplitDeviceList != null && !(SplitDeviceList.Count == 1 && SplitDeviceList[0].CompareTo("All") == 0))
-				{
-					throw new AutomationException("When passing -Device to VerifySdk without platform specifiers ('Platform:Device'), a single platform must be specified (unless -Device=All is used)");
-				}
+				PlatformsAndDevices = TurnkeyUtils.GetDevicesFromCommandLineOrUser(CommandOptions, null);
 			}
 
-			if (PlatformsAndDevices.Count == 0)
+			if (PlatformsAndDevices == null)
 			{
-				TurnkeyUtils.Log("Platform(s) needed for VerifySdk command. Ending command.");
+				TurnkeyUtils.Log("Platform(s) and/or device(s) needed for VerifySdk command. Check parameters or selections.");
 				return;
 			}
 
@@ -212,79 +168,59 @@ namespace Turnkey.Commands
 					// @todo turnkey: validate!
 				}
 
-				// use the per-platform device list, unless it's not specifed, then use the global device list (which is set when not using platform specifiers)
-				List<string> DeviceNames = Pair.Value != null ? Pair.Value : SplitDeviceList;
-				if (DeviceNames != null && DeviceNames.Count > 0)
+				// now check software verison of each device
+				foreach (DeviceInfo Device in Pair.Value)
 				{
-					DeviceInfo[] Devices = AutomationPlatform.GetDevices();
-					if (Devices == null)
+					bool bArePrerequisitesValid = AutomationPlatform.UpdateDevicePrerequisites(Device, TurnkeyUtils.CommandUtilHelper, Retriever, !bUpdateIfNeeded);
+					bool bIsSoftwareValid = PlatformSDK.IsSoftwareVersionValid(Device.SoftwareVersion);
+
+					SdkUtils.LocalAvailability DeviceState = SdkUtils.LocalAvailability.None;
+					if (!bArePrerequisitesValid)
 					{
-						TurnkeyUtils.Log("Platform {0} didn't have any devices, ignoring any devices specified", Platform);
-						continue;
+						StatusString = "Invalid";
+						DeviceState |= SdkUtils.LocalAvailability.Device_InvalidPrerequisites;
 					}
-
-					TurnkeyUtils.Log("Installed Device validity:");
-
-					// a single device named all means all devices
-					if (!(DeviceNames.Count == 1 && DeviceNames[0].ToLower() == "all"))
+					else if (bIsSoftwareValid)
 					{
-						Devices = Devices.Where(x => DeviceNames.Contains(x.Name, StringComparer.OrdinalIgnoreCase)).ToArray();
+						StatusString = "Valid";
+						DeviceState |= SdkUtils.LocalAvailability.Device_InstallSoftwareValid;
 					}
-
-					// now check software verison of each device
-					foreach (DeviceInfo Device in Devices)
+					else
 					{
-						bool bArePrerequisitesValid = AutomationPlatform.UpdateDevicePrerequisites(Device, TurnkeyUtils.CommandUtilHelper, Retriever, !bUpdateIfNeeded);
-						bool bIsSoftwareValid = PlatformSDK.IsSoftwareVersionValid(Device.SoftwareVersion);
-
-						SdkUtils.LocalAvailability DeviceState = SdkUtils.LocalAvailability.None;
-						if (!bArePrerequisitesValid)
-						{
-							StatusString = "Invalid";
-							DeviceState |= SdkUtils.LocalAvailability.Device_InvalidPrerequisites;
-						}
-						else if (bIsSoftwareValid)
-						{
-							StatusString = "Valid";
-							DeviceState |= SdkUtils.LocalAvailability.Device_InstallSoftwareValid;
-						}
-						else
-						{
-							StatusString = "Invalid";
-							DeviceState |= SdkUtils.LocalAvailability.Device_InstallSoftwareInvalid;
-						}
+						StatusString = "Invalid";
+						DeviceState |= SdkUtils.LocalAvailability.Device_InstallSoftwareInvalid;
+					}
 
-						if (Device.bCanConnect== false)
-						{
-							DeviceState |= SdkUtils.LocalAvailability.Device_CannotConnect;
-						}
+					if (Device.bCanConnect== false)
+					{
+						DeviceState |= SdkUtils.LocalAvailability.Device_CannotConnect;
+					}
 
-						TurnkeyUtils.Report("{0}@{1}: (Status={2}, Installed={3}, MinAllowed={4}, MaxAllowed={5}, Flags=\"{6}\")", Platform, Device.Name, StatusString, Device.SoftwareVersion, 
-							MinSoftwareAllowedVersion, MaxSoftwareAllowedVersion, DeviceState.ToString());
+					TurnkeyUtils.Report("{0}@{1}: (Status={2}, Installed={3}, MinAllowed={4}, MaxAllowed={5}, Flags=\"{6}\")", Platform, Device.Name, StatusString, Device.SoftwareVersion, 
+						MinSoftwareAllowedVersion, MaxSoftwareAllowedVersion, DeviceState.ToString());
 
-						if (bForceDeviceInstall || !bIsSoftwareValid)
+					if (bForceDeviceInstall || !bIsSoftwareValid)
+					{
+						if (bUpdateIfNeeded)
 						{
-							if (bUpdateIfNeeded)
+							if (Device.bCanConnect)
 							{
-								if (Device.bCanConnect)
+								FileSource MatchingInstallableSdk = FileSource.FindMatchingSdk(AutomationPlatform, new FileSource.SourceType[] { FileSource.SourceType.Flash }, bSelectBest: bUnattended, DeviceType: Device.Type);
+
+								if (MatchingInstallableSdk == null)
 								{
-									FileSource MatchingInstallableSdk = FileSource.FindMatchingSdk(AutomationPlatform, new FileSource.SourceType[] { FileSource.SourceType.Flash }, bSelectBest: bUnattended, DeviceType: Device.Type);
-
-									if (MatchingInstallableSdk == null)
-									{
-										TurnkeyUtils.Log("ERROR: {0}: Unable top find any Sdks that could be installed on {1}", Platform, Device.Name);
-										TurnkeyUtils.ExitCode = AutomationTool.ExitCode.Error_SDKNotFound;
-									}
-									else
-									{
-										MatchingInstallableSdk.DownloadOrInstall(Platform, Device, bUnattended);
-									}
+									TurnkeyUtils.Log("ERROR: {0}: Unable top find any Sdks that could be installed on {1}", Platform, Device.Name);
+									TurnkeyUtils.ExitCode = AutomationTool.ExitCode.Error_SDKNotFound;
 								}
 								else
 								{
-									TurnkeyUtils.Log("Skipping device {0} because it cannot connect.", Device.Name);
+									MatchingInstallableSdk.DownloadOrInstall(Platform, Device, bUnattended);
 								}
 							}
+							else
+							{
+								TurnkeyUtils.Log("Skipping device {0} because it cannot connect.", Device.Name);
+							}
 						}
 					}
 				}
diff --git a/Engine/Source/Programs/AutomationTool/Turnkey/Turnkey.Automation.cs b/Engine/Source/Programs/AutomationTool/Turnkey/Turnkey.Automation.cs
--- a/Engine/Source/Programs/AutomationTool/Turnkey/Turnkey.Automation.cs
+++ b/Engine/Source/Programs/AutomationTool/Turnkey/Turnkey.Automation.cs
@@ -34,29 +34,37 @@ namespace Turnkey
 			return Turnkey.Execute(IOProvider, this);
 		}
 
-
+		static bool bHasBeenInitialized = false;
 		public static AutomationTool.ExitCode Execute(IOProvider IOProvider, BuildCommand CommandUtilHelper)
 		{
-			SetupVisuals();
+			if (!bHasBeenInitialized)
+			{
+				SetupVisuals();
 
-			// cache some settings for other classes
-			TurnkeyUtils.SetVariable("EngineDir", EngineDirectory.FullName);
-			TurnkeyUtils.SetVariable("Project", CommandUtilHelper.ParseParamValue("Project="));
-			
-			TurnkeyUtils.Initialize(IOProvider, CommandUtilHelper);
-			TurnkeySettings.Initialize();
+				// cache some settings for other classes
+				TurnkeyUtils.SetVariable("EngineDir", EngineDirectory.FullName);
+				TurnkeyUtils.SetVariable("Project", CommandUtilHelper.ParseParamValue("Project="));
 
-			//			Microsoft.Win32.SystemEvents.UserPreferenceChanged += new Microsoft.Win32.UserPreferenceChangedEventHandler((sender, args) => TurnkeyUtils.Log("Got a change! {0}", args.Category));
+				TurnkeySettings.Initialize();
 
-			Console.CancelKeyPress += delegate
-			{
-				TurnkeyUtils.CleanupPaths();
+				//			Microsoft.Win32.SystemEvents.UserPreferenceChanged += new Microsoft.Win32.UserPreferenceChangedEventHandler((sender, args) => TurnkeyUtils.Log("Got a change! {0}", args.Category));
+
+				Console.CancelKeyPress += delegate
+				{
+					TurnkeyUtils.CleanupPaths();
+
+					TurnkeyUtils.Log("");
+					TurnkeyUtils.Log("If you installed an SDK, you should NOT \"Terminate batch job\"!");
+					TurnkeyUtils.Log("");
+					TurnkeyUtils.ExitCode = ExitCode.Success;
+				};
+
+				bHasBeenInitialized = true;
+			}
+
+			// IOProvider could change between multiple executions
+			TurnkeyUtils.Initialize(IOProvider, CommandUtilHelper);
 
-				TurnkeyUtils.Log("");
-				TurnkeyUtils.Log("If you installed an SDK, you should NOT \"Terminate batch job\"!");
-				TurnkeyUtils.Log("");
-				TurnkeyUtils.ExitCode = ExitCode.Success;
-			};
 
 			try
 			{
diff --git a/Engine/Source/Programs/AutomationTool/Turnkey/TurnkeyUtils.cs b/Engine/Source/Programs/AutomationTool/Turnkey/TurnkeyUtils.cs
--- a/Engine/Source/Programs/AutomationTool/Turnkey/TurnkeyUtils.cs
+++ b/Engine/Source/Programs/AutomationTool/Turnkey/TurnkeyUtils.cs
@@ -34,8 +34,8 @@ namespace Turnkey
 			IEnumerable<string> AutoSDKPlatforms = UEBuildPlatformSDK.AllSDKs.Select(x => x.GetAutoSDKPlatformName()).Distinct();
 			SetVariable("AutoSDKPlatforms", string.Join(",", AutoSDKPlatforms));
 
-// 			TurnkeyUtils.Log("AllPlatforms = {0}", GetVariableValue("AllPlatforms"));
-// 			TurnkeyUtils.Log("AutoSDKPlatforms = {0}", GetVariableValue("AutoSDKPlatforms"));
+			// 			TurnkeyUtils.Log("AllPlatforms = {0}", GetVariableValue("AllPlatforms"));
+			// 			TurnkeyUtils.Log("AutoSDKPlatforms = {0}", GetVariableValue("AutoSDKPlatforms"));
 
 			SetVariable("HOST_PLATFORM_NAME", HostPlatform.Current.HostEditorPlatform.ToString());
 		}
@@ -70,7 +70,7 @@ namespace Turnkey
 			return TurnkeyVariables.ContainsKey(Key);
 		}
 
-		public static string ExpandVariables(string Str, bool bUseOnlyTurnkeyVariables=false)
+		public static string ExpandVariables(string Str, bool bUseOnlyTurnkeyVariables = false)
 		{
 			// don't crash on null
 			if (Str == null)
@@ -95,7 +95,7 @@ namespace Turnkey
 		public static string ParseParamValue(string Param, string Default, string[] ExtraOptions)
 		{
 			// our internal extraoptions still have - in front, but CommandUtilHelper won't have the dashes
-			string Value = CommandUtils.ParseParamValue(ExtraOptions, "-" + Param, Default); 
+			string Value = CommandUtils.ParseParamValue(ExtraOptions, "-" + Param, Default);
 			if (Value == null)
 			{
 				Value = CommandUtilHelper.ParseParamValue(Param, Default);
@@ -104,19 +104,23 @@ namespace Turnkey
 			return Value == null ? Default : Value;
 		}
 
+		private static List<UnrealTargetPlatform> GetAllValidPlatforms(List<UnrealTargetPlatform> SourcePlatforms=null)
+		{
+			if (SourcePlatforms == null)
+			{
+				SourcePlatforms = UnrealTargetPlatform.GetValidPlatforms().ToList();
+			}
+
+			return SourcePlatforms.Where(x => x != UnrealTargetPlatform.Win32 && x != UnrealTargetPlatform.XboxOne && UEBuildPlatformSDK.GetSDKForPlatform(x.ToString()) != null).ToList();
+		}
+
 		public static List<UnrealTargetPlatform> GetPlatformsFromCommandLineOrUser(string[] CommandOptions, List<UnrealTargetPlatform> PossiblePlatforms)
 		{
 			string PlatformString = TurnkeyUtils.ParseParamValue("Platform", null, CommandOptions);
 			bool bUnattended = TurnkeyUtils.ParseParam("Unattended", CommandOptions);
 
-			// use default set of platforms if none specified
-			if (PossiblePlatforms == null)
-			{
-				PossiblePlatforms = UnrealTargetPlatform.GetValidPlatforms().ToList();
-			}
-
-			// Remove known bad platforms and platforms that don't have a UEPlatformSDK
-			PossiblePlatforms = PossiblePlatforms.Where(x => x != UnrealTargetPlatform.Win32 && x != UnrealTargetPlatform.XboxOne && UEBuildPlatformSDK.GetSDKForPlatform(x.ToString()) != null).ToList();
+			// Remove known bad platforms
+			PossiblePlatforms = GetAllValidPlatforms(PossiblePlatforms);
 
 			// sort by name
 			PossiblePlatforms.Sort((x, y) => string.Compare(x.ToString(), y.ToString()));
@@ -376,44 +380,235 @@ namespace Turnkey
 			return ProjectUtils.FindProjectFileFromName(Project);
 		}
 
-		public static DeviceInfo GetDeviceFromCommandLineOrUser(string[] CommandOptions, UnrealTargetPlatform Platform)
+		private static DeviceInfo GetDeviceByPlatformAndName(UnrealTargetPlatform Platform, string DeviceName)
+		{
+			DeviceInfo[] Devices = AutomationTool.Platform.GetPlatform(Platform).GetDevices();
+			if (Devices == null)
+			{
+				return null;
+			}
+			return Array.Find(Devices, x => string.Compare(x.Name, DeviceName, true) == 0);
+		}
+		public static Dictionary<UnrealTargetPlatform, List<DeviceInfo>> GetDevicesFromCommandLineOrUser(string[] CommandOptions, UnrealTargetPlatform Platform)
 		{
-			string DeviceName = TurnkeyUtils.ParseParamValue("Device", null, CommandOptions);
+			return GetDevicesFromCommandLineOrUser(CommandOptions, new List<UnrealTargetPlatform>() { Platform });
+		}
+		//public static DeviceInfo GetDeviceFromCommandLineOrUser(string[] CommandOptions, UnrealTargetPlatform Platform)
+		//{
+		//	Dictionary<UnrealTargetPlatform, List<DeviceInfo>> PlatformsAndDevices = GetDevicesFromCommandLineOrUser(CommandOptions, Platform);
+		//}
+
+		public static Dictionary<UnrealTargetPlatform, List<DeviceInfo>> GetDevicesFromCommandLineOrUser(string[] CommandOptions, List<UnrealTargetPlatform> PossiblePlatforms)
+		{
+			Dictionary<UnrealTargetPlatform, List<DeviceInfo>> PlatformsAndDevices = null;
+
+			// look at any devices on the commandline, and see if they have platforms or not
+			string DeviceList = TurnkeyUtils.ParseParamValue("Device", null, CommandOptions);
+			List<string> SplitDeviceList = null;
+			if (DeviceList != null && DeviceList.ToLower() != "all")
+			{
+				SplitDeviceList = DeviceList.Split("+".ToCharArray()).ToList();
+
+				// look if they have platform@ tags
+				bool bAnyHavePlatform = SplitDeviceList.Any(x => x.Contains("@"));
+				if (bAnyHavePlatform)
+				{
+					if (!SplitDeviceList.All(x => x.Contains("@")))
+					{
+						throw new AutomationException("If any device in -device has a platform indicator ('Platform@Device'), they must all have a platform indicator");
+					}
+
+					// now split it up for devices for each platform
+					foreach (string DeviceToken in SplitDeviceList)
+					{
+						string[] Tokens = DeviceToken.Split("@".ToCharArray(), StringSplitOptions.RemoveEmptyEntries);
+						if (Tokens.Length != 2)
+						{
+							throw new AutomationException("{0} did not have the Platform@Device format", DeviceToken);
+						}
+						UnrealTargetPlatform Platform;
+						if (!UnrealTargetPlatform.TryParse(Tokens[0], out Platform))
+						{
+							TurnkeyUtils.Log("Platform indicator {0} is an invalid platform, skipping", Tokens[0]);
+							continue;
+						}
 
-			AutomationTool.Platform AutomationPlatform = AutomationTool.Platform.Platforms[new TargetPlatformDescriptor(Platform)];
-			
-			if (DeviceName == null)
+						string DeviceName = Tokens[1];
+
+						// track it
+						if (PlatformsAndDevices == null)
+						{
+							PlatformsAndDevices = new Dictionary<UnrealTargetPlatform, List<DeviceInfo>>();
+						}
+						if (!PlatformsAndDevices.ContainsKey(Platform))
+						{
+							PlatformsAndDevices[Platform] = new List<DeviceInfo>();
+						}
+						
+						if (DeviceName.ToLower() == "all")
+						{
+							PlatformsAndDevices[Platform].AddRange(AutomationTool.Platform.GetPlatform(Platform).GetDevices());
+						}
+						else
+						{
+							DeviceInfo Device = GetDeviceByPlatformAndName(Platform, DeviceName);
+							if (Device != null)
+							{
+								PlatformsAndDevices[Platform].Add(Device);
+							}
+						}
+					}
+					SplitDeviceList = null;
+				}
+			}
+
+			// if we didn't get some platforms already from -device list, then get or ask the user for platforms
+			if (PlatformsAndDevices == null)
 			{
-				List<string> Options = new List<string>();
-				// @todo turnkey: filter devices that have valid flash versions
-				//				DeviceInfo[] PossibleDevices = Array.FindAll(AutomationPlatform.GetDevices(), x => TurnkeyUtils.IsValueValid(x.SoftwareVersion, AutomationPlatform.GetAllowedSoftwareVersions(), AutomationPlatforms[Platform]));
-				DeviceInfo[] PossibleDevices = AutomationPlatform.GetDevices();
-				if (PossibleDevices == null)
+				PlatformsAndDevices = new Dictionary<UnrealTargetPlatform, List<DeviceInfo>>();
+				List<UnrealTargetPlatform> ChosenPlatforms;
+
+				// use all platforms (with -device=all), or ask user if needed (GetPlatformsFromCommandLineOrUser would look at -platform=all, not -device=all)
+				// if -platform was specified, use the function to get just the specified ones
+				if (DeviceList != null && DeviceList.ToLower() == "all" && TurnkeyUtils.ParseParamValue("Platform", null, CommandOptions) == null)
+				{
+					ChosenPlatforms = GetAllValidPlatforms(PossiblePlatforms);
+				}
+				else
+				{
+					ChosenPlatforms = TurnkeyUtils.GetPlatformsFromCommandLineOrUser(CommandOptions, PossiblePlatforms);
+				}
+
+				if (ChosenPlatforms == null)
 				{
 					return null;
 				}
 
-				foreach (DeviceInfo Device in PossibleDevices)
+				// set up PlatformsAndDevices for all platforms
+				ChosenPlatforms.ForEach(x => PlatformsAndDevices.Add(x, new List<DeviceInfo>()));
+
+				if (ChosenPlatforms.Count > 1 && SplitDeviceList != null && !(SplitDeviceList.Count == 1 && SplitDeviceList[0].ToLower() == "all"))
 				{
-					Options.Add(string.Format("[{0} {1}] {2}", Platform, Device.Type, Device.Name));
+					throw new AutomationException("When using -Device without platform specifiers ('Platform@Device'), a single platform must be specified (unless -Device=All is used)");
 				}
 
-				// get the choice
-				int Choice = TurnkeyUtils.ReadInputInt("Select the number of a device to flash:", Options, true);
+				// get all the devices for the platforms
+				if (!string.IsNullOrEmpty(DeviceList) && DeviceList.ToLower() == "all")
+				{
+					foreach (UnrealTargetPlatform Platform in ChosenPlatforms)
+					{
+						DeviceInfo[] Devices = AutomationTool.Platform.GetPlatform(Platform).GetDevices();
+						if (Devices != null)
+						{
+							PlatformsAndDevices[Platform].AddRange(Devices);
+						}
+					}
+				}
+				// now if the list of devices was given, then attempt to find them in the platform
+				else if (SplitDeviceList != null && SplitDeviceList.Count > 0)
+				{
+					foreach (UnrealTargetPlatform Platform in ChosenPlatforms)
+					{
+						foreach (string DeviceName in SplitDeviceList)
+						{
+							DeviceInfo Device = GetDeviceByPlatformAndName(Platform, DeviceName);
+							if (Device != null)
+							{
+								PlatformsAndDevices[Platform].Add(Device);
+							}
+						}
+					}
+				}
+				// otherwise ask user for device
+				else
+				{
+					List<string> Options = new List<string>();
+					List<DeviceInfo> PossibleDevices = new List<DeviceInfo>();
+					List<UnrealTargetPlatform> PossibleDevicePlatforms = new List<UnrealTargetPlatform>();
+
+					foreach (UnrealTargetPlatform Platform in ChosenPlatforms)
+					{
+						Platform AutomationPlatform = AutomationTool.Platform.GetPlatform(Platform);
+						DeviceInfo[] Devices = AutomationPlatform.GetDevices();
+						if (Devices != null)
+						{
+							foreach (DeviceInfo Device in Devices)
+							{
+								PossibleDevices.Add(Device);
+								PossibleDevicePlatforms.Add(Platform);
+
+								Options.Add(string.Format("[{0} {1}] {2}", Platform, Device.Type, Device.Name));
+							}
+						}
+					}
+
+					if (PossibleDevices.Count == 0)
+					{
+						Log("Unable to find any devices for platform(s): {0}", string.Join(", ", ChosenPlatforms));
+						return null;
+					}
+
+					// get the choice
+					int Choice = TurnkeyUtils.ReadInputInt("Select a device:", Options, true);
 
+					if (Choice == 0)
+					{
+						return null;
+					}
+
+					// finally, add it to the proper list
+					PlatformsAndDevices[PossibleDevicePlatforms[Choice - 1]].Add(PossibleDevices[Choice - 1]);
+				}
+			}
+
+			// if we ended up with some platforms, but no devices, just return null
+			foreach (var Pair in PlatformsAndDevices)
+			{
+				if (Pair.Value != null && Pair.Value.Count > 0)
+				{
+					return PlatformsAndDevices;
+				}
+			}
+			return null;
+		}
+
+		public static string GetGenericOption(string[] CommandOptions, List<string> Options, string CommandLineOption)
+		{
+			return GetGenericOption(CommandOptions, Options, CommandLineOption, out _);
+		}
+
+		public static string GetGenericOption(string[] CommandOptions, List<string> Options, string CommandLineOption, out bool bWasOnCommandLine)
+		{
+			string ChosenValue = TurnkeyUtils.ParseParamValue(CommandLineOption, null, CommandOptions);
+
+			if (ChosenValue != null)
+			{
+				bWasOnCommandLine = true;
+			}
+			else
+			{
+				bWasOnCommandLine = false;
+
+				// default to previous selection if any
+				string CachedOptionName = "User_LastSelectedGeneric_" + CommandLineOption;
+				string LastSelectedOption = TurnkeySettings.GetUserSettingIfSet(CachedOptionName, "");
+				int Default = Options.FindIndex(x => x.Equals(LastSelectedOption, StringComparison.OrdinalIgnoreCase));
+
+				int Choice = ReadInputInt($"Choose the {CommandLineOption}:", Options, true, Default == -1 ? -1 : Default + 1);
 				if (Choice == 0)
 				{
 					return null;
 				}
 
-				// get the name of the device chosen
-				DeviceName = PossibleDevices[Choice - 1].Name;
+				ChosenValue = Options[Choice - 1];
+
+				// remember for next time
+				TurnkeySettings.SetUserSetting(CachedOptionName, ChosenValue);
 			}
 
-			// get device info of the chosen or supplied device
-			DeviceInfo InstallDevice = Array.Find(AutomationPlatform.GetDevices(), x => string.Compare(x.Name, DeviceName, true) == 0);
-			return InstallDevice;
+			return ChosenValue;
 		}
+
 #endregion
 
 #region Env vars
