[Turnkey] Add the ability to print a message before installing and SDK.  Add 'PauseForUser' action to facilitate displaying a message (without capturing input).  It displays a prompt to the user and waits for acknowledgement.

#rnx

diff --git a/Engine/Source/Programs/AutomationTool/AutomationUtils/Platform.cs b/Engine/Source/Programs/AutomationTool/AutomationUtils/Platform.cs
--- a/Engine/Source/Programs/AutomationTool/AutomationUtils/Platform.cs
+++ b/Engine/Source/Programs/AutomationTool/AutomationUtils/Platform.cs
@@ -17,6 +17,7 @@ namespace AutomationTool
 		string RetrieveFileSource(object HintObject);
 		string GetVariable(string VariableName);
 		bool RunExternalCommand(string Command, string Params, string Preamble = "", string SuccessPostamble = "", string FailurePostamble = "");
+		void PauseForUser(string Message); 
 	}
 
 	//public interface InputOutput
diff --git a/Engine/Source/Programs/AutomationTool/Turnkey/CopyProvider.cs b/Engine/Source/Programs/AutomationTool/Turnkey/CopyProvider.cs
--- a/Engine/Source/Programs/AutomationTool/Turnkey/CopyProvider.cs
+++ b/Engine/Source/Programs/AutomationTool/Turnkey/CopyProvider.cs
@@ -314,5 +314,10 @@ namespace Turnkey
 		{
 			return TurnkeyUtils.GetVariableValue(VariableName);
 		}
+
+		public void PauseForUser(string Message)
+		{
+			TurnkeyUtils.PauseForUser(Message);
+		}
 	}
 }
diff --git a/Engine/Source/Programs/AutomationTool/Turnkey/IOProvider.cs b/Engine/Source/Programs/AutomationTool/Turnkey/IOProvider.cs
--- a/Engine/Source/Programs/AutomationTool/Turnkey/IOProvider.cs
+++ b/Engine/Source/Programs/AutomationTool/Turnkey/IOProvider.cs
@@ -17,6 +17,8 @@ namespace Turnkey
 			Log(Message, bAppendNewLine);
 		}
 
+		public abstract void PauseForUser(string Message, bool bAppendNewLine);
+
 		public abstract string ReadInput(string Prompt, string DefaultValue, bool bAppendNewLine);
 		
 		public abstract int ReadInputInt(string Prompt, List<string> Options, bool bIsCancellable, int DefaultValue, bool bAppendNewLine);
diff --git a/Engine/Source/Programs/AutomationTool/Turnkey/Providers/ConsoleIOProvider.cs b/Engine/Source/Programs/AutomationTool/Turnkey/Providers/ConsoleIOProvider.cs
--- a/Engine/Source/Programs/AutomationTool/Turnkey/Providers/ConsoleIOProvider.cs
+++ b/Engine/Source/Programs/AutomationTool/Turnkey/Providers/ConsoleIOProvider.cs
@@ -21,6 +21,13 @@ namespace Turnkey
 			}
 		}
 
+		public override void PauseForUser(string Message, bool bAppendNewLine)
+		{
+			Log(Message, bAppendNewLine);
+			Log("Press Enter to continue", bAppendNewLine);
+			Console.ReadLine();
+		}
+
 		public override string ReadInput(string Prompt, string Default, bool bAppendNewLine)
 		{
 			Log(Prompt, bAppendNewLine);
diff --git a/Engine/Source/Programs/AutomationTool/Turnkey/Providers/HybridIOProvider.cs b/Engine/Source/Programs/AutomationTool/Turnkey/Providers/HybridIOProvider.cs
--- a/Engine/Source/Programs/AutomationTool/Turnkey/Providers/HybridIOProvider.cs
+++ b/Engine/Source/Programs/AutomationTool/Turnkey/Providers/HybridIOProvider.cs
@@ -166,6 +166,11 @@ namespace Turnkey
 #pragma warning restore
 		}
 
+		public override void PauseForUser(string Message, bool bAppendNewLine)
+		{
+			ShowDialog(Message, "", false);
+		}
+
 		public override string ReadInput(string Prompt, string Default, bool bAppendNewLine)
 		{
 			return ShowDialog(Prompt, Default, false);
diff --git a/Engine/Source/Programs/AutomationTool/Turnkey/TurnkeyUtils.cs b/Engine/Source/Programs/AutomationTool/Turnkey/TurnkeyUtils.cs
--- a/Engine/Source/Programs/AutomationTool/Turnkey/TurnkeyUtils.cs
+++ b/Engine/Source/Programs/AutomationTool/Turnkey/TurnkeyUtils.cs
@@ -779,6 +779,11 @@ namespace Turnkey
 			IOProvider.Report(string.Format(Message, Params), bAppendNewLine: true);
 		}
 
+		public static void PauseForUser(string Message, params object[] Params)
+		{
+			IOProvider.PauseForUser(string.Format(Message, Params), bAppendNewLine: true);
+		}
+
 		public static string ReadInput(string Prompt, string Default = "")
 		{
 			return IOProvider.ReadInput(Prompt, Default, bAppendNewLine: true);
