Improve logging for failures to chmod.

#rnx

diff --git a/Engine/Source/Programs/AutomationTool/AutomationUtils/BuildUtils.cs b/Engine/Source/Programs/AutomationTool/AutomationUtils/BuildUtils.cs
--- a/Engine/Source/Programs/AutomationTool/AutomationUtils/BuildUtils.cs
+++ b/Engine/Source/Programs/AutomationTool/AutomationUtils/BuildUtils.cs
@@ -1,4 +1,4 @@
-﻿// Copyright Epic Games, Inc. All Rights Reserved.
+// Copyright Epic Games, Inc. All Rights Reserved.
 using System;
 using System.Collections.Generic;
 using System.Text;
@@ -111,7 +111,7 @@ namespace AutomationTool
 			var Result = CommandUtils.Run("sh", string.Format("-c 'chmod {0} \"{1}\"'", Permissions, Filename.Replace("'", "'\"'\"'")), Options:ERunOptions.SpewIsVerbose);
 			if (Result.ExitCode != 0)
 			{
-				throw new AutomationException(String.Format("Failed to chmod \"{0}\"", Filename));
+				throw new AutomationException(String.Format("Failed to chmod \"{0}\" (Exit Code {1}): {2}", Filename, Result.ExitCode, Result.Output));
 			}
 		}
 	}
