Log out actual command used to fix unzipped file.

#rnx

diff --git a/Engine/Source/Programs/AutomationTool/AutomationUtils/BuildUtils.cs b/Engine/Source/Programs/AutomationTool/AutomationUtils/BuildUtils.cs
--- a/Engine/Source/Programs/AutomationTool/AutomationUtils/BuildUtils.cs
+++ b/Engine/Source/Programs/AutomationTool/AutomationUtils/BuildUtils.cs
@@ -108,10 +108,11 @@ namespace AutomationTool
 		public static void FixUnixFilePermissions(string Filename)
 		{
 			string Permissions = IsProbablyAMacOrIOSExe(Filename) ? "0755" : "0644";
-			var Result = CommandUtils.Run("sh", string.Format("-c 'chmod {0} \"{1}\"'", Permissions, Filename.Replace("'", "'\"'\"'")), Options:ERunOptions.SpewIsVerbose);
+			string CommandString = string.Format("-c 'chmod {0} \"{1}\"'", Permissions, Filename.Replace("'", "'\"'\"'"));
+			var Result = CommandUtils.Run("sh", CommandString, Options:ERunOptions.SpewIsVerbose);
 			if (Result.ExitCode != 0)
 			{
-				throw new AutomationException(String.Format("Failed to chmod \"{0}\" (Exit Code {1}): {2}", Filename, Result.ExitCode, Result.Output));
+				throw new AutomationException(String.Format("Command 'sh {0}' failed to chmod \"{1}\" (Exit Code {2}): {3}", CommandString, Filename, Result.ExitCode, Result.Output));
 			}
 		}
 	}
