[Turnkey] Add the ability to specify the target of RunExternalCommand requires administrative elevation.   For use when the target is known to need elevation and may not return the correct error code to trigger automatic elevation.

diff --git a/Engine/Source/Programs/AutomationTool/AutomationUtils/Platform.cs b/Engine/Source/Programs/AutomationTool/AutomationUtils/Platform.cs
--- a/Engine/Source/Programs/AutomationTool/AutomationUtils/Platform.cs
+++ b/Engine/Source/Programs/AutomationTool/AutomationUtils/Platform.cs
@@ -16,7 +16,7 @@ namespace AutomationTool
 		string RetrieveFileSource(string Name, string InType = "Misc", string InPlatform = null, string SubType = null);
 		string RetrieveFileSource(object HintObject);
 		string GetVariable(string VariableName);
-		bool RunExternalCommand(string Command, string Params, string Preamble = "", string SuccessPostamble = "", string FailurePostamble = "");
+		bool RunExternalCommand(string Command, string Params, string Preamble = "", string SuccessPostamble = "", string FailurePostamble = "", bool bRequiresPrivilegeElevation = false);
 		void PauseForUser(string Message); 
 	}
 
@@ -193,30 +193,33 @@ namespace AutomationTool
 			}
 
 			string Command, Params, Preamble="", SuccessPostamble="", FailurePostamble="";
-			
-			if (Device != null && GetDeviceUpdateSoftwareCommand(out Command, out Params, ref Preamble, ref SuccessPostamble, ref FailurePostamble, Retriever, Device))
+
+			bool bRequiresPrivilegeElevation;
+			if (Device != null && GetDeviceUpdateSoftwareCommand(out Command, out Params, out bRequiresPrivilegeElevation, ref Preamble, ref SuccessPostamble, ref FailurePostamble, Retriever, Device))
 			{
-				return Retriever.RunExternalCommand(Command, Params, Preamble, SuccessPostamble, FailurePostamble);
+				return Retriever.RunExternalCommand(Command, Params, Preamble, SuccessPostamble, FailurePostamble, bRequiresPrivilegeElevation);
 			}
-			else if (Device == null && GetSDKInstallCommand(out Command, out Params, ref Preamble, ref SuccessPostamble, ref FailurePostamble, Retriever))
+			else if (Device == null && GetSDKInstallCommand(out Command, out Params, out bRequiresPrivilegeElevation, ref Preamble, ref SuccessPostamble, ref FailurePostamble, Retriever))
 			{
-				return Retriever.RunExternalCommand(Command, Params, Preamble, SuccessPostamble, FailurePostamble);
+				return Retriever.RunExternalCommand(Command, Params, Preamble, SuccessPostamble, FailurePostamble, bRequiresPrivilegeElevation);
 			}
 
 			return false;
 		}
 
-		public virtual bool GetSDKInstallCommand(out string Command, out string Params, ref string Preamble, ref string SuccessPostamble, ref string FailurePostamble, FileRetriever Retriever)
+		public virtual bool GetSDKInstallCommand(out string Command, out string Params, out bool bRequiresPrivilegeElevation, ref string Preamble, ref string SuccessPostamble, ref string FailurePostamble, FileRetriever Retriever)
 		{
 			Command = null;
 			Params = null;
+			bRequiresPrivilegeElevation = false;
 			return false;
 		}
 
-		public virtual bool GetDeviceUpdateSoftwareCommand(out string Command, out string Params, ref string Preamble, ref string SuccessPostamble, ref string FailurePostamble, FileRetriever Retriever, DeviceInfo Device = null)
+		public virtual bool GetDeviceUpdateSoftwareCommand(out string Command, out string Params, out bool bRequiresPrivilegeElevation, ref string Preamble, ref string SuccessPostamble, ref string FailurePostamble, FileRetriever Retriever, DeviceInfo Device = null)
 		{
 			Command = null;
 			Params = null;
+			bRequiresPrivilegeElevation = false;
 			return false;
 		}
 
diff --git a/Engine/Source/Programs/AutomationTool/IOS/IOSPlatform.Automation.cs b/Engine/Source/Programs/AutomationTool/IOS/IOSPlatform.Automation.cs
--- a/Engine/Source/Programs/AutomationTool/IOS/IOSPlatform.Automation.cs
+++ b/Engine/Source/Programs/AutomationTool/IOS/IOSPlatform.Automation.cs
@@ -169,10 +169,11 @@ public class IOSPlatform : Platform
 		SDKName = (TargetPlatform == UnrealTargetPlatform.TVOS) ? "appletvos" : "iphoneos";
 	}
 
-	public override bool GetSDKInstallCommand(out string Command, out string Params, ref string Preamble, ref string SuccessPostamble, ref string FailurePostamble, FileRetriever Retriever)
+	public override bool GetSDKInstallCommand(out string Command, out string Params, out bool bRequiresPrivilegeElevation, ref string Preamble, ref string SuccessPostamble, ref string FailurePostamble, FileRetriever Retriever)
 	{
 		Preamble = "Moving your original Xcode application from /Applications to the Trash, and unzipping the new version into /Applications";
 		SuccessPostamble = "If you had Xcode in your Dock, you will need to remove it and add the new one (even though it was in the same location). macOS follows the move to the Trash for the Dock icon";
+		bRequiresPrivilegeElevation = false;
 
 		// put current Xcode in the trash, and unzip a new one. Xcode in the dock will have to be fixed up tho!
 		Command = "osascript";
@@ -188,7 +189,7 @@ public class IOSPlatform : Platform
 		return true;
 	}
 
-	public override bool GetDeviceUpdateSoftwareCommand(out string Command, out string Params, ref string Preamble, ref string SuccessPostamble, ref string FailurePostamble, FileRetriever Retriever, DeviceInfo Device)
+	public override bool GetDeviceUpdateSoftwareCommand(out string Command, out string Params, out bool bRequiresPrivilegeElevation, ref string Preamble, ref string SuccessPostamble, ref string FailurePostamble, FileRetriever Retriever, DeviceInfo Device)
 	{
 		Preamble = "Installing an offline downloaded .ipsw onto your device using the Apple Configurator application.";
 
@@ -197,6 +198,7 @@ public class IOSPlatform : Platform
 
 		string CfgUtilParams = string.Format("-c '{0} list | grep {1}'", Configurator, Device.Id);
 		string CfgUtilOutput = UnrealBuildTool.Utils.RunLocalProcessAndReturnStdOut("sh", CfgUtilParams);
+		bRequiresPrivilegeElevation = false;
 
 		Match Result = Regex.Match(CfgUtilOutput, @"Type: (\S*).*ECID: (\S*)");
 		if (!Result.Success)
diff --git a/Engine/Source/Programs/AutomationTool/Turnkey/CopyAndRun.cs b/Engine/Source/Programs/AutomationTool/Turnkey/CopyAndRun.cs
--- a/Engine/Source/Programs/AutomationTool/Turnkey/CopyAndRun.cs
+++ b/Engine/Source/Programs/AutomationTool/Turnkey/CopyAndRun.cs
@@ -64,7 +64,7 @@ namespace Turnkey
 			CommandLine = TurnkeyUtils.ExpandVariables(CommandLine, bUseOnlyTurnkeyVariables: true);
 		}
 
-		public static bool RunExternalCommand(string CommandPath, string CommandLine)
+		public static bool RunExternalCommand(string CommandPath, string CommandLine, bool bRequiresPrivilegeElevation)
 		{
 			string FixedCommandPath = TurnkeyUtils.ExpandVariables(CommandPath);
 			string FixedCommandLine = TurnkeyUtils.ExpandVariables(CommandLine);
@@ -103,6 +103,15 @@ namespace Turnkey
 			{
 				try
 				{
+					if (bRequiresPrivilegeElevation)
+					{
+						TurnkeyUtils.Log("The installer {0} requires elevated permissions, trying with Admin privileges (output may be hidden)", FixedCommandPath);
+
+						InstallProcess.StartInfo.UseShellExecute = true;
+						InstallProcess.StartInfo.Verb = "runas";
+						InstallProcess.StartInfo.WindowStyle = ProcessWindowStyle.Normal;
+					}
+
 					InstallProcess.Start();
 					InstallProcess.WaitForExit();
 					ExitCode = InstallProcess.ExitCode;
@@ -190,7 +199,7 @@ namespace Turnkey
 			// run an external command
 			if (CommandPath != null)
 			{
-				return RunExternalCommand(CommandPath, CommandLine);
+				return RunExternalCommand(CommandPath, CommandLine, false);
 			}
 
 			return true;
diff --git a/Engine/Source/Programs/AutomationTool/Turnkey/CopyProvider.cs b/Engine/Source/Programs/AutomationTool/Turnkey/CopyProvider.cs
--- a/Engine/Source/Programs/AutomationTool/Turnkey/CopyProvider.cs
+++ b/Engine/Source/Programs/AutomationTool/Turnkey/CopyProvider.cs
@@ -230,7 +230,7 @@ namespace Turnkey
 			return null;
 		}
 
-		public bool RunExternalCommand(string Command, string Params, string Preamble, string SuccessPostAmble, string FailurePostamble)
+		public bool RunExternalCommand(string Command, string Params, string Preamble, string SuccessPostAmble, string FailurePostamble, bool bRequiresPrivilegeElevation)
 		{
 			TurnkeyUtils.Log("----------------------------------------------");
 			TurnkeyUtils.Log("Running '{0} {1}'", Command, Params);
@@ -242,7 +242,7 @@ namespace Turnkey
 			}
 			TurnkeyUtils.Log("----------------------------------------------", Command);
 
-			bool bSuccess = CopyAndRun.RunExternalCommand(Command, Params);
+			bool bSuccess = CopyAndRun.RunExternalCommand(Command, Params, bRequiresPrivilegeElevation);
 
 			TurnkeyUtils.Log("----------------------------------------------");
 			TurnkeyUtils.Log("Finished with {0}", bSuccess ? "Success" : "Failure");
