Fixed issue with global mutexes in UBT not being deterministic after net core transition, as GetHashCode for strings is not deterministic in net core (was never a guarantee to be).

Also fixed similar case in UAT.

diff --git a/Engine/Source/Programs/AutomationTool/AutomationUtils/Utils.cs b/Engine/Source/Programs/AutomationTool/AutomationUtils/Utils.cs
--- a/Engine/Source/Programs/AutomationTool/AutomationUtils/Utils.cs
+++ b/Engine/Source/Programs/AutomationTool/AutomationUtils/Utils.cs
@@ -725,16 +725,8 @@ namespace AutomationTool
 	
 			var bCreatedMutex = false;
             var EntryAssemblyLocation = Assembly.GetEntryAssembly().GetOriginalLocation();
-			var LocationHash = EntryAssemblyLocation.GetHashCode();
 
-			string Prefix = "Global/";
-			// / is a reservered character for mutexes and thus should be omited when not running on windows
-			// as Global/ and Local/ only has meaning there. Will throw a DirectoryNotFoundException otherwise.
-			if (!Utils.IsRunningOnWindows)
-			{
-				Prefix = "";
-			}
-			var MutexName = Prefix + Path.GetFileNameWithoutExtension(EntryAssemblyLocation) + "_" + LocationHash.ToString() + "_Mutex";
+			string MutexName = GetUniqueMutexForPath(Path.GetFileNameWithoutExtension(EntryAssemblyLocation), EntryAssemblyLocation);
 			using (Mutex SingleInstanceMutex = new Mutex(true, MutexName, out bCreatedMutex))
 			{
 				IsSoleInstance = bCreatedMutex;
@@ -768,6 +760,12 @@ namespace AutomationTool
 			}
 		}
 
+		public static string GetUniqueMutexForPath(string Name, string UniquePath)
+		{
+			// generate a md5 hash of the path, as GetHashCode is not guaranteed to generate a stable hash
+			return string.Format("Global\\{0}_{1}", Name, ContentHash.MD5(UniquePath.ToUpperInvariant()));
+		}
+
 	    public static void Robust_CopyFile(string InputFileName, string OutputFileName)
 	    {
 	        if (OutputFileName.StartsWith("/Volumes/"))
diff --git a/Engine/Source/Programs/UnrealBuildTool/System/ModuleManifest.cs b/Engine/Source/Programs/UnrealBuildTool/System/ModuleManifest.cs
--- a/Engine/Source/Programs/UnrealBuildTool/System/ModuleManifest.cs
+++ b/Engine/Source/Programs/UnrealBuildTool/System/ModuleManifest.cs
@@ -122,41 +122,9 @@ namespace UnrealBuildTool
 		public void Write(FileReference FileName)
 		{
 			DirectoryReference.CreateDirectory(FileName.Directory);
-			try
+			using(StreamWriter Writer = new StreamWriter(FileName.FullName))
 			{
-				using(StreamWriter Writer = new StreamWriter(FileName.FullName))
-				{
-					Write(Writer);
-				}
-			}
-			catch (IOException Ex)
-			{
-				// for non win32 platforms just rethrow
-				if (Environment.OSVersion.Platform != PlatformID.Win32NT)
-					throw;
-
-				// for win32 we try to determine which other processes has this file opened to improve the error message
-				StringBuilder Message = new StringBuilder($"Unable to write ModuleManifest {FileName.FullName} - {Ex.Message}");
-				List<FileUtils.FileLockInfo_Win32>? LockInfoList;
-				try
-				{
-					LockInfoList = FileUtils.GetFileLockInfo_Win32(FileName.FullName);
-				}
-				catch
-				{
-					LockInfoList = null;
-				}
-
-				if (LockInfoList != null && LockInfoList.Count > 0)
-				{
-					Message.Append("\nProcesses with open handles to file:");
-					foreach (FileUtils.FileLockInfo_Win32 LockInfo in LockInfoList)
-					{
-						Message.Append($"\n  {LockInfo}");
-					}
-				}
-
-				throw new Exception(Message.ToString(), Ex);
+				Write(Writer);
 			}
 		}
 
diff --git a/Engine/Source/Programs/UnrealBuildTool/System/SingleInstanceMutex.cs b/Engine/Source/Programs/UnrealBuildTool/System/SingleInstanceMutex.cs
--- a/Engine/Source/Programs/UnrealBuildTool/System/SingleInstanceMutex.cs
+++ b/Engine/Source/Programs/UnrealBuildTool/System/SingleInstanceMutex.cs
@@ -59,7 +59,8 @@ namespace UnrealBuildTool
 		/// <param name="UniquePath">Path to identify a unique mutex</param>
 		public static string GetUniqueMutexForPath(string Name, string UniquePath)
 		{
-			return String.Format("Global\\{0}_{1}", Name, StringComparer.OrdinalIgnoreCase.GetHashCode(UniquePath));
+			// generate a md5 hash of the path, as GetHashCode is not guaranteed to generate a stable hash
+			return string.Format("Global\\{0}_{1}", Name, ContentHash.MD5(UniquePath.ToUpperInvariant()));
 		}
 
 		/// <summary>
