Adding WaitTask.

#rnx

diff --git a/Engine/Source/Programs/AutomationTool/BuildGraph/Tasks/WaitTask.cs b/Engine/Source/Programs/AutomationTool/BuildGraph/Tasks/WaitTask.cs
new file mode 100644
--- /dev/null
+++ b/Engine/Source/Programs/AutomationTool/BuildGraph/Tasks/WaitTask.cs
@@ -0,0 +1,82 @@
+// Copyright Epic Games, Inc. All Rights Reserved.
+
+using System;
+using System.Collections.Generic;
+using System.IO;
+using System.Threading;
+using System.Xml;
+using AutomationTool;
+using Tools.DotNETCommon;
+
+namespace BuildGraph.Tasks
+{
+	/// <summary>
+	/// Parameters for a wait task
+	/// </summary>
+	public class WaitTaskParameters
+	{
+		/// <summary>
+		/// Number of seconds to wait.
+		/// </summary>
+		[TaskParameter]
+		public int Seconds;
+	}
+
+	/// <summary>
+	/// Waits a defined number of seconds.
+	/// </summary>
+	[TaskElement("Wait", typeof(WaitTaskParameters))]
+	public class WaitTask : CustomTask
+	{
+		/// <summary>
+		/// Parameters for this task
+		/// </summary>
+		WaitTaskParameters Parameters;
+
+		/// <summary>
+		/// Construct a wait task
+		/// </summary>
+		/// <param name="InParameters">Parameters for the task</param>
+		public WaitTask(WaitTaskParameters InParameters)
+		{
+			Parameters = InParameters;
+		}
+
+		/// <summary>
+		/// Execute the task.
+		/// </summary>
+		/// <param name="Job">Information about the current job</param>
+		/// <param name="BuildProducts">Set of build products produced by this node.</param>
+		/// <param name="TagNameToFileSet">Mapping from tag names to the set of files they include</param>
+		public override void Execute(JobContext Job, HashSet<FileReference> BuildProducts, Dictionary<string, HashSet<FileReference>> TagNameToFileSet)
+		{
+			Thread.Sleep(TimeSpan.FromSeconds(Parameters.Seconds));
+		}
+
+		/// <summary>
+		/// Output this task out to an XML writer.
+		/// </summary>
+		public override void Write(XmlWriter Writer)
+		{
+			Write(Writer, Parameters);
+		}
+
+		/// <summary>
+		/// Find all the tags which are used as inputs to this task
+		/// </summary>
+		/// <returns>The tag names which are read by this task</returns>
+		public override IEnumerable<string> FindConsumedTagNames()
+		{
+			yield break;
+		}
+
+		/// <summary>
+		/// Find all the tags which are modified by this task
+		/// </summary>
+		/// <returns>The tag names which are modified by this task</returns>
+		public override IEnumerable<string> FindProducedTagNames()
+		{
+			yield break;
+		}
+	}
+}
