Fix compile errors in WaitTask.cs.

#rnx

diff --git a/Engine/Source/Programs/AutomationTool/BuildGraph/Tasks/WaitTask.cs b/Engine/Source/Programs/AutomationTool/BuildGraph/Tasks/WaitTask.cs
--- a/Engine/Source/Programs/AutomationTool/BuildGraph/Tasks/WaitTask.cs
+++ b/Engine/Source/Programs/AutomationTool/BuildGraph/Tasks/WaitTask.cs
@@ -6,7 +6,7 @@ using System.IO;
 using System.Threading;
 using System.Xml;
 using AutomationTool;
-using Tools.DotNETCommon;
+using EpicGames.Core;
 
 namespace BuildGraph.Tasks
 {
