Turnkey for device management and SDK apis for platforms coming online soon:
- Retooled some turnkey device management (DeviceInfo now knows its platform, so we don't need to associate a platform externally)
- Changed Control command to use Gauntlet devices to do the PowerOn etc type stuff since there is already support for device control in Gauntlet. Now Turnkey is more of an interactive/scriptable frontend to Gauntlet
- Allow for a platform to do a manual Sdk installation, which doesn't depend on finding a Turnkey file source
- Allow for a platform SDK to return custom versions. This is solely used by platform-specific code

diff --git a/Engine/Source/Programs/AutomationTool/AutomationUtils/Platform.cs b/Engine/Source/Programs/AutomationTool/AutomationUtils/Platform.cs
--- a/Engine/Source/Programs/AutomationTool/AutomationUtils/Platform.cs
+++ b/Engine/Source/Programs/AutomationTool/AutomationUtils/Platform.cs
@@ -27,6 +27,23 @@ namespace AutomationTool
 
 	public class DeviceInfo
 	{
+		public DeviceInfo(UnrealTargetPlatform Platform)
+		{
+			this.Platform = Platform;
+		}
+
+		public DeviceInfo(UnrealTargetPlatform Platform, string Name, string Id, string SoftwareVersion, string Type, bool bIsDefault, bool bCanConnect)
+		{
+			this.Platform = Platform;
+			this.Name = Name;
+			this.Id = Id;
+			this.SoftwareVersion = SoftwareVersion;
+			this.Type = Type;
+			this.bIsDefault = bIsDefault;
+			this.bCanConnect = bCanConnect;
+		}
+
+		public UnrealTargetPlatform Platform;
 		public string Name;
 		public string Id;
 		public string SoftwareVersion;
@@ -35,12 +52,6 @@ namespace AutomationTool
 		// is the device able to be connected to (this is more about able to flash SDK or run, not about matching SDK version)
 		// if false, any of the above fields are suspect, especually SoftwareVersion
 		public bool bCanConnect = true;
-
-		public virtual void Connect() { }
-		public virtual void Disconnect() { }
-		public virtual void PowerOn() { }
-		public virtual void PowerOff() { }
-		public virtual void Reboot() { }
 	}
 
 	/// <summary>
@@ -183,7 +194,7 @@ namespace AutomationTool
 			return null;
 		}
 
-		public virtual bool InstallSDK(BuildCommand BuildCommand, FileRetriever Retriever, object HintObject, DeviceInfo Device = null)
+		public virtual bool InstallSDK(BuildCommand BuildCommand, FileRetriever Retriever, object HintObject, DeviceInfo Device)
 		{
 			string DownloadedSDK = Retriever.RetrieveFileSource(HintObject);
 
@@ -228,6 +239,19 @@ namespace AutomationTool
 			return null;
 		}
 
+		public virtual bool ShouldPerformManualSDKInstall()
+		{
+			return false;
+		}
+		public virtual bool ShouldPerformManualDeviceSoftwareInstall()
+		{
+			return false;
+		}
+
+		public virtual bool ManualInstallSDK(BuildCommand BuildCommand, FileRetriever Retriever, DeviceInfo Device = null)
+		{
+			return false;
+		}
 
 		public virtual bool UpdateHostPrerequisites(BuildCommand Command, FileRetriever Retriever, bool bVerifyOnly)
 		{
diff --git a/Engine/Source/Programs/AutomationTool/IOS/IOSPlatform.Automation.cs b/Engine/Source/Programs/AutomationTool/IOS/IOSPlatform.Automation.cs
--- a/Engine/Source/Programs/AutomationTool/IOS/IOSPlatform.Automation.cs
+++ b/Engine/Source/Programs/AutomationTool/IOS/IOSPlatform.Automation.cs
@@ -154,7 +154,6 @@ public class IOSPlatform : Platform
 
 	private string PlatformName = null;
 	private string SDKName = null;
-	private UnrealTargetPlatform Platform;
 
 	public IOSPlatform()
 		: this(UnrealTargetPlatform.IOS)
@@ -165,7 +164,6 @@ public class IOSPlatform : Platform
 		: base(TargetPlatform)
 	{
 		PlatformName = TargetPlatform.ToString();
-		Platform = TargetPlatform;
 		SDKName = (TargetPlatform == UnrealTargetPlatform.TVOS) ? "appletvos" : "iphoneos";
 	}
 
@@ -509,7 +507,7 @@ public class IOSPlatform : Platform
 			return base.GetDevices();
 		}
 
-		string DeviceType = Platform == UnrealTargetPlatform.TVOS ? "tvOS" : "iOS";
+		string DeviceType = TargetPlatformType == UnrealTargetPlatform.TVOS ? "tvOS" : "iOS";
 
 		// print out each connected device's udid, os version, and name (which might have spaces)
 		string Params = "-e 'require \"fastlane\"; FastlaneCore::DeviceManager.connected_devices(\"" + DeviceType + "\").each { |x| puts \"#{x.udid} #{x.os_type} #{x.os_version} #{x.name}\" }'";
@@ -520,7 +518,7 @@ public class IOSPlatform : Platform
 		{
 			string[] Tokens = Line.Split(" ".ToCharArray());
 
-			DeviceInfo Device = new DeviceInfo();
+			DeviceInfo Device = new DeviceInfo(TargetPlatformType);
 			Device.Id = Tokens[0];
 			Device.Type = Tokens[1];
 			Device.SoftwareVersion = Tokens[2];
@@ -625,7 +623,7 @@ public class IOSPlatform : Platform
 	private bool IsBuiltAsFramework(ProjectParams Params, DeploymentContext SC)
 	{
 		UnrealTargetConfiguration Config = SC.StageTargetConfigurations[0];
-		string InExecutablePath = CombinePaths(Path.GetDirectoryName(Params.GetProjectExeForPlatform(Platform).ToString()), SC.StageExecutables[0]);
+		string InExecutablePath = CombinePaths(Path.GetDirectoryName(Params.GetProjectExeForPlatform(TargetPlatformType).ToString()), SC.StageExecutables[0]);
 		DirectoryReference InEngineDir = DirectoryReference.Combine(SC.LocalRoot, "Engine");
 		DirectoryReference InProjectDirectory = Params.RawProjectPath.Directory;
 		bool bIsUE4Game = !SC.IsCodeBasedProject;
@@ -724,7 +722,7 @@ public class IOSPlatform : Platform
 
 		// ensure the ue4game binary exists, if applicable
 #if !PLATFORM_MAC
-		string ProjectGameExeFilename = Params.GetProjectExeForPlatform(Platform).ToString();
+		string ProjectGameExeFilename = Params.GetProjectExeForPlatform(TargetPlatformType).ToString();
 		string FullExePath = CombinePaths(Path.GetDirectoryName(ProjectGameExeFilename), SC.StageExecutables[0] + (UnrealBuildTool.BuildHostPlatform.Current.Platform != UnrealTargetPlatform.Mac ? ".stub" : ""));
 		if (!SC.IsCodeBasedProject && !FileExists_NoExceptions(FullExePath) && !bIsBuiltAsFramework)
 		{
@@ -805,7 +803,7 @@ public class IOSPlatform : Platform
 
 		StageLaunchScreenStoryboard(Params, SC);
 
-		IOSExports.GenerateAssetCatalog(Params.RawProjectPath, new FileReference(FullExePath), new DirectoryReference(CombinePaths(Params.BaseStageDirectory, (Platform == UnrealTargetPlatform.IOS ? "IOS" : "TVOS"))), Platform);
+		IOSExports.GenerateAssetCatalog(Params.RawProjectPath, new FileReference(FullExePath), new DirectoryReference(CombinePaths(Params.BaseStageDirectory, (TargetPlatformType == UnrealTargetPlatform.IOS ? "IOS" : "TVOS"))), TargetPlatformType);
 
 		bCreatedIPA = false;
 		bool bNeedsIPA = false;
@@ -2162,7 +2160,7 @@ public class IOSPlatform : Platform
 		bool bForDistribution = Params.Distribution;
 
 		// intermediate directory
-		string IntermediateDir = SC.ProjectRoot + "/Intermediate/" + (Platform == UnrealTargetPlatform.IOS ? "IOS" : "TVOS");
+		string IntermediateDir = SC.ProjectRoot + "/Intermediate/" + (TargetPlatformType == UnrealTargetPlatform.IOS ? "IOS" : "TVOS");
 
 		//	entitlements file name
 		string OutputFilename = Path.Combine(IntermediateDir, AppName + ".entitlements");
@@ -2171,7 +2169,7 @@ public class IOSPlatform : Platform
 		ConfigHierarchy PlatformGameConfig;
 		if (Params.EngineConfigs.TryGetValue(SC.StageTargetPlatform.PlatformType, out PlatformGameConfig))
 		{
-			IOSExports.WriteEntitlements(Platform, PlatformGameConfig, AppName, MobileProvisionFile, bForDistribution, IntermediateDir);
+			IOSExports.WriteEntitlements(TargetPlatformType, PlatformGameConfig, AppName, MobileProvisionFile, bForDistribution, IntermediateDir);
 		}
 	}
 
diff --git a/Engine/Source/Programs/AutomationTool/Turnkey/Commands/Control.cs b/Engine/Source/Programs/AutomationTool/Turnkey/Commands/Control.cs
new file mode 100644
--- /dev/null
+++ b/Engine/Source/Programs/AutomationTool/Turnkey/Commands/Control.cs
@@ -0,0 +1,61 @@
+// Copyright Epic Games, Inc. All Rights Reserved.
+
+using System;
+using System.Collections.Generic;
+using AutomationTool;
+using UnrealBuildTool;
+using System.Linq;
+using Gauntlet;
+
+namespace Turnkey.Commands
+{
+	class Control : TurnkeyCommand
+	{
+		private delegate void PerformOperation(ITargetDevice Device);
+
+		protected override void Execute(string[] CommandOptions)
+		{
+			List<DeviceInfo> Devices = TurnkeyUtils.GetDevicesFromCommandLineOrUser(CommandOptions, null);
+			if (Devices == null)
+			{
+				return;
+			}
+
+			Dictionary<string, PerformOperation> Ops = new Dictionary<string, PerformOperation>();
+			Ops.Add("Connect", x => x.Connect());
+			Ops.Add("Disconnect", x => x.Disconnect());
+			Ops.Add("PowerOn", x => x.PowerOn());
+			Ops.Add("PowerOff", x => x.PowerOff());
+			Ops.Add("Reboot", x => x.Reboot());
+
+			bool bWasScripted = false;
+			// break if the GetGenericOption was scripted, otherwise we would loop forever
+			while (!bWasScripted)
+			{
+				string Operation = TurnkeyUtils.GetGenericOption(CommandOptions, Ops.Keys.ToList(), "Operation", out bWasScripted);
+				if (Operation == null)
+				{
+					return;
+				}
+
+				foreach (DeviceInfo Device in Devices)
+				{
+					TurnkeyUtils.Log($"Performing {Operation} on {Device.Platform}@{Device.Name}:");
+
+					try
+					{
+						ITargetDevice GauntletDevice = TurnkeyGauntletUtils.GetGauntletDevice(Device);
+						if (GauntletDevice != null)
+						{
+							Ops[Operation](GauntletDevice);
+						}
+					}
+					catch (Exception Ex)
+					{
+						TurnkeyUtils.Log($"Operation {Operation} failed : {Ex.Message}\n");
+					}
+				}
+			}
+		}
+	}
+}
diff --git a/Engine/Source/Programs/AutomationTool/Turnkey/Commands/InstallBuild.cs b/Engine/Source/Programs/AutomationTool/Turnkey/Commands/InstallBuild.cs
--- a/Engine/Source/Programs/AutomationTool/Turnkey/Commands/InstallBuild.cs
+++ b/Engine/Source/Programs/AutomationTool/Turnkey/Commands/InstallBuild.cs
@@ -107,25 +107,15 @@ namespace Turnkey.Commands
 			Config.Build = Build;
 			Config.ProjectName = Project;
 
-			Dictionary<UnrealTargetPlatform, List<DeviceInfo>> TurnkeyDevices = TurnkeyUtils.GetDevicesFromCommandLineOrUser(CommandOptions, Platform);
-			List<ITargetDevice> GauntletDevices = new List<ITargetDevice>();
-
+			List<DeviceInfo> TurnkeyDevices = TurnkeyUtils.GetDevicesFromCommandLineOrUser(CommandOptions, Platform);
 			if (TurnkeyDevices == null)
 			{
-				// @todo turnkey: probably don't want this, just cancel?
-				IEnumerable<IDefaultDeviceSource> DeviceSources = Gauntlet.Utils.InterfaceHelpers.FindImplementations<IDefaultDeviceSource>();
-				GauntletDevices.AddRange(DeviceSources.Where(S => S.CanSupportPlatform(Platform)).SelectMany(S => S.GetDefaultDevices()));
-			}
-			else
-			{
-				IDeviceFactory Factory = Gauntlet.Utils.InterfaceHelpers.FindImplementations<IDeviceFactory>()
-					.Where(F => F.CanSupportPlatform(Platform))
-					.FirstOrDefault();
-
-				GauntletDevices.AddRange(TurnkeyDevices[Platform].Select(x => Factory.CreateDevice(x.Name, null)));
+				TurnkeyUtils.Log("Could not find a device to install on!");
+				return;
 			}
-
-			if (GauntletDevices.Count == 0 || GauntletDevices[0] == null)
+			
+			List<ITargetDevice> GauntletDevices = TurnkeyGauntletUtils.GetGauntletDevices(TurnkeyDevices);
+			if (GauntletDevices == null)
 			{
 				TurnkeyUtils.Log("Could not find a device to install on!");
 				return;
diff --git a/Engine/Source/Programs/AutomationTool/Turnkey/Commands/VerifySdk.cs b/Engine/Source/Programs/AutomationTool/Turnkey/Commands/VerifySdk.cs
--- a/Engine/Source/Programs/AutomationTool/Turnkey/Commands/VerifySdk.cs
+++ b/Engine/Source/Programs/AutomationTool/Turnkey/Commands/VerifySdk.cs
@@ -23,25 +23,33 @@ namespace Turnkey.Commands
 
 			// track each platform to check, and if -device was specified, get the devices
 			string DeviceString = TurnkeyUtils.ParseParamValue("Device", null, CommandOptions);
-			Dictionary<UnrealTargetPlatform, List<DeviceInfo>> PlatformsAndDevices = null;
+			List<UnrealTargetPlatform> PlatformsToCheck = new List<UnrealTargetPlatform>();
+			List<DeviceInfo> DevicesToCheck = new List<DeviceInfo>();
 			if (string.IsNullOrEmpty(DeviceString))
 			{
-				List<UnrealTargetPlatform> Platforms = TurnkeyUtils.GetPlatformsFromCommandLineOrUser(CommandOptions, null);
-				if (Platforms != null)
+				PlatformsToCheck = TurnkeyUtils.GetPlatformsFromCommandLineOrUser(CommandOptions, null);
+			}
+			else
+			{
+				DevicesToCheck = TurnkeyUtils.GetDevicesFromCommandLineOrUser(CommandOptions, null);
+				// if we failed to get any devices, fallback to getting just platforms anyway
+				if (DevicesToCheck == null)
 				{
-					PlatformsAndDevices = new Dictionary<UnrealTargetPlatform, List<DeviceInfo>>();
-					foreach (UnrealTargetPlatform Platform in Platforms)
+					if (TurnkeyUtils.ParseParamValue("Platform", null, CommandOptions) == null)
 					{
-						PlatformsAndDevices[Platform] = new List<DeviceInfo>();
+						TurnkeyUtils.Log("Devices were requested, but none of them were found. Since -platforms was not specified, exiting command...");
+						return;
 					}
+					TurnkeyUtils.Log("Devices were requested, but none of them were found. Will continue with just platform verification");
+					PlatformsToCheck = TurnkeyUtils.GetPlatformsFromCommandLineOrUser(CommandOptions, null);
+				}
+				else
+				{
+					PlatformsToCheck = DevicesToCheck.Select(x => x.Platform).ToHashSet().ToList();
 				}
-			}
-			else
-			{
-				PlatformsAndDevices = TurnkeyUtils.GetDevicesFromCommandLineOrUser(CommandOptions, null);
 			}
 
-			if (PlatformsAndDevices == null)
+			if (PlatformsToCheck == null || PlatformsToCheck.Count == 0)
 			{
 				TurnkeyUtils.Log("Platform(s) and/or device(s) needed for VerifySdk command. Check parameters or selections.");
 				return;
@@ -56,14 +64,14 @@ namespace Turnkey.Commands
 			TurnkeyUtils.StartTrackingExternalEnvVarChanges();
 
 			// check all the platforms
-			foreach (var Pair in PlatformsAndDevices)
+			foreach (UnrealTargetPlatform Platform in PlatformsToCheck)
 			{
-				UnrealTargetPlatform Platform = Pair.Key;
 				UEBuildPlatformSDK PlatformSDK = UEBuildPlatformSDK.GetSDKForPlatform(Platform.ToString());
 
 				// get the platform object
 				AutomationTool.Platform AutomationPlatform = AutomationTool.Platform.GetPlatform(Platform);
 
+				// checking availability may generate errors, if prerequisites are failing
 				SdkUtils.LocalAvailability LocalState = SdkUtils.GetLocalAvailability(AutomationPlatform, bUpdateIfNeeded);
 
 				string ManualSDKVersion = "", AutoSDKVersion = "";
@@ -93,7 +101,6 @@ namespace Turnkey.Commands
 				{
 					StatusString = "Valid";
 					ReportedState &= (SdkUtils.LocalAvailability.AutoSdk_ValidVersionExists | SdkUtils.LocalAvailability.InstalledSdk_ValidVersionExists | SdkUtils.LocalAvailability.Support_FullSdk | SdkUtils.LocalAvailability.Support_AutoSdk);
-					//					TurnkeyUtils.Log("{0}: Valid [Installed: '{1}', Required: '{2}']", Platform, PlatformObject.GetInstalledSdk(), PlatformObject.GetAllowedSdks());
 				}
 
 				TurnkeyUtils.Report("{0}: (Status={1}, Installed={2}, AutoSDK={3}, MinAllowed={4}, MaxAllowed={5}, Flags=\"{6}\")", Platform, StatusString, ManualSDKVersion, AutoSDKVersion,
@@ -116,6 +123,15 @@ namespace Turnkey.Commands
 // 							}
 // 						}
 
+					if (AutomationPlatform.ShouldPerformManualSDKInstall())
+					{
+						if (AutomationPlatform.ManualInstallSDK(TurnkeyUtils.CommandUtilHelper, new CopyProviderRetriever(), null) == false)
+						{
+							TurnkeyUtils.Log("Install failed!");
+						}
+						continue;
+					}
+
 					// if the platform has a valid sdk but isn't the "installed one", then try to switch to it
 					if ((LocalState & SdkUtils.LocalAvailability.InstalledSdk_ValidInactiveVersionExists) != 0)
 					{
@@ -169,58 +185,62 @@ namespace Turnkey.Commands
 				}
 
 				// now check software verison of each device
-				foreach (DeviceInfo Device in Pair.Value)
+				if (DevicesToCheck != null)
 				{
-					bool bArePrerequisitesValid = AutomationPlatform.UpdateDevicePrerequisites(Device, TurnkeyUtils.CommandUtilHelper, Retriever, !bUpdateIfNeeded);
-					bool bIsSoftwareValid = PlatformSDK.IsSoftwareVersionValid(Device.SoftwareVersion);
-
-					SdkUtils.LocalAvailability DeviceState = SdkUtils.LocalAvailability.None;
-					if (!bArePrerequisitesValid)
-					{
-						StatusString = "Invalid";
-						DeviceState |= SdkUtils.LocalAvailability.Device_InvalidPrerequisites;
-					}
-					else if (bIsSoftwareValid)
+					foreach (DeviceInfo Device in DevicesToCheck.Where(x => x.Platform == Platform))
 					{
-						StatusString = "Valid";
-						DeviceState |= SdkUtils.LocalAvailability.Device_InstallSoftwareValid;
-					}
-					else
-					{
-						StatusString = "Invalid";
-						DeviceState |= SdkUtils.LocalAvailability.Device_InstallSoftwareInvalid;
-					}
+						List<string> DeviceErrorMessages = new List<string>();
+						bool bArePrerequisitesValid = AutomationPlatform.UpdateDevicePrerequisites(Device, TurnkeyUtils.CommandUtilHelper, Retriever, !bUpdateIfNeeded);
+						bool bIsSoftwareValid = PlatformSDK.IsSoftwareVersionValid(Device.SoftwareVersion);
 
-					if (Device.bCanConnect== false)
-					{
-						DeviceState |= SdkUtils.LocalAvailability.Device_CannotConnect;
-					}
+						SdkUtils.LocalAvailability DeviceState = SdkUtils.LocalAvailability.None;
+						if (!bArePrerequisitesValid)
+						{
+							StatusString = "Invalid";
+							DeviceState |= SdkUtils.LocalAvailability.Device_InvalidPrerequisites;
+						}
+						else if (bIsSoftwareValid)
+						{
+							StatusString = "Valid";
+							DeviceState |= SdkUtils.LocalAvailability.Device_InstallSoftwareValid;
+						}
+						else
+						{
+							StatusString = "Invalid";
+							DeviceState |= SdkUtils.LocalAvailability.Device_InstallSoftwareInvalid;
+						}
 
-					TurnkeyUtils.Report("{0}@{1}: (Status={2}, Installed={3}, MinAllowed={4}, MaxAllowed={5}, Flags=\"{6}\")", Platform, Device.Name, StatusString, Device.SoftwareVersion, 
-						MinSoftwareAllowedVersion, MaxSoftwareAllowedVersion, DeviceState.ToString());
+						if (Device.bCanConnect == false)
+						{
+							DeviceState |= SdkUtils.LocalAvailability.Device_CannotConnect;
+						}
 
-					if (bForceDeviceInstall || !bIsSoftwareValid)
-					{
-						if (bUpdateIfNeeded)
+						TurnkeyUtils.Report("{0}@{1}: (Status={2}, Installed={3}, MinAllowed={4}, MaxAllowed={5}, Flags=\"{6}\")", Platform, Device.Name, StatusString, Device.SoftwareVersion,
+							MinSoftwareAllowedVersion, MaxSoftwareAllowedVersion, DeviceState.ToString());
+
+						if (bForceDeviceInstall || !bIsSoftwareValid)
 						{
-							if (Device.bCanConnect)
+							if (bUpdateIfNeeded)
 							{
-								FileSource MatchingInstallableSdk = FileSource.FindMatchingSdk(AutomationPlatform, new FileSource.SourceType[] { FileSource.SourceType.Flash }, bSelectBest: bUnattended, DeviceType: Device.Type);
-
-								if (MatchingInstallableSdk == null)
+								if (Device.bCanConnect)
 								{
-									TurnkeyUtils.Log("ERROR: {0}: Unable top find any Sdks that could be installed on {1}", Platform, Device.Name);
-									TurnkeyUtils.ExitCode = AutomationTool.ExitCode.Error_SDKNotFound;
+									FileSource MatchingInstallableSdk = FileSource.FindMatchingSdk(AutomationPlatform, new FileSource.SourceType[] { FileSource.SourceType.Flash }, bSelectBest: bUnattended, DeviceType: Device.Type);
+
+									if (MatchingInstallableSdk == null)
+									{
+										TurnkeyUtils.Log("ERROR: {0}: Unable top find any Sdks that could be installed on {1}", Platform, Device.Name);
+										TurnkeyUtils.ExitCode = AutomationTool.ExitCode.Error_SDKNotFound;
+									}
+									else
+									{
+										MatchingInstallableSdk.DownloadOrInstall(Platform, Device, bUnattended);
+									}
 								}
 								else
 								{
-									MatchingInstallableSdk.DownloadOrInstall(Platform, Device, bUnattended);
+									TurnkeyUtils.Log("Skipping device {0} because it cannot connect.", Device.Name);
 								}
 							}
-							else
-							{
-								TurnkeyUtils.Log("Skipping device {0} because it cannot connect.", Device.Name);
-							}
 						}
 					}
 				}
diff --git a/Engine/Source/Programs/AutomationTool/Turnkey/CopyAndRun.cs b/Engine/Source/Programs/AutomationTool/Turnkey/CopyAndRun.cs
--- a/Engine/Source/Programs/AutomationTool/Turnkey/CopyAndRun.cs
+++ b/Engine/Source/Programs/AutomationTool/Turnkey/CopyAndRun.cs
@@ -160,54 +160,5 @@ namespace Turnkey
 
 			return ExitCode == 0;
 		}
-
-
-
-		public bool Execute(CopyExecuteSpecialMode SpecialMode=CopyExecuteSpecialMode.None, string ModeHint=null)
-		{
-			if (ShouldExecute() == false)
-			{
-				return false;
-			}
-
-			// do the copy operation, capturing the output
-			string OutputPath = null;
-			if (Copy != null)
-			{
-				// if there was a copy operation do it
-				OutputPath = CopyProvider.ExecuteCopy(Copy, SpecialMode, ModeHint);
-
-				// if it returned null, then something went wrong, don't continue with the command. 
-				// assume the CopyProvider supplied enough information, so we quietly skip
-				if (OutputPath == null)
-				{
-					return false;
-				}
-
-				if (SpecialMode == CopyExecuteSpecialMode.DownloadOnly)
-				{
-					// if we only wanted to download if needed, then we are done!
-					return true;
-				}
-			}
-
-			if (OutputPath == null && ((CommandPath != null && CommandPath.Contains("$(CopyOutputPath)")) || (CommandLine != null && CommandLine.Contains("$(CopyOutputPath)"))))
-			{
-				throw new AutomationException("CopyAndRun required an valid $(CopyOutputPath) from the Copy operation, but it failed");
-			}
-
-			// run an external command
-			if (CommandPath != null)
-			{
-				return RunExternalCommand(CommandPath, CommandLine, false);
-			}
-
-			return true;
-		}
-
-		public bool ShouldExecute()
-		{
-			return !Platform.HasValue || Platform == HostPlatform.Current.HostEditorPlatform;
-		}
 	}
 }
diff --git a/Engine/Source/Programs/AutomationTool/Turnkey/CopyProvider.cs b/Engine/Source/Programs/AutomationTool/Turnkey/CopyProvider.cs
--- a/Engine/Source/Programs/AutomationTool/Turnkey/CopyProvider.cs
+++ b/Engine/Source/Programs/AutomationTool/Turnkey/CopyProvider.cs
@@ -233,7 +233,7 @@ namespace Turnkey
 		public bool RunExternalCommand(string Command, string Params, string Preamble, string SuccessPostAmble, string FailurePostamble, bool bRequiresPrivilegeElevation)
 		{
 			TurnkeyUtils.Log("----------------------------------------------");
-			TurnkeyUtils.Log("Running '{0} {1}'", Command, Params);
+			TurnkeyUtils.Log("Running '{0} {1}'", TurnkeyUtils.ExpandVariables(Command), TurnkeyUtils.ExpandVariables(Params));
 
 			if (!string.IsNullOrEmpty(Preamble))
 			{
diff --git a/Engine/Source/Programs/AutomationTool/Turnkey/FileSource.cs b/Engine/Source/Programs/AutomationTool/Turnkey/FileSource.cs
--- a/Engine/Source/Programs/AutomationTool/Turnkey/FileSource.cs
+++ b/Engine/Source/Programs/AutomationTool/Turnkey/FileSource.cs
@@ -333,10 +333,8 @@ namespace Turnkey
 				return SdkUtils.SetupAutoSdk(this, Platform, bUnattended);
 			}
 
-			// let the platform decide how to 
-			AutomationTool.Platform.GetPlatform(Platform).InstallSDK(TurnkeyUtils.CommandUtilHelper, new CopyProviderRetriever(), this, Device);
-
-			return false;
+			// let the platform decide how to install
+			return AutomationTool.Platform.GetPlatform(Platform).InstallSDK(TurnkeyUtils.CommandUtilHelper, new CopyProviderRetriever(), this, Device);
 		}
 
 
diff --git a/Engine/Source/Programs/AutomationTool/Turnkey/TurnkeyGauntletUtils.cs b/Engine/Source/Programs/AutomationTool/Turnkey/TurnkeyGauntletUtils.cs
new file mode 100644
--- /dev/null
+++ b/Engine/Source/Programs/AutomationTool/Turnkey/TurnkeyGauntletUtils.cs
@@ -0,0 +1,78 @@
+// Copyright Epic Games, Inc. All Rights Reserved.
+
+using System;
+using System.Collections.Generic;
+using System.Linq;
+using AutomationTool;
+using UnrealBuildTool;
+using Gauntlet;
+
+namespace Turnkey
+{
+	static class TurnkeyGauntletUtils
+	{
+
+		static public ITargetDevice GetGauntletDevice(UnrealTargetPlatform Platform, string DeviceName)
+		{
+			IDeviceFactory Factory = Gauntlet.Utils.InterfaceHelpers.FindImplementations<IDeviceFactory>()
+				.Where(F => F.CanSupportPlatform(Platform))
+				.FirstOrDefault();
+
+			if (Factory == null)
+			{
+				return null;
+			}
+
+			return Factory.CreateDevice(DeviceName, null);
+		}
+
+		static public ITargetDevice GetGauntletDevice(DeviceInfo Device)
+		{
+			IDeviceFactory Factory = Gauntlet.Utils.InterfaceHelpers.FindImplementations<IDeviceFactory>()
+				.Where(F => F.CanSupportPlatform(Device.Platform))
+				.FirstOrDefault();
+
+			if (Factory == null)
+			{
+				return null;
+			}
+
+			return Factory.CreateDevice(Device.Name, null);
+		}
+
+		static public List<ITargetDevice> GetGauntletDevices(UnrealTargetPlatform Platform, List<string> DeviceNames)
+		{
+			IDeviceFactory Factory = Gauntlet.Utils.InterfaceHelpers.FindImplementations<IDeviceFactory>()
+				.Where(F => F.CanSupportPlatform(Platform))
+				.FirstOrDefault();
+
+			if (Factory == null)
+			{
+				return new List<ITargetDevice>();
+			}
+
+			return DeviceNames.Select(x => Factory.CreateDevice(x, null)).ToList();
+		}
+
+		static public List<ITargetDevice> GetGauntletDevices(List<DeviceInfo> Devices)
+		{
+			return Devices.Select(x => GetGauntletDevice(x)).Where(x => x != null).ToList();
+		}
+
+		static public List<ITargetDevice> GetDefaultGauntletDevices(UnrealTargetPlatform Platform)
+		{
+			IDefaultDeviceSource DeviceSource = Gauntlet.Utils.InterfaceHelpers.FindImplementations<IDefaultDeviceSource>()
+				.Where(S => S.CanSupportPlatform(Platform))
+				.FirstOrDefault();
+
+			if (DeviceSource == null)
+			{
+				return new List<ITargetDevice>();
+			}
+
+			return DeviceSource.GetDefaultDevices().ToList();
+		}
+
+
+	}
+}
diff --git a/Engine/Source/Programs/AutomationTool/Turnkey/TurnkeyUtils.cs b/Engine/Source/Programs/AutomationTool/Turnkey/TurnkeyUtils.cs
--- a/Engine/Source/Programs/AutomationTool/Turnkey/TurnkeyUtils.cs
+++ b/Engine/Source/Programs/AutomationTool/Turnkey/TurnkeyUtils.cs
@@ -40,7 +40,7 @@ namespace Turnkey
 			SetVariable("HOST_PLATFORM_NAME", HostPlatform.Current.HostEditorPlatform.ToString());
 		}
 
-		#region Turnkey Variables
+#region Turnkey Variables
 
 		static Dictionary<string, string> TurnkeyVariables = new Dictionary<string, string>();
 
@@ -82,9 +82,9 @@ namespace Turnkey
 			return UnrealBuildTool.Utils.ExpandVariables(ExpandedUserVariables, TurnkeyVariables, bUseOnlyTurnkeyVariables);
 		}
 
-		#endregion
+#endregion
 
-		#region Commandline Handling
+#region Commandline Handling
 
 		public static bool ParseParam(string Param, string[] ExtraOptions)
 		{
@@ -387,9 +387,17 @@ namespace Turnkey
 			{
 				return null;
 			}
-			return Array.Find(Devices, x => string.Compare(x.Name, DeviceName, true) == 0);
+			// look by Id first
+			DeviceInfo Device = Array.Find(Devices, x => string.Compare(x.Id, DeviceName, true) == 0);
+			// if that fails, use Name
+			if (Device == null)
+			{
+				Device = Array.Find(Devices, x => string.Compare(x.Name, DeviceName, true) == 0);
+			}
+			return Device;
 		}
-		public static Dictionary<UnrealTargetPlatform, List<DeviceInfo>> GetDevicesFromCommandLineOrUser(string[] CommandOptions, UnrealTargetPlatform Platform)
+
+		public static List<DeviceInfo> GetDevicesFromCommandLineOrUser(string[] CommandOptions, UnrealTargetPlatform Platform)
 		{
 			return GetDevicesFromCommandLineOrUser(CommandOptions, new List<UnrealTargetPlatform>() { Platform });
 		}
@@ -398,9 +406,9 @@ namespace Turnkey
 		//	Dictionary<UnrealTargetPlatform, List<DeviceInfo>> PlatformsAndDevices = GetDevicesFromCommandLineOrUser(CommandOptions, Platform);
 		//}
 
-		public static Dictionary<UnrealTargetPlatform, List<DeviceInfo>> GetDevicesFromCommandLineOrUser(string[] CommandOptions, List<UnrealTargetPlatform> PossiblePlatforms)
+		public static List<DeviceInfo> GetDevicesFromCommandLineOrUser(string[] CommandOptions, List<UnrealTargetPlatform> PossiblePlatforms)
 		{
-			Dictionary<UnrealTargetPlatform, List<DeviceInfo>> PlatformsAndDevices = null;
+			List<DeviceInfo> ChosenDevices = null;
 
 			// look at any devices on the commandline, and see if they have platforms or not
 			string DeviceList = TurnkeyUtils.ParseParamValue("Device", null, CommandOptions);
@@ -436,25 +444,21 @@ namespace Turnkey
 						string DeviceName = Tokens[1];
 
 						// track it
-						if (PlatformsAndDevices == null)
+						if (ChosenDevices == null)
 						{
-							PlatformsAndDevices = new Dictionary<UnrealTargetPlatform, List<DeviceInfo>>();
-						}
-						if (!PlatformsAndDevices.ContainsKey(Platform))
-						{
-							PlatformsAndDevices[Platform] = new List<DeviceInfo>();
+							ChosenDevices = new List<DeviceInfo>();
 						}
 						
 						if (DeviceName.ToLower() == "all")
 						{
-							PlatformsAndDevices[Platform].AddRange(AutomationTool.Platform.GetPlatform(Platform).GetDevices());
+							ChosenDevices.AddRange(AutomationTool.Platform.GetPlatform(Platform).GetDevices());
 						}
 						else
 						{
 							DeviceInfo Device = GetDeviceByPlatformAndName(Platform, DeviceName);
 							if (Device != null)
 							{
-								PlatformsAndDevices[Platform].Add(Device);
+								ChosenDevices.Add(Device);
 							}
 						}
 					}
@@ -463,9 +467,9 @@ namespace Turnkey
 			}
 
 			// if we didn't get some platforms already from -device list, then get or ask the user for platforms
-			if (PlatformsAndDevices == null)
+			if (ChosenDevices == null)
 			{
-				PlatformsAndDevices = new Dictionary<UnrealTargetPlatform, List<DeviceInfo>>();
+				ChosenDevices = new List<DeviceInfo>();
 				List<UnrealTargetPlatform> ChosenPlatforms;
 
 				// use all platforms (with -device=all), or ask user if needed (GetPlatformsFromCommandLineOrUser would look at -platform=all, not -device=all)
@@ -484,9 +488,6 @@ namespace Turnkey
 					return null;
 				}
 
-				// set up PlatformsAndDevices for all platforms
-				ChosenPlatforms.ForEach(x => PlatformsAndDevices.Add(x, new List<DeviceInfo>()));
-
 				if (ChosenPlatforms.Count > 1 && SplitDeviceList != null && !(SplitDeviceList.Count == 1 && SplitDeviceList[0].ToLower() == "all"))
 				{
 					throw new AutomationException("When using -Device without platform specifiers ('Platform@Device'), a single platform must be specified (unless -Device=All is used)");
@@ -500,7 +501,7 @@ namespace Turnkey
 						DeviceInfo[] Devices = AutomationTool.Platform.GetPlatform(Platform).GetDevices();
 						if (Devices != null)
 						{
-							PlatformsAndDevices[Platform].AddRange(Devices);
+							ChosenDevices.AddRange(Devices);
 						}
 					}
 				}
@@ -514,7 +515,7 @@ namespace Turnkey
 							DeviceInfo Device = GetDeviceByPlatformAndName(Platform, DeviceName);
 							if (Device != null)
 							{
-								PlatformsAndDevices[Platform].Add(Device);
+								ChosenDevices.Add(Device);
 							}
 						}
 					}
@@ -524,7 +525,6 @@ namespace Turnkey
 				{
 					List<string> Options = new List<string>();
 					List<DeviceInfo> PossibleDevices = new List<DeviceInfo>();
-					List<UnrealTargetPlatform> PossibleDevicePlatforms = new List<UnrealTargetPlatform>();
 
 					foreach (UnrealTargetPlatform Platform in ChosenPlatforms)
 					{
@@ -535,8 +535,6 @@ namespace Turnkey
 							foreach (DeviceInfo Device in Devices)
 							{
 								PossibleDevices.Add(Device);
-								PossibleDevicePlatforms.Add(Platform);
-
 								Options.Add(string.Format("[{0} {1}] {2}", Platform, Device.Type, Device.Name));
 							}
 						}
@@ -557,19 +555,12 @@ namespace Turnkey
 					}
 
 					// finally, add it to the proper list
-					PlatformsAndDevices[PossibleDevicePlatforms[Choice - 1]].Add(PossibleDevices[Choice - 1]);
+					ChosenDevices.Add(PossibleDevices[Choice - 1]);
 				}
 			}
 
 			// if we ended up with some platforms, but no devices, just return null
-			foreach (var Pair in PlatformsAndDevices)
-			{
-				if (Pair.Value != null && Pair.Value.Count > 0)
-				{
-					return PlatformsAndDevices;
-				}
-			}
-			return null;
+			return (ChosenDevices != null && ChosenDevices.Count > 0) ? ChosenDevices : null;
 		}
 
 		public static string GetGenericOption(string[] CommandOptions, List<string> Options, string CommandLineOption)
@@ -662,6 +653,7 @@ namespace Turnkey
 				foreach (var Pair in EnvVarsToSaveToBatchFile)
 				{
 					BatchContents.AppendLine("set {0}={1}", Pair.Key, Pair.Value);
+					TurnkeyUtils.Log("  Recording variable to set by caller {0}={1}", Pair.Key, Pair.Value);
 				}
 
 				// write out a batch file for the caller of this to call to update vars, including all changes from previous commands
diff --git a/Engine/Source/Programs/AutomationTool/Win/WinPlatform.Automation.cs b/Engine/Source/Programs/AutomationTool/Win/WinPlatform.Automation.cs
--- a/Engine/Source/Programs/AutomationTool/Win/WinPlatform.Automation.cs
+++ b/Engine/Source/Programs/AutomationTool/Win/WinPlatform.Automation.cs
@@ -23,10 +23,8 @@ public abstract class BaseWinPlatform : Platform
 
 		if (HostPlatform.Current.HostEditorPlatform == UnrealTargetPlatform.Win64)
 		{
-			DeviceInfo LocalMachine = new DeviceInfo();
-			LocalMachine.Name = LocalMachine.Id = Environment.MachineName;
-			LocalMachine.Type = "Computer";
-			LocalMachine.SoftwareVersion = Environment.OSVersion.Version.ToString();
+			DeviceInfo LocalMachine = new DeviceInfo(UnrealTargetPlatform.Win64, Environment.MachineName, Environment.MachineName,
+				Environment.OSVersion.Version.ToString(), "Computer", true, true);
 
 			Devices.Add(LocalMachine);
 		}
diff --git a/Engine/Source/Programs/DotNETCommon/BuildUtilities/UEBuildPlatformSDK.cs b/Engine/Source/Programs/DotNETCommon/BuildUtilities/UEBuildPlatformSDK.cs
--- a/Engine/Source/Programs/DotNETCommon/BuildUtilities/UEBuildPlatformSDK.cs
+++ b/Engine/Source/Programs/DotNETCommon/BuildUtilities/UEBuildPlatformSDK.cs
@@ -171,6 +171,16 @@ namespace EpicGames.Core
 		/// <returns></returns>
 		public abstract string GetMainVersion();
 
+		/// <summary>
+		/// Returns a platform-specific version string that can be retrieved from anywhere without needing to typecast the SDK object
+		/// </summary>
+		/// <param name="VersionType">Descriptor of the version you want</param>
+		/// <returns>Version string, or empty string if not known/supported</returns>
+		public virtual string GetPlatformSpecificVersion(string VersionType)
+		{
+			return "";
+		}
+
 		/// <summary>
 		/// Return a list of full SDK versions that are already installed and can be quickly switched to without running any installers.
 		/// </summary>
