When shelling to chmod fails, try invoking chmod directly.

#rnx

diff --git a/Engine/Source/Programs/AutomationTool/AutomationUtils/BuildUtils.cs b/Engine/Source/Programs/AutomationTool/AutomationUtils/BuildUtils.cs
--- a/Engine/Source/Programs/AutomationTool/AutomationUtils/BuildUtils.cs
+++ b/Engine/Source/Programs/AutomationTool/AutomationUtils/BuildUtils.cs
@@ -112,7 +112,12 @@ namespace AutomationTool
 			var Result = CommandUtils.Run("sh", CommandString, Options:ERunOptions.SpewIsVerbose);
 			if (Result.ExitCode != 0)
 			{
-				throw new AutomationException(String.Format("Command 'sh {0}' failed to chmod \"{1}\" (Exit Code {2}): {3}", CommandString, Filename, Result.ExitCode, Result.Output));
+				// Try simply invoking chmod directly.
+				var DirectResult = CommandUtils.Run("chmod", string.Format("{0} \"{1}\"", Permissions, Filename.Replace("'", "'\"'\"'")));
+				if (DirectResult.ExitCode != 0)
+				{
+					throw new AutomationException(String.Format("Command 'sh {0}' failed to chmod \"{1}\" (Exit Code {2}): {3}", CommandString, Filename, Result.ExitCode, Result.Output));
+				}
 			}
 		}
 	}
