Use actual filename rather than an underscore for recursively added gauntlet cs files.  Was resulting in a single underscore for each platform rather than the several cs files that should have been there.
#rnx

diff --git a/Engine/Source/Programs/AutomationTool/Gauntlet/Gauntlet.Automation.csproj b/Engine/Source/Programs/AutomationTool/Gauntlet/Gauntlet.Automation.csproj
--- a/Engine/Source/Programs/AutomationTool/Gauntlet/Gauntlet.Automation.csproj
+++ b/Engine/Source/Programs/AutomationTool/Gauntlet/Gauntlet.Automation.csproj
@@ -44,7 +44,7 @@
       <!-- RecursiveDir metadata expands to [PlatformName]/Source/Programs/[etc]. The Replace() will
       replace everything from the first slash to the end leaving just the platform name. Should the
       Replace() fail then RecursiveDir is still an agreeable value -->
-      <Link>Platform/$([System.Text.RegularExpressions.Regex]::Replace(%(Compile.RecursiveDir), [\\/].+$, ``))/_</Link>
+      <Link>Platform/$([System.Text.RegularExpressions.Regex]::Replace(%(Compile.RecursiveDir), [\\/].+$, ``))/%(Compile.FileName).cs</Link>
     </Compile>
   </ItemGroup>
 
