Fix non-deterministic use of GetHashCode() in ProjectUtils.cs.

#rnx

diff --git a/Engine/Source/Programs/AutomationTool/AutomationUtils/ProjectUtils.cs b/Engine/Source/Programs/AutomationTool/AutomationUtils/ProjectUtils.cs
--- a/Engine/Source/Programs/AutomationTool/AutomationUtils/ProjectUtils.cs
+++ b/Engine/Source/Programs/AutomationTool/AutomationUtils/ProjectUtils.cs
@@ -656,7 +656,7 @@ namespace AutomationTool
 			{
 				CommandUtils.LogVerbose("Looking for targets for project {0}", Properties.RawProjectPath);
 
-				TargetsDllFilename = FileReference.Combine(RulesFolder, String.Format("UATRules{0}.dll", Properties.RawProjectPath.GetHashCode()));
+				TargetsDllFilename = FileReference.Combine(RulesFolder, String.Format("UATRules-{0}.dll", ContentHash.MD5(Properties.RawProjectPath.FullName.ToUpperInvariant()).ToString()));
 
 				FullProjectPath = CommandUtils.GetDirectoryName(Properties.RawProjectPath.FullName);
 				GameFolders.Add(new DirectoryReference(FullProjectPath));
