Fix non-deterministic use of GetHashCode() in UBT and Gauntlet

#rnx

diff --git a/Engine/Source/Programs/AutomationTool/Gauntlet/Platform/Android/Gauntlet.TargetDeviceAndroid.cs b/Engine/Source/Programs/AutomationTool/Gauntlet/Platform/Android/Gauntlet.TargetDeviceAndroid.cs
--- a/Engine/Source/Programs/AutomationTool/Gauntlet/Platform/Android/Gauntlet.TargetDeviceAndroid.cs
+++ b/Engine/Source/Programs/AutomationTool/Gauntlet/Platform/Android/Gauntlet.TargetDeviceAndroid.cs
@@ -9,6 +9,7 @@ using System.Text.RegularExpressions;
 using System.Linq;
 using System.Text;
 using System.Threading;
+using EpicGames.Core;
 
 namespace Gauntlet
 {
@@ -705,9 +706,9 @@ namespace Gauntlet
             }
 
             // dependency info is a hash of the destination name, saved under a folder on /sdcard
-            int DestHash = DestPath.GetHashCode();
+            string DestHash = ContentHash.MD5(DestPath).ToString();
 			string DependencyCacheDir = "/sdcard/gdeps";
-			string DepFile = string.Format("{0}/{1:X}", DependencyCacheDir, DestHash);	
+			string DepFile = string.Format("{0}/{1}", DependencyCacheDir, DestHash);
 
 			IProcessResult AdbResult = null;
 
diff --git a/Engine/Source/Programs/UnrealBuildTool/System/UnrealPluginLanguage.cs b/Engine/Source/Programs/UnrealBuildTool/System/UnrealPluginLanguage.cs
--- a/Engine/Source/Programs/UnrealBuildTool/System/UnrealPluginLanguage.cs
+++ b/Engine/Source/Programs/UnrealBuildTool/System/UnrealPluginLanguage.cs
@@ -574,7 +574,7 @@ namespace UnrealBuildTool
 
 		public string GetUPLHash()
 		{
-			return XDoc.ToString().GetHashCode().ToString();
+			return ContentHash.MD5(XDoc.ToString()).ToString();
 		}
 
 		public bool MergeXML(XDocument MergeDoc, string PluginDir, List<string> Architectures)
