libimobiledevice integration for ios / tvos deployment

diff --git a/Engine/Source/Programs/AutomationTool/IOS/IOSPlatform.Automation.cs b/Engine/Source/Programs/AutomationTool/IOS/IOSPlatform.Automation.cs
--- a/Engine/Source/Programs/AutomationTool/IOS/IOSPlatform.Automation.cs
+++ b/Engine/Source/Programs/AutomationTool/IOS/IOSPlatform.Automation.cs
@@ -14,7 +14,6 @@ using System.Threading;
 using System.Diagnostics;
 using EpicGames.Core;
 using System.Xml;
-//using Manzana;
 
 static class IOSEnvVarNames
 {
@@ -1624,7 +1623,6 @@ public class IOSPlatform : Platform
 		bool Result = true;
 		UFSManifests = new List<string>();
 		NonUFSManifests = new List<string>();
-		var DeployServer = CombinePaths(CmdEnv.LocalRoot, "Engine/Binaries/DotNET/IOS/DeploymentServer.exe");
 		try
 		{
 			var TargetConfiguration = SC.StageTargetConfigurations[0];
@@ -1637,16 +1635,28 @@ public class IOSPlatform : Platform
 				int EndPos = Contents.IndexOf("</string>", Pos);
 				BundleIdentifier = Contents.Substring(Pos, EndPos - Pos);
 			}
-			
+
+			var DeviceInstaller = GetPathToLibiMobileDeviceTool("ideviceinstaller");
 			LogInformation("Checking if bundle {0} is installed", BundleIdentifier);
-			string Output = CommandUtils.RunAndLog(DeployServer, "ListApplications" + (String.IsNullOrEmpty(Params.DeviceNames[0]) ? "" : " -device " + Params.DeviceNames[0]), GetRunAndLogOnlyName(CmdEnv, DeployServer, null), 0, null, ERunOptions.SpewIsVerbose);
-			bool bBundleIsInstalled = Output.Contains(string.Format("CFBundleIdentifier -> {0}{1}", BundleIdentifier, Environment.NewLine)); 
+
+			string Output = CommandUtils.RunAndLog(DeviceInstaller, "--list-apps");
+			bool bBundleIsInstalled = Output.Contains(string.Format("CFBundleIdentifier -> {0}{1}", BundleIdentifier, Environment.NewLine));
+			int ExitCode = 0;
 
 			if (bBundleIsInstalled)
 			{
 				LogInformation("Bundle {0} found, retrieving deployed manifests...", BundleIdentifier);
 
-				RunAndLog(CmdEnv, DeployServer, "Backup -file \"" + CombinePaths(Params.BaseStageDirectory, PlatformName, SC.GetUFSDeployedManifestFileName(null)) + "\" -file \"" + CombinePaths(Params.BaseStageDirectory, PlatformName, SC.GetNonUFSDeployedManifestFileName(null)) + "\"" + (String.IsNullOrEmpty(Params.DeviceNames[0]) ? "" : " -device " + Params.DeviceNames[0]) + " -bundle " + BundleIdentifier);
+				var DeviceFS = GetPathToLibiMobileDeviceTool("idevicefs");
+
+				string AllCommandsToPush = " push " + CombinePaths(Params.BaseStageDirectory, PlatformName, SC.GetUFSDeployedManifestFileName(null)) + "\n"
+										+ " push " + CombinePaths(Params.BaseStageDirectory, PlatformName, SC.GetNonUFSDeployedManifestFileName(null));		
+				System.IO.File.WriteAllText(Directory.GetCurrentDirectory() + "\\CommandsToPush.txt", AllCommandsToPush);
+				Utils.RunLocalProcessAndReturnStdOut(DeviceFS, "-b " + BundleIdentifier + " -x " + Directory.GetCurrentDirectory() + "\\CommandsToPush.txt", out ExitCode, true);
+				if (ExitCode != 0)
+				{
+					LogError("Deployer error : failed to push files. See above log for more details");
+				}
 
 				string[] ManifestFiles = Directory.GetFiles(CombinePaths(Params.BaseStageDirectory, PlatformName), "*_Manifest_UFS*.txt");
 				UFSManifests.AddRange(ManifestFiles);
@@ -1673,6 +1683,59 @@ public class IOSPlatform : Platform
 		return Result;
 	}
 
+	private string GetPathToLibiMobileDeviceTool(string LibimobileExec)
+	{
+		string ExecWithPath = "";
+		if (UnrealBuildTool.BuildHostPlatform.Current.Platform == UnrealTargetPlatform.Win32 ||
+			UnrealBuildTool.BuildHostPlatform.Current.Platform == UnrealTargetPlatform.Win64)
+		{
+			ExecWithPath = CombinePaths(CmdEnv.LocalRoot, "Engine/Extras/ThirdPartyNotUE/libimobiledevice/x64/" + LibimobileExec + ".exe");
+		}
+		else if (UnrealBuildTool.BuildHostPlatform.Current.Platform == UnrealTargetPlatform.Mac)
+		{
+			ExecWithPath = CombinePaths(CmdEnv.LocalRoot, "Engine/Extras/ThirdPartyNotUE/libimobiledevice/Mac/" + LibimobileExec);
+		}
+		if (!File.Exists(ExecWithPath) || ExecWithPath == "")
+		{
+			LogError("Failed to locate Libimobile exec. Deploment will fail.");
+		}
+		return ExecWithPath;
+	}
+
+	private void DeployManifestContent(string BaseFolder, DeploymentContext SC, ProjectParams Params, ref string Files, ref string BundleIdentifier)
+	{
+		var DeviceFS = GetPathToLibiMobileDeviceTool("idevicefs");
+
+
+		string[] FileList = Files.Split('\n');
+		string AllCommandsToPush = "";
+		foreach (string Filename in FileList)
+		{
+			if (!string.IsNullOrEmpty(Filename) && !string.IsNullOrWhiteSpace(Filename))
+			{
+				string Trimmed = Filename.Trim();
+				string SourceFilename = BaseFolder + "\\" + Trimmed;
+				SourceFilename = SourceFilename.Replace('/', '\\');
+				string DestFilename = "/Library/Caches/" + Trimmed.Replace("cookeddata/", "");
+				DestFilename = DestFilename.Replace('\\', '/');
+				DestFilename = "\"" + DestFilename + "\"";
+				SourceFilename = SourceFilename.Replace('\\', Path.DirectorySeparatorChar);
+				string CommandToPush = "push -p \"" + SourceFilename + "\" " + DestFilename + "\n";
+				AllCommandsToPush += CommandToPush;
+			}
+		}
+		System.IO.File.WriteAllText(Directory.GetCurrentDirectory() + "\\CommandsToPush.txt", AllCommandsToPush);
+		LogWarning("BEFORE CRASH " + Directory.GetCurrentDirectory() + "\\CommandsToPush.txt");
+		int ExitCode = 0;
+		Utils.RunLocalProcessAndReturnStdOut(DeviceFS, "-b " + BundleIdentifier + " -x " + Directory.GetCurrentDirectory() + "\\CommandsToPush.txt", out ExitCode, true);
+		if (ExitCode != 0)
+		{
+			LogError("Deployer error : failed to push files. See above log for more details");
+		}
+
+		File.Delete(Directory.GetCurrentDirectory() + "\\CommandsToPush.txt");
+	}
+
 	public override void Deploy(ProjectParams Params, DeploymentContext SC)
 	{
 		if (Params.Devices.Count != 1)
@@ -1729,29 +1792,34 @@ public class IOSPlatform : Platform
 		// Add a commandline for this deploy, if the config allows it.
 		string AdditionalCommandline = (Params.FileServer || Params.CookOnTheFly) ? "" : (" -additionalcommandline " + "\"" + Params.RunCommandline + "\"");
 
+
 		// deploy the .ipa
-		var DeployServer = CombinePaths(CmdEnv.LocalRoot, "Engine/Binaries/DotNET/IOS/DeploymentServer.exe");
+		var DeviceInstaller = GetPathToLibiMobileDeviceTool("ideviceinstaller");
 
 		// check for it in the stage directory
 		string CurrentDir = Directory.GetCurrentDirectory();
 		Directory.SetCurrentDirectory(CombinePaths(CmdEnv.LocalRoot, "Engine/Binaries/DotNET/IOS/"));
 		if (!Params.IterativeDeploy || bCreatedIPA || bNeedsIPA)
 		{
-			RunAndLog(CmdEnv, DeployServer, "Install -ipa \"" + Path.GetFullPath(StagedIPA) + "\"" + (String.IsNullOrEmpty(Params.DeviceNames[0]) ? "" : " -device " + Params.DeviceNames[0]) + AdditionalCommandline);
+			RunAndLog(CmdEnv, DeviceInstaller, "-i " + Path.GetFullPath(StagedIPA));
 		}
+
+		// deploy the assets
 		if (Params.IterativeDeploy)
 		{
-			// push over the changed files
-			RunAndLog(CmdEnv, DeployServer, "Deploy -manifest \"" + SC.GetUFSDeploymentDeltaPath(Params.Devices.Count == 0 ? null : Params.DeviceNames[0]) + "\"" + (Params.Devices.Count == 0 ? "" : " -device " + Params.DeviceNames[0]) + AdditionalCommandline + " -bundle " + BundleIdentifier);
-
-			// IPA for launch-on doesn't include these files - they are a delta and still need uploading
-			if(bNeedsIPA)
+			string BaseFolder = Path.GetDirectoryName(SC.GetUFSDeploymentDeltaPath(Params.DeviceNames[0]));
+			string FilesString = File.ReadAllText(SC.GetUFSDeploymentDeltaPath(Params.DeviceNames[0]));
+			DeployManifestContent(BaseFolder, SC, Params, ref FilesString, ref BundleIdentifier);
+			if (bNeedsIPA)
 			{
-				RunAndLog(CmdEnv, DeployServer, "Deploy -manifest \"" + SC.GetNonUFSDeploymentDeltaPath(Params.Devices.Count == 0 ? null : Params.DeviceNames[0]) + "\"" + (Params.Devices.Count == 0 ? "" : " -device " + Params.DeviceNames[0]) + AdditionalCommandline + " -bundle " + BundleIdentifier);
+				BaseFolder = Path.GetDirectoryName(SC.GetNonUFSDeploymentDeltaPath(Params.DeviceNames[0]));
+				FilesString = File.ReadAllText(SC.GetNonUFSDeploymentDeltaPath(Params.DeviceNames[0]));
+				DeployManifestContent(BaseFolder, SC, Params, ref FilesString, ref BundleIdentifier);
 			}
+			
+			Directory.SetCurrentDirectory(CurrentDir);
+			PrintRunTime();
 		}
-		Directory.SetCurrentDirectory(CurrentDir);
-		PrintRunTime();
 	}
 
 	public override string GetCookPlatform(bool bDedicatedServer, bool bIsClientOnly)
@@ -1790,15 +1858,6 @@ public class IOSPlatform : Platform
 	{
 		return new List<string> { ".dsym", ".udebugsymbols" };
 	}
-
-	// 	void MobileDeviceConnected(object sender, ConnectEventArgs args)
-	// 	{
-	// 	}
-	// 	
-	// 	void MobileDeviceDisconnected(object sender, ConnectEventArgs args)
-	// 	{
-	// 	}
-
 	public override IProcessResult RunClient(ERunOptions ClientRunFlags, string ClientApp, string ClientCmdLine, ProjectParams Params)
 	{
 		if (UnrealBuildTool.BuildHostPlatform.Current.Platform == UnrealTargetPlatform.Mac)
@@ -2180,7 +2239,6 @@ public class IOSPlatform : Platform
 		if (UnrealBuildTool.BuildHostPlatform.Current.Platform != UnrealTargetPlatform.Mac && !CommandUtils.IsEngineInstalled())
 		{
 			Agenda.DotNetProjects.Add(@"Engine\Source\Programs\IOS\MobileDeviceInterface\MobileDeviceInterface.csproj");
-			Agenda.DotNetProjects.Add(@"Engine\Source\Programs\IOS\DeploymentServer\DeploymentServer.csproj");
 		}
 	}
 
