Avoid matching trashed files when searching google drive.
#rnx

diff --git a/Engine/Source/Programs/AutomationTool/Turnkey/Providers/GoogleDriveCopyProvider.cs b/Engine/Source/Programs/AutomationTool/Turnkey/Providers/GoogleDriveCopyProvider.cs
--- a/Engine/Source/Programs/AutomationTool/Turnkey/Providers/GoogleDriveCopyProvider.cs
+++ b/Engine/Source/Programs/AutomationTool/Turnkey/Providers/GoogleDriveCopyProvider.cs
@@ -367,7 +367,7 @@ namespace Turnkey
 					// otherwise, this is a normal component
 					if (bProcessComponent)
 					{
-						List<Google.Apis.Drive.v3.Data.File> Files = ServiceHelper.SearchFiles(string.Format("name = '{0}' and '{1}' in parents", CurrentComponent, ParentFile.Id));
+						List<Google.Apis.Drive.v3.Data.File> Files = ServiceHelper.SearchFiles(string.Format("name = '{0}' and '{1}' in parents and trashed = false", CurrentComponent, ParentFile.Id));
 						if (Files.Count > 1)
 						{
 							TurnkeyUtils.Log("Found more than one item with the name {0}", CurrentComponent);
@@ -405,7 +405,7 @@ namespace Turnkey
 			}
 			else
 			{
-				List<Google.Apis.Drive.v3.Data.File> FilesInDirectory = ServiceHelper.SearchFiles(string.Format("'{0}' in parents", File.Id));
+				List<Google.Apis.Drive.v3.Data.File> FilesInDirectory = ServiceHelper.SearchFiles(string.Format("'{0}' in parents and trashed = false", File.Id));
 				foreach (Google.Apis.Drive.v3.Data.File FileInDir in FilesInDirectory)
 				{
 					DownloadFile(FileInDir, Path.Combine(LocalPath, FileInDir.Name));
