Gauntlet - Update UE.Automation to use new test report API

#rnx

diff --git a/Engine/Source/Programs/AutomationTool/Gauntlet/Unreal/Automation/UE.Automation.cs b/Engine/Source/Programs/AutomationTool/Gauntlet/Unreal/Automation/UE.Automation.cs
--- a/Engine/Source/Programs/AutomationTool/Gauntlet/Unreal/Automation/UE.Automation.cs
+++ b/Engine/Source/Programs/AutomationTool/Gauntlet/Unreal/Automation/UE.Automation.cs
@@ -72,6 +72,12 @@ namespace UE
 		[AutoParam]
 		public string ReportURL = "";
 
+		/// <summary>
+		/// Use Simple Horde Report instead of Unreal Automated Tests
+		/// </summary>
+		[AutoParam]
+		public virtual bool SimpleHordeReport { get; set; } = false;
+
 		/// <summary>
 		/// Validate DDC during tests
 		/// </summary>
@@ -110,6 +116,10 @@ namespace UE
 			// Are we writing out info for Horde?
 			if (WriteTestResultsForHorde)
 			{
+				if (string.IsNullOrEmpty(ReportExportPath))
+				{
+					ReportExportPath = Path.Combine(Globals.TempDir, "TestReport");
+				}
 				if (string.IsNullOrEmpty(HordeTestDataPath))
 				{
 					HordeTestDataPath = HordeReport.DefaultTestDataDir;
@@ -120,7 +130,6 @@ namespace UE
 				}
 			}
 
-		
 			// Arguments for writing out the report and providing a URL where it can be viewed
 			string ReportArgs = string.IsNullOrEmpty(ReportExportPath) ? "" : string.Format("-ReportExportPath=\"{0}\"", ReportExportPath);
 
@@ -399,47 +408,29 @@ namespace UE
 		/// Optional function that is called on test completion and gives an opportunity to create a report
 		/// </summary>
 		/// <param name="Result"></param>
-		/// <param name="Context"></param>
-		/// <param name="Build"></param>
-		public override void CreateReport(TestResult Result, UnrealTestContext Context, UnrealBuildSource Build, IEnumerable<UnrealRoleArtifacts> Artifacts, string ArtifactPath)
+		/// <returns>ITestReport</returns>
+		public override ITestReport CreateReport(TestResult Result)
 		{
 			// Save test result data for Horde build system
-			string ReportPath = GetConfiguration().ReportExportPath;
 			bool WriteTestResultsForHorde = GetConfiguration().WriteTestResultsForHorde;
 			if (WriteTestResultsForHorde)
 			{
-				ITestReport Report = null;
-				if (!string.IsNullOrEmpty(ReportPath))
+				if (GetConfiguration().SimpleHordeReport)
 				{
-					string JsonReportPath = Path.Combine(ReportPath, "index.json");
-					if (File.Exists(JsonReportPath))
-					{
-						Log.Verbose("Reading json Unreal Automated test report from {0}", JsonReportPath);
-						UnrealAutomatedTestPassResults JsonTestPassResults = UnrealAutomatedTestPassResults.LoadFromJson(JsonReportPath);
-						// write test results for Horde
-						string HordeArtifactPath = GetConfiguration().HordeArtifactPath;
-						HordeReport.UnrealEngineTestPassResults HordeTestPassResults = HordeReport.UnrealEngineTestPassResults.FromUnrealAutomatedTests(JsonTestPassResults, ReportPath, GetConfiguration().ReportURL);
-						HordeTestPassResults.CopyTestResultsArtifacts(HordeArtifactPath);
-						Report = HordeTestPassResults;
-					}
-					else
-					{
-						Log.Info("Could not find Unreal Automated test report at {0}. Generating a simple report instead.", JsonReportPath);
-					}
+					return base.CreateReport(Result);
 				}
-				if (Report == null)
+				else
 				{
-					Report = CreateSimpleReportForHorde(Result);
+					string ReportPath = GetConfiguration().ReportExportPath;
+					if (!string.IsNullOrEmpty(ReportPath))
+					{
+						return CreateUnrealEngineTestPassReport(ReportPath, GetConfiguration().ReportURL);
+					}
 				}
-				// write test data collection for Horde
-				string HordeTestDataKey = GetConfiguration().HordeTestDataKey;
-				string HordeTestDataFilePath = Path.Combine(GetConfiguration().HordeTestDataPath, "TestDataCollection.json");
-				HordeReport.TestDataCollection HordeTestDataCollection = new HordeReport.TestDataCollection();
-				HordeTestDataCollection.AddNewTestReport(HordeTestDataKey, Report);
-				HordeTestDataCollection.WriteToJson(HordeTestDataFilePath);
 			}
-		}
 
+			return null;
+		}
 
 		/// <summary>
 		/// Override the summary report so we can create a custom summary with info about our tests and
@@ -574,7 +565,7 @@ namespace UE
 						Parser.GetResults().Where(R => !R.Passed)
 							.SelectMany(R => R.Events
 								.Where(E => E.ToLower().Contains("error"))
-								.Distinct()
+								.Distinct().Select(E => string.Format("[test={0}] {1}", R.DisplayName, E))
 							)
 						);
 				}
