- Fixed cooking for Win64 error

diff --git a/Engine/Source/Programs/AutomationTool/Win/WinPlatform.Automation.cs b/Engine/Source/Programs/AutomationTool/Win/WinPlatform.Automation.cs
--- a/Engine/Source/Programs/AutomationTool/Win/WinPlatform.Automation.cs
+++ b/Engine/Source/Programs/AutomationTool/Win/WinPlatform.Automation.cs
@@ -10,7 +10,7 @@ using Microsoft.Win32;
 using System.Diagnostics;
 using EpicGames.Core;
 
-public abstract class Win64Platform : Platform
+public class Win64Platform : Platform
 {
 	public Win64Platform()
 		: base(UnrealTargetPlatform.Win64)
@@ -126,9 +126,19 @@ public abstract class Win64Platform : Platform
 				}
 			}
 		}
+
+		if (Params.Prereqs)
+		{
+			SC.StageFile(StagedFileType.NonUFS, FileReference.Combine(SC.EngineRoot, "Extras", "Redist", "en-us", "UE4PrereqSetup_x64.exe"));
+		}
+
+		if (!string.IsNullOrWhiteSpace(Params.AppLocalDirectory))
+		{
+			StageAppLocalDependencies(Params, SC, "Win64");
+		}
 	}
 
-    public override void ExtractPackage(ProjectParams Params, string SourcePath, string DestinationPath)
+	public override void ExtractPackage(ProjectParams Params, string SourcePath, string DestinationPath)
     {
     }
 
