- Android Turnkey support for prerequisites (making sure Android Studio has been installed and run once sdk directory creation)
- Updated versioning to pass to SetupAndroid scripts instead of the script hardcoding it
- Added idea of code-based "FileSource" objects that can be used to install an SDK without needing to copy an installer, it's all in the SDKInstall command
- Prereqs are not working yet for Mac/Linux (logic will need to be ported on those platforms to get it right)

diff --git a/Engine/Extras/Android/SetupAndroid.bat b/Engine/Extras/Android/SetupAndroid.bat
--- a/Engine/Extras/Android/SetupAndroid.bat
+++ b/Engine/Extras/Android/SetupAndroid.bat
@@ -4,11 +4,31 @@ set KEY_NAME=HKLM\SOFTWARE\Android Studio
 set VALUE_NAME=Path
 set STUDIO_PATH=
 
+IF "%5" == "-noninteractive" (
+	set PAUSE=
+) ELSE (
+	set PAUSE=pause
+)
+
+SET PLATFORMS_VERSION=%1
+SET BUILDTOOLS_VERSION=%2
+SET CMAKE_VERSION=%3
+SET NDK_VERSION=%4
+
+rem hardcoded versions for compatibility with non-Turnkey manual running
+if "%PLATFORMS_VERSION%" == "" SET PLATFORMS_VERSION=android-28
+if "%BUILDTOOLS_VERSION%" == "" SET BUILDTOOLS_VERSION=28.0.3
+if "%CMAKE_VERSION%" == "" SET CMAKE_VERSION=3.10.2.4988404
+if "%NDK_VERSION%" == "" SET NDK_VERSION=21.1.6352462
+
+
 FOR /F "tokens=2*" %%A IN ('REG.exe query "%KEY_NAME%" /v "%VALUE_NAME%"') DO (set STUDIO_PATH=%%B)
 
-if "%STUDIO_PATH%" == "" (
+IF EXIST "%STUDIO_PATH%" (
+	echo.
+	) ELSE (
 	echo Android Studio not installed, please download Android Studio 3.5.3 from https://developer.android.com/studio
-	pause
+	%PAUSE%
 	exit /b 1
 )
 echo Android Studio Path: %STUDIO_PATH%
@@ -27,8 +47,8 @@ if "%STUDIO_SDK_PATH%" == "" (
 			set STUDIO_SDK_PATH=%ANDROID_HOME%
 		) ELSE (
 			echo Unable to locate local Android SDK location. Did you run Android Studio after installing?
-			pause
-			exit /b 1
+			%PAUSE%
+			exit /b 2
 		)
 	)
 )
@@ -42,7 +62,7 @@ if DEFINED JAVA_HOME (set a=1) ELSE (
 	set JAVA_HOME=%STUDIO_PATH%\jre
 	setx JAVA_HOME "%STUDIO_PATH%\jre"
 )
-set NDKINSTALLPATH=%STUDIO_SDK_PATH%\ndk\21.1.6352462
+set NDKINSTALLPATH=%STUDIO_SDK_PATH%\ndk\%NDK_VERSION%
 set PLATFORMTOOLS=%STUDIO_SDK_PATH%\platform-tools;%STUDIO_SDK_PATH%\tools
 
 set KEY_NAME=HKCU\Environment
@@ -67,17 +87,17 @@ IF EXIST "%SDKMANAGER%" (
 		echo Using sdkmanager: %SDKMANAGER%
 	) ELSE (
 		echo Unable to locate sdkmanager.bat. Did you run Android Studio and install cmdline-tools after installing?
-		pause
-		exit /b 1
+		%PAUSE%
+		exit /b 3
 	)
 )
 
-call "%SDKMANAGER%" "platform-tools" "platforms;android-28" "build-tools;28.0.3" "cmake;3.10.2.4988404" "ndk;21.1.6352462"
+call "%SDKMANAGER%" "platform-tools" "platforms;%PLATFORMS_VERSION%" "build-tools;%BUILDTOOLS_VERSION%" "cmake;%CMAKE_VERSION%" "ndk;%NDK_VERSION%"
 
 IF /I "%ERRORLEVEL%" NEQ "0" (
 	echo Update failed. Please check the Android Studio install.
-	pause
-	exit /b 1
+	%PAUSE%
+	exit /b 4
 )
 
 if EXIST "%NDKINSTALLPATH%" (
@@ -86,9 +106,9 @@ if EXIST "%NDKINSTALLPATH%" (
 	setx NDK_ROOT "%NDKINSTALLPATH%"
 ) ELSE (
 	echo Update failed. Did you accept the license agreement?
-	pause
-	exit /b 1
+	%PAUSE%
+	exit /b 5
 )
 
-pause
+%PAUSE%
 exit /b 0
diff --git a/Engine/Extras/Android/SetupAndroid.command b/Engine/Extras/Android/SetupAndroid.command
--- a/Engine/Extras/Android/SetupAndroid.command
+++ b/Engine/Extras/Android/SetupAndroid.command
@@ -1,11 +1,43 @@
-#!/bin/sh
+#!/usr/bin/env bash
+
+PLATFORMS_VERSION=${1:-}
+BUILDTOOLS_VERSION=${2:-}
+CMAKE_VERSION=${3:-}
+NDK_VERSION=${4:-}
+
+ARG5=${5:-}
+pausefunc()
+{
+    read -rsp $'Press any key to continue...\n' -n1 key
+}
+
+if [[ ${ARG5} == "-noninteractive" ]]; then
+    PAUSE=
+else
+    PAUSE="pausefunc"
+fi
+
+rem hardcoded versions for compatibility with non-Turnkey manual running
+if [[ -z "${PLATFORMS_VERSION}" ]]; then
+    PLATFORMS_VERSION="android-28"
+fi
+if [[ -z "${BUILDTOOLS_VERSION}" ]]; then
+    BUILDTOOLS_VERSION="28.0.3"
+fi
+if [[ -z "${CMAKE_VERSION}" ]]; then
+    CMAKE_VERSION="3.10.2.4988404"
+fi
+if [[ -z "${NDK_VERSION}" ]]; then
+    NDK_VERSION="21.1.6352462"
+fi
+
 
 STUDIO_PATH="~/Applications/Android Studio.app"
 if [ ! -x "$STUDIO_PATH" ]; then
 	STUDIO_PATH="/Applications/Android Studio.app"
 	if [ ! -x "$STUDIO_PATH" ]; then
 		echo Android Studio not installed, please download Android Studio 3.5.3 from https://developer.android.com/studio
-		read -rsp $'Press any key to continue...\n' -n1 key
+		${PAUSE}
 		exit 1
 	fi
 fi
@@ -21,7 +53,7 @@ if [ ! -d "$STUDIO_SDK_PATH" ]; then
 	echo Android SDK not found at: $STUDIO_SDK_PATH
 	echo Unable to locate local Android SDK location. Did you run Android Studio after installing?
 	echo If Android Studio is installed, please run again with SDK path as parameter, otherwise download Android Studio 3.5.3 from https://developer.android.com/studio
-	read -rsp $'Press any key to continue...\n' -n1 key
+	${PAUSE}
 	exit 1
 fi
 echo Android Studio SDK Path: $STUDIO_SDK_PATH
@@ -38,8 +70,8 @@ then
 	echo >> ~/.bash_profile
 	echo "export JAVA_HOME=\"$JAVA_HOME\"" >>~/.bash_profile
 fi
-NDKINSTALLPATH="$STUDIO_SDK_PATH/ndk/21.1.6352462"
-PLATFORMTOOLS="$STUDIO_SDK_PATH/platform-tools:$STUDIO_SDK_PATH/build-tools/28.0.3:$%STUDIO_SDK_PATH/tools/bin"
+NDKINSTALLPATH="$STUDIO_SDK_PATH/ndk/${NDK_VERSION}"
+PLATFORMTOOLS="$STUDIO_SDK_PATH/platform-tools:$STUDIO_SDK_PATH/build-tools/${BUILDTOOLS_VERSION}:$%STUDIO_SDK_PATH/tools/bin"
 
 retVal=$(type -P "adb")
 if [ "$retVal" == "" ]; then
@@ -53,27 +85,27 @@ if [ ! -d "$SDKMANAGERPATH" ]; then
 	SDKMANAGERPATH="$STUDIO_SDK_PATH/cmdline-tools/latest/bin"
 	if [ ! -d "$SDKMANAGERPATH" ]; then
 		echo Unable to locate sdkmanager. Did you run Android Studio and install cmdline-tools after installing?
-		read -rsp $'Press any key to continue...\n' -n1 key
+		${PAUSE}
 		exit 1
 	fi
 fi
 
-"$SDKMANAGERPATH/sdkmanager" "platform-tools" "platforms;android-28" "build-tools;28.0.3" "cmake;3.10.2.4988404" "ndk;21.1.6352462"
+"$SDKMANAGERPATH/sdkmanager" "platform-tools" "platforms;${PLATFORMS_VERSION}" "build-tools;${BUILDTOOLS_VERSION}" "cmake;${CMAKE_VERSION}" "ndk;${NDK_VERSION}"
 
 retVal=$?
 if [ $retVal -ne 0 ]; then
 	echo Update failed. Please check the Android Studio install.
-	read -rsp $'Press any key to continue...\n' -n1 key
+	${PAUSE}
 	exit $retVal
 fi
 
 if [ ! -d "$STUDIO_SDK_PATH/platform-tools" ]; then
 	retVal=1
 fi
-if [ ! -d "$STUDIO_SDK_PATH/platforms/android-28" ]; then
+if [ ! -d "$STUDIO_SDK_PATH/platforms/${PLATFORMS_VERSION}" ]; then
 	retVal=1
 fi
-if [ ! -d "$STUDIO_SDK_PATH/build-tools/28.0.3" ]; then
+if [ ! -d "$STUDIO_SDK_PATH/build-tools/${BUILDTOOLS_VERSION}" ]; then
 	retVal=1
 fi
 if [ ! -d "$NDKINSTALLPATH" ]; then
@@ -82,7 +114,7 @@ fi
 
 if [ $retVal -ne 0 ]; then
 	echo Update failed. Did you accept the license agreement?
-	read -rsp $'Press any key to continue...\n' -n1 key
+	${PAUSE}
 	exit $retVal
 fi
 
diff --git a/Engine/Extras/Android/SetupAndroid.sh b/Engine/Extras/Android/SetupAndroid.sh
--- a/Engine/Extras/Android/SetupAndroid.sh
+++ b/Engine/Extras/Android/SetupAndroid.sh
@@ -3,15 +3,48 @@
 if [[ "$OSTYPE" == "darwin"* ]]; then
 	echo "Please run SetupAndroid.command on MacOSX; attempting to run it for you."
 	DIR="$( cd "$( dirname "${BASH_SOURCE[0]}" )" >/dev/null 2>&1 && pwd )"
-	exec "$DIR"/SetupAndroid.command
+	exec "$DIR"/SetupAndroid.command $@
 	exit 1
 fi
 
+
+PLATFORMS_VERSION=${1:-}
+BUILDTOOLS_VERSION=${2:-}
+CMAKE_VERSION=${3:-}
+NDK_VERSION=${4:-}
+
+ARG5=${5:-}
+pausefunc()
+{
+    read -rsp $'Press any key to continue...\n' -n1 key
+}
+
+if [[ ${ARG5} == "-noninteractive" ]]; then
+    PAUSE=
+else
+    PAUSE="pausefunc"
+fi
+
+rem hardcoded versions for compatibility with non-Turnkey manual running
+if [[ -z "${PLATFORMS_VERSION}" ]]; then
+    PLATFORMS_VERSION="android-28"
+fi
+if [[ -z "${BUILDTOOLS_VERSION}" ]]; then
+    BUILDTOOLS_VERSION="28.0.3"
+fi
+if [[ -z "${CMAKE_VERSION}" ]]; then
+    CMAKE_VERSION="3.10.2.4988404"
+fi
+if [[ -z "${NDK_VERSION}" ]]; then
+    NDK_VERSION="21.1.6352462"
+fi
+
+
 STUDIO_PATH="$HOME/android-studio"
 if [ ! -d "$STUDIO_PATH" ]; then
 	echo "Android Studio not installed, please download Android Studio 3.5.3 from https://developer.android.com/studio"
 	echo "Please download, extract and move to $HOME/android-studio"
-	read -rsp $'Press any key to continue...\n' -n1 key
+	${PAUSE}
 	exit 1
 fi
 echo Android Studio Path: $STUDIO_PATH
@@ -26,7 +59,7 @@ if [ ! -d "$STUDIO_SDK_PATH" ]; then
 	echo Android SDK not found at: $STUDIO_SDK_PATH
 	echo Unable to locate local Android SDK location. Did you run Android Studio after installing?
 	echo If Android Studio is installed, please run again with SDK path as parameter, otherwise download Android Studio 3.5.3 from https://developer.android.com/studio
-	read -rsp $'Press any key to continue...\n' -n1 key
+	${PAUSE}
 	exit 1
 fi
 echo Android Studio SDK Path: $STUDIO_SDK_PATH
@@ -43,8 +76,8 @@ then
 	echo >>$HOME/.bashrc
 	echo "export JAVA_HOME=\"$JAVA_HOME\"" >>$HOME/.bashrc
 fi
-NDKINSTALLPATH="$STUDIO_SDK_PATH/ndk/21.1.6352462"
-PLATFORMTOOLS="$STUDIO_SDK_PATH/platform-tools:$STUDIO_SDK_PATH/build-tools/28.0.3:$%STUDIO_SDK_PATH/tools/bin"
+NDKINSTALLPATH="$STUDIO_SDK_PATH/ndk/${NDK_VERSION}"
+PLATFORMTOOLS="$STUDIO_SDK_PATH/platform-tools:$STUDIO_SDK_PATH/build-tools/${BUILDTOOLS_VERSION}:$%STUDIO_SDK_PATH/tools/bin"
 
 retVal=$(type -P "adb")
 if [ "$retVal" == "" ]; then
@@ -58,27 +91,27 @@ if [ ! -d "$SDKMANAGERPATH" ]; then
 	SDKMANAGERPATH="$STUDIO_SDK_PATH/cmdline-tools/latest/bin"
 	if [ ! -d "$SDKMANAGERPATH" ]; then
 		echo Unable to locate sdkmanager. Did you run Android Studio and install cmdline-tools after installing?
-		read -rsp $'Press any key to continue...\n' -n1 key
+		${PAUSE}
 		exit 1
 	fi
 fi
 
-"$SDKMANAGERPATH/sdkmanager" "platform-tools" "platforms;android-28" "build-tools;28.0.3" "cmake;3.10.2.4988404" "ndk;21.1.6352462"
+"$SDKMANAGERPATH/sdkmanager" "platform-tools" "platforms;${PLATFORMS_VERSION}" "build-tools;${BUILDTOOLS_VERSION}" "cmake;${CMAKE_VERSION}" "ndk;${NDK_VERSION}"
 
 retVal=$?
 if [ $retVal -ne 0 ]; then
 	echo Update failed. Please check the Android Studio install.
-	read -rsp $'Press any key to continue...\n' -n1 key
+	${PAUSE}
 	exit $retVal
 fi
 
 if [ ! -d "$STUDIO_SDK_PATH/platform-tools" ]; then
 	retVal=1
 fi
-if [ ! -d "$STUDIO_SDK_PATH/platforms/android-28" ]; then
+if [ ! -d "$STUDIO_SDK_PATH/platforms/${PLATFORMS_VERSION}" ]; then
 	retVal=1
 fi
-if [ ! -d "$STUDIO_SDK_PATH/build-tools/28.0.3" ]; then
+if [ ! -d "$STUDIO_SDK_PATH/build-tools/${BUILDTOOLS_VERSION}" ]; then
 	retVal=1
 fi
 if [ ! -d "$NDKINSTALLPATH" ]; then
@@ -87,7 +120,7 @@ fi
 
 if [ $retVal -ne 0 ]; then
 	echo Update failed. Did you accept the license agreement?
-	read -rsp $'Press any key to continue...\n' -n1 key
+	${PAUSE}
 	exit $retVal
 fi
 
diff --git a/Engine/Source/Programs/AutomationTool/Android/AndroidPlatform.Automation.cs b/Engine/Source/Programs/AutomationTool/Android/AndroidPlatform.Automation.cs
--- a/Engine/Source/Programs/AutomationTool/Android/AndroidPlatform.Automation.cs
+++ b/Engine/Source/Programs/AutomationTool/Android/AndroidPlatform.Automation.cs
@@ -29,6 +29,164 @@ public class AndroidPlatform : Platform
 
 	}
 
+
+	public override string[] GetCodeSpecifiedSdkVersions()
+	{
+		UEBuildPlatformSDK AndroidSDK = UEBuildPlatformSDK.GetSDKForPlatform("Android");
+
+		return new string[] { AndroidSDK.GetMainVersion() };
+	}
+
+	public override bool GetSDKInstallCommand(out string Command, out string Params, out bool bRequiresPrivilegeElevation, ref string Preamble, ref string SuccessPostamble, ref string FailurePostamble, FileRetriever Retriever)
+	{
+		// run the Setup.bat in the engine, not coming from a normal FileSource
+
+		if (HostPlatform.Current.HostEditorPlatform == UnrealTargetPlatform.Win64)
+		{
+			Command = "$(EngineDir)/Extras/Android/SetupAndroid.bat";
+		}
+		else if (HostPlatform.Current.HostEditorPlatform == UnrealTargetPlatform.Mac)
+		{
+			Command = "$(EngineDir)/Extras/Android/SetupAndroid.command";
+		}
+		else
+		{
+			Command = "$(EngineDir)/Extras/Android/SetupAndroid.sh";
+		}
+
+		// pull the desired version numbers to install
+		UEBuildPlatformSDK AndroidSDK = UEBuildPlatformSDK.GetSDKForPlatform("Android");
+		string PlatformsVersion = AndroidSDK.GetPlatformSpecificVersion("platforms");
+		string BuildToolsVersion = AndroidSDK.GetPlatformSpecificVersion("build-tools");
+		string CMakeVersion = AndroidSDK.GetPlatformSpecificVersion("cmake");
+		string NDKVersion = AndroidSDK.GetPlatformSpecificVersion("ndk");
+
+		Params = $"{PlatformsVersion} {BuildToolsVersion} {CMakeVersion} {NDKVersion} -noninteractive";
+
+		bRequiresPrivilegeElevation = false;
+
+		return true;
+	}
+
+	private static string GetAndroidStudioExe()
+	{
+		string DefaultAndroidStudioInstallDir = Path.Combine(Environment.GetFolderPath(Environment.SpecialFolder.ProgramFiles), "Android", "Android Studio");
+		string RegValue = Microsoft.Win32.Registry.GetValue("HKEY_LOCAL_MACHINE\\SOFTWARE\\Android", "Path", null) as string;
+		string AndroidStudioInstallDir = RegValue == null ? DefaultAndroidStudioInstallDir : RegValue;
+		return Path.Combine(AndroidStudioInstallDir, "bin", "studio64.exe");
+	}
+	private static string GetSdkDir()
+	{
+		string DefaultSdkDir = Path.Combine(Environment.GetFolderPath(Environment.SpecialFolder.LocalApplicationData), "Android", "Sdk");
+		string RegValue = Microsoft.Win32.Registry.GetValue("HKEY_LOCAL_MACHINE\\SOFTWARE\\Android", "SdkPath", null) as string;
+		return RegValue == null ? DefaultSdkDir : RegValue;
+	}
+
+	public override bool UpdateHostPrerequisites(BuildCommand Command, FileRetriever Retriever, bool bVerifyOnly)
+	{
+		// @todo turnkey: Handle Mac/Linux
+		if (HostPlatform.Current.HostEditorPlatform != UnrealTargetPlatform.Win64)
+		{
+			return false;
+		}
+
+		string AndroidStudioExe = GetAndroidStudioExe();
+		string SdkDir = GetSdkDir();
+
+		bool bIsInstalled = File.Exists(AndroidStudioExe) && Directory.Exists(SdkDir);
+
+		// if we are only verifying, just return the status, and if it's installed, we are done!
+		if (bVerifyOnly || bIsInstalled)
+		{
+			if (!File.Exists(AndroidStudioExe))
+			{
+//				Retriever.ReportError("Android Studio is not installed correctly.");
+				Retriever.PauseForUser("Android Studio is not installed correctly.");
+			}
+			if (!Directory.Exists(SdkDir))
+			{
+//				Retriever.ReportError("Android SDK directory is not set correctly.");
+				Retriever.PauseForUser("Android SDK directory is not set correctly.");
+			}
+			return bIsInstalled;
+		}
+
+		if (!File.Exists(AndroidStudioExe))
+		{
+			// get AS installer
+			string OutputPath = Retriever.RetrieveFileSource("AndroidStudio");
+
+			// Unset some envvars in case autosdk ran - they will mess up the first run of Android Studio
+			string[] Vars = new string[]
+			{
+					"ANDROID_HOME",
+					"ANDROID_SDK_HOME",
+					"JAVA_HOME",
+					"NDKROOT",
+					"NDK_ROOT",
+					"ANDROID_NDK_ROOT",
+					"ANDROID_SWT"
+			};
+			Array.ForEach(Vars, x => Environment.SetEnvironmentVariable(x, null));
+
+			if (OutputPath == null)
+			{
+				Retriever.PauseForUser("Unable to find Android Studio installer. Please download and install Android Studio 3.5.3 from https://developer.android.com/studio before continuing.\n\nMake sure to use the Run Android Studio and complete the first-time setup!\n\nChoose all default options unless you know what you are doing.");
+			}
+			else
+			{
+				Retriever.PauseForUser("Running the Android Studio installer, and then Android Studio for first-time setup!\n\nChoose all default options unless you know what you are doing.");
+				int ExitCode;
+				Utils.RunLocalProcessAndReturnStdOut(OutputPath, "/S", out ExitCode, true);
+				//				Utils.RunLocalProcessAndReturnStdOut(OutputPath, "", out ExitCode, true);
+
+				// AS installer returns 1223 even on success ("user canceled" even tho there's no UI to cancel it...) when running with /S
+				if (ExitCode != 0 && ExitCode != 1223)
+				{
+//					Retriever.ReportError($"Android Studio installer failed. ExitCode = {ExitCode}");
+					Retriever.PauseForUser($"Android Studio installer failed. ExitCode = {ExitCode}");
+					return false;
+				}
+
+				Utils.RunLocalProcessAndReturnStdOut(GetAndroidStudioExe(), "", out _, false);
+			}
+		}
+		else
+		{
+			Retriever.PauseForUser("The Sdk directory was not found. Running the Android Studio to perform first-time setup.\n\nChoose all default options unless you know what you are doing.");
+
+			Utils.RunLocalProcessAndReturnStdOut(AndroidStudioExe, "", out _, false);
+		}
+
+		// check to see if the installation worked. If so, continue on!
+		AndroidStudioExe = GetAndroidStudioExe();
+		SdkDir = GetSdkDir();
+
+		bIsInstalled = File.Exists(AndroidStudioExe) && Directory.Exists(SdkDir);
+
+		if (!File.Exists(AndroidStudioExe))
+		{
+//			Retriever.ReportError("Android Studio is not installed correctly, after attempted installation.");
+			Retriever.PauseForUser("Android Studio is not installed correctly, after attempted installation.");
+		}
+		if (!Directory.Exists(SdkDir))
+		{
+//			Retriever.ReportError("Android SDK directory is not set correctly, after attempted installation.");
+			Retriever.PauseForUser("Android SDK directory is not set correctly, after attempted installation.");
+		}
+
+		return bIsInstalled;
+	}
+
+
+
+
+
+
+
+
+
+
 	private static string GetSONameWithoutArchitecture(ProjectParams Params, string DecoratedExeName)
 	{
 		return Path.Combine(Path.GetDirectoryName(Params.GetProjectExeForPlatform(UnrealTargetPlatform.Android).ToString()), DecoratedExeName) + ".so";
diff --git a/Engine/Source/Programs/AutomationTool/AutomationUtils/Platform.cs b/Engine/Source/Programs/AutomationTool/AutomationUtils/Platform.cs
--- a/Engine/Source/Programs/AutomationTool/AutomationUtils/Platform.cs
+++ b/Engine/Source/Programs/AutomationTool/AutomationUtils/Platform.cs
@@ -194,15 +194,8 @@ namespace AutomationTool
 			return null;
 		}
 
-		public virtual bool InstallSDK(BuildCommand BuildCommand, FileRetriever Retriever, object HintObject, DeviceInfo Device)
+		public virtual bool InstallSDK(BuildCommand BuildCommand, FileRetriever Retriever, DeviceInfo Device)
 		{
-			string DownloadedSDK = Retriever.RetrieveFileSource(HintObject);
-
-			if (string.IsNullOrEmpty(DownloadedSDK))
-			{
-				return false;
-			}
-
 			string Command, Params, Preamble="", SuccessPostamble="", FailurePostamble="";
 
 			bool bRequiresPrivilegeElevation;
@@ -218,6 +211,16 @@ namespace AutomationTool
 			return false;
 		}
 
+		/// <summary>
+		/// Return a list of versions that will be used to create "fake" FileSource objects which are used
+		/// for install Sdks where no file downloads are needed
+		/// </summary>
+		/// <returns></returns>
+		public virtual string[] GetCodeSpecifiedSdkVersions()
+		{
+			return new string[] { };
+		}
+
 		public virtual bool GetSDKInstallCommand(out string Command, out string Params, out bool bRequiresPrivilegeElevation, ref string Preamble, ref string SuccessPostamble, ref string FailurePostamble, FileRetriever Retriever)
 		{
 			Command = null;
@@ -239,20 +242,6 @@ namespace AutomationTool
 			return null;
 		}
 
-		public virtual bool ShouldPerformManualSDKInstall()
-		{
-			return false;
-		}
-		public virtual bool ShouldPerformManualDeviceSoftwareInstall()
-		{
-			return false;
-		}
-
-		public virtual bool ManualInstallSDK(BuildCommand BuildCommand, FileRetriever Retriever, DeviceInfo Device = null)
-		{
-			return false;
-		}
-
 		public virtual bool UpdateHostPrerequisites(BuildCommand Command, FileRetriever Retriever, bool bVerifyOnly)
 		{
 			return true;
diff --git a/Engine/Source/Programs/AutomationTool/Turnkey/Commands/InstallSdk.cs b/Engine/Source/Programs/AutomationTool/Turnkey/Commands/InstallSdk.cs
--- a/Engine/Source/Programs/AutomationTool/Turnkey/Commands/InstallSdk.cs
+++ b/Engine/Source/Programs/AutomationTool/Turnkey/Commands/InstallSdk.cs
@@ -242,6 +242,20 @@ namespace Turnkey.Commands
 						}
 					}
 
+					// make sure prerequisites are in good shape
+					CopyProviderRetriever Retriever = new CopyProviderRetriever();
+					if (AutomationPlatform.UpdateHostPrerequisites(TurnkeyUtils.CommandUtilHelper, Retriever, false) == false)
+					{
+						// Retriever.ReportError("Failed to update host prerequisites");
+						return;
+					}
+					if (InstallDevice != null && 
+						AutomationPlatform.UpdateDevicePrerequisites(InstallDevice, TurnkeyUtils.CommandUtilHelper, Retriever, false) == false)
+					{
+						// Retriever.ReportError("Failed to update device prerequisites");
+						return;
+					}
+
 					Sdk.DownloadOrInstall(Platform, InstallDevice);
 				}
 			}
diff --git a/Engine/Source/Programs/AutomationTool/Turnkey/Commands/VerifySdk.cs b/Engine/Source/Programs/AutomationTool/Turnkey/Commands/VerifySdk.cs
--- a/Engine/Source/Programs/AutomationTool/Turnkey/Commands/VerifySdk.cs
+++ b/Engine/Source/Programs/AutomationTool/Turnkey/Commands/VerifySdk.cs
@@ -114,24 +114,6 @@ namespace Turnkey.Commands
 				// install if out of date, or if forcing it
 				if (bForceSdkInstall || (bUpdateIfNeeded && bHasOutOfDateSDK))
 				{
-// 						if (!bUnattended)
-// 						{
-// 							string Response = TurnkeyUtils.ReadInput("Your Sdk installation is not up to date. Would you like to install a valid Sdk? [Y/n]", "Y");
-// 							if (string.Compare(Response, "Y", true) != 0)
-// 							{
-// 								continue;
-// 							}
-// 						}
-
-					if (AutomationPlatform.ShouldPerformManualSDKInstall())
-					{
-						if (AutomationPlatform.ManualInstallSDK(TurnkeyUtils.CommandUtilHelper, new CopyProviderRetriever(), null) == false)
-						{
-							TurnkeyUtils.Log("Install failed!");
-						}
-						continue;
-					}
-
 					// if the platform has a valid sdk but isn't the "installed one", then try to switch to it
 					if ((LocalState & SdkUtils.LocalAvailability.InstalledSdk_ValidInactiveVersionExists) != 0)
 					{
diff --git a/Engine/Source/Programs/AutomationTool/Turnkey/FileSource.cs b/Engine/Source/Programs/AutomationTool/Turnkey/FileSource.cs
--- a/Engine/Source/Programs/AutomationTool/Turnkey/FileSource.cs
+++ b/Engine/Source/Programs/AutomationTool/Turnkey/FileSource.cs
@@ -113,7 +113,17 @@ namespace Turnkey
 		public string BuildPlatformEnumerationSuffix;
 
 
-
+		public static FileSource CreateCodeSpecifiedSource(string Name, string Version, UnrealTargetPlatform Platform)
+		{
+			FileSource NewSource = new FileSource();
+			NewSource.PlatformString = Platform.ToString();
+			NewSource.Platforms = new List<UnrealTargetPlatform>() { Platform };
+			NewSource.Name = Name;
+			NewSource.Version = Version;
+			NewSource.Sources = new CopySource[] { };
+
+			return NewSource;
+		}
 
 		public FileSource CloneForExpansion(string NewValue, Action<FileSource, string> Setter)
 		{
@@ -333,8 +343,24 @@ namespace Turnkey
 				return SdkUtils.SetupAutoSdk(this, Platform, bUnattended);
 			}
 
+			CopyProviderRetriever Retriever = new CopyProviderRetriever();
+
+			// if we have sources, make sure we can download something. if we have a null Sources, that indicates
+			// it was code generated, and we don't need to download anything
+			if (Sources.Length != 0)
+			{
+				// this will set the $(CopyOutputPath) variable which the Sdks will generally need to use
+				string DownloadedSDK = Retriever.RetrieveFileSource(this);
+
+				if (string.IsNullOrEmpty(DownloadedSDK))
+				{
+//					Retriever.ReportError($"Unable to download anInstaller for {Platform}. Your Studio's TurnkeyManifest.xml file(s) may need to be fixed.");
+					return false;
+				}
+			}
+
 			// let the platform decide how to install
-			return AutomationTool.Platform.GetPlatform(Platform).InstallSDK(TurnkeyUtils.CommandUtilHelper, new CopyProviderRetriever(), this, Device);
+			return AutomationTool.Platform.GetPlatform(Platform).InstallSDK(TurnkeyUtils.CommandUtilHelper, Retriever, Device);
 		}
 
 
diff --git a/Engine/Source/Programs/AutomationTool/Turnkey/TurnkeyManifest.cs b/Engine/Source/Programs/AutomationTool/Turnkey/TurnkeyManifest.cs
--- a/Engine/Source/Programs/AutomationTool/Turnkey/TurnkeyManifest.cs
+++ b/Engine/Source/Programs/AutomationTool/Turnkey/TurnkeyManifest.cs
@@ -166,6 +166,17 @@ namespace Turnkey
 						DiscoveredFileSources.AddRange(x.FileSources);
 					}
 				});
+
+				// also manually create local "FileSource" objects specified via code
+				foreach (UnrealTargetPlatform Platform in UnrealTargetPlatform.GetValidPlatforms())
+				{
+					// this is usually going to be empty
+					foreach (string SpecifiedVersion in AutomationTool.Platform.GetPlatform(Platform).GetCodeSpecifiedSdkVersions())
+					{
+						FileSource CodeSource = FileSource.CreateCodeSpecifiedSource($"{Platform} SDK {SpecifiedVersion}", SpecifiedVersion, Platform);
+						DiscoveredFileSources.Add(CodeSource);
+					}
+				}
 			}
 		}
 
diff --git a/Engine/Source/Programs/UnrealBuildTool/Platform/Android/AndroidPlatformSDK.cs b/Engine/Source/Programs/UnrealBuildTool/Platform/Android/AndroidPlatformSDK.cs
--- a/Engine/Source/Programs/UnrealBuildTool/Platform/Android/AndroidPlatformSDK.cs
+++ b/Engine/Source/Programs/UnrealBuildTool/Platform/Android/AndroidPlatformSDK.cs
@@ -33,10 +33,19 @@ namespace UnrealBuildTool
 			MinVersion = MaxVersion = null;
 		}
 
+		public override string GetPlatformSpecificVersion(string VersionType)
+		{
+			switch (VersionType.ToLower())
+			{
+				case "platforms": return "android-28";
+				case "build-tools": return "28.0.3";
+				case "cmake": return "3.10.2.4988404";
+				case "ndk": return "21.4.7075529";
+			}
 
+			return "";
+		}
 
-		// the version we prefer
-		private string RecommendedVersion = "r21b";
 
 		public override string GetInstalledSDKVersion()
 		{
@@ -120,9 +129,9 @@ namespace UnrealBuildTool
 		{
 			SDKStatus Validity = base.PrintSDKInfoAndReturnValidity(Verbosity, Options, ErrorVerbosity, ErrorOptions);
 
-			if (GetInstalledVersion() != RecommendedVersion)
+			if (GetInstalledVersion() != GetMainVersion())
 			{
-				Log.WriteLine(Verbosity, Options, "Note: Android toolchain NDK {0} recommended", RecommendedVersion);
+				Log.WriteLine(Verbosity, Options, "Note: Android toolchain NDK {0} recommended", GetMainVersion());
 			}
 
 			return Validity;
