Android shouldn't PauseForUser with -VerifySDK

diff --git a/Engine/Source/Programs/AutomationTool/Android/AndroidPlatform.Automation.cs b/Engine/Source/Programs/AutomationTool/Android/AndroidPlatform.Automation.cs
--- a/Engine/Source/Programs/AutomationTool/Android/AndroidPlatform.Automation.cs
+++ b/Engine/Source/Programs/AutomationTool/Android/AndroidPlatform.Automation.cs
@@ -101,12 +101,12 @@ public class AndroidPlatform : Platform
 			if (!File.Exists(AndroidStudioExe))
 			{
 //				Retriever.ReportError("Android Studio is not installed correctly.");
-				Retriever.PauseForUser("Android Studio is not installed correctly.");
+				//Retriever.PauseForUser("Android Studio is not installed correctly.");
 			}
 			if (!Directory.Exists(SdkDir))
 			{
 //				Retriever.ReportError("Android SDK directory is not set correctly.");
-				Retriever.PauseForUser("Android SDK directory is not set correctly.");
+				//Retriever.PauseForUser("Android SDK directory is not set correctly.");
 			}
 			return bIsInstalled;
 		}
