Fixed P4 usage error in UAT localization script

#rnx

diff --git a/Engine/Source/Programs/AutomationTool/Scripts/Localisation.Automation.cs b/Engine/Source/Programs/AutomationTool/Scripts/Localisation.Automation.cs
--- a/Engine/Source/Programs/AutomationTool/Scripts/Localisation.Automation.cs
+++ b/Engine/Source/Programs/AutomationTool/Scripts/Localisation.Automation.cs
@@ -428,7 +428,7 @@ class Localize : BuildCommand
 						}
 					}
 
-					P4.LogP4("", String.Format("-x{0} revert", P4RevertArgsFilename));
+					P4.LogP4(String.Format("-x {0}", P4RevertArgsFilename), "revert");
 					DeleteFile_NoExceptions(P4RevertArgsFilename);
 				}
 			}
