Fix scheme error when packaging Blueprint iOS projects from Mac

diff --git a/Engine/Source/Programs/AutomationTool/IOS/IOSPlatform.Automation.cs b/Engine/Source/Programs/AutomationTool/IOS/IOSPlatform.Automation.cs
--- a/Engine/Source/Programs/AutomationTool/IOS/IOSPlatform.Automation.cs
+++ b/Engine/Source/Programs/AutomationTool/IOS/IOSPlatform.Automation.cs
@@ -877,7 +877,7 @@ public class IOSPlatform : Platform
 		Params.AutomaticSigning = bAutomaticSigning;
 
 		// Scheme name and configuration for code signing with Xcode project
-		string SchemeName = Params.IsCodeBasedProject ? Params.RawProjectPath.GetFileNameWithoutExtension() : "UE4";
+		string SchemeName = Params.IsCodeBasedProject ? Params.RawProjectPath.GetFileNameWithoutExtension() : "UE5";
 		string SchemeConfiguration = TargetConfiguration.ToString();
 		if (Params.Client)
 		{
