- Renamed FileRetriever class to ITurnkeyContext
- Enhanced Turnkey error reporting to the editor (TurnkeyContext now has Log and ReportError)
- Improved the SDK display in the editor
- Various bug fixes
- Renamed GetAutoSDKDirectoryForMasterVersion() to GetAutoSDKDirectoryForMainVersion()

diff --git a/Engine/Source/Editor/TurnkeySupport/Private/TurnkeySupportModule.cpp b/Engine/Source/Editor/TurnkeySupport/Private/TurnkeySupportModule.cpp
--- a/Engine/Source/Editor/TurnkeySupport/Private/TurnkeySupportModule.cpp
+++ b/Engine/Source/Editor/TurnkeySupport/Private/TurnkeySupportModule.cpp
@@ -765,7 +765,47 @@ static FSlateIcon MakePlatformSdkIconAttribute(FName IniPlatformName, TSharedPtr
 //		));
 }
 
+static FText FormatSdkInfo(const FTurnkeySdkInfo& SdkInfo, bool bIncludeAutoSdk)
+{
+	FFormatOrderedArguments Args;
+	Args.Add(FText::FromString(SdkInfo.InstalledVersion));
+	Args.Add(FText::FromString(SdkInfo.AutoSDKVersion));
+	Args.Add(FText::FromString(SdkInfo.MinAllowedVersion));
+	Args.Add(FText::FromString(SdkInfo.MaxAllowedVersion));
+	Args.Add(SdkInfo.SdkErrorInformation);
+
+	TArray<FText> Lines;
+	Lines.Add(FText::Format(LOCTEXT("SdkInfo_Installed", "Installed SDK: {0}"), Args));
+	if (bIncludeAutoSdk)
+	{
+		Lines.Add(FText::Format(LOCTEXT("SdkInfo_AutoSDK", "AutoSDK: {1}"), Args));
+	}
+	
+	if (SdkInfo.MinAllowedVersion == SdkInfo.MaxAllowedVersion)
+	{
+		Lines.Add(FText::Format(LOCTEXT("SdkInfo_AllowedSDK_Single", "Allowed Version: {2}"), Args));
+	}
+	else if (SdkInfo.MinAllowedVersion == TEXT(""))
+	{
+		Lines.Add(FText::Format(LOCTEXT("SdkInfo_AllowedSDK_MaxOnly", "Allowed Versions: Up to {3}"), Args));
+	}
+	else if (SdkInfo.MaxAllowedVersion == TEXT(""))
+	{
+		Lines.Add(FText::Format(LOCTEXT("SdkInfo_AllowedSDK_MinOnly", "Allowed Versions: {2} and up"), Args));
+	}
+	else
+	{
+		Lines.Add(FText::Format(LOCTEXT("SdkInfo_AllowedSDK_Range", "Allowed Versions: {2} through {3}"), Args));
+	}
+
+	if (!SdkInfo.SdkErrorInformation.IsEmpty())
+	{
+		Lines.Add(FText::Format(LOCTEXT("SdkInfo_Error", "Error Info:\n{4}"), Args));
+	}
 
+	// now make a single \n delimted text
+	return FText::Join(FText::FromString(TEXT("\n")), Lines);
+}
 
 
 
@@ -917,27 +957,33 @@ static void MakeTurnkeyPlatformMenu(FMenuBuilder& MenuBuilder, FName IniPlatform
 				FNewMenuDelegate::CreateLambda([UBTPlatformString, IniPlatformName, DeviceName, DeviceId](FMenuBuilder& SubMenuBuilder)
 				{
 					FTurnkeySdkInfo SdkInfo = ITurnkeySupportModule::Get().GetSdkInfoForDeviceId(DeviceId);
-					FFormatOrderedArguments Args = { FText::FromString(SdkInfo.InstalledVersion), FText::FromString(SdkInfo.MinAllowedVersion),FText::FromString(SdkInfo.MaxAllowedVersion) };
+
 					SubMenuBuilder.AddWidget(
 						SNew(STextBlock)
 						.ColorAndOpacity(FSlateColor::UseSubduedForeground())
-						.Text(FText::Format(LOCTEXT("SdkInfo", "Manual Installed SDK: {0}\nAllowedVersions: {1}-{2}"), Args)),
+						.Text(FormatSdkInfo(SdkInfo, false)),
 						FText::GetEmpty()
 					);
 
-					SubMenuBuilder.AddMenuEntry(
-						LOCTEXT("Turnkey_RepairDevice", "Repair Device as Needed"),
-						LOCTEXT("TurnkeyTooltip_RepairDevice", "Perform any fixup that may be needed on this device. If up to date already, nothing will be done."),
-						FSlateIcon(),
-						FExecuteAction::CreateStatic(TurnkeyInstallSdk, UBTPlatformString, false, false, DeviceId)
-					);
+					if (SdkInfo.Status == ETurnkeyPlatformSdkStatus::FlashValid)
+					{
+						SubMenuBuilder.AddMenuEntry(
+							LOCTEXT("Turnkey_ForceRepairDevice", "Repair Device"),
+							LOCTEXT("TurnkeyTooltip_ForceRepairDevice", "Force repairing anything on the device needed (update firmware, etc). Will perform all steps possible, even if not needed."),
+							FSlateIcon(),
+							FExecuteAction::CreateStatic(TurnkeyInstallSdk, UBTPlatformString, true, false, DeviceId)
+						);
+					}
+					else
+					{
+						SubMenuBuilder.AddMenuEntry(
+							LOCTEXT("Turnkey_RepairDevice", "Update Device"),
+							LOCTEXT("TurnkeyTooltip_RepairDevice", "Perform any fixup that may be needed on this device. If up to date already, nothing will be done."),
+							FSlateIcon(),
+							FExecuteAction::CreateStatic(TurnkeyInstallSdk, UBTPlatformString, false, false, DeviceId)
+						);
+					}
 
-					SubMenuBuilder.AddMenuEntry(
-						LOCTEXT("Turnkey_ForceRepairDevice", "Force Repair Device"),
-						LOCTEXT("TurnkeyTooltip_ForceRepairDevice", "Force repairing anything on the device needed (update firmware, etc). Will perform all steps possible, even if not needed."),
-						FSlateIcon(),
-						FExecuteAction::CreateStatic(TurnkeyInstallSdk, UBTPlatformString, true, false, DeviceId)
-					);
 				}),
 				false, MakePlatformSdkIconAttribute(IniPlatformName, Proxy)
 			);
@@ -950,11 +996,26 @@ static void MakeTurnkeyPlatformMenu(FMenuBuilder& MenuBuilder, FName IniPlatform
 	MenuBuilder.BeginSection("SdkManagement", LOCTEXT("TurnkeySection_Sdks", "Sdk Managment"));
 
 	const FTurnkeySdkInfo& SdkInfo = ITurnkeySupportModule::Get().GetSdkInfo(IniPlatformName, true);
-	FFormatOrderedArguments Args = { FText::FromString(SdkInfo.InstalledVersion), FText::FromString(SdkInfo.AutoSDKVersion),FText::FromString(SdkInfo.MinAllowedVersion),FText::FromString(SdkInfo.MaxAllowedVersion) };
+	FFormatOrderedArguments Args;
+	Args.Add(FText::FromString(SdkInfo.InstalledVersion));
+	Args.Add(FText::FromString(SdkInfo.AutoSDKVersion));
+	Args.Add(FText::FromString(SdkInfo.MinAllowedVersion));
+	Args.Add(FText::FromString(SdkInfo.MaxAllowedVersion));
+	if (SdkInfo.SdkErrorInformation.IsEmpty())
+	{
+		Args.Add(FText::GetEmpty());
+		Args.Add(FText::GetEmpty());
+	}
+	else
+	{
+		Args.Add(LOCTEXT("ErrorPrefix", "\nErrors:\n"));
+		Args.Add(FText::GetEmpty());
+	}
+
 	MenuBuilder.AddWidget(
 		SNew(STextBlock)
 		.ColorAndOpacity(FSlateColor::UseSubduedForeground())
-		.Text(FText::Format(LOCTEXT("SdkInfo", "Manual Installed SDK: {0}\nAutoSDK: {1}\nAllowedVersions: {2}-{3}"), Args)),
+		.Text(FormatSdkInfo(SdkInfo, true)),
 		FText::GetEmpty()
 	);
 
@@ -964,7 +1025,7 @@ static void MakeTurnkeyPlatformMenu(FMenuBuilder& MenuBuilder, FName IniPlatform
 		if (SdkInfo.Status == ETurnkeyPlatformSdkStatus::OutOfDate)
 		{
 			MenuBuilder.AddMenuEntry(
-				LOCTEXT("Turnkey_InstallSdkMinimal", "Update Sdk"),
+				LOCTEXT("Turnkey_UpdateSdkMinimal", "Update Sdk"),
 				LOCTEXT("TurnkeyTooltip_InstallSdkMinimal", "Attempt to update an Sdk, as hosted by your studio. Will attempt to install a minimal Sdk (useful for building/running only)"),
 				FSlateIcon(),
 				FExecuteAction::CreateStatic(TurnkeyInstallSdk, UBTPlatformString, false, false, NoDevice)
@@ -973,8 +1034,8 @@ static void MakeTurnkeyPlatformMenu(FMenuBuilder& MenuBuilder, FName IniPlatform
 			if (SdkInfo.bCanInstallFullSdk && SdkInfo.bCanInstallAutoSdk)
 			{
 				MenuBuilder.AddMenuEntry(
-					LOCTEXT("Turnkey_InstallSdkFull", "Update Sdk (Prefer Full)"),
-					LOCTEXT("TurnkeyTooltip_InstallSdkMinimal", "Attempt to update an Sdk, as hosted by your studio. Will attempt to install a full Sdk (useful profiling or other use cases)"),
+					LOCTEXT("Turnkey_UpdateSdkFull", "Update Sdk (Full Platform Installer)"),
+					LOCTEXT("TurnkeyTooltip_UpdateSdkFull", "Attempt to update an Sdk, as hosted by your studio. Will attempt to install a full Sdk (useful profiling or other use cases)"),
 					FSlateIcon(),
 					FExecuteAction::CreateStatic(TurnkeyInstallSdk, UBTPlatformString, true, false, NoDevice)
 				);
@@ -983,31 +1044,34 @@ static void MakeTurnkeyPlatformMenu(FMenuBuilder& MenuBuilder, FName IniPlatform
 		else if (SdkInfo.Status == ETurnkeyPlatformSdkStatus::Valid)
 		{
 			MenuBuilder.AddMenuEntry(
-				LOCTEXT("Turnkey_InstallSdkMinimal", "Force Reinstall Sdk (Prefer Minimal)"),
-				LOCTEXT("TurnkeyTooltip_InstallSdkMinimal", "Attempt to force re-install an Sdk, as hosted by your studio. Will attempt to install a minimal Sdk (useful for building/running only)"),
+				LOCTEXT("Turnkey_ForceSdkMinimal", "Force Reinstall Sdk"),
+				LOCTEXT("TurnkeyTooltip_ForceSdkMinimal", "Attempt to force re-install an Sdk, as hosted by your studio. Will attempt to install a minimal Sdk (useful for building/running only)"),
 				FSlateIcon(),
 				FExecuteAction::CreateStatic(TurnkeyInstallSdk, UBTPlatformString, false, true, NoDevice)
 			);
 
-			MenuBuilder.AddMenuEntry(
-				LOCTEXT("Turnkey_InstallSdkFull", "Force Reinstall (Prefer Full)"),
-				LOCTEXT("TurnkeyTooltip_InstallSdkMinimal", "Attempt to force re-install an Sdk, as hosted by your studio. Will attempt to install a full Sdk (useful profiling or other use cases)"),
-				FSlateIcon(),
-				FExecuteAction::CreateStatic(TurnkeyInstallSdk, UBTPlatformString, true, true, NoDevice)
-			);
+			if (SdkInfo.bCanInstallFullSdk && SdkInfo.bCanInstallAutoSdk)
+			{
+				MenuBuilder.AddMenuEntry(
+					LOCTEXT("Turnkey_ForceSdkFull", "Force Reinstall (Full Platform Installer)"),
+					LOCTEXT("TurnkeyTooltip_ForceSdkForce", "Attempt to force re-install an Sdk, as hosted by your studio. Will attempt to install a full Sdk (useful profiling or other use cases)"),
+					FSlateIcon(),
+					FExecuteAction::CreateStatic(TurnkeyInstallSdk, UBTPlatformString, true, true, NoDevice)
+				);
+			}
 		}
 		else
 		{
 			MenuBuilder.AddMenuEntry(
-				LOCTEXT("Turnkey_InstallSdkMinimal", "Install Sdk (Prefer Minimal)"),
+				LOCTEXT("Turnkey_InstallSdkMinimal", "Install Sdk"),
 				LOCTEXT("TurnkeyTooltip_InstallSdkMinimal", "Attempt to install an Sdk, as hosted by your studio. Will attempt to install a minimal Sdk (useful for building/running only)"),
 				FSlateIcon(),
 				FExecuteAction::CreateStatic(TurnkeyInstallSdk, UBTPlatformString, false, false, NoDevice)
 			);
 
 			MenuBuilder.AddMenuEntry(
-				LOCTEXT("Turnkey_InstallSdkFull", "Install Sdk (Prefer Full)"),
-				LOCTEXT("TurnkeyTooltip_InstallSdkMinimal", "Attempt to install an Sdk, as hosted by your studio. Will attempt to install a full Sdk (useful profiling or other use cases)"),
+				LOCTEXT("Turnkey_InstallSdkFull", "Install Sdk (Full Platform Installer)"),
+				LOCTEXT("TurnkeyTooltip_InstallSdkFull", "Attempt to install an Sdk, as hosted by your studio. Will attempt to install a full Sdk (useful profiling or other use cases)"),
 				FSlateIcon(),
 				FExecuteAction::CreateStatic(TurnkeyInstallSdk, UBTPlatformString, true, false, NoDevice)
 			);
@@ -1495,6 +1559,9 @@ bool GetSdkInfoFromTurnkey(FString Line, FName& PlatformName, FString& DeviceId,
 	FParse::Value(*Info, TEXT("AutoSDK="), SdkInfo.AutoSDKVersion);
 	FParse::Value(*Info, TEXT("MinAllowed="), SdkInfo.MinAllowedVersion);
 	FParse::Value(*Info, TEXT("MaxAllowed="), SdkInfo.MaxAllowedVersion);
+	FString ErrorString;
+	FParse::Value(*Info, TEXT("Error="), ErrorString);
+	SdkInfo.SdkErrorInformation = FText::FromString(ErrorString.Replace(TEXT("|"), TEXT("\n")));
 
 	SdkInfo.Status = ETurnkeyPlatformSdkStatus::Unknown;
 	if (StatusString == TEXT("Valid"))
@@ -1717,9 +1784,6 @@ void FTurnkeySupportModule::UpdateSdkInfoForDevices(TArray<FString> PlatformDevi
 						UE_LOG(LogTurnkeySupport, Log, TEXT("Turnkey Device: %s"), *Line);
 
 						PerDeviceSdkInfo[DDPIDeviceId] = SdkInfo;
-
-						UE_LOG(LogTurnkeySupport, Log, TEXT("[TEST] Turnkey Device: %s - %d, Installed: %s, Allowed: %s-%s"), *DDPIDeviceId, (int)SdkInfo.Status, *SdkInfo.InstalledVersion,
-							*SdkInfo.MinAllowedVersion, *SdkInfo.MaxAllowedVersion);
 					}
 				}
 			}
diff --git a/Engine/Source/Editor/UnrealEd/Public/EditorAnalytics.h b/Engine/Source/Editor/UnrealEd/Public/EditorAnalytics.h
--- a/Engine/Source/Editor/UnrealEd/Public/EditorAnalytics.h
+++ b/Engine/Source/Editor/UnrealEd/Public/EditorAnalytics.h
@@ -67,6 +67,8 @@ namespace EAnalyticsErrorCodes
 		SymbolizedSONotFound = 153,
 		LicenseNotAccepted = 154,
 		AndroidOBBError = 155,
+		SDKInstallFailed = 200,
+		DeviceUpdateFailed = 201,
 	};
 };
 
diff --git a/Engine/Source/Programs/AutomationTool/Android/AndroidPlatform.Automation.cs b/Engine/Source/Programs/AutomationTool/Android/AndroidPlatform.Automation.cs
--- a/Engine/Source/Programs/AutomationTool/Android/AndroidPlatform.Automation.cs
+++ b/Engine/Source/Programs/AutomationTool/Android/AndroidPlatform.Automation.cs
@@ -37,7 +37,7 @@ public class AndroidPlatform : Platform
 		return new string[] { AndroidSDK.GetMainVersion() };
 	}
 
-	public override bool GetSDKInstallCommand(out string Command, out string Params, out bool bRequiresPrivilegeElevation, ref string Preamble, ref string SuccessPostamble, ref string FailurePostamble, FileRetriever Retriever)
+	public override bool GetSDKInstallCommand(out string Command, out string Params, ref bool bRequiresPrivilegeElevation, ref bool bCreateWindow, ITurnkeyContext TurnkeyContext)
 	{
 		// run the Setup.bat in the engine, not coming from a normal FileSource
 
@@ -63,11 +63,14 @@ public class AndroidPlatform : Platform
 
 		Params = $"{PlatformsVersion} {BuildToolsVersion} {CMakeVersion} {NDKVersion} -noninteractive";
 
-		bRequiresPrivilegeElevation = false;
+		// because this may bring up a license acceptance message that needs the user to respond, so we make a new window
+		bCreateWindow = true;
 
 		return true;
 	}
 
+
+
 	private static string GetAndroidStudioExe()
 	{
 		string DefaultAndroidStudioInstallDir = Path.Combine(Environment.GetFolderPath(Environment.SpecialFolder.ProgramFiles), "Android", "Android Studio");
@@ -82,7 +85,7 @@ public class AndroidPlatform : Platform
 		return RegValue == null ? DefaultSdkDir : RegValue;
 	}
 
-	public override bool UpdateHostPrerequisites(BuildCommand Command, FileRetriever Retriever, bool bVerifyOnly)
+	public override bool UpdateHostPrerequisites(BuildCommand Command, ITurnkeyContext TurnkeyContext, bool bVerifyOnly)
 	{
 		// @todo turnkey: Handle Mac/Linux
 		if (HostPlatform.Current.HostEditorPlatform != UnrealTargetPlatform.Win64)
@@ -100,13 +103,11 @@ public class AndroidPlatform : Platform
 		{
 			if (!File.Exists(AndroidStudioExe))
 			{
-//				Retriever.ReportError("Android Studio is not installed correctly.");
-				//Retriever.PauseForUser("Android Studio is not installed correctly.");
+				TurnkeyContext.ReportError("Android Studio is not installed correctly.");
 			}
 			if (!Directory.Exists(SdkDir))
 			{
-//				Retriever.ReportError("Android SDK directory is not set correctly.");
-				//Retriever.PauseForUser("Android SDK directory is not set correctly.");
+				TurnkeyContext.ReportError("Android SDK directory is not set correctly.");
 			}
 			return bIsInstalled;
 		}
@@ -114,7 +115,7 @@ public class AndroidPlatform : Platform
 		if (!File.Exists(AndroidStudioExe))
 		{
 			// get AS installer
-			string OutputPath = Retriever.RetrieveFileSource("AndroidStudio");
+			string OutputPath = TurnkeyContext.RetrieveFileSource("AndroidStudio");
 
 			// Unset some envvars in case autosdk ran - they will mess up the first run of Android Studio
 			string[] Vars = new string[]
@@ -131,31 +132,30 @@ public class AndroidPlatform : Platform
 
 			if (OutputPath == null)
 			{
-				Retriever.PauseForUser("Unable to find Android Studio installer. Please download and install Android Studio 3.5.3 from https://developer.android.com/studio before continuing.\n\nMake sure to use the Run Android Studio and complete the first-time setup!\n\nChoose all default options unless you know what you are doing.");
+				TurnkeyContext.PauseForUser("Unable to find Android Studio installer. Please download and install Android Studio 3.5.3 from https://developer.android.com/studio/archive before continuing.\n\nMake sure to use the Run Android Studio and complete the first-time setup!\n\nChoose all default options unless you know what you are doing.");
 			}
 			else
 			{
-				Retriever.PauseForUser("Running the Android Studio installer, and then Android Studio for first-time setup!\n\nChoose all default options unless you know what you are doing.");
-				int ExitCode;
-				Utils.RunLocalProcessAndReturnStdOut(OutputPath, "/S", out ExitCode, true);
+				TurnkeyContext.PauseForUser("Running the Android Studio installer, and then Android Studio for first-time setup!\n\nChoose all default options unless you know what you are doing.");
+				int ExitCode = TurnkeyContext.RunExternalCommand(OutputPath, "/S", false, true, true);
 				//				Utils.RunLocalProcessAndReturnStdOut(OutputPath, "", out ExitCode, true);
 
 				// AS installer returns 1223 even on success ("user canceled" even tho there's no UI to cancel it...) when running with /S
 				if (ExitCode != 0 && ExitCode != 1223)
 				{
-//					Retriever.ReportError($"Android Studio installer failed. ExitCode = {ExitCode}");
-					Retriever.PauseForUser($"Android Studio installer failed. ExitCode = {ExitCode}");
+					TurnkeyContext.ReportError($"Android Studio installer failed. ExitCode = {ExitCode}");
 					return false;
 				}
 
-				Utils.RunLocalProcessAndReturnStdOut(GetAndroidStudioExe(), "", out _, false);
+				// re-query for AS location
+				TurnkeyContext.RunExternalCommand(GetAndroidStudioExe(), "", false, true, false);
 			}
 		}
 		else
 		{
-			Retriever.PauseForUser("The Sdk directory was not found. Running the Android Studio to perform first-time setup.\n\nChoose all default options unless you know what you are doing.");
+			TurnkeyContext.PauseForUser("The Sdk directory was not found. Running the Android Studio to perform first-time setup.\n\nChoose all default options unless you know what you are doing.");
 
-			Utils.RunLocalProcessAndReturnStdOut(AndroidStudioExe, "", out _, false);
+			TurnkeyContext.RunExternalCommand(AndroidStudioExe, "", false, true, false);
 		}
 
 		// check to see if the installation worked. If so, continue on!
@@ -166,13 +166,11 @@ public class AndroidPlatform : Platform
 
 		if (!File.Exists(AndroidStudioExe))
 		{
-//			Retriever.ReportError("Android Studio is not installed correctly, after attempted installation.");
-			Retriever.PauseForUser("Android Studio is not installed correctly, after attempted installation.");
+			TurnkeyContext.ReportError("Android Studio is not installed correctly, after attempted installation.");
 		}
 		if (!Directory.Exists(SdkDir))
 		{
-//			Retriever.ReportError("Android SDK directory is not set correctly, after attempted installation.");
-			Retriever.PauseForUser("Android SDK directory is not set correctly, after attempted installation.");
+			TurnkeyContext.ReportError("Android SDK directory is not set correctly, after attempted installation.");
 		}
 
 		return bIsInstalled;
diff --git a/Engine/Source/Programs/AutomationTool/AutomationUtils/AutomationException.cs b/Engine/Source/Programs/AutomationTool/AutomationUtils/AutomationException.cs
--- a/Engine/Source/Programs/AutomationTool/AutomationUtils/AutomationException.cs
+++ b/Engine/Source/Programs/AutomationTool/AutomationUtils/AutomationException.cs
@@ -61,6 +61,8 @@ namespace AutomationTool
 		Error_SymbolizedSONotFound = 153,
 		Error_LicenseNotAccepted = 154,
 		Error_AndroidOBBError = 155,
+		Error_SDKInstallFailed = 200,
+		Error_DeviceUpdateFailed = 201,
 	};
 
 	/// <summary>
diff --git a/Engine/Source/Programs/AutomationTool/AutomationUtils/Platform.cs b/Engine/Source/Programs/AutomationTool/AutomationUtils/Platform.cs
--- a/Engine/Source/Programs/AutomationTool/AutomationUtils/Platform.cs
+++ b/Engine/Source/Programs/AutomationTool/AutomationUtils/Platform.cs
@@ -11,12 +11,14 @@ using EpicGames.Core;
 
 namespace AutomationTool
 {
-	public interface FileRetriever
+	public interface ITurnkeyContext
 	{
 		string RetrieveFileSource(string Name, string InType = "Misc", string InPlatform = null, string SubType = null);
 		string RetrieveFileSource(object HintObject);
 		string GetVariable(string VariableName);
-		bool RunExternalCommand(string Command, string Params, string Preamble = "", string SuccessPostamble = "", string FailurePostamble = "", bool bRequiresPrivilegeElevation = false);
+		int RunExternalCommand(string Command, string Params, bool bRequiresPrivilegeElevation, bool bUnattended, bool bCreateWindow);
+		void Log(string Message);
+		void ReportError(string Message);
 		void PauseForUser(string Message); 
 	}
 
@@ -194,18 +196,21 @@ namespace AutomationTool
 			return null;
 		}
 
-		public virtual bool InstallSDK(BuildCommand BuildCommand, FileRetriever Retriever, DeviceInfo Device)
+		public virtual bool InstallSDK(BuildCommand BuildCommand, ITurnkeyContext TurnkeyContext, DeviceInfo Device, bool bUnattended)
 		{
-			string Command, Params, Preamble="", SuccessPostamble="", FailurePostamble="";
+			string Command, Params;
 
-			bool bRequiresPrivilegeElevation;
-			if (Device != null && GetDeviceUpdateSoftwareCommand(out Command, out Params, out bRequiresPrivilegeElevation, ref Preamble, ref SuccessPostamble, ref FailurePostamble, Retriever, Device))
+			bool bRequiresPrivilegeElevation = false;
+			bool bCreateWindow = false;
+			if (Device != null && GetDeviceUpdateSoftwareCommand(out Command, out Params, ref bRequiresPrivilegeElevation, ref bCreateWindow, TurnkeyContext, Device))
 			{
-				return Retriever.RunExternalCommand(Command, Params, Preamble, SuccessPostamble, FailurePostamble, bRequiresPrivilegeElevation);
+				int ExitCode = TurnkeyContext.RunExternalCommand(Command, Params, bRequiresPrivilegeElevation, bUnattended, bCreateWindow);
+				return OnSDKInstallComplete(ExitCode, TurnkeyContext, Device);
 			}
-			else if (Device == null && GetSDKInstallCommand(out Command, out Params, out bRequiresPrivilegeElevation, ref Preamble, ref SuccessPostamble, ref FailurePostamble, Retriever))
+			else if (Device == null && GetSDKInstallCommand(out Command, out Params, ref bRequiresPrivilegeElevation, ref bCreateWindow, TurnkeyContext))
 			{
-				return Retriever.RunExternalCommand(Command, Params, Preamble, SuccessPostamble, FailurePostamble, bRequiresPrivilegeElevation);
+				int ExitCode = TurnkeyContext.RunExternalCommand(Command, Params, bRequiresPrivilegeElevation, bUnattended, bCreateWindow);
+				return OnSDKInstallComplete(ExitCode, TurnkeyContext, null);
 			}
 
 			return false;
@@ -221,33 +226,42 @@ namespace AutomationTool
 			return new string[] { };
 		}
 
-		public virtual bool GetSDKInstallCommand(out string Command, out string Params, out bool bRequiresPrivilegeElevation, ref string Preamble, ref string SuccessPostamble, ref string FailurePostamble, FileRetriever Retriever)
+		public virtual bool GetSDKInstallCommand(out string Command, out string Params, ref bool bRequiresPrivilegeElevation, ref bool bCreateWindow, ITurnkeyContext TurnkeyContext)
 		{
 			Command = null;
 			Params = null;
-			bRequiresPrivilegeElevation = false;
 			return false;
 		}
 
-		public virtual bool GetDeviceUpdateSoftwareCommand(out string Command, out string Params, out bool bRequiresPrivilegeElevation, ref string Preamble, ref string SuccessPostamble, ref string FailurePostamble, FileRetriever Retriever, DeviceInfo Device = null)
+		public virtual bool GetDeviceUpdateSoftwareCommand(out string Command, out string Params, ref bool bRequiresPrivilegeElevation, ref bool bCreateWindow, ITurnkeyContext TurnkeyContext, DeviceInfo Device = null)
 		{
 			Command = null;
 			Params = null;
-			bRequiresPrivilegeElevation = false;
 			return false;
 		}
 
+		/// <summary>
+		/// Let's the platform handle the result of 
+		/// </summary>
+		/// <param name="ExitCode"></param>
+		/// <param name="Device"></param>
+		/// <returns>True if the installation was a success (defaults to ExitCode == 0)</returns>
+		public virtual bool OnSDKInstallComplete(int ExitCode, ITurnkeyContext TurnkeyContext, DeviceInfo Device)
+		{
+			return ExitCode == 0;
+		}
+
 		public virtual string GetSDKCreationHelp()
 		{
 			return null;
 		}
 
-		public virtual bool UpdateHostPrerequisites(BuildCommand Command, FileRetriever Retriever, bool bVerifyOnly)
+		public virtual bool UpdateHostPrerequisites(BuildCommand Command, ITurnkeyContext TurnkeyContext, bool bVerifyOnly)
 		{
 			return true;
 		}
 
-		public virtual bool UpdateDevicePrerequisites(DeviceInfo Device, BuildCommand Command, FileRetriever Retriever, bool bVerifyOnly)
+		public virtual bool UpdateDevicePrerequisites(DeviceInfo Device, BuildCommand Command, ITurnkeyContext TurnkeyContext, bool bVerifyOnly)
 		{
 			return true;
 		}
diff --git a/Engine/Source/Programs/AutomationTool/IOS/IOSPlatform.Automation.cs b/Engine/Source/Programs/AutomationTool/IOS/IOSPlatform.Automation.cs
--- a/Engine/Source/Programs/AutomationTool/IOS/IOSPlatform.Automation.cs
+++ b/Engine/Source/Programs/AutomationTool/IOS/IOSPlatform.Automation.cs
@@ -166,11 +166,9 @@ public class IOSPlatform : Platform
 		SDKName = (TargetPlatform == UnrealTargetPlatform.TVOS) ? "appletvos" : "iphoneos";
 	}
 
-	public override bool GetSDKInstallCommand(out string Command, out string Params, out bool bRequiresPrivilegeElevation, ref string Preamble, ref string SuccessPostamble, ref string FailurePostamble, FileRetriever Retriever)
+	public override bool GetSDKInstallCommand(out string Command, out string Params, ref bool bRequiresPrivilegeElevation, ref bool bCreateWindow, ITurnkeyContext TurnkeyContext)
 	{
-		Preamble = "Moving your original Xcode application from /Applications to the Trash, and unzipping the new version into /Applications";
-		SuccessPostamble = "If you had Xcode in your Dock, you will need to remove it and add the new one (even though it was in the same location). macOS follows the move to the Trash for the Dock icon";
-		bRequiresPrivilegeElevation = false;
+		TurnkeyContext.Log("Moving your original Xcode application from /Applications to the Trash, and unzipping the new version into /Applications");
 
 		// put current Xcode in the trash, and unzip a new one. Xcode in the dock will have to be fixed up tho!
 		Command = "osascript";
@@ -186,9 +184,9 @@ public class IOSPlatform : Platform
 		return true;
 	}
 
-	public override bool GetDeviceUpdateSoftwareCommand(out string Command, out string Params, out bool bRequiresPrivilegeElevation, ref string Preamble, ref string SuccessPostamble, ref string FailurePostamble, FileRetriever Retriever, DeviceInfo Device)
+	public override bool GetDeviceUpdateSoftwareCommand(out string Command, out string Params, ref bool bRequiresPrivilegeElevation, ref bool bCreateWindow, ITurnkeyContext TurnkeyContext, DeviceInfo Device)
 	{
-		Preamble = "Installing an offline downloaded .ipsw onto your device using the Apple Configurator application.";
+		TurnkeyContext.Log("Installing an offline downloaded .ipsw onto your device using the Apple Configurator application.");
 
 		// cfgtool needs ECID, not UDID, so find it
 		string Configurator = Path.Combine(GetConfiguratorLocation().Replace(" ", "\\ "), "Contents/MacOS/cfgutil");
@@ -200,7 +198,7 @@ public class IOSPlatform : Platform
 		Match Result = Regex.Match(CfgUtilOutput, @"Type: (\S*).*ECID: (\S*)");
 		if (!Result.Success)
 		{
-			Console.WriteLine("Unable to find the given deviceid: {0} in cfgutil output", Device);
+			TurnkeyContext.ReportError($"Unable to find the given deviceid: {Device} in cfgutil output");
 			Command = Params = null;
 			return false;
 		}
@@ -211,6 +209,19 @@ public class IOSPlatform : Platform
 		return true;
 	}
 
+	public override bool OnSDKInstallComplete(int ExitCode, ITurnkeyContext TurnkeyContext, DeviceInfo Device)
+	{
+		if (Device == null)
+		{
+			if (ExitCode == 0)
+			{
+				TurnkeyContext.PauseForUser("If you had Xcode in your Dock, you will need to remove it and add the new one (even though it was in the same location). macOS follows the move to the Trash for the Dock icon");
+			}
+		}
+
+		return ExitCode == 0;
+	}
+
 
 
 
@@ -227,8 +238,12 @@ public class IOSPlatform : Platform
 		public string RubyScript = Path.Combine(CommandUtils.EngineDirectory.FullName, "Build/Turnkey/VerifyIOS.ru");
 		public string InstallCertScript = Path.Combine(CommandUtils.EngineDirectory.FullName, "Build/Turnkey/InstallCert.ru");
 
-		public VerifyIOSSettings(BuildCommand Command, FileRetriever Retriever)
+		private ITurnkeyContext TurnkeyContext;
+
+		public VerifyIOSSettings(BuildCommand Command, ITurnkeyContext TurnkeyContext)
 		{
+			this.TurnkeyContext = TurnkeyContext;
+
 			FileReference ProjectPath = Command.ParseProjectParam();
 			string ProjectName = ProjectPath == null ? "" : ProjectPath.GetFileNameWithoutAnyExtensions();
 
@@ -242,9 +257,9 @@ public class IOSPlatform : Platform
 			Team = Command.ParseParamValue("teamid");
 			Provision = Command.ParseParamValue("provision");
 
-			if (string.IsNullOrEmpty(Team)) Team = Retriever.GetVariable("User_AppleDevCenterTeamID");
-			if (string.IsNullOrEmpty(Account)) Account = Retriever.GetVariable("User_AppleDevCenterUsername");
-			if (string.IsNullOrEmpty(Provision)) Provision = Retriever.GetVariable("User_IOSProvisioningProfile");
+			if (string.IsNullOrEmpty(Team)) Team = TurnkeyContext.GetVariable("User_AppleDevCenterTeamID");
+			if (string.IsNullOrEmpty(Account)) Account = TurnkeyContext.GetVariable("User_AppleDevCenterUsername");
+			if (string.IsNullOrEmpty(Provision)) Provision = TurnkeyContext.GetVariable("User_IOSProvisioningProfile");
 
 			// fall back to ini for anything else
 			if (string.IsNullOrEmpty(CodeSigningIdentity)) EngineConfig.GetString("/Script/IOSRuntimeSettings.IOSRuntimeSettings", "DevCodeSigningIdentity", out CodeSigningIdentity);
@@ -263,7 +278,7 @@ public class IOSPlatform : Platform
 			}
 		}
 
-		public int RunCommandMaybeInteractive(string Command, string Params, bool bInteractive)
+		public bool RunCommandMaybeInteractive(string Command, string Params, bool bInteractive)
 		{
 			Console.WriteLine("Running Command '{0} {1}'", Command, Params);
 
@@ -309,10 +324,37 @@ public class IOSPlatform : Platform
 				}
 				
 			}
-			return ExitCode;
+
+			if (ExitCode != 0)
+			{
+				// only ExitCode 3 (needs cert) can be handled. Any other error can't be fixed (Interactive means it can't be fixed)
+				if (!bInteractive || ExitCode != 3)
+				{
+					if (ExitCode == 3)
+					{
+						TurnkeyContext.ReportError("Signing certificate is required.");
+					}
+					else
+					{
+						// @todo turnkey: turn exitcodes into useful messages
+						TurnkeyContext.ReportError($"Ruby command exited with code {ExitCode}");
+					}
+
+					return false;
+				}
+
+				// only here with ExitCode 3
+				if (!InstallCert())
+				{
+					TurnkeyContext.ReportError($"Certificate installation failed.");
+					return false;
+				}
+			}
+
+			return ExitCode == 0;
 		}
 
-		public int RunRubyCommand(bool bVerifyOnly, string DeviceName)
+		public bool RunRubyCommand(bool bVerifyOnly, string DeviceName)
 		{
 			string Params;
 
@@ -353,58 +395,60 @@ public class IOSPlatform : Platform
 
 			return RunCommandMaybeInteractive(RubyScript, Params, !bVerifyOnly);
 		}
-	}
 
-	private bool InstallCert(FileRetriever Retriever, VerifyIOSSettings Settings)
-	{
-//		string ProjectName = Retriever.GetVariable("Project");
+		private bool InstallCert()
+		{
+			//		string ProjectName = TurnkeyContext.GetVariable("Project");
 
-		string CertLoc = null;
+			string CertLoc = null;
 
-		if (!string.IsNullOrEmpty(Settings.BundleId))
-		{
-			CertLoc = Retriever.RetrieveFileSource("DevCert: " + Settings.BundleId);
-		}
-		if (CertLoc == null)
-		{
-			CertLoc = Retriever.RetrieveFileSource("DevCert");
-		}
+			if (!string.IsNullOrEmpty(BundleId))
+			{
+				CertLoc = TurnkeyContext.RetrieveFileSource("DevCert: " + BundleId);
+			}
+			if (CertLoc == null)
+			{
+				CertLoc = TurnkeyContext.RetrieveFileSource("DevCert");
+			}
 
-		if (CertLoc != null)
-		{
-			// get the cert password from Studio settings
-			string CertPassword = Retriever.GetVariable("Studio_AppleSigningCertPassword");
+			if (CertLoc != null)
+			{
+				// get the cert password from Studio settings
+				string CertPassword = TurnkeyContext.GetVariable("Studio_AppleSigningCertPassword");
 
-			Console.WriteLine("Will install cert from: '{0}'", CertLoc);
+				TurnkeyContext.Log($"Will install cert from: '{CertLoc}'");
 
-			// osascript -e 'Tell application "System Events" to display dialog "Enter the network password:" with hidden answer default answer ""' -e 'text returned of result' 2>/dev/null
-			string CommandLine = string.Format("'{0}' '{1}'", CertLoc, CertPassword);
+				// osascript -e 'Tell application "System Events" to display dialog "Enter the network password:" with hidden answer default answer ""' -e 'text returned of result' 2>/dev/null
+				string CommandLine = string.Format("'{0}' '{1}'", CertLoc, CertPassword);
 
-			// run ruby script to install cert
-			return Settings.RunCommandMaybeInteractive(Settings.InstallCertScript, CommandLine, true) == 0;
-		}
-		else
-		{
-			Console.WriteLine("Unable to find a tagged source for DevCert");
+				// run ruby script to install cert
+				return RunCommandMaybeInteractive(InstallCertScript, CommandLine, true);
+			}
+			else
+			{
+				TurnkeyContext.ReportError("Unable to find a tagged source for DevCert");
 
-			return false;
-		}
+				return false;
+			}
 
+		}
 	}
 
+
+
 	string GetConfiguratorLocation()
 	{
 		string FindCommand = "-c 'mdfind \"kMDItemKind == Application\" | grep \"Apple Configurator 2.app\"'";
 		return UnrealBuildTool.Utils.RunLocalProcessAndReturnStdOut("sh", FindCommand);
 	}
 
-	public override bool UpdateHostPrerequisites(BuildCommand Command, FileRetriever Retriever, bool bVerifyOnly)
+	public override bool UpdateHostPrerequisites(BuildCommand Command, ITurnkeyContext TurnkeyContext, bool bVerifyOnly)
 	{
 		int ExitCode;
 
 		if (HostPlatform.Current.HostEditorPlatform != UnrealTargetPlatform.Mac)
 		{
-			return base.UpdateHostPrerequisites(Command, Retriever, bVerifyOnly);
+			return base.UpdateHostPrerequisites(Command, TurnkeyContext, bVerifyOnly);
 		}
 
 		// make sure the Configurator is installed
@@ -414,10 +458,11 @@ public class IOSPlatform : Platform
 		{
 			if (bVerifyOnly)
 			{
+				TurnkeyContext.ReportError("Apple Configurator 2 is required.");
 				return false;
 			}
 
-			Console.WriteLine("Apple Configurator 2 is required for some automation to work. You should install it from the App Store. Launching...");
+			TurnkeyContext.PauseForUser("Apple Configurator 2 is required for some automation to work. You should install it from the App Store. Launching...");
 
 			// we need to install Configurator 2, and we will block until it's done
 			UnrealBuildTool.Utils.RunLocalProcessAndReturnStdOut("open", "macappstore://apps.apple.com/us/app/apple-configurator-2/id1037126344?mt=12");
@@ -434,10 +479,11 @@ public class IOSPlatform : Platform
 			Console.WriteLine("Fastlane is not installed");
 			if (bVerifyOnly)
 			{
+				TurnkeyContext.ReportError("Fastlane is not installed.");
 				return false;
 			}
 
-			Console.WriteLine("Installing Fastlane from internet source. You may ignore the error about the bin directory not in your path.");
+			TurnkeyContext.PauseForUser("Installing Fastlane from internet source. You may ignore the error about the bin directory not in your path.");
 
 			// install missing fastlane without needing sudo
 			UnrealBuildTool.Utils.RunLocalProcessAndReturnStdOut("/usr/bin/gem", "install fastlane --user-install --no-document", out ExitCode, true);
@@ -448,54 +494,26 @@ public class IOSPlatform : Platform
 			}
 		}
 
-		VerifyIOSSettings Settings = new VerifyIOSSettings(Command, Retriever);
+		VerifyIOSSettings Settings = new VerifyIOSSettings(Command, TurnkeyContext);
 
 		// look if we have a cert that matches it
-		ExitCode = Settings.RunRubyCommand(bVerifyOnly, null);
-
-		Console.WriteLine("Host VerifyIOS exited with code {0}", ExitCode);
-
-		if (ExitCode != 0)
-		{
-			if (bVerifyOnly)
-			{
-				return false;
-			}
-
-			// ExitCode 3 means we need to install a cert
-			if (ExitCode != 3 || !InstallCert(Retriever, Settings))
-			{
-				return false;
-			}
-		}
-
-		return true;
+		return Settings.RunRubyCommand(bVerifyOnly, null);
 	}
 
-	public override bool UpdateDevicePrerequisites(DeviceInfo Device, BuildCommand Command, FileRetriever Retriever, bool bVerifyOnly)
+	public override bool UpdateDevicePrerequisites(DeviceInfo Device, BuildCommand Command, ITurnkeyContext TurnkeyContext, bool bVerifyOnly)
 	{
 		if (HostPlatform.Current.HostEditorPlatform != UnrealTargetPlatform.Mac)
 		{
-			return base.UpdateDevicePrerequisites(Device, Command, Retriever, bVerifyOnly);
+			return base.UpdateDevicePrerequisites(Device, Command, TurnkeyContext, bVerifyOnly);
 		}
 
-		VerifyIOSSettings Settings = new VerifyIOSSettings(Command, Retriever);
+		VerifyIOSSettings Settings = new VerifyIOSSettings(Command, TurnkeyContext);
 
 		// @todo turnkey - better to use the device's udid if it's set properly in DeviceInfo
 		string DeviceName = Device == null ? null : Device.Name;
 
 		// now look for a provision that can be used with a (maybe newly) instally cert
-		int ExitCode = Settings.RunRubyCommand(bVerifyOnly, DeviceName);
-
-		Console.WriteLine("Device VerifyIOS exited with code {0}", ExitCode);
-
-		// ExitCode 3 means we need to install a cert
-		if (ExitCode == 3 && !InstallCert(Retriever, Settings))
-		{
-			return false;
-		}
-
-		return ExitCode == 0;
+		return Settings.RunRubyCommand(bVerifyOnly, DeviceName);
 	}
 
 
diff --git a/Engine/Source/Programs/AutomationTool/Turnkey/Commands/InstallSdk.cs b/Engine/Source/Programs/AutomationTool/Turnkey/Commands/InstallSdk.cs
--- a/Engine/Source/Programs/AutomationTool/Turnkey/Commands/InstallSdk.cs
+++ b/Engine/Source/Programs/AutomationTool/Turnkey/Commands/InstallSdk.cs
@@ -18,7 +18,7 @@ namespace Turnkey.Commands
 			return new Dictionary<string, string[]>()
 			{
 				{ "Auto Install All Needed Sdks", new string[] { "-platform=All", "-NeededOnly", "-BestAvailable" } },
-				{ "Auto Update All Sdks", new string[] { "-platform=All", "-UpdateOnly", "-BestAvailable" } },
+				{ "Auto Update Installed Sdks", new string[] { "-platform=All", "-UpdateOnly", "-BestAvailable" } },
 			};
 		}
 
@@ -28,9 +28,9 @@ namespace Turnkey.Commands
 			string DeviceName = TurnkeyUtils.ParseParamValue("Device", null, CommandOptions);
 			string SdkTypeString = TurnkeyUtils.ParseParamValue("SdkType", null, CommandOptions);
 
+			bool bUnattended = TurnkeyUtils.ParseParam("Unattended", CommandOptions);
 			bool bBestAvailable = TurnkeyUtils.ParseParam("BestAvailable", CommandOptions);
 			bool bUpdateOnly = TurnkeyUtils.ParseParam("UpdateOnly", CommandOptions);
-			bool bAllowAutoSdk = TurnkeyUtils.ParseParam("AllowAutoSdk", CommandOptions);
 
 			// best available installation requires valid and needed Sdks
 			bool bValidOnly = bBestAvailable || !TurnkeyUtils.ParseParam("AllowInvalid", CommandOptions);
@@ -48,17 +48,64 @@ namespace Turnkey.Commands
 				}
 				DesiredType = OutType;
 			}
+			// if a device is specified, assume Flash type
+			else if (!string.IsNullOrEmpty(DeviceName))
+			{
+				DesiredType = FileSource.SourceType.Flash;
+			}
+			// otherwise ask
+			else if (!bUnattended)
+			{
+				List<string> Options = new List<string>() { "Full or Auto Sdk", "Full Sdk", "AutoSdk", "Device Software / Flash" };
+				int Choice = TurnkeyUtils.ReadInputInt("Choose a type of Sdk to install:", Options, true, 1);
+
+				if (Choice == 0)
+				{
+					return;
+				}
+				// Choice 1 means two types, so we use null DesiredType to mean Full|AutoSdk
+				if (Choice == 2)
+				{
+					DesiredType = FileSource.SourceType.Full;
+				}
+				else if (Choice == 3)
+				{
+					DesiredType = FileSource.SourceType.AutoSdk;
+				}
+				else if (Choice == 4)
+				{
+					DesiredType = FileSource.SourceType.Flash;
+				}
+			}
 
 			// we need all sdks we can find
 			List<UnrealTargetPlatform> PlatformsWithSdks = TurnkeyManifest.GetPlatformsWithSdks();
 
-			// get the platforms to install either from user or from commandline
-			List<UnrealTargetPlatform> PlatformsLeftToInstall = TurnkeyUtils.GetPlatformsFromCommandLineOrUser(CommandOptions, PlatformsWithSdks);
-			if (PlatformsLeftToInstall == null)
+			List<DeviceInfo> ChosenDevices = null;
+			List<UnrealTargetPlatform> PlatformsLeftToInstall;
+
+			if (DesiredType == FileSource.SourceType.Flash)
+			{
+				TurnkeyUtils.GetPlatformsAndDevicesFromCommandLineOrUser(CommandOptions, false, out PlatformsLeftToInstall, out ChosenDevices, PlatformsWithSdks);
+				
+				if (PlatformsLeftToInstall.Count == 0 || ChosenDevices.Count == 0)
+				{
+					return;
+				}
+			}
+			else
 			{
-				return;
+				PlatformsLeftToInstall = TurnkeyUtils.GetPlatformsFromCommandLineOrUser(CommandOptions, PlatformsWithSdks);
+
+				if (PlatformsLeftToInstall == null || PlatformsLeftToInstall.Count == 0)
+				{
+					return;
+				}
 			}
 
+
+			TurnkeyContextImpl TurnkeyContext = new TurnkeyContextImpl();
+
 			List<FileSource> SdksAlreadyInstalled = new List<FileSource>();
 
 			// keep going while we have Sdks left to install
@@ -69,18 +116,28 @@ namespace Turnkey.Commands
 				// remove the platform (we may remove more later if an Sdk supports multiple platforms)
 				PlatformsLeftToInstall.RemoveAt(0);
 
-
 				// cache the automation platform object
 				AutomationTool.Platform AutomationPlatform = AutomationTool.Platform.GetPlatform(Platform);
 				UEBuildPlatformSDK SDK = UEBuildPlatformSDK.GetSDKForPlatform(Platform.ToString());
 
 				// filter the Sdks if a platform was given, and type if it was given
-				List<FileSource> Sdks = TurnkeyManifest.FilterDiscoveredFileSources(Platform, DesiredType);
+				List<FileSource> Sdks;
+				if (DesiredType == null)
+				{
+					// get Full and Auto together
+					Sdks = TurnkeyManifest.FilterDiscoveredFileSources(Platform, FileSource.SourceType.Full);
+					Sdks.AddRange(TurnkeyManifest.FilterDiscoveredFileSources(Platform, FileSource.SourceType.AutoSdk));
+				}
+				else
+				{
+					Sdks = TurnkeyManifest.FilterDiscoveredFileSources(Platform, DesiredType);
+				}
 
-				// strip out flash Sdks where there are no devices
-				int SdkCount = Sdks.Count;
-				Sdks = Sdks.FindAll(x => x.Type != FileSource.SourceType.Flash || (AutomationPlatform.GetDevices() != null && AutomationPlatform.GetDevices().Length > 0));
-				bool bStrippedDevices = Sdks.Count != SdkCount;
+				DeviceInfo InstallDevice = null;
+				if (DesiredType == FileSource.SourceType.Flash)
+				{
+					InstallDevice = ChosenDevices.First();
+				}
 
 				// skip Misc FileSources, as we dont know how they are sued
 				Sdks = Sdks.FindAll(x => x.IsSdkType());
@@ -91,17 +148,22 @@ namespace Turnkey.Commands
 					Sdks = Sdks.FindAll(x => !string.IsNullOrEmpty(SDK.GetInstalledVersion()));
 				}
 
-// 				// strip out Sdks not in the allowed range of the platform
-// 				if (bValidOnly)
-// 				{
-// 					Sdks = Sdks.FindAll(x => x.IsValid(Platform, DeviceName));
-// 				}
-// 
-// 				// strip out Sdks for platforms that are not needed to be updated
-// 				if (bNeededOnly)
-// 				{
-// 					Sdks = Sdks.FindAll(x => x.IsNeeded(Platform, DeviceName));
-// 				}
+				// strip out Sdks not in the allowed range of the platform
+				bool bHadInvalidOptions = false;
+				if (bValidOnly)
+				{
+					int BeforeCount = Sdks.Count;
+					Sdks = Sdks.FindAll(x => x.IsVersionValid(Platform, InstallDevice));
+					
+					// track if we removed some invalid ones
+					bHadInvalidOptions = Sdks.Count != BeforeCount;
+				}
+
+				//// strip out Sdks for platforms that are not needed to be updated
+				//if (bNeededOnly)
+				//{
+				//	Sdks = Sdks.FindAll(x => x.IsNeeded(Platform, DeviceName));
+				//}
 
 				Sdks.Sort((a,b) => 
 				{
@@ -157,13 +219,15 @@ namespace Turnkey.Commands
 				// if we are not doing best available Sdks, then let used choose one
 				else
 				{
-					if (Sdks.Count == 0)
+					bool bAddShowMoreOption = bHadInvalidOptions && bValidOnly;
+
+					if (Sdks.Count == 0 && !bAddShowMoreOption)
 					{
 						TurnkeyUtils.Log("No Sdks found for platform {0}. Skipping", Platform);
-						if (bStrippedDevices)
-						{
-							TurnkeyUtils.Log("NOTE: Some Flash Sdks were removed because no devices were found");
-						}
+						//if (bStrippedDevices)
+						//{
+						//	TurnkeyUtils.Log("NOTE: Some Flash Sdks were removed because no devices were found");
+						//}
 						continue;
 					}
 
@@ -180,11 +244,16 @@ namespace Turnkey.Commands
 						Options.Add(string.Format("[{0}] {1} [Current: {2}]", string.Join(",", Sdk.GetPlatforms()), Sdk.Name, Current));
 					}
 
-					string Prompt = "Select an Sdk to install";
-					if (bStrippedDevices)
+					if (bAddShowMoreOption)
 					{
-						Prompt += "\nNOTE: Some Flash Sdks were removed because no devices were found";
+						Options.Add("Show Invalid Sdks");
 					}
+
+					string Prompt = $"Select an Sdk to install:";
+					//if (bStrippedDevices)
+					//{
+					//	Prompt += "\nNOTE: Some Flash Sdks were removed because no devices were found";
+					//}
 					int Choice = TurnkeyUtils.ReadInputInt(Prompt, Options, true);
 
 					if (Choice == 0)
@@ -192,6 +261,14 @@ namespace Turnkey.Commands
 						continue;
 					}
 
+					// if they chose the show more option, then add thius platform back in, and loop with bAllowInvaid to be true
+					if (bAddShowMoreOption && Choice == Options.Count)
+					{
+						PlatformsLeftToInstall.Insert(0, Platform);
+						bValidOnly = false;
+						continue;
+					}
+
 					// only install the chosen one
 					FileSource ChosenSdk = Sdks[Choice - 1];
 					Sdks = new List<FileSource>() { ChosenSdk };
@@ -206,57 +283,20 @@ namespace Turnkey.Commands
 					}
 					SdksAlreadyInstalled.Add(Sdk);
 
-					DeviceInfo InstallDevice = null;
-					// set variables
-					if (Sdk.Type == FileSource.SourceType.Flash)
-					{
-						string InstallDeviceName = DeviceName;
-						if (InstallDeviceName == null)
-						{
-							List<string> Options = new List<string>();
-							DeviceInfo[] PossibleDevices = Array.FindAll(AutomationPlatform.GetDevices(), x => TurnkeyUtils.IsValueValid(x.Type, Sdk.AllowedFlashDeviceTypes, AutomationPlatform));
-							foreach (DeviceInfo Device in PossibleDevices)
-							{
-								Options.Add(string.Format("[{0} {1}] {2}", Platform, Device.Type, Device.Name));
-							}
-
-							// get the choice
-							int Choice = TurnkeyUtils.ReadInputInt("Select the number of a device to flash:", Options, true);
-
-							if (Choice == 0)
-							{
-								return;
-							}
-
-							// get the name of the device chosen
-							InstallDeviceName = PossibleDevices[Choice - 1].Name;
-						}
-
-						// get device info of the chosen or supplied device
-						InstallDevice = Array.Find(AutomationPlatform.GetDevices(), x => string.Compare(x.Name, InstallDeviceName, true) == 0);
-
-						if (InstallDevice == null)
-						{
-							TurnkeyUtils.Log("Unable to find {0} device {1}", Platform, InstallDeviceName);
-							return;
-						}
-					}
-
 					// make sure prerequisites are in good shape
-					CopyProviderRetriever Retriever = new CopyProviderRetriever();
-					if (AutomationPlatform.UpdateHostPrerequisites(TurnkeyUtils.CommandUtilHelper, Retriever, false) == false)
+					if (AutomationPlatform.UpdateHostPrerequisites(TurnkeyUtils.CommandUtilHelper, TurnkeyContext, false) == false)
 					{
-						// Retriever.ReportError("Failed to update host prerequisites");
-						return;
+						TurnkeyContext.ReportError("Failed to update host prerequisites");
+						continue;
 					}
 					if (InstallDevice != null && 
-						AutomationPlatform.UpdateDevicePrerequisites(InstallDevice, TurnkeyUtils.CommandUtilHelper, Retriever, false) == false)
+						AutomationPlatform.UpdateDevicePrerequisites(InstallDevice, TurnkeyUtils.CommandUtilHelper, TurnkeyContext, false) == false)
 					{
-						// Retriever.ReportError("Failed to update device prerequisites");
-						return;
+						TurnkeyContext.ReportError("Failed to update device prerequisites");
+						continue;
 					}
 
-					Sdk.DownloadOrInstall(Platform, InstallDevice);
+					Sdk.DownloadOrInstall(Platform, TurnkeyContext, InstallDevice, false);
 				}
 			}
 		}
diff --git a/Engine/Source/Programs/AutomationTool/Turnkey/Commands/VerifySdk.cs b/Engine/Source/Programs/AutomationTool/Turnkey/Commands/VerifySdk.cs
--- a/Engine/Source/Programs/AutomationTool/Turnkey/Commands/VerifySdk.cs
+++ b/Engine/Source/Programs/AutomationTool/Turnkey/Commands/VerifySdk.cs
@@ -20,46 +20,29 @@ namespace Turnkey.Commands
 			bool bForceSdkInstall = TurnkeyUtils.ParseParam("ForceSdkInstall", CommandOptions);
 			bool bForceDeviceInstall = TurnkeyUtils.ParseParam("ForceDeviceInstall", CommandOptions);
 			bool bUpdateIfNeeded = bForceSdkInstall || bForceDeviceInstall || TurnkeyUtils.ParseParam("UpdateIfNeeded", CommandOptions);
-
-			// track each platform to check, and if -device was specified, get the devices
 			string DeviceString = TurnkeyUtils.ParseParamValue("Device", null, CommandOptions);
-			List<UnrealTargetPlatform> PlatformsToCheck = new List<UnrealTargetPlatform>();
-			List<DeviceInfo> DevicesToCheck = new List<DeviceInfo>();
-			if (string.IsNullOrEmpty(DeviceString))
-			{
-				PlatformsToCheck = TurnkeyUtils.GetPlatformsFromCommandLineOrUser(CommandOptions, null);
-			}
-			else
+
+			List<UnrealTargetPlatform> PlatformsToCheck;
+			List<DeviceInfo> DevicesToCheck;
+			TurnkeyUtils.GetPlatformsAndDevicesFromCommandLineOrUser(CommandOptions, true, out PlatformsToCheck, out DevicesToCheck);
+
+			// if we got no devices, requested some, and platforms were not specified, then we don't want to continue.
+			// if -device and -platform were specified, and no devices found, we will still continue with the platforms
+			if (DevicesToCheck.Count == 0 && TurnkeyUtils.ParseParamValue("Device", null, CommandOptions) != null && TurnkeyUtils.ParseParamValue("Platform", null, CommandOptions) == null)
 			{
-				DevicesToCheck = TurnkeyUtils.GetDevicesFromCommandLineOrUser(CommandOptions, null);
-				// if we failed to get any devices, fallback to getting just platforms anyway
-				if (DevicesToCheck == null)
-				{
-					if (TurnkeyUtils.ParseParamValue("Platform", null, CommandOptions) == null)
-					{
-						TurnkeyUtils.Log("Devices were requested, but none of them were found. Since -platforms was not specified, exiting command...");
-						return;
-					}
-					TurnkeyUtils.Log("Devices were requested, but none of them were found. Will continue with just platform verification");
-					PlatformsToCheck = TurnkeyUtils.GetPlatformsFromCommandLineOrUser(CommandOptions, null);
-				}
-				else
-				{
-					PlatformsToCheck = DevicesToCheck.Select(x => x.Platform).ToHashSet().ToList();
-				}
+				TurnkeyUtils.Log("Devices were requested, but none of them were found. Since -platforms was not specified, exiting command...");
 			}
 
-			if (PlatformsToCheck == null || PlatformsToCheck.Count == 0)
+			if (PlatformsToCheck.Count == 0)
 			{
 				TurnkeyUtils.Log("Platform(s) and/or device(s) needed for VerifySdk command. Check parameters or selections.");
 				return;
 			}
 
-
 			TurnkeyUtils.Log("Installed Sdk validity:");
 			TurnkeyUtils.ExitCode = AutomationTool.ExitCode.Success;
 
-			CopyProviderRetriever Retriever = new CopyProviderRetriever();
+			TurnkeyContextImpl TurnkeyContext = new TurnkeyContextImpl();
 
 			TurnkeyUtils.StartTrackingExternalEnvVarChanges();
 
@@ -71,8 +54,11 @@ namespace Turnkey.Commands
 				// get the platform object
 				AutomationTool.Platform AutomationPlatform = AutomationTool.Platform.GetPlatform(Platform);
 
+				// reset the errors for each device
+				TurnkeyContext.ErrorMessages.Clear();
+
 				// checking availability may generate errors, if prerequisites are failing
-				SdkUtils.LocalAvailability LocalState = SdkUtils.GetLocalAvailability(AutomationPlatform, bUpdateIfNeeded);
+				SdkUtils.LocalAvailability LocalState = SdkUtils.GetLocalAvailability(AutomationPlatform, bUpdateIfNeeded, TurnkeyContext);
 
 				string ManualSDKVersion = "", AutoSDKVersion = "";
 				string MinAllowedVersion = "", MaxAllowedVersion = "";
@@ -103,8 +89,14 @@ namespace Turnkey.Commands
 					ReportedState &= (SdkUtils.LocalAvailability.AutoSdk_ValidVersionExists | SdkUtils.LocalAvailability.InstalledSdk_ValidVersionExists | SdkUtils.LocalAvailability.Support_FullSdk | SdkUtils.LocalAvailability.Support_AutoSdk);
 				}
 
-				TurnkeyUtils.Report("{0}: (Status={1}, Installed={2}, AutoSDK={3}, MinAllowed={4}, MaxAllowed={5}, Flags=\"{6}\")", Platform, StatusString, ManualSDKVersion, AutoSDKVersion,
-					MinAllowedVersion, MaxAllowedVersion, ReportedState.ToString());
+				string ErrorString = "";
+				if (TurnkeyContext.ErrorMessages.Count() > 0)
+				{
+					ErrorString = string.Format(", Error=\"{0}\"", string.Join("|", TurnkeyContext.ErrorMessages));
+				}
+
+				TurnkeyUtils.Report("{0}: (Status={1}, Installed={2}, AutoSDK={3}, MinAllowed={4}, MaxAllowed={5}, Flags=\"{6}\"{7})", Platform, StatusString, ManualSDKVersion, AutoSDKVersion,
+					MinAllowedVersion, MaxAllowedVersion, ReportedState.ToString(), ErrorString);
 
 				if (PlatformSDK == null)
 				{
@@ -152,13 +144,12 @@ namespace Turnkey.Commands
 
 					TurnkeyUtils.Log("Will install {0}", BestSdk.Name);
 
-					if (BestSdk.DownloadOrInstall(Platform) == false)
+					if (BestSdk.DownloadOrInstall(Platform, TurnkeyContext, null, bUnattended) == false)
 					{
-						TurnkeyUtils.Log("Install failed!");
+						TurnkeyUtils.Log("Failed to install {0}", BestSdk.Name);
+						TurnkeyUtils.ExitCode = ExitCode.Error_SDKInstallFailed;
 						continue;
 					}
-// 					AutomationPlatform.InstallSDK(TurnkeyUtils.CommandUtilHelper, Retriever, BestSdk);
-
 
 					// update LocalState
 //					LocalState = SdkUtils.GetLocalAvailability(AutomationPlatform, false);
@@ -171,8 +162,10 @@ namespace Turnkey.Commands
 				{
 					foreach (DeviceInfo Device in DevicesToCheck.Where(x => x.Platform == Platform))
 					{
-						List<string> DeviceErrorMessages = new List<string>();
-						bool bArePrerequisitesValid = AutomationPlatform.UpdateDevicePrerequisites(Device, TurnkeyUtils.CommandUtilHelper, Retriever, !bUpdateIfNeeded);
+						// reset the errors for each device
+						TurnkeyContext.ErrorMessages.Clear();
+
+						bool bArePrerequisitesValid = AutomationPlatform.UpdateDevicePrerequisites(Device, TurnkeyUtils.CommandUtilHelper, TurnkeyContext, !bUpdateIfNeeded);
 						bool bIsSoftwareValid = PlatformSDK.IsSoftwareVersionValid(Device.SoftwareVersion);
 
 						SdkUtils.LocalAvailability DeviceState = SdkUtils.LocalAvailability.None;
@@ -197,9 +190,14 @@ namespace Turnkey.Commands
 							DeviceState |= SdkUtils.LocalAvailability.Device_CannotConnect;
 						}
 
-						// @todo turnkey: if Name has a comma in it, we are in trouble, and maybe a )
-						TurnkeyUtils.Report("{0}@{1}: (Name={2}, Status={3}, Installed={4}, MinAllowed={5}, MaxAllowed={6}, Flags=\"{7}\")", Platform, Device.Id, Device.Name, StatusString, Device.SoftwareVersion,
-							MinSoftwareAllowedVersion, MaxSoftwareAllowedVersion, DeviceState.ToString());
+						string DeviceErrorString = "";
+						if (TurnkeyContext.ErrorMessages.Count() > 0)
+						{
+							ErrorString = string.Format(", Error=\"{0}\"", string.Join("|", TurnkeyContext.ErrorMessages));
+						}
+
+						TurnkeyUtils.Report("{0}@{1}: (Name={2}, Status={3}, Installed={4}, MinAllowed={5}, MaxAllowed={6}, Flags=\"{7}\"{8})", Platform, Device.Id, Device.Name, StatusString, Device.SoftwareVersion,
+							MinSoftwareAllowedVersion, MaxSoftwareAllowedVersion, DeviceState.ToString(), DeviceErrorString);
 
 						if (bForceDeviceInstall || !bIsSoftwareValid)
 						{
@@ -216,7 +214,11 @@ namespace Turnkey.Commands
 									}
 									else
 									{
-										MatchingInstallableSdk.DownloadOrInstall(Platform, Device, bUnattended);
+										if (MatchingInstallableSdk.DownloadOrInstall(Platform, TurnkeyContext, Device, bUnattended) == false)
+										{
+											TurnkeyUtils.Log("Failed to update Device '{0}' with '{0}'", Device.Name, MatchingInstallableSdk.Name);
+											TurnkeyUtils.ExitCode = ExitCode.Error_DeviceUpdateFailed;
+										}
 									}
 								}
 								else
diff --git a/Engine/Source/Programs/AutomationTool/Turnkey/CopyAndRun.cs b/Engine/Source/Programs/AutomationTool/Turnkey/CopyAndRun.cs
--- a/Engine/Source/Programs/AutomationTool/Turnkey/CopyAndRun.cs
+++ b/Engine/Source/Programs/AutomationTool/Turnkey/CopyAndRun.cs
@@ -64,7 +64,7 @@ namespace Turnkey
 			CommandLine = TurnkeyUtils.ExpandVariables(CommandLine, bUseOnlyTurnkeyVariables: true);
 		}
 
-		public static bool RunExternalCommand(string CommandPath, string CommandLine, bool bRequiresPrivilegeElevation)
+		public static int RunExternalCommand(string CommandPath, string CommandLine, ITurnkeyContext TurnkeyContext, bool bUnattended, bool bRequiresPrivilegeElevation, bool bCreateWindow)
 		{
 			string FixedCommandPath = TurnkeyUtils.ExpandVariables(CommandPath);
 			string FixedCommandLine = TurnkeyUtils.ExpandVariables(CommandLine);
@@ -85,10 +85,10 @@ namespace Turnkey
 
 			// run installer as administrator, as some need it
 			Process InstallProcess = new Process();
-			InstallProcess.StartInfo.UseShellExecute = false;
+			InstallProcess.StartInfo.UseShellExecute = bCreateWindow;
 			InstallProcess.StartInfo.FileName = FixedCommandPath;
 			InstallProcess.StartInfo.Arguments = FixedCommandLine;
-			InstallProcess.StartInfo.WindowStyle = ProcessWindowStyle.Hidden;
+			InstallProcess.StartInfo.WindowStyle = bCreateWindow ? ProcessWindowStyle.Normal : ProcessWindowStyle.Hidden;
 
 			InstallProcess.OutputDataReceived += (Sender, Args) => { if (Args != null && Args.Data != null) TurnkeyUtils.Log(Args.Data.TrimEnd()); };
 			InstallProcess.ErrorDataReceived += (Sender, Args) => { if (Args != null && Args.Data != null) TurnkeyUtils.Log("Error: {0}", Args.Data.TrimEnd()); };
@@ -103,7 +103,7 @@ namespace Turnkey
 			{
 				try
 				{
-					if (bRequiresPrivilegeElevation)
+					if (bRequiresPrivilegeElevation && InstallProcess.StartInfo.Verb != "runas")
 					{
 						TurnkeyUtils.Log("The installer {0} requires elevated permissions, trying with Admin privileges (output may be hidden)", FixedCommandPath);
 
@@ -122,28 +122,29 @@ namespace Turnkey
 					// this will not allow capturing stdout, so run with window as Normal
 					if (InstallProcess.StartInfo.UseShellExecute == false && Ex is Win32Exception && ((Win32Exception)Ex).NativeErrorCode == 740)
 					{
-						InstallProcess.StartInfo.UseShellExecute = true;
-						InstallProcess.StartInfo.Verb = "runas";
-						InstallProcess.StartInfo.WindowStyle = ProcessWindowStyle.Normal;
-
-						TurnkeyUtils.Log("The installer {0} needed to run with elevated permissions, trying with Admin privileges (output may be hidden)", FixedCommandPath);
-
-						// try again
+						bRequiresPrivilegeElevation = true;
 						continue;
 					}
 
-					TurnkeyUtils.Log("Error: {0} caused an exception: {1}", FixedCommandPath, Ex.Message);
+					TurnkeyContext.ReportError($"Error: {FixedCommandPath} caused an exception: {Ex.Message}");
 					ExitCode = -1;
 				}
 
 				if (ExitCode != 0)
 				{
 					TurnkeyUtils.Log("");
-					TurnkeyUtils.Log("Command {0} {1} failed [Exit code {2}, working dir = {3}]", FixedCommandPath, FixedCommandLine, ExitCode, CommandWorkingDir);
+					TurnkeyContext.ReportError($"Command {FixedCommandPath} {FixedCommandLine} failed [Exit code {ExitCode}, working dir = {CommandWorkingDir}]");
 					TurnkeyUtils.Log("");
 
-					string Response = TurnkeyUtils.ReadInput("Do you want to attempt again? [y/N]", "N");
-					if (string.Compare(Response, "Y", true) != 0)
+					if (!bUnattended)
+					{
+						string Response = TurnkeyUtils.ReadInput("Do you want to attempt again? [y/N]", "N");
+						if (string.Compare(Response, "Y", true) != 0)
+						{
+							bDone = true;
+						}
+					}
+					else
 					{
 						bDone = true;
 					}
@@ -158,7 +159,7 @@ namespace Turnkey
 
 			Environment.CurrentDirectory = PreviousCWD;
 
-			return ExitCode == 0;
+			return ExitCode;
 		}
 	}
 }
diff --git a/Engine/Source/Programs/AutomationTool/Turnkey/CopyProvider.cs b/Engine/Source/Programs/AutomationTool/Turnkey/CopyProvider.cs
--- a/Engine/Source/Programs/AutomationTool/Turnkey/CopyProvider.cs
+++ b/Engine/Source/Programs/AutomationTool/Turnkey/CopyProvider.cs
@@ -212,8 +212,10 @@ namespace Turnkey
 		}
 	}
 
-	class CopyProviderRetriever : AutomationTool.FileRetriever
+	class TurnkeyContextImpl : ITurnkeyContext
 	{
+		public List<string> ErrorMessages = new List<string>();
+
 		public string RetrieveFileSource(object HintObject)
 		{
 			FileSource Sdk = (FileSource)HintObject;
@@ -230,43 +232,21 @@ namespace Turnkey
 			return null;
 		}
 
-		public bool RunExternalCommand(string Command, string Params, string Preamble, string SuccessPostAmble, string FailurePostamble, bool bRequiresPrivilegeElevation)
+		public int RunExternalCommand(string Command, string Params, bool bRequiresPrivilegeElevation, bool bUnattended, bool bCreateWindow)
 		{
 			TurnkeyUtils.Log("----------------------------------------------");
 			TurnkeyUtils.Log("Running '{0} {1}'", TurnkeyUtils.ExpandVariables(Command), TurnkeyUtils.ExpandVariables(Params));
 
-			if (!string.IsNullOrEmpty(Preamble))
-			{
-				TurnkeyUtils.Log("");
-				TurnkeyUtils.Log(Preamble);
-			}
 			TurnkeyUtils.Log("----------------------------------------------", Command);
 
-			bool bSuccess = CopyAndRun.RunExternalCommand(Command, Params, bRequiresPrivilegeElevation);
+			int ExitCode = CopyAndRun.RunExternalCommand(Command, Params, this, bUnattended, bRequiresPrivilegeElevation, bCreateWindow);
 
 			TurnkeyUtils.Log("----------------------------------------------");
-			TurnkeyUtils.Log("Finished with {0}", bSuccess ? "Success" : "Failure");
-
-			if (bSuccess)
-			{
-				if (!string.IsNullOrEmpty(SuccessPostAmble))
-				{
-					TurnkeyUtils.Log("");
-					TurnkeyUtils.Log(SuccessPostAmble);
-				}
-			}
-			else
-			{
-				if (!string.IsNullOrEmpty(FailurePostamble))
-				{
-					TurnkeyUtils.Log("");
-					TurnkeyUtils.Log(FailurePostamble);
-				}
-			}
+			TurnkeyUtils.Log($"Finished with {ExitCode}");
 
 			TurnkeyUtils.Log("----------------------------------------------", Command);
 
-			return bSuccess;
+			return ExitCode;
 		}
 
 		public string RetrieveFileSource(string Name, string InType, string InPlatform, string SubType)
@@ -315,6 +295,17 @@ namespace Turnkey
 			return TurnkeyUtils.GetVariableValue(VariableName);
 		}
 
+		public void Log(string Message)
+		{
+			TurnkeyUtils.Log(Message);
+		}
+
+		public void ReportError(string Message)
+		{
+			Log("ERROR: " + Message);
+			ErrorMessages.Add(Message);
+		}
+
 		public void PauseForUser(string Message)
 		{
 			TurnkeyUtils.PauseForUser(Message);
diff --git a/Engine/Source/Programs/AutomationTool/Turnkey/FileSource.cs b/Engine/Source/Programs/AutomationTool/Turnkey/FileSource.cs
--- a/Engine/Source/Programs/AutomationTool/Turnkey/FileSource.cs
+++ b/Engine/Source/Programs/AutomationTool/Turnkey/FileSource.cs
@@ -331,7 +331,7 @@ namespace Turnkey
 			return null;
 		}
 
-		public bool DownloadOrInstall(UnrealTargetPlatform Platform, DeviceInfo Device = null, bool bUnattended = false)
+		public bool DownloadOrInstall(UnrealTargetPlatform Platform, ITurnkeyContext TurnkeyContext, DeviceInfo Device, bool bUnattended)
 		{
 			// standard variables
 			TurnkeyUtils.SetVariable("Platform", Platform.ToString());
@@ -340,27 +340,25 @@ namespace Turnkey
 			if (Type == SourceType.AutoSdk)
 			{
 				// AutoSdk has some extra setup needed
-				return SdkUtils.SetupAutoSdk(this, Platform, bUnattended);
+				return SdkUtils.SetupAutoSdk(this, TurnkeyContext, Platform, bUnattended);
 			}
 
-			CopyProviderRetriever Retriever = new CopyProviderRetriever();
-
 			// if we have sources, make sure we can download something. if we have a null Sources, that indicates
 			// it was code generated, and we don't need to download anything
 			if (Sources.Length != 0)
 			{
 				// this will set the $(CopyOutputPath) variable which the Sdks will generally need to use
-				string DownloadedSDK = Retriever.RetrieveFileSource(this);
+				string DownloadedSDK = TurnkeyContext.RetrieveFileSource(this);
 
 				if (string.IsNullOrEmpty(DownloadedSDK))
 				{
-//					Retriever.ReportError($"Unable to download anInstaller for {Platform}. Your Studio's TurnkeyManifest.xml file(s) may need to be fixed.");
+//					TurnkeyContext.ReportError($"Unable to download anInstaller for {Platform}. Your Studio's TurnkeyManifest.xml file(s) may need to be fixed.");
 					return false;
 				}
 			}
 
 			// let the platform decide how to install
-			return AutomationTool.Platform.GetPlatform(Platform).InstallSDK(TurnkeyUtils.CommandUtilHelper, Retriever, Device);
+			return AutomationTool.Platform.GetPlatform(Platform).InstallSDK(TurnkeyUtils.CommandUtilHelper, TurnkeyContext, Device, bUnattended);
 		}
 
 
diff --git a/Engine/Source/Programs/AutomationTool/Turnkey/SdkInfo.cs b/Engine/Source/Programs/AutomationTool/Turnkey/SdkInfo.cs
--- a/Engine/Source/Programs/AutomationTool/Turnkey/SdkInfo.cs
+++ b/Engine/Source/Programs/AutomationTool/Turnkey/SdkInfo.cs
@@ -38,7 +38,7 @@ namespace Turnkey
 			Support_AutoSdk = 8192,
 		}
 
-		static public LocalAvailability GetLocalAvailability(AutomationTool.Platform AutomationPlatform, bool bAllowUpdatingPrerequisites)
+		static public LocalAvailability GetLocalAvailability(AutomationTool.Platform AutomationPlatform, bool bAllowUpdatingPrerequisites, TurnkeyContextImpl TurnkeyContext)
 		{
 			LocalAvailability Result = LocalAvailability.None;
 
@@ -49,10 +49,7 @@ namespace Turnkey
 				return Result;
 			}
 
-			// no need to retrieve anything if we aren't allowing updates
-			CopyProviderRetriever Retriever = new CopyProviderRetriever();
-
-			if (AutomationPlatform.UpdateHostPrerequisites(TurnkeyUtils.CommandUtilHelper, Retriever, !bAllowUpdatingPrerequisites))
+			if (AutomationPlatform.UpdateHostPrerequisites(TurnkeyUtils.CommandUtilHelper, TurnkeyContext, !bAllowUpdatingPrerequisites))
 			{
 				Result |= LocalAvailability.Platform_ValidHostPrerequisites;
 			}
@@ -143,105 +140,92 @@ namespace Turnkey
 		}
 
 
-		public static bool SetupAutoSdk(FileSource Source, UnrealTargetPlatform Platform, bool bUnattended)
+		public static bool SetupAutoSdk(FileSource Source, ITurnkeyContext TurnkeyContext, UnrealTargetPlatform Platform, bool bUnattended)
 		{
-			bool bAttemptAutoSdkSetup = false;
 			bool bSetupEnvVarAfterInstall = false;
-			if (Environment.GetEnvironmentVariable("UE_SDKS_ROOT") != null)
-			{
-				bAttemptAutoSdkSetup = true;
-			}
-			else
+			if (Environment.GetEnvironmentVariable("UE_SDKS_ROOT") == null)
 			{
-				if (!bUnattended)
+				if (bUnattended)
 				{
-					string Response = TurnkeyUtils.ReadInput("The AutoSdk system is not setup on this machine. Would you like to set it up now? [Y/n]", "Y");
-					if (string.Compare(Response, "Y", true) == 0)
-					{
-						bAttemptAutoSdkSetup = true;
-						bSetupEnvVarAfterInstall = true;
-					}
+					TurnkeyContext.ReportError($"Unable to install an AutoSDK without UE_SDKS_ROOT setup (can use Turnkey interactively to set it up)");
+					return false;
+				}
+
+				string Response = TurnkeyUtils.ReadInput("The AutoSdk system is not setup on this machine. Would you like to set it up now? [Y/n]", "Y");
+				if (string.Compare(Response, "Y", true) == 0)
+				{
+					bSetupEnvVarAfterInstall = true;
 				}
 			}
 
-			if (bAttemptAutoSdkSetup)
-			{
-				AutomationTool.Platform AutomationPlatform = AutomationTool.Platform.GetPlatform(Platform);
+			AutomationTool.Platform AutomationPlatform = AutomationTool.Platform.GetPlatform(Platform);
 
-				TurnkeyUtils.Log("{0}: AutoSdk is setup on this computer, will look for available AutoSdk to download", Platform);
+			TurnkeyUtils.Log("{0}: AutoSdk is setup on this computer, will look for available AutoSdk to download", Platform);
 
-				// make sure this is unset so that we can know if it worked or not after install
-				TurnkeyUtils.ClearVariable("CopyOutputPath");
+			// make sure this is unset so that we can know if it worked or not after install
+			TurnkeyUtils.ClearVariable("CopyOutputPath");
 
-				// now download it (AutoSdks don't "install") on download
-				// @todo turnkey: handle errors, handle p4 going to wrong location, handle one Sdk for multiple platforms
-				string CopyOperation = Source.GetCopySourceOperation();
-				if (CopyOperation == null)
+			// now download it (AutoSdks don't "install") on download
+			// @todo turnkey: handle errors, handle p4 going to wrong location, handle one Sdk for multiple platforms
+			string CopyOperation = Source.GetCopySourceOperation();
+			if (CopyOperation == null)
+			{
+				TurnkeyContext.ReportError($"Unable to find AutoSDK FileSource fopr {Platform}. Your Studio's TurnkeyManifest.xml file(s) may need to be fixed.");
+				return false;
+			}
+
+			// download the AutoSDK
+			string DownloadedRoot = CopyProvider.ExecuteCopy(CopyOperation);
+			if (string.IsNullOrEmpty(DownloadedRoot))
+			{
+				TurnkeyContext.ReportError($"Unable to download the AutoSDK for {Platform}. Your Studio's TurnkeyManifest.xml file(s) may need to be fixed.");
+				return false;
+			}
+
+			if (bSetupEnvVarAfterInstall)
+			{
+				// walk up to one above Host* directory
+				DirectoryInfo AutoSdkSearch;
+				if (Directory.Exists(DownloadedRoot))
 				{
-					return false;
+					AutoSdkSearch = new DirectoryInfo(DownloadedRoot);
 				}
-
-				// download the AutoSDK
-				string DownloadedRoot = CopyProvider.ExecuteCopy(CopyOperation);
-				if (DownloadedRoot == null)
+				else
 				{
-					return false;
+					AutoSdkSearch = new FileInfo(DownloadedRoot).Directory;
 				}
-
-				if (bSetupEnvVarAfterInstall)
+				while (AutoSdkSearch.Name != "Host" + HostPlatform.Current.HostEditorPlatform.ToString())
 				{
-					// failed to install, nothing we can do
-					if (string.IsNullOrEmpty(DownloadedRoot))
-					{
-						TurnkeyUtils.ExitCode = AutomationTool.ExitCode.Error_SDKNotFound;
-						return false;
-					}
-
-					// walk up to one above Host* directory
-					DirectoryInfo AutoSdkSearch;
-					if (Directory.Exists(DownloadedRoot))
-					{
-						AutoSdkSearch = new DirectoryInfo(DownloadedRoot);
-					}
-					else
-					{
-						AutoSdkSearch = new FileInfo(DownloadedRoot).Directory;
-					}
-					while (AutoSdkSearch.Name != "Host" + HostPlatform.Current.HostEditorPlatform.ToString())
-					{
-						AutoSdkSearch = AutoSdkSearch.Parent;
-					}
-
-					// now go one up to the parent of Host
 					AutoSdkSearch = AutoSdkSearch.Parent;
+				}
 
-					string AutoSdkDir = AutoSdkSearch.FullName;
-					if (!bUnattended)
+				// now go one up to the parent of Host
+				AutoSdkSearch = AutoSdkSearch.Parent;
+
+				string AutoSdkDir = AutoSdkSearch.FullName;
+				if (!bUnattended)
+				{
+					string Response = TurnkeyUtils.ReadInput("Enter directory for root of AutoSdks. Use detected value, or enter another:", AutoSdkSearch.FullName);
+					if (string.IsNullOrEmpty(Response))
 					{
-						string Response = TurnkeyUtils.ReadInput("Enter directory for root of AutoSdks. Use detected value, or enter another:", AutoSdkSearch.FullName);
-						if (string.IsNullOrEmpty(Response))
-						{
-							return false;
-						}
+						return false;
 					}
-
-					// set the env var, globally
-					TurnkeyUtils.StartTrackingExternalEnvVarChanges();
-					Environment.SetEnvironmentVariable("UE_SDKS_ROOT", AutoSdkDir);
-					Environment.SetEnvironmentVariable("UE_SDKS_ROOT", AutoSdkDir, EnvironmentVariableTarget.User);
-					TurnkeyUtils.EndTrackingExternalEnvVarChanges();
-
 				}
 
-				// and now activate it in case we need it this run
-				TurnkeyUtils.Log("Re-activating AutoSDK '{0}'...", Source.Name);
-
-				UEBuildPlatformSDK.GetSDKForPlatform(Platform.ToString()).ReactivateAutoSDK();
+				// set the env var, globally
+				TurnkeyUtils.StartTrackingExternalEnvVarChanges();
+				Environment.SetEnvironmentVariable("UE_SDKS_ROOT", AutoSdkDir);
+				Environment.SetEnvironmentVariable("UE_SDKS_ROOT", AutoSdkDir, EnvironmentVariableTarget.User);
+				TurnkeyUtils.EndTrackingExternalEnvVarChanges();
 
-				return true;
 			}
 
-			return false;
+			// and now activate it in case we need it this run
+			TurnkeyUtils.Log("Re-activating AutoSDK '{0}'...", Source.Name);
+
+			UEBuildPlatformSDK.GetSDKForPlatform(Platform.ToString()).ReactivateAutoSDK();
+
+			return true;
 		}
 	}
 }
diff --git a/Engine/Source/Programs/AutomationTool/Turnkey/TurnkeyUtils.cs b/Engine/Source/Programs/AutomationTool/Turnkey/TurnkeyUtils.cs
--- a/Engine/Source/Programs/AutomationTool/Turnkey/TurnkeyUtils.cs
+++ b/Engine/Source/Programs/AutomationTool/Turnkey/TurnkeyUtils.cs
@@ -478,6 +478,11 @@ namespace Turnkey
 				{
 					ChosenPlatforms = GetAllValidPlatforms(PossiblePlatforms);
 				}
+				// if there's only one platform possible, use it
+				else if (PossiblePlatforms.Count == 1)
+				{
+					ChosenPlatforms = PossiblePlatforms;
+				}
 				else
 				{
 					ChosenPlatforms = TurnkeyUtils.GetPlatformsFromCommandLineOrUser(CommandOptions, PossiblePlatforms);
@@ -563,6 +568,33 @@ namespace Turnkey
 			return (ChosenDevices != null && ChosenDevices.Count > 0) ? ChosenDevices : null;
 		}
 
+		public static void GetPlatformsAndDevicesFromCommandLineOrUser(string[] CommandOptions, bool bSkipAskUserForDevice, out List<UnrealTargetPlatform> Platforms, out List<DeviceInfo> Devices, List<UnrealTargetPlatform> AllowedPlatforms=null)
+		{
+			Platforms = new List<UnrealTargetPlatform>();
+			Devices = new List<DeviceInfo>();
+
+			string DeviceString = TurnkeyUtils.ParseParamValue("Device", null, CommandOptions);
+			if (string.IsNullOrEmpty(DeviceString))
+			{
+				Platforms.AddRange(TurnkeyUtils.GetPlatformsFromCommandLineOrUser(CommandOptions, AllowedPlatforms));
+				// restrict devices to this platform
+				AllowedPlatforms = Platforms;
+			}
+
+			// if there's no -device param, and we don't want to ask user for a device, skip this
+			if (!string.IsNullOrEmpty(DeviceString) || !bSkipAskUserForDevice)
+			{
+				List<DeviceInfo> ChosenDevices = TurnkeyUtils.GetDevicesFromCommandLineOrUser(CommandOptions, AllowedPlatforms);
+				
+				if (ChosenDevices != null)
+				{
+					Devices = ChosenDevices;
+					// pull the platforms out of the devices we have chosen
+					Platforms = Devices.Select(x => x.Platform).ToHashSet().ToList();
+				}
+			}
+		}
+
 		public static string GetGenericOption(string[] CommandOptions, List<string> Options, string CommandLineOption)
 		{
 			return GetGenericOption(CommandOptions, Options, CommandLineOption, out _);
diff --git a/Engine/Source/Programs/DotNETCommon/BuildUtilities/UEBuildPlatformSDK.cs b/Engine/Source/Programs/DotNETCommon/BuildUtilities/UEBuildPlatformSDK.cs
--- a/Engine/Source/Programs/DotNETCommon/BuildUtilities/UEBuildPlatformSDK.cs
+++ b/Engine/Source/Programs/DotNETCommon/BuildUtilities/UEBuildPlatformSDK.cs
@@ -219,10 +219,10 @@ namespace EpicGames.Core
 		/// <summary>
 		/// For a platform that doesn't use properly named AutoSDK directories, the directory name may not be convertible to an integer,
 		/// and IsVersionValid checks could fail when checking AutoSDK version for an exact match. GetMainVersion() would return the 
-		/// proper, integer-convertible version number of the SDK inside of the directory returned by GetAutoSDKDirectoryForMasterVersion()
+		/// proper, integer-convertible version number of the SDK inside of the directory returned by GetAutoSDKDirectoryForMainVersion()
 		/// </summary>
 		/// <returns></returns>
-		public virtual string GetAutoSDKDirectoryForMasterVersion()
+		public virtual string GetAutoSDKDirectoryForMainVersion()
 		{
 			return GetMainVersion();
 		}
@@ -294,7 +294,8 @@ namespace EpicGames.Core
 				// AutoSDK must match the desired version exactly, since that is the only one we will use
 				if (bForAutoSDK)
 				{
-					return false;
+					// if integer version checking failed, then we can detect valid autosdk if the version matches the autosdk directory by name
+					return string.Compare(Version, GetAutoSDKDirectoryForMainVersion(), true) == 0;
 				}
 			}
 
@@ -676,7 +677,7 @@ namespace EpicGames.Core
 		{
 			if (PlatformSupportsAutoSDKs() && HasAutoSDKSystemEnabled())
 			{
-				String InstalledSDKVersionString = GetAutoSDKDirectoryForMasterVersion();
+				String InstalledSDKVersionString = GetAutoSDKDirectoryForMainVersion();
 				String PlatformSDKRoot = GetPathToPlatformAutoSDKs();
                 if (!Directory.Exists(PlatformSDKRoot))
                 {
@@ -1061,7 +1062,7 @@ namespace EpicGames.Core
 					bool bEnvVarFileExists = File.Exists(EnvVarFile);
 
 					string CurrentSDKString;
-					if (bEnvVarFileExists && GetCurrentlyInstalledSDKString(AutoSDKRoot, out CurrentSDKString) && CurrentSDKString == GetAutoSDKDirectoryForMasterVersion() && bScriptVersionMatches)
+					if (bEnvVarFileExists && GetCurrentlyInstalledSDKString(AutoSDKRoot, out CurrentSDKString) && CurrentSDKString == GetAutoSDKDirectoryForMainVersion() && bScriptVersionMatches)
 					{
 						return SDKStatus.Valid;
 					}
@@ -1169,22 +1170,22 @@ namespace EpicGames.Core
 					// delete Manifest file to avoid multiple uninstalls
 					InvalidateCurrentlyInstalledAutoSDK();
 
-					if (!RunAutoSDKHooks(AutoSDKRoot, GetAutoSDKDirectoryForMasterVersion(), SDKHookType.Install, false))
+					if (!RunAutoSDKHooks(AutoSDKRoot, GetAutoSDKDirectoryForMainVersion(), SDKHookType.Install, false))
 					{
-						Log.TraceLog("Failed to install required SDK {0}.  Attemping to uninstall", GetAutoSDKDirectoryForMasterVersion());
-						RunAutoSDKHooks(AutoSDKRoot, GetAutoSDKDirectoryForMasterVersion(), SDKHookType.Uninstall, false);
+						Log.TraceLog("Failed to install required SDK {0}.  Attemping to uninstall", GetAutoSDKDirectoryForMainVersion());
+						RunAutoSDKHooks(AutoSDKRoot, GetAutoSDKDirectoryForMainVersion(), SDKHookType.Uninstall, false);
 						return;
 					}
 
 					string EnvVarFile = Path.Combine(AutoSDKRoot, SDKEnvironmentVarsFile);
 					if (!File.Exists(EnvVarFile))
 					{
-						Log.TraceLog("Installation of required SDK {0}.  Did not generate Environment file {1}", GetAutoSDKDirectoryForMasterVersion(), EnvVarFile);
-						RunAutoSDKHooks(AutoSDKRoot, GetAutoSDKDirectoryForMasterVersion(), SDKHookType.Uninstall, false);
+						Log.TraceLog("Installation of required SDK {0}.  Did not generate Environment file {1}", GetAutoSDKDirectoryForMainVersion(), EnvVarFile);
+						RunAutoSDKHooks(AutoSDKRoot, GetAutoSDKDirectoryForMainVersion(), SDKHookType.Uninstall, false);
 						return;
 					}
 
-					SetCurrentlyInstalledAutoSDKString(GetAutoSDKDirectoryForMasterVersion());
+					SetCurrentlyInstalledAutoSDKString(GetAutoSDKDirectoryForMainVersion());
 					SetLastRunAutoSDKScriptVersion(GetRequiredScriptVersionString());
 				}
 
@@ -1225,7 +1226,7 @@ namespace EpicGames.Core
 			// load environment variables from current SDK
 			if (!SetupEnvironmentFromAutoSDK(PlatformSDKRoot))
 			{
-				Log.TraceLog("Failed to load environment from required SDK {0}", GetAutoSDKDirectoryForMasterVersion());
+				Log.TraceLog("Failed to load environment from required SDK {0}", GetAutoSDKDirectoryForMainVersion());
 				InvalidateCurrentlyInstalledAutoSDK();
 				return SDKStatus.Invalid;
 			}
diff --git a/Engine/Source/Programs/IOS/iPhonePackager/iPhonePackager.cs b/Engine/Source/Programs/IOS/iPhonePackager/iPhonePackager.cs
--- a/Engine/Source/Programs/IOS/iPhonePackager/iPhonePackager.cs
+++ b/Engine/Source/Programs/IOS/iPhonePackager/iPhonePackager.cs
@@ -74,7 +74,9 @@ namespace iPhonePackager
         Error_SymbolizedSONotFound = 153,
         Error_LicenseNotAccepted = 154,
         Error_AndroidOBBError = 155,
-    };
+		Error_SDKInstallFailed = 200,
+		Error_DeviceUpdateFailed = 201,
+	};
 
     public partial class Program
     {
diff --git a/Engine/Source/Programs/UnrealBuildTool/Platform/Android/AndroidPlatformSDK.cs b/Engine/Source/Programs/UnrealBuildTool/Platform/Android/AndroidPlatformSDK.cs
--- a/Engine/Source/Programs/UnrealBuildTool/Platform/Android/AndroidPlatformSDK.cs
+++ b/Engine/Source/Programs/UnrealBuildTool/Platform/Android/AndroidPlatformSDK.cs
@@ -17,7 +17,7 @@ namespace UnrealBuildTool
 		{
 			return "r21b";
 		}
-		public override string GetAutoSDKDirectoryForMasterVersion()
+		public override string GetAutoSDKDirectoryForMainVersion()
 		{
 			return "-23";
 		}
diff --git a/Engine/Source/Programs/UnrealBuildTool/Platform/Linux/LinuxPlatformSDK.cs b/Engine/Source/Programs/UnrealBuildTool/Platform/Linux/LinuxPlatformSDK.cs
--- a/Engine/Source/Programs/UnrealBuildTool/Platform/Linux/LinuxPlatformSDK.cs
+++ b/Engine/Source/Programs/UnrealBuildTool/Platform/Linux/LinuxPlatformSDK.cs
@@ -125,7 +125,7 @@ namespace UnrealBuildTool
 			string SDKRoot = Environment.GetEnvironmentVariable(SDKRootEnvVar);
 			if (!String.IsNullOrEmpty(SDKRoot))
 			{
-				string AutoSDKPath = Path.Combine(SDKRoot, "Host" + BuildHostPlatform.Current.Platform, TargetPlatformName, GetAutoSDKDirectoryForMasterVersion(), LinuxPlatform.DefaultHostArchitecture);
+				string AutoSDKPath = Path.Combine(SDKRoot, "Host" + BuildHostPlatform.Current.Platform, TargetPlatformName, GetAutoSDKDirectoryForMainVersion(), LinuxPlatform.DefaultHostArchitecture);
 				if (DirectoryReference.Exists(new DirectoryReference(AutoSDKPath)))
 				{
 					return AutoSDKPath;
