Integrating CL# 15321861 from //UE4/Main

diff --git a/Engine/Source/Programs/AutomationTool/BuildGraph/Tasks/NotarizeTask.cs b/Engine/Source/Programs/AutomationTool/BuildGraph/Tasks/NotarizeTask.cs
--- a/Engine/Source/Programs/AutomationTool/BuildGraph/Tasks/NotarizeTask.cs
+++ b/Engine/Source/Programs/AutomationTool/BuildGraph/Tasks/NotarizeTask.cs
@@ -49,7 +49,7 @@ namespace BuildGraph.Tasks
 		/// When true the notarization ticket will be stapled
 		/// </summary>
 		[TaskParameter(Optional = true)]
-		public bool RequireStapling = true;
+		public bool RequireStapling = false;
 	}
 
 	[TaskElement("Notarize", typeof(NotarizeTaskParameters))]
