UE5EA - Fix Turnkey crash with null PossiblePlatforms
Fix a typo

diff --git a/Engine/Source/Programs/AutomationTool/Turnkey/Commands/VerifySdk.cs b/Engine/Source/Programs/AutomationTool/Turnkey/Commands/VerifySdk.cs
--- a/Engine/Source/Programs/AutomationTool/Turnkey/Commands/VerifySdk.cs
+++ b/Engine/Source/Programs/AutomationTool/Turnkey/Commands/VerifySdk.cs
@@ -209,7 +209,7 @@ namespace Turnkey.Commands
 
 									if (MatchingInstallableSdk == null)
 									{
-										TurnkeyUtils.Log("ERROR: {0}: Unable top find any Sdks that could be installed on {1}", Platform, Device.Name);
+										TurnkeyUtils.Log("ERROR: {0}: Unable to find any Sdks that could be installed on {1}", Platform, Device.Name);
 										TurnkeyUtils.ExitCode = AutomationTool.ExitCode.Error_SDKNotFound;
 									}
 									else
diff --git a/Engine/Source/Programs/AutomationTool/Turnkey/TurnkeyUtils.cs b/Engine/Source/Programs/AutomationTool/Turnkey/TurnkeyUtils.cs
--- a/Engine/Source/Programs/AutomationTool/Turnkey/TurnkeyUtils.cs
+++ b/Engine/Source/Programs/AutomationTool/Turnkey/TurnkeyUtils.cs
@@ -479,7 +479,7 @@ namespace Turnkey
 					ChosenPlatforms = GetAllValidPlatforms(PossiblePlatforms);
 				}
 				// if there's only one platform possible, use it
-				else if (PossiblePlatforms.Count == 1)
+				else if (PossiblePlatforms != null && PossiblePlatforms.Count == 1)
 				{
 					ChosenPlatforms = PossiblePlatforms;
 				}
