Use build Agent with GPU for ShooterGame BootTest
+ Use UE.BootTest gauntlet test to run boot test
+ Improve UE.BootTest CreateReport()
+ Improve Errors and Warnings reported to Horde in UnrealTestNode

#rnx

diff --git a/Engine/Source/Programs/AutomationTool/Gauntlet/Unreal/Automation/UE.Automation.cs b/Engine/Source/Programs/AutomationTool/Gauntlet/Unreal/Automation/UE.Automation.cs
--- a/Engine/Source/Programs/AutomationTool/Gauntlet/Unreal/Automation/UE.Automation.cs
+++ b/Engine/Source/Programs/AutomationTool/Gauntlet/Unreal/Automation/UE.Automation.cs
@@ -294,8 +294,6 @@ namespace UE
 		public AutomationNodeBase(Gauntlet.UnrealTestContext InContext)
 			: base(InContext)
 		{
-			// We format warnings ourselves so don't show these
-			LogWarningsAndErrorsAfterSummary = false;
 		}
 
 	
@@ -591,7 +589,7 @@ namespace UE
 
 			if (SessionArtifacts == null)
 			{
-				return Enumerable.Empty<string>();
+				return AllWarnings;
 			}
 
 			foreach (var Artifact in SessionArtifacts)
@@ -603,7 +601,7 @@ namespace UE
 						Parser.GetResults()
 							.SelectMany(R => R.Events
 								.Where(E => E.ToLower().Contains("warning"))
-								.Distinct()
+								.Distinct().Select(E => string.Format("[test={0}] {1}", R.DisplayName, E))
 							)
 						);
 				}
diff --git a/Engine/Source/Programs/AutomationTool/Gauntlet/Unreal/Automation/UE.BootTest.cs b/Engine/Source/Programs/AutomationTool/Gauntlet/Unreal/Automation/UE.BootTest.cs
--- a/Engine/Source/Programs/AutomationTool/Gauntlet/Unreal/Automation/UE.BootTest.cs
+++ b/Engine/Source/Programs/AutomationTool/Gauntlet/Unreal/Automation/UE.BootTest.cs
@@ -123,59 +123,55 @@ namespace UE
 		/// Called after a test finishes to create an overall summary based on looking at the artifacts
 		/// </summary>
 		/// <param name="Result"></param>
-		/// <param name="Context"></param>
-		/// <param name="Build"></param>
-		/// <param name="Artifacts"></param>
-		/// <param name="InArtifactPath"></param>
-		public override void CreateReport(TestResult Result, UnrealTestContext Context, UnrealBuildSource Build, IEnumerable<UnrealRoleArtifacts> Artifacts, string InArtifactPath)
+		/// <returns>ITestReport</returns>
+		public override ITestReport CreateReport(TestResult Result)
 		{
-			// only check for artifacts if the test passed
-			if (Result !=TestResult.Passed)
+			if (Result == TestResult.Passed)
 			{
-				return;
-			}
+				if (!DidDetectLaunch)
+				{
+					ReportError("Failed to detect completion of launch");
+				}
+				else
+				{
+					// find a logfile or something that indicates the process ran successsfully
+					bool MissingFiles = false;
 
-			if (!DidDetectLaunch)
-			{
-				SetUnrealTestResult(TestResult.Failed);
-				Log.Error("Failed to detect completion of launch");
-				return;
-			}
+					foreach (var RoleArtifact in SessionArtifacts)
+					{
+						DirectoryInfo RoleDir = new DirectoryInfo(RoleArtifact.ArtifactPath);
 
-			// find a logfile or something that indicates the process ran successsfully
-			bool MissingFiles = false;
+						IEnumerable<FileInfo> ArtifactFiles = RoleDir.EnumerateFiles("*.*", SearchOption.AllDirectories);
 
-			foreach(var RoleArtifact in Artifacts)
-			{
-				DirectoryInfo RoleDir = new DirectoryInfo(RoleArtifact.ArtifactPath);
+						// user may not have cleared paths between runs, so throw away anything that's older than 2m
+						ArtifactFiles = ArtifactFiles.Where(F => (DateTime.Now - F.LastWriteTime).TotalMinutes < 2);
 
-				IEnumerable<FileInfo> ArtifactFiles = RoleDir.EnumerateFiles("*.*", SearchOption.AllDirectories);
+						if (ArtifactFiles.Any() == false)
+						{
+							MissingFiles = true;
+							ReportError("No artifact files found for {0}. Were they not retrieved from the device?", RoleArtifact.SessionRole);
+						}
 
-				// user may not have cleared paths between runs, so throw away anything that's older than 2m
-				ArtifactFiles = ArtifactFiles.Where(F => (DateTime.Now - F.LastWriteTime).TotalMinutes < 2);
+						IEnumerable<FileInfo> LogFiles = ArtifactFiles.Where(F => F.Extension.Equals(".log", StringComparison.OrdinalIgnoreCase));
 
-				if (ArtifactFiles.Any() == false)
-				{
-					MissingFiles = true;
-					Log.Error("No artifact files found for {0}. Were they not retrieved from the device?", RoleArtifact.SessionRole);
-				}
+						if (LogFiles.Any() == false)
+						{
+							MissingFiles = true;
+							ReportError("No log files found for {0}. Were they not retrieved from the device?", RoleArtifact.SessionRole);
+						}
+					}
 
-				IEnumerable<FileInfo> LogFiles = ArtifactFiles.Where(F => F.Extension.Equals(".log", StringComparison.OrdinalIgnoreCase));
+					if (MissingFiles)
+					{
+						SetUnrealTestResult(TestResult.Failed);
+						ReportError("One or more roles did not generated any artifacts");
+					}
 
-				if (LogFiles.Any() == false)
-				{
-					MissingFiles = true;
-					Log.Error("No log files found for {0}. Were they not retrieved from the device?", RoleArtifact.SessionRole);
+					Log.Info("Found valid artifacts for test");
 				}
 			}
 
-			if (MissingFiles)
-			{
-				SetUnrealTestResult(TestResult.Failed);
-				Log.Error("One or more roles did not generated any artifacts");
-			}
-
-			Log.Info("Found valid artifacts for test");
+			return base.CreateReport(GetTestResult());
 		}
 	}
 }
diff --git a/Engine/Source/Programs/AutomationTool/Gauntlet/Unreal/Automation/UE.ErrorTest.cs b/Engine/Source/Programs/AutomationTool/Gauntlet/Unreal/Automation/UE.ErrorTest.cs
--- a/Engine/Source/Programs/AutomationTool/Gauntlet/Unreal/Automation/UE.ErrorTest.cs
+++ b/Engine/Source/Programs/AutomationTool/Gauntlet/Unreal/Automation/UE.ErrorTest.cs
@@ -80,7 +80,6 @@ namespace UE
 			ErrorType = ErrorTypes.Check;
 			ErrorDelay = 5;
 			Server = false;
-			LogWarningsAndErrorsAfterSummary = false;
 		}
 
 		public override UnrealGame.UE4TestConfig GetConfiguration()
diff --git a/Engine/Source/Programs/AutomationTool/Gauntlet/Unreal/Base/Gauntlet.UnrealTestNode.cs b/Engine/Source/Programs/AutomationTool/Gauntlet/Unreal/Base/Gauntlet.UnrealTestNode.cs
--- a/Engine/Source/Programs/AutomationTool/Gauntlet/Unreal/Base/Gauntlet.UnrealTestNode.cs
+++ b/Engine/Source/Programs/AutomationTool/Gauntlet/Unreal/Base/Gauntlet.UnrealTestNode.cs
@@ -100,7 +100,7 @@ namespace Gauntlet
 		public override TestPriority Priority { get { return GetPriority(); } }
 
 		/// <summary>
-		/// Returns Warnings found during tests. By default only ensures are considered
+		/// Returns Warnings found during tests. By default only ensures and warnings in case of abnormal exit are considered
 		/// </summary>
 		public override IEnumerable<string> GetWarnings()
 		{
@@ -109,14 +109,13 @@ namespace Gauntlet
 				return Warnings;
 			}
 
-			return SessionArtifacts.SelectMany(A =>
-			{
-				return A.LogSummary.Ensures.Select(E => E.Message);
-			}).Union(Warnings); 
+			return SessionArtifacts.SelectMany(A => A.LogSummary.Ensures.Select(E => E.Message))
+					.Union(SessionArtifacts.Where(A => A.AppInstance.WasKilled == false && A.LogSummary.HasAbnormalExit).SelectMany(A => A.LogSummary.Warnings))
+					.Union(Warnings); 
 		}
 
 		/// <summary>
-		/// Returns Errors found during tests. By default only fatal errors are considered
+		/// Returns Errors found during tests. By default only fatal errors and errors in case of abnormal exit are considered
 		/// </summary>
 		public override IEnumerable<string> GetErrors()
 		{
@@ -125,18 +124,27 @@ namespace Gauntlet
 				return Errors;
 			}
 
-			var FailedArtifacts = GetArtifactsWithFailures();
+			IEnumerable<UnrealRoleArtifacts> FailedArtifacts = GetArtifactsWithFailures();
+
+			IEnumerable<string> AbnormalExitErrors = FailedArtifacts.Where(A => A.AppInstance.WasKilled == false && A.LogSummary.HasAbnormalExit).SelectMany(A => A.LogSummary.Errors).Distinct();
 
-			 return FailedArtifacts.Where(A => A.LogSummary.FatalError != null).Select(A => A.LogSummary.FatalError.Message).Union(Errors);
+			IEnumerable<string> FatalErrors = FailedArtifacts.Where(A => A.LogSummary.FatalError != null).Select(A => A.LogSummary.FatalError.Message);
+				
+			return FatalErrors.Union(AbnormalExitErrors).Union(Errors);
 		}
 
 		/// <summary>
 		/// Returns Errors found during tests. Including Abnornal Exit reasons
 		/// </summary>
-		public virtual IEnumerable<string> GetErrorsAndAbnornalExits()
+		public virtual IEnumerable<string> GetErrorsAndAbnornalExitReasons()
 		{
 			IEnumerable<string>  Errors = GetErrors();
 
+			if (SessionArtifacts == null)
+			{
+				return Errors;
+			}
+
 			Dictionary<UnrealRoleArtifacts, Tuple<int, string>> ErrorCodesAndReasons = new Dictionary<UnrealRoleArtifacts, Tuple<int, string>>();
 
 			var FailedArtifacts = SessionArtifacts.Where(
@@ -307,6 +315,8 @@ namespace Gauntlet
 			TestVersion = new Version("1.0.0");
 			ArtifactPath = string.Empty;
 			PopulateCommandlineInfo();
+			// We format warnings ourselves so don't show these
+			LogWarningsAndErrorsAfterSummary = false;
 		}
 
 		 ~UnrealTestNode()
@@ -1033,15 +1043,19 @@ namespace Gauntlet
 			HordeTestReport.ReportCreatedOn = DateTime.Now.ToString();
 			HordeTestReport.TotalDurationSeconds = (float) (DateTime.Now - SessionStartTime).TotalSeconds;
 			HordeTestReport.Description = Context.ToString();
-			HordeTestReport.Status = Result.ToString();
 			HordeTestReport.URLLink = GetURLLink();
-			HordeTestReport.Errors.AddRange(GetErrorsAndAbnornalExits());
+			HordeTestReport.Errors.AddRange(GetErrorsAndAbnornalExitReasons());
 			if (!string.IsNullOrEmpty(CancellationReason))
 			{
 				HordeTestReport.Errors.Add(CancellationReason);
 			}
 			HordeTestReport.Warnings.AddRange(GetWarnings());
 			HordeTestReport.HasSucceeded = !(Result == TestResult.Failed || Result == TestResult.TimedOut || HordeTestReport.Errors.Count > 0);
+			if (HordeTestReport.Errors.Count > 0 && Result == TestResult.Passed)
+			{
+				SetUnrealTestResult(TestResult.Failed);
+			}
+			HordeTestReport.Status = GetTestResult().ToString();
 			string HordeArtifactPath = string.IsNullOrEmpty(GetConfiguration().HordeArtifactPath) ? HordeReport.DefaultArtifactsDir : GetConfiguration().HordeArtifactPath;
 			HordeTestReport.SetOutputArtifactPath(HordeArtifactPath);
 			if (SessionArtifacts != null)
