Expose automation controller idle/tick timeout to AutomationTool.

= testing =
Verified fixed by running locally via RunUAT

diff --git a/Engine/Source/Programs/AutomationTool/Gauntlet/Unreal/Editor/EditorTests.cs b/Engine/Source/Programs/AutomationTool/Gauntlet/Unreal/Editor/EditorTests.cs
--- a/Engine/Source/Programs/AutomationTool/Gauntlet/Unreal/Editor/EditorTests.cs
+++ b/Engine/Source/Programs/AutomationTool/Gauntlet/Unreal/Editor/EditorTests.cs
@@ -49,6 +49,12 @@ namespace UnrealEditor
 		[AutoParam]
 		public string TreatLogErrorsAsTestErrors = "true";
 
+		/// <summary>
+		/// Control for interpretation of log errors as test failures
+		/// </summary>
+		[AutoParam]
+		public string GameInstanceLostTimerSeconds = string.Empty;
+
 		/// <summary>
 		/// Applies these options to the provided app config
 		/// </summary>
@@ -79,6 +85,11 @@ namespace UnrealEditor
 			{
 				AppConfig.CommandLineParams.Add("ini:Engine:[/Script/AutomationController.AutomationControllerSettings]:bTreatLogErrorsAsTestErrors=false");
 			}
+
+			if (GameInstanceLostTimerSeconds != string.Empty)
+			{
+				AppConfig.CommandLineParams.Add(string.Format("ini:Engine:[/Script/AutomationController.AutomationControllerSettings]:GameInstanceLostTimerSeconds={0}", GameInstanceLostTimerSeconds));
+			}
 		}
 	}
 
