Libimobiledevice support for Network AppleTVs #jira UE-103878 #rb jack.porter #fyi josh.adams

diff --git a/Engine/Source/Developer/IOS/IOSTargetPlatform/Private/IOSDeviceHelper.cpp b/Engine/Source/Developer/IOS/IOSTargetPlatform/Private/IOSDeviceHelper.cpp
--- a/Engine/Source/Developer/IOS/IOSTargetPlatform/Private/IOSDeviceHelper.cpp
+++ b/Engine/Source/Developer/IOS/IOSTargetPlatform/Private/IOSDeviceHelper.cpp
@@ -12,13 +12,22 @@
 
 DEFINE_LOG_CATEGORY_STATIC(LogIOSDeviceHelper, Log, All);
 
-struct FDeviceNotificationCallbackInformation
-{
-    FString UDID;
-    FString    DeviceName;
-    FString ProductType;
-    uint32 msgType;
-};
+    enum DeviceConnectionInterface
+    {
+        NoValue = 0,
+        USB = 1,
+        Network = 2,
+        Max = 4
+    };
+
+    struct FDeviceNotificationCallbackInformation
+    {
+        FString UDID;
+        FString DeviceName;
+        FString ProductType;
+        DeviceConnectionInterface DeviceInterface;
+        uint32 msgType;
+    };
 
 
     struct LibIMobileDevice
@@ -26,6 +35,7 @@ struct FDeviceNotificationCallbackInformation
         FString DeviceName;
         FString DeviceID;
         FString DeviceType;
+        DeviceConnectionInterface DeviceInterface;
     };
 
     static FString GetLibImobileDeviceExe(const FString &ExeName)
@@ -53,45 +63,54 @@ struct FDeviceNotificationCallbackInformation
         
 
         TArray<LibIMobileDevice> ToReturn;
-
         // separate out each line
-        TArray<FString> CurrentDeviceIds;
-        OutStdOut.ParseIntoArray(CurrentDeviceIds, TEXT("\n"), true);
-        for (int32 Index = 0; Index != CurrentDeviceIds.Num(); ++Index)
-        {
-            if (CurrentDeviceIds[Index].Contains("Network"))
-            {
-                CurrentDeviceIds.RemoveAt(Index);
-                Index--;
-                continue;
-            }
-            CurrentDeviceIds[Index].Split(TEXT(" "), &CurrentDeviceIds[Index], nullptr, ESearchCase::CaseSensitive, ESearchDir::FromStart);
-        }
+        TArray<FString> OngoingDeviceIds;
+        
+        OutStdOut.ParseIntoArray(OngoingDeviceIds, TEXT("\n"), true);
         TArray<FString> DeviceStrings;
-        for (int32 StringIndex = 0; StringIndex < CurrentDeviceIds.Num(); ++StringIndex)
+        for (int32 StringIndex = 0; StringIndex < OngoingDeviceIds.Num(); ++StringIndex)
         {
-
-            const FString& DeviceID = CurrentDeviceIds[StringIndex];
+            const FString& DeviceID = OngoingDeviceIds[StringIndex];
+            DeviceConnectionInterface OngoingDeviceInterface = DeviceConnectionInterface::NoValue;
 
             FString OutStdOutInfo;
             FString OutStdErrInfo;
             FString LibimobileDeviceInfo = GetLibImobileDeviceExe("ideviceinfo");
             int ReturnCodeInfo;
-            FString Arguments = "-u " + DeviceID;
-            // get the list of devices UDID
-            FPlatformProcess::ExecProcess(*LibimobileDeviceInfo, *Arguments, &ReturnCodeInfo, &OutStdOutInfo, &OutStdErrInfo);
+            FString Arguments;
+
+            if (OngoingDeviceIds[StringIndex].Contains("USB"))
+            {
+                OngoingDeviceInterface = DeviceConnectionInterface::USB;
+            }
+            else if (OngoingDeviceIds[StringIndex].Contains("Network"))
+            {
+                OngoingDeviceInterface = DeviceConnectionInterface::Network;
+            }
+            OngoingDeviceIds[StringIndex].Split(TEXT(" "), &OngoingDeviceIds[StringIndex], nullptr, ESearchCase::CaseSensitive, ESearchDir::FromStart);
 
+            if (OngoingDeviceInterface == DeviceConnectionInterface::USB)
+            {
+                Arguments = "-u " + DeviceID;
+            }
+            else if (OngoingDeviceInterface == DeviceConnectionInterface::Network)
+            {
+                Arguments = "-n -u " + DeviceID;
+            }
+
+            FPlatformProcess::ExecProcess(*LibimobileDeviceInfo, *Arguments, &ReturnCodeInfo, &OutStdOutInfo, &OutStdErrInfo);
             // parse product type and device name
             FString DeviceName;
             OutStdOutInfo.Split(TEXT("DeviceName: "), nullptr, &DeviceName, ESearchCase::CaseSensitive, ESearchDir::FromStart);
             DeviceName.Split(LINE_TERMINATOR, &DeviceName, nullptr, ESearchCase::CaseSensitive, ESearchDir::FromStart);
             FString ProductType;
             OutStdOutInfo.Split(TEXT("ProductType: "), nullptr, &ProductType, ESearchCase::CaseSensitive, ESearchDir::FromStart);
-            DeviceName.Split(LINE_TERMINATOR, &DeviceName, nullptr, ESearchCase::CaseSensitive, ESearchDir::FromStart);
+            ProductType.Split(LINE_TERMINATOR, &ProductType, nullptr, ESearchCase::CaseSensitive, ESearchDir::FromStart);
             LibIMobileDevice ToAdd;
             ToAdd.DeviceID = DeviceID;
             ToAdd.DeviceName = DeviceName;
             ToAdd.DeviceType = ProductType;
+            ToAdd.DeviceInterface = OngoingDeviceInterface;
             ToReturn.Add(ToAdd);
         }
         return ToReturn;
@@ -271,6 +290,7 @@ private:
             FDeviceNotificationCallbackInformation CallbackInfo;
             CallbackInfo.DeviceName = ParsedDevices[Index].DeviceName;
             CallbackInfo.UDID = ParsedDevices[Index].DeviceID;
+            CallbackInfo.DeviceInterface = ParsedDevices[Index].DeviceInterface;
             CallbackInfo.ProductType = ParsedDevices[Index].DeviceType;
             CallbackInfo.msgType = 1;
             DeviceNotification.Broadcast(&CallbackInfo);
diff --git a/Engine/Source/Programs/AutomationTool/IOS/IOSPlatform.Automation.cs b/Engine/Source/Programs/AutomationTool/IOS/IOSPlatform.Automation.cs
--- a/Engine/Source/Programs/AutomationTool/IOS/IOSPlatform.Automation.cs
+++ b/Engine/Source/Programs/AutomationTool/IOS/IOSPlatform.Automation.cs
@@ -153,7 +153,6 @@ public class IOSPlatform : Platform
 
 	private string PlatformName = null;
 	private string SDKName = null;
-
 	public IOSPlatform()
 		: this(UnrealTargetPlatform.IOS)
 	{
@@ -519,48 +518,35 @@ public class IOSPlatform : Platform
 
 	public override DeviceInfo[] GetDevices()
 	{
+
 		List<DeviceInfo> Devices = new List<DeviceInfo>();
 
-		if (HostPlatform.Current.HostEditorPlatform == UnrealTargetPlatform.Mac)
-		{
-			string DeviceType = TargetPlatformType == UnrealTargetPlatform.TVOS ? "tvOS" : "iOS";
+		var IdeviceIdPath = GetPathToLibiMobileDeviceTool("idevice_id");
+		string Output = Utils.RunLocalProcessAndReturnStdOut(IdeviceIdPath, "");
+		var ConnectedDevicesUDIDs = Output.Split(new string[] { Environment.NewLine }, StringSplitOptions.None);
 
-			// print out each connected device's udid, os version, and name (which might have spaces)
-			string Params = "-e 'require \"fastlane\"; FastlaneCore::DeviceManager.connected_devices(\"" + DeviceType + "\").each { |x| puts \"#{x.udid} #{x.os_type} #{x.os_version} #{x.name}\" }'";
-			string StdOut = UnrealBuildTool.Utils.RunLocalProcessAndReturnStdOut("ruby", Params);
+		DeviceInfo CurrentDevice = null;
 
-			foreach (string Line in StdOut.Split("\n".ToCharArray(), StringSplitOptions.RemoveEmptyEntries))
+		foreach (string UnparsedUDID in ConnectedDevicesUDIDs)
+		{
+			var IdeviceInfoPath = GetPathToLibiMobileDeviceTool("ideviceinfo");
+			String ParsedUDID = UnparsedUDID.Split(" ").First();
+			String IdeviceInfoArgs = "-u " + ParsedUDID;
+
+			if (UnparsedUDID.Contains("Network"))
 			{
-				string[] Tokens = Line.Split(" ".ToCharArray());
+				IdeviceInfoArgs = "-n " + IdeviceInfoArgs;
+			}
 
-				DeviceInfo Device = new DeviceInfo(TargetPlatformType);
-				Device.Id = Tokens[0];
-				Device.Type = Tokens[1];
-				Device.SoftwareVersion = Tokens[2];
-				Device.Name = string.Join(" ", Tokens.Skip(3));
+			string OutputInfo = Utils.RunLocalProcessAndReturnStdOut(IdeviceInfoPath, IdeviceInfoArgs);
 
-				Devices.Add(Device);
-			}
-		}
-		else if (HostPlatform.Current.HostEditorPlatform == UnrealTargetPlatform.Win64)
-		{
-			string UtilPath = Path.Combine(CommandUtils.EngineDirectory.FullName, "Extras/ThirdPartyNotUE/libimobiledevice/x64/ideviceinfo");
-			string Output = Utils.RunLocalProcessAndReturnStdOut(UtilPath, "");
+			CurrentDevice = new DeviceInfo(TargetPlatformType);
+			CurrentDevice.Name = UnparsedUDID;
 
-			DeviceInfo CurrentDevice = null;
-			foreach (string Line in Output.Split("\r\n".ToCharArray()))
+			foreach (string Line in OutputInfo.Split(Environment.NewLine.ToCharArray()))
 			{
-				if (Line.StartsWith("ActivationState:"))
-				{
-					CurrentDevice = new DeviceInfo(TargetPlatformType);
-					Devices.Add(CurrentDevice);
-				}
-				else if (Line.StartsWith("DeviceName:"))
-				{
-					CurrentDevice.Name = Line.Split(": ").Last();
-				}
 				// check we are returning the proper device for this class
-				else if (Line.StartsWith("DeviceClass:"))
+				if (Line.StartsWith("DeviceClass:"))
 				{
 					bool bIsDeviceTVOS = Line.Split(": ").Last().ToLower() == "tvos";
 					if (bIsDeviceTVOS != (TargetPlatformType == UnrealTargetPlatform.TVOS))
@@ -581,6 +567,7 @@ public class IOSPlatform : Platform
 					CurrentDevice.SoftwareVersion = Line.Split(": ").Last();
 				}
 			}
+			Devices.Add(CurrentDevice);
 		}
 		return Devices.ToArray();
 	}
@@ -1692,10 +1679,13 @@ public class IOSPlatform : Platform
 				BundleIdentifier = Contents.Substring(Pos, EndPos - Pos);
 			}
 
+			string IdeviceInstallerArgs = "--list-apps -u " + Params.DeviceNames[0];
+			IdeviceInstallerArgs = GetLibimobileDeviceNetworkedArgument(IdeviceInstallerArgs, Params.DeviceNames[0]);
+
 			var DeviceInstaller = GetPathToLibiMobileDeviceTool("ideviceinstaller");
 			LogInformation("Checking if bundle {0} is installed", BundleIdentifier);
 
-			string Output = CommandUtils.RunAndLog(DeviceInstaller, "--list-apps");
+			string Output = CommandUtils.RunAndLog(DeviceInstaller, IdeviceInstallerArgs);
 			bool bBundleIsInstalled = Output.Contains(string.Format("CFBundleIdentifier -> {0}{1}", BundleIdentifier, Environment.NewLine));
 			int ExitCode = 0;
 
@@ -1708,7 +1698,11 @@ public class IOSPlatform : Platform
 				string AllCommandsToPush = " push " + CombinePaths(Params.BaseStageDirectory, PlatformName, SC.GetUFSDeployedManifestFileName(null)) + "\n"
 										+ " push " + CombinePaths(Params.BaseStageDirectory, PlatformName, SC.GetNonUFSDeployedManifestFileName(null));		
 				System.IO.File.WriteAllText(Directory.GetCurrentDirectory() + "\\CommandsToPush.txt", AllCommandsToPush);
-				Utils.RunLocalProcessAndReturnStdOut(DeviceFS, "-b " + BundleIdentifier + " -x " + Directory.GetCurrentDirectory() + "\\CommandsToPush.txt", out ExitCode, true);
+
+				string IdeviceFSArgs = "-b " + BundleIdentifier + " -x " + Directory.GetCurrentDirectory() + "\\CommandsToPush.txt -u " + Params.DeviceNames[0];
+				IdeviceFSArgs = GetLibimobileDeviceNetworkedArgument(IdeviceFSArgs, Params.DeviceNames[0]);
+
+				Utils.RunLocalProcessAndReturnStdOut(DeviceFS, IdeviceFSArgs, out ExitCode, true);
 				if (ExitCode != 0)
 				{
 					LogError("Deployer error : failed to push files. See above log for more details");
@@ -1757,11 +1751,21 @@ public class IOSPlatform : Platform
 		return ExecWithPath;
 	}
 
+	private string GetLibimobileDeviceNetworkedArgument(string EntryArguments, string UDID)
+	{
+		DeviceInfo[] CachedDevices = GetDevices();
+
+		if (CachedDevices.Where(CachedDevice => CachedDevice.Name.Contains(UDID) && CachedDevice.Name.Contains("Network")).Count() > 0)
+		{
+			return "-n " + EntryArguments;
+		}
+		return EntryArguments;
+	}
+
 	private void DeployManifestContent(string BaseFolder, DeploymentContext SC, ProjectParams Params, ref string Files, ref string BundleIdentifier)
 	{
 		var DeviceFS = GetPathToLibiMobileDeviceTool("idevicefs");
 
-
 		string[] FileList = Files.Split('\n');
 		string AllCommandsToPush = "";
 		foreach (string Filename in FileList)
@@ -1780,9 +1784,12 @@ public class IOSPlatform : Platform
 			}
 		}
 		System.IO.File.WriteAllText(Directory.GetCurrentDirectory() + "\\CommandsToPush.txt", AllCommandsToPush);
-		LogWarning("BEFORE CRASH " + Directory.GetCurrentDirectory() + "\\CommandsToPush.txt");
 		int ExitCode = 0;
-		Utils.RunLocalProcessAndReturnStdOut(DeviceFS, "-b " + BundleIdentifier + " -x " + Directory.GetCurrentDirectory() + "\\CommandsToPush.txt", out ExitCode, true);
+
+		string IdeviceFSArgs = "-u " + Params.DeviceNames[0] + " -b " + BundleIdentifier + " -x " + Directory.GetCurrentDirectory() + "\\CommandsToPush.txt";
+		IdeviceFSArgs = GetLibimobileDeviceNetworkedArgument(IdeviceFSArgs, Params.DeviceNames[0]);
+
+		Utils.RunLocalProcessAndReturnStdOut(DeviceFS, IdeviceFSArgs, out ExitCode, true);
 		if (ExitCode != 0)
 		{
 			LogError("Deployer error : failed to push files. See above log for more details");
@@ -1847,16 +1854,18 @@ public class IOSPlatform : Platform
 		// Add a commandline for this deploy, if the config allows it.
 		string AdditionalCommandline = (Params.FileServer || Params.CookOnTheFly) ? "" : (" -additionalcommandline " + "\"" + Params.RunCommandline + "\"");
 
-
 		// deploy the .ipa
 		var DeviceInstaller = GetPathToLibiMobileDeviceTool("ideviceinstaller");
 
+		string LibimobileDeviceArguments =  "-u " + Params.DeviceNames[0] + " -i " + Path.GetFullPath(StagedIPA);
+		LibimobileDeviceArguments = GetLibimobileDeviceNetworkedArgument(LibimobileDeviceArguments, Params.DeviceNames[0]);
+
 		// check for it in the stage directory
 		string CurrentDir = Directory.GetCurrentDirectory();
 		Directory.SetCurrentDirectory(CombinePaths(CmdEnv.LocalRoot, "Engine/Binaries/DotNET/IOS/"));
 		if (!Params.IterativeDeploy || bCreatedIPA || bNeedsIPA)
 		{
-			RunAndLog(CmdEnv, DeviceInstaller, "-i " + Path.GetFullPath(StagedIPA));
+			RunAndLog(CmdEnv, DeviceInstaller, LibimobileDeviceArguments);
 		}
 
 		// deploy the assets
