Fix DeviceInfo substring for iOS and tvOS
Fix deployment error when there is a space in the project path

diff --git a/Engine/Source/Programs/AutomationTool/AutomationUtils/Platform.cs b/Engine/Source/Programs/AutomationTool/AutomationUtils/Platform.cs
--- a/Engine/Source/Programs/AutomationTool/AutomationUtils/Platform.cs
+++ b/Engine/Source/Programs/AutomationTool/AutomationUtils/Platform.cs
@@ -34,13 +34,14 @@ namespace AutomationTool
 			this.Platform = Platform;
 		}
 
-		public DeviceInfo(UnrealTargetPlatform Platform, string Name, string Id, string SoftwareVersion, string Type, bool bIsDefault, bool bCanConnect)
+		public DeviceInfo(UnrealTargetPlatform Platform, string Name, string Id, string SoftwareVersion, string Type, bool bIsDefault, bool bCanConnect, string SubType = "")
 		{
 			this.Platform = Platform;
 			this.Name = Name;
 			this.Id = Id;
 			this.SoftwareVersion = SoftwareVersion;
 			this.Type = Type;
+			this.SubType = SubType;
 			this.bIsDefault = bIsDefault;
 			this.bCanConnect = bCanConnect;
 		}
@@ -50,6 +51,7 @@ namespace AutomationTool
 		public string Id;
 		public string SoftwareVersion;
 		public string Type;
+		public string SubType;
 		public bool bIsDefault = false;
 		// is the device able to be connected to (this is more about able to flash SDK or run, not about matching SDK version)
 		// if false, any of the above fields are suspect, especually SoftwareVersion
diff --git a/Engine/Source/Programs/AutomationTool/IOS/IOSPlatform.Automation.cs b/Engine/Source/Programs/AutomationTool/IOS/IOSPlatform.Automation.cs
--- a/Engine/Source/Programs/AutomationTool/IOS/IOSPlatform.Automation.cs
+++ b/Engine/Source/Programs/AutomationTool/IOS/IOSPlatform.Automation.cs
@@ -525,7 +525,7 @@ public class IOSPlatform : Platform
 		string Output = Utils.RunLocalProcessAndReturnStdOut(IdeviceIdPath, "");
 		var ConnectedDevicesUDIDs = Output.Split(new string[] { Environment.NewLine }, StringSplitOptions.None);
 
-		DeviceInfo CurrentDevice = null;
+		DeviceInfo CurrentDevice = new DeviceInfo(TargetPlatformType);
 
 		foreach (string UnparsedUDID in ConnectedDevicesUDIDs)
 		{
@@ -535,14 +535,16 @@ public class IOSPlatform : Platform
 
 			if (UnparsedUDID.Contains("Network"))
 			{
+				CurrentDevice.SubType = "Network";
 				IdeviceInfoArgs = "-n " + IdeviceInfoArgs;
 			}
+			else
+			{
+				CurrentDevice.SubType = "USB";
+			}
 
 			string OutputInfo = Utils.RunLocalProcessAndReturnStdOut(IdeviceInfoPath, IdeviceInfoArgs);
 
-			CurrentDevice = new DeviceInfo(TargetPlatformType);
-			CurrentDevice.Name = UnparsedUDID;
-
 			foreach (string Line in OutputInfo.Split(Environment.NewLine.ToCharArray()))
 			{
 				// check we are returning the proper device for this class
@@ -554,6 +556,10 @@ public class IOSPlatform : Platform
 						Devices.Remove(CurrentDevice);
 					}
 				}
+				else if (Line.StartsWith("DeviceName: "))
+				{
+					CurrentDevice.Name = Line.Split(": ").Last();
+				}
 				else if (Line.StartsWith("UniqueDeviceID: "))
 				{
 					CurrentDevice.Id = Line.Split(": ").Last();
@@ -1755,7 +1761,7 @@ public class IOSPlatform : Platform
 	{
 		DeviceInfo[] CachedDevices = GetDevices();
 
-		if (CachedDevices.Where(CachedDevice => CachedDevice.Name.Contains(UDID) && CachedDevice.Name.Contains("Network")).Count() > 0)
+		if (CachedDevices.Where(CachedDevice => CachedDevice.Name == UDID && CachedDevice.SubType == "Network").Count() > 0)
 		{
 			return "-n " + EntryArguments;
 		}
@@ -1857,7 +1863,7 @@ public class IOSPlatform : Platform
 		// deploy the .ipa
 		var DeviceInstaller = GetPathToLibiMobileDeviceTool("ideviceinstaller");
 
-		string LibimobileDeviceArguments =  "-u " + Params.DeviceNames[0] + " -i " + Path.GetFullPath(StagedIPA);
+		string LibimobileDeviceArguments =  "-u " + Params.DeviceNames[0] + " -i " + "\"" +  Path.GetFullPath(StagedIPA) + "\"";
 		LibimobileDeviceArguments = GetLibimobileDeviceNetworkedArgument(LibimobileDeviceArguments, Params.DeviceNames[0]);
 
 		// check for it in the stage directory
