Rename UE4BuildUtils to UEBuildUtils.

#rnx

diff --git a/Engine/Source/Programs/AutomationTool/Scripts/CopyBuildToStagingDirectory.Automation.cs b/Engine/Source/Programs/AutomationTool/Scripts/CopyBuildToStagingDirectory.Automation.cs
--- a/Engine/Source/Programs/AutomationTool/Scripts/CopyBuildToStagingDirectory.Automation.cs
+++ b/Engine/Source/Programs/AutomationTool/Scripts/CopyBuildToStagingDirectory.Automation.cs
@@ -116,7 +116,7 @@ public partial class Project : CommandUtils
 			{
 				if (!CommandUtils.FileExists_NoExceptions(BuildPatchToolExe))
 				{
-					UE4BuildUtils.BuildBuildPatchTool(null, UnrealBuildTool.BuildHostPlatform.Current.Platform);
+					UEBuildUtils.BuildBuildPatchTool(null, UnrealBuildTool.BuildHostPlatform.Current.Platform);
 				}
 			}
 		}
diff --git a/Engine/Source/Programs/AutomationTool/Scripts/UE4BuildUtils.cs b/Engine/Source/Programs/AutomationTool/Scripts/UE4BuildUtils.cs
deleted file mode 100644
diff --git a/Engine/Source/Programs/AutomationTool/Scripts/UEBuildUtils.cs b/Engine/Source/Programs/AutomationTool/Scripts/UEBuildUtils.cs
new file mode 100644
--- /dev/null
+++ b/Engine/Source/Programs/AutomationTool/Scripts/UEBuildUtils.cs
@@ -0,0 +1,77 @@
+﻿// Copyright Epic Games, Inc. All Rights Reserved.
+
+using System;
+using System.Collections.Generic;
+using System.Linq;
+using System.Text;
+using AutomationTool;
+using UnrealBuildTool;
+
+/// <summary>
+/// Common UEBuild utilities
+/// </summary>
+public class UEBuildUtils : CommandUtils
+{
+	/// <summary>
+	/// Builds BuildPatchTool for the specified platform.
+	/// </summary>
+	/// <param name="Command"></param>
+	/// <param name="InPlatform"></param>
+	public static void BuildBuildPatchTool(BuildCommand Command, UnrealBuildTool.UnrealTargetPlatform InPlatform)
+	{
+		BuildProduct(Command, new UE4Build.BuildTarget()
+			{
+				UprojectPath = null,
+				TargetName = "BuildPatchTool",
+				Platform = InPlatform,
+				Config = UnrealBuildTool.UnrealTargetConfiguration.Shipping,
+			});
+	}
+
+	/// <summary>
+	/// Builds UnrealHeaderTool for the specified platform.
+	/// </summary>
+	/// <param name="Command"></param>
+	/// <param name="InPlatform"></param>
+	public static void BuildUnrealHeaderTool(BuildCommand Command, UnrealBuildTool.UnrealTargetPlatform InPlatform)
+	{
+		BuildProduct(Command, new UE4Build.BuildTarget()
+		{
+			UprojectPath = null,
+			TargetName = "UnrealHeaderTool",
+			Platform = InPlatform,
+			Config = UnrealBuildTool.UnrealTargetConfiguration.Development,
+		});
+	}
+
+	private static void BuildProduct(BuildCommand Command, UE4Build.BuildTarget Target)
+	{
+		if (Target == null)
+		{
+			throw new AutomationException("Target is required when calling UE4BuildUtils.BuildProduct");
+		}
+
+		LogInformation("Building {0}", Target.TargetName);
+
+		if (Command == null)
+		{
+			Command = new UE4BuildUtilDummyBuildCommand();
+		}
+
+		var UE4Build = new UE4Build(Command);
+
+		var Agenda = new UE4Build.BuildAgenda();
+		Agenda.Targets.Add(Target);
+
+		UE4Build.Build(Agenda, InDeleteBuildProducts: true, InUpdateVersionFiles: true);
+		UE4Build.CheckBuildProducts(UE4Build.BuildProductFiles);
+	}
+
+	class UE4BuildUtilDummyBuildCommand : BuildCommand
+	{
+		public override void ExecuteBuild()
+		{
+			// noop
+		}
+	}
+}
