Ignore abandoned mutex exception

diff --git a/Engine/Source/Programs/AutomationTool/AutomationUtils/Utils.cs b/Engine/Source/Programs/AutomationTool/AutomationUtils/Utils.cs
--- a/Engine/Source/Programs/AutomationTool/AutomationUtils/Utils.cs
+++ b/Engine/Source/Programs/AutomationTool/AutomationUtils/Utils.cs
@@ -737,7 +737,7 @@ namespace AutomationTool
 					{
 						CommandUtils.LogWarning("Another instance of UAT at '{0}' is running, and the -WaitForUATMutex parameter has been used. Waiting for other UAT to finish...", EntryAssemblyLocation);
 						int Seconds = 0;
-						while (SingleInstanceMutex.WaitOne(15 * 1000) == false)
+						while (WaitMutexNoExceptions(SingleInstanceMutex, 15 * 1000) == false)
 						{
 							Seconds += 15;
 							CommandUtils.LogInformation("Still waiting for Mutex. {0} seconds has passed...", Seconds);
@@ -760,6 +760,18 @@ namespace AutomationTool
 			}
 		}
 
+		static bool WaitMutexNoExceptions(Mutex Mutex, int TimeoutMs)
+		{
+			try
+			{
+				return Mutex.WaitOne(15 * 1000);
+			}
+			catch (AbandonedMutexException)
+			{
+				return true;
+			}
+		}
+
 		public static string GetUniqueMutexForPath(string Name, string UniquePath)
 		{
 			// generate a md5 hash of the path, as GetHashCode is not guaranteed to generate a stable hash
