Gauntlet - Fix support for extra arguments

#rnx

diff --git a/Engine/Source/Programs/AutomationTool/Gauntlet/Unreal/Base/Gauntlet.UnrealTestConfiguration.cs b/Engine/Source/Programs/AutomationTool/Gauntlet/Unreal/Base/Gauntlet.UnrealTestConfiguration.cs
--- a/Engine/Source/Programs/AutomationTool/Gauntlet/Unreal/Base/Gauntlet.UnrealTestConfiguration.cs
+++ b/Engine/Source/Programs/AutomationTool/Gauntlet/Unreal/Base/Gauntlet.UnrealTestConfiguration.cs
@@ -37,7 +37,7 @@ namespace Gauntlet
 		public void AddRawCommandline(string InRawCommandline, bool bOverrideExistingValues = true)
 		{
 			// turn Name(p1,etc) into a collection of Name|(p1,etc) groups
-			MatchCollection Matches = Regex.Matches(InRawCommandline, "-(?<option>\\-?[\\w\\d.:\\[\\]\\/\\\\]+)(=(?<value>(\"([^\"]*)\")|(\\S+)))?");
+			MatchCollection Matches = Regex.Matches(InRawCommandline, "-(?<option>\\-?[\\w\\d.:!\\[\\]\\/\\\\]+)(=(?<value>(\"([^\"]*)\")|(\\S+)))?");
 
 			foreach (Match M in Matches)
 			{
diff --git a/Engine/Source/Programs/AutomationTool/Gauntlet/Unreal/BuildSource/Gauntlet.UnrealBuildSource.cs b/Engine/Source/Programs/AutomationTool/Gauntlet/Unreal/BuildSource/Gauntlet.UnrealBuildSource.cs
--- a/Engine/Source/Programs/AutomationTool/Gauntlet/Unreal/BuildSource/Gauntlet.UnrealBuildSource.cs
+++ b/Engine/Source/Programs/AutomationTool/Gauntlet/Unreal/BuildSource/Gauntlet.UnrealBuildSource.cs
@@ -466,7 +466,7 @@ namespace Gauntlet
 			// Break down Commandline into individual tokens 
 			Dictionary<string, string> CommandlineTokens = new Dictionary<string, string>(StringComparer.OrdinalIgnoreCase);
 			// turn Name(p1,etc) into a collection of Name|(p1,etc) groups
-			MatchCollection Matches = Regex.Matches(InCommandLine, "(?<option>\\-?[\\w\\d.:\\[\\]\\/\\\\\\?]+)(=(?<value>(\"([^\"]*)\")|(\\S+)))?");
+			MatchCollection Matches = Regex.Matches(InCommandLine, "(?<option>\\-?[\\w\\d.:!\\[\\]\\/\\\\\\?]+)(=(?<value>(\"([^\"]*)\")|(\\S+)))?");
 
 			foreach (Match M in Matches)
 			{
diff --git a/Engine/Source/Programs/AutomationTool/Gauntlet/Unreal/RunUnreal.cs b/Engine/Source/Programs/AutomationTool/Gauntlet/Unreal/RunUnreal.cs
--- a/Engine/Source/Programs/AutomationTool/Gauntlet/Unreal/RunUnreal.cs
+++ b/Engine/Source/Programs/AutomationTool/Gauntlet/Unreal/RunUnreal.cs
@@ -187,12 +187,11 @@ namespace Gauntlet
 					}
 
 					// look for-args= and then -clientargs= and -editorargs etc
-					Role.ExtraArgs = Globals.Params.ParseValue("Args", "");
-					string ExtraRoleArgs = Globals.Params.ParseValue(Type.ToString() + "Args", "");
-
-					if (!string.IsNullOrEmpty(ExtraRoleArgs))
+					Role.ExtraArgs = string.Join(' ', Globals.Params.ParseValues("Args", true));
+					IEnumerable<string> ExtraRoleArgs = Globals.Params.ParseValues(Type.ToString() + "Args", true);
+					if (ExtraRoleArgs.Count() > 0)
 					{
-						Role.ExtraArgs += ExtraRoleArgs;
+						Role.ExtraArgs += " " + string.Join(' ', ExtraRoleArgs);
 					}
 
 					// look for -clientexeccmds=, -editorexeccmds= etc, these are separate from clientargs for sanity
