Fix Networked iOS / tvOS deployment (DeviceINfo UDID string check)

diff --git a/Engine/Source/Programs/AutomationTool/IOS/IOSPlatform.Automation.cs b/Engine/Source/Programs/AutomationTool/IOS/IOSPlatform.Automation.cs
--- a/Engine/Source/Programs/AutomationTool/IOS/IOSPlatform.Automation.cs
+++ b/Engine/Source/Programs/AutomationTool/IOS/IOSPlatform.Automation.cs
@@ -1761,7 +1761,7 @@ public class IOSPlatform : Platform
 	{
 		DeviceInfo[] CachedDevices = GetDevices();
 
-		if (CachedDevices.Where(CachedDevice => CachedDevice.Name == UDID && CachedDevice.SubType == "Network").Count() > 0)
+		if (CachedDevices.Where(CachedDevice => CachedDevice.Id == UDID && CachedDevice.SubType == "Network").Count() > 0)
 		{
 			return "-n " + EntryArguments;
 		}
