Attached UE logs to report and push them on Horde

#rnx

diff --git a/Engine/Source/Programs/AutomationTool/Gauntlet/Unreal/Base/Gauntlet.UnrealTestNode.cs b/Engine/Source/Programs/AutomationTool/Gauntlet/Unreal/Base/Gauntlet.UnrealTestNode.cs
--- a/Engine/Source/Programs/AutomationTool/Gauntlet/Unreal/Base/Gauntlet.UnrealTestNode.cs
+++ b/Engine/Source/Programs/AutomationTool/Gauntlet/Unreal/Base/Gauntlet.UnrealTestNode.cs
@@ -1097,6 +1097,15 @@ namespace Gauntlet
 				string HordeArtifactPath = GetConfiguration().HordeArtifactPath;
 				HordeReport.UnrealEngineTestPassResults HordeTestPassResults = HordeReport.UnrealEngineTestPassResults.FromUnrealAutomatedTests(JsonTestPassResults, UnrealAutomatedTestReportPath, ReportURL);
 				HordeTestPassResults.CopyTestResultsArtifacts(HordeArtifactPath);
+				// Attached test Artifacts
+				if (SessionArtifacts != null)
+				{
+					foreach (UnrealRoleArtifacts Artifact in SessionArtifacts)
+					{
+						string LogName = Path.GetFullPath(Artifact.LogPath).Replace(Path.GetFullPath(Path.Combine(ArtifactPath, "..")), "").TrimStart(Path.DirectorySeparatorChar);
+						HordeTestPassResults.AttachArtifact(Artifact.LogPath, LogName);
+					}
+				}
 				return HordeTestPassResults;
 			}
 			else
