- Categorized Turnkey commands
- Added ability to have input list options be headers (prefix with ;)

diff --git a/Engine/Source/Programs/AutomationTool/Turnkey/Commands/Control.cs b/Engine/Source/Programs/AutomationTool/Turnkey/Commands/Control.cs
--- a/Engine/Source/Programs/AutomationTool/Turnkey/Commands/Control.cs
+++ b/Engine/Source/Programs/AutomationTool/Turnkey/Commands/Control.cs
@@ -11,6 +11,8 @@ namespace Turnkey.Commands
 {
 	class Control : TurnkeyCommand
 	{
+		protected override CommandGroup Group => CommandGroup.Misc;
+
 		private delegate void PerformOperation(ITargetDevice Device);
 
 		protected override void Execute(string[] CommandOptions)
diff --git a/Engine/Source/Programs/AutomationTool/Turnkey/Commands/ExecuteBuild.cs b/Engine/Source/Programs/AutomationTool/Turnkey/Commands/ExecuteBuild.cs
--- a/Engine/Source/Programs/AutomationTool/Turnkey/Commands/ExecuteBuild.cs
+++ b/Engine/Source/Programs/AutomationTool/Turnkey/Commands/ExecuteBuild.cs
@@ -13,6 +13,8 @@ namespace Turnkey.Commands
 {
 	class ExecuteBuild : TurnkeyCommand
 	{
+		protected override CommandGroup Group => CommandGroup.Builds;
+
 		// could move this to TurnkeyUtils!
 		static string GetStructEntry(string Input, string Property, bool bIsArrayProperty)
 		{
diff --git a/Engine/Source/Programs/AutomationTool/Turnkey/Commands/Help.cs b/Engine/Source/Programs/AutomationTool/Turnkey/Commands/Help.cs
--- a/Engine/Source/Programs/AutomationTool/Turnkey/Commands/Help.cs
+++ b/Engine/Source/Programs/AutomationTool/Turnkey/Commands/Help.cs
@@ -11,6 +11,8 @@ namespace Turnkey.Commands
 {
 	class Help : TurnkeyCommand
 	{
+		protected override CommandGroup Group => CommandGroup.Informational;
+
 		protected override void Execute(string[] CommandOptions)
 		{
 			bool bShowUsingHelp = TurnkeyUtils.ParseParam("Using", CommandOptions);
diff --git a/Engine/Source/Programs/AutomationTool/Turnkey/Commands/InstallBuild.cs b/Engine/Source/Programs/AutomationTool/Turnkey/Commands/InstallBuild.cs
--- a/Engine/Source/Programs/AutomationTool/Turnkey/Commands/InstallBuild.cs
+++ b/Engine/Source/Programs/AutomationTool/Turnkey/Commands/InstallBuild.cs
@@ -11,6 +11,8 @@ namespace Turnkey.Commands
 {
 	class InstallBuild : TurnkeyCommand
 	{
+		protected override CommandGroup Group => CommandGroup.Builds;
+
 		protected override void Execute(string[] CommandOptions)
 		{
 			// get distinct project names
diff --git a/Engine/Source/Programs/AutomationTool/Turnkey/Commands/InstallSdk.cs b/Engine/Source/Programs/AutomationTool/Turnkey/Commands/InstallSdk.cs
--- a/Engine/Source/Programs/AutomationTool/Turnkey/Commands/InstallSdk.cs
+++ b/Engine/Source/Programs/AutomationTool/Turnkey/Commands/InstallSdk.cs
@@ -12,7 +12,8 @@ namespace Turnkey.Commands
 {
 	class InstallSdk : TurnkeyCommand
 	{
-	
+		protected override CommandGroup Group => CommandGroup.Sdk;
+
 		protected override Dictionary<string, string[]> GetExtendedCommandsWithOptions()
 		{
 			return new Dictionary<string, string[]>()
diff --git a/Engine/Source/Programs/AutomationTool/Turnkey/Commands/ListPlatforms.cs b/Engine/Source/Programs/AutomationTool/Turnkey/Commands/ListPlatforms.cs
--- a/Engine/Source/Programs/AutomationTool/Turnkey/Commands/ListPlatforms.cs
+++ b/Engine/Source/Programs/AutomationTool/Turnkey/Commands/ListPlatforms.cs
@@ -9,6 +9,8 @@ namespace Turnkey.Commands
 {
 	class ListPlatforms : TurnkeyCommand
 	{
+		protected override CommandGroup Group => CommandGroup.Informational;
+
 		protected override void Execute(string[] CommandOptions)
 		{
 			TurnkeyUtils.Log("");
diff --git a/Engine/Source/Programs/AutomationTool/Turnkey/Commands/ListSdks.cs b/Engine/Source/Programs/AutomationTool/Turnkey/Commands/ListSdks.cs
--- a/Engine/Source/Programs/AutomationTool/Turnkey/Commands/ListSdks.cs
+++ b/Engine/Source/Programs/AutomationTool/Turnkey/Commands/ListSdks.cs
@@ -10,6 +10,8 @@ namespace Turnkey.Commands
 {
 	class ListSdks : TurnkeyCommand
 	{
+		protected override CommandGroup Group => CommandGroup.Informational;
+
 		protected override void Execute(string[] CommandOptions)
 		{
 			TurnkeyUtils.Log("");
diff --git a/Engine/Source/Programs/AutomationTool/Turnkey/Commands/ManageSettings.cs b/Engine/Source/Programs/AutomationTool/Turnkey/Commands/ManageSettings.cs
--- a/Engine/Source/Programs/AutomationTool/Turnkey/Commands/ManageSettings.cs
+++ b/Engine/Source/Programs/AutomationTool/Turnkey/Commands/ManageSettings.cs
@@ -10,6 +10,8 @@ namespace Turnkey.Commands
 {
 	class ManageSettings : TurnkeyCommand
 	{
+		protected override CommandGroup Group => CommandGroup.Settings;
+
 		protected override Dictionary<string, string[]> GetExtendedCommandsWithOptions()
 		{
 			return new Dictionary<string, string[]>()
diff --git a/Engine/Source/Programs/AutomationTool/Turnkey/Commands/Test.cs b/Engine/Source/Programs/AutomationTool/Turnkey/Commands/Test.cs
--- a/Engine/Source/Programs/AutomationTool/Turnkey/Commands/Test.cs
+++ b/Engine/Source/Programs/AutomationTool/Turnkey/Commands/Test.cs
@@ -10,6 +10,8 @@ namespace Turnkey.Commands
 {
 	class Test : TurnkeyCommand
 	{
+		protected override CommandGroup Group => CommandGroup.Misc;
+
 		protected override void Execute(string[] CommandOptions)
 		{
 			TurnkeyUtils.Log("Testing:");
diff --git a/Engine/Source/Programs/AutomationTool/Turnkey/Commands/VerifySdk.cs b/Engine/Source/Programs/AutomationTool/Turnkey/Commands/VerifySdk.cs
--- a/Engine/Source/Programs/AutomationTool/Turnkey/Commands/VerifySdk.cs
+++ b/Engine/Source/Programs/AutomationTool/Turnkey/Commands/VerifySdk.cs
@@ -13,6 +13,8 @@ namespace Turnkey.Commands
 {
 	class VerifySdk : TurnkeyCommand
 	{
+		protected override CommandGroup Group => CommandGroup.Sdk;
+
 		protected override void Execute(string[] CommandOptions)
 		{
 			bool bUnattended = TurnkeyUtils.ParseParam("Unattended", CommandOptions);
diff --git a/Engine/Source/Programs/AutomationTool/Turnkey/Providers/ConsoleIOProvider.cs b/Engine/Source/Programs/AutomationTool/Turnkey/Providers/ConsoleIOProvider.cs
--- a/Engine/Source/Programs/AutomationTool/Turnkey/Providers/ConsoleIOProvider.cs
+++ b/Engine/Source/Programs/AutomationTool/Turnkey/Providers/ConsoleIOProvider.cs
@@ -30,6 +30,8 @@ namespace Turnkey
 
 		public override string ReadInput(string Prompt, string Default, bool bAppendNewLine)
 		{
+			TurnkeyUtils.Log("");
+
 			Log(Prompt, bAppendNewLine);
 			if (!string.IsNullOrEmpty(Default))
 			{
@@ -46,6 +48,8 @@ namespace Turnkey
 
 		public override int ReadInputInt(string Prompt, List<string> Options, bool bIsCancellable, int DefaultValue, bool bAppendNewLine)
 		{
+			TurnkeyUtils.Log("");
+
 			StringBuilder FullPromptBuilder  = new StringBuilder();
 
 			// start with given prompt
@@ -65,7 +69,15 @@ namespace Turnkey
 			// now add the options given
 			foreach (string Option in Options)
 			{
-				FullPromptBuilder.AppendLine(" {0}{1}{2} {3}", DefaultValue == Index ? "[" : " ", Index, DefaultValue == Index++ ? "]" : ")", Option);
+				// don't make an index for commented lines
+				if (Option.StartsWith(";"))
+				{
+					FullPromptBuilder.AppendLine(" {0}", Option.Substring(1));
+				}
+				else
+				{
+					FullPromptBuilder.AppendLine(" {0}{1}{2} {3}", DefaultValue == Index ? "[" : " ", Index, DefaultValue == Index++ ? "]" : ")", Option);
+				}
 			}
 
 			string FullPrompt = FullPromptBuilder.ToString();
diff --git a/Engine/Source/Programs/AutomationTool/Turnkey/Providers/HybridIOProvider.cs b/Engine/Source/Programs/AutomationTool/Turnkey/Providers/HybridIOProvider.cs
--- a/Engine/Source/Programs/AutomationTool/Turnkey/Providers/HybridIOProvider.cs
+++ b/Engine/Source/Programs/AutomationTool/Turnkey/Providers/HybridIOProvider.cs
@@ -191,7 +191,14 @@ namespace Turnkey
 			int Index = 1;
 			foreach (string Option in Options)
 			{
-				FullPromptBuilder.AppendLine(" ({0}) {1}", Index++, Option);
+				if (Option.StartsWith(";"))
+				{
+					FullPromptBuilder.AppendLine(" {0}", Option.Substring(1));
+				}
+				else
+				{
+					FullPromptBuilder.AppendLine(" ({0}) {1}", Index++, Option);
+				}
 			}
 
 			string FullPrompt = FullPromptBuilder.ToString();
diff --git a/Engine/Source/Programs/AutomationTool/Turnkey/TurnkeyCommand.cs b/Engine/Source/Programs/AutomationTool/Turnkey/TurnkeyCommand.cs
--- a/Engine/Source/Programs/AutomationTool/Turnkey/TurnkeyCommand.cs
+++ b/Engine/Source/Programs/AutomationTool/Turnkey/TurnkeyCommand.cs
@@ -31,7 +31,16 @@ namespace Turnkey
 
 	abstract class TurnkeyCommand
 	{
+		protected enum CommandGroup
+		{
+			Sdk,
+			Builds,
+			Informational,
+			Settings,
+			Misc,
+		}
 
+		protected abstract CommandGroup Group { get; }
 		protected abstract void Execute(string[] CommandOptions);
 		protected virtual Dictionary<string, string[]> GetExtendedCommandsWithOptions()
 		{
@@ -46,29 +55,51 @@ namespace Turnkey
 		// cached commands
 		private static Dictionary<string, TurnkeyCommand> CachedCommandsByName = new Dictionary<string, TurnkeyCommand>(StringComparer.OrdinalIgnoreCase);
 		private static List<CommandInfo> CachedCommandsByIndex = new List<CommandInfo>();
+		private static Dictionary<CommandGroup, List<CommandInfo>> CachedCommandsByGroup = new Dictionary<CommandGroup, List<CommandInfo>>();
 
 		static TurnkeyCommand()
 		{
+			List<TurnkeyCommand> BaseCommands = new List<TurnkeyCommand>();
 			// look for all subclasses, and cache by their ProviderToken
 			foreach (Type AssemType in Assembly.GetExecutingAssembly().GetTypes())
 			{
 				if (typeof(TurnkeyCommand).IsAssignableFrom(AssemType) && AssemType != typeof(TurnkeyCommand))
 				{
 					TurnkeyCommand Provider = (TurnkeyCommand)Activator.CreateInstance(AssemType);
-					CachedCommandsByName[AssemType.Name] = Provider;
+					BaseCommands.Add(Provider);
+				}
+			}
+
+			BaseCommands.SortBy(x => x.GetType().Name);
 
-					// always add base option
-					CachedCommandsByIndex.Add(new CommandInfo(Provider));
-					Dictionary<string, string[]> ExtendedCommands = Provider.GetExtendedCommandsWithOptions();
-					if (ExtendedCommands != null)
+			foreach (CommandGroup Group in Enum.GetValues(typeof(CommandGroup)))
+			{
+				CachedCommandsByGroup[Group] = new List<CommandInfo>();
+			}
+			foreach (TurnkeyCommand Command in BaseCommands)
+			{
+				CachedCommandsByName[Command.GetType().Name] = Command;
+
+				// always add base option
+				CachedCommandsByGroup[Command.Group].Add(new CommandInfo(Command));
+				Dictionary<string, string[]> ExtendedCommands = Command.GetExtendedCommandsWithOptions();
+				if (ExtendedCommands != null)
+				{
+					foreach (var Pair in ExtendedCommands)
 					{
-						foreach (var Pair in ExtendedCommands)
-						{
-							CachedCommandsByIndex.Add(new CommandInfo(Provider, Pair.Key, Pair.Value));
-						}
+						CachedCommandsByGroup[Command.Group].Add(new CommandInfo(Command, Pair.Key, Pair.Value));
 					}
 				}
 			}
+
+			// now flatten the lists to match the choices in ChooseCommand
+			foreach (var Pair in CachedCommandsByGroup)
+			{
+				foreach (CommandInfo Command in Pair.Value)
+				{
+					CachedCommandsByIndex.Add(Command);
+				}
+			}
 		}
 
 		static bool bExitAfterCommand = false;
@@ -126,7 +157,23 @@ namespace Turnkey
 
 		private static int PromptForCommand()
 		{
-			return TurnkeyUtils.ReadInputInt("Enter command", CachedCommandsByIndex.ConvertAll(x => x.DisplayString), true);
+			// build up categorized list
+			List<string> Options = new List<string>();
+			foreach (var Pair in CachedCommandsByGroup)
+			{
+				CommandGroup Group = Pair.Key;
+				List<CommandInfo> Commands = Pair.Value;
+
+				if (Commands.Count > 0)
+				{
+					Options.Add($";{Group}:");
+					foreach (CommandInfo Command in Commands)
+					{
+						Options.Add($"  {Command.DisplayString}");
+					}
+				}
+			}
+			return TurnkeyUtils.ReadInputInt("Enter command", Options, true);
 		}
 	}
 }
