Fix UAT no longer respecting the -nocompileeditor argument.

diff --git a/Engine/Source/Programs/AutomationTool/AutomationUtils/ProjectParams.cs b/Engine/Source/Programs/AutomationTool/AutomationUtils/ProjectParams.cs
--- a/Engine/Source/Programs/AutomationTool/AutomationUtils/ProjectParams.cs
+++ b/Engine/Source/Programs/AutomationTool/AutomationUtils/ProjectParams.cs
@@ -587,7 +587,7 @@ namespace AutomationTool
 			}
 
 			this.SkipBuildClient = GetParamValueIfNotSpecified(Command, SkipBuildClient, this.SkipBuildClient, "skipbuildclient");
-			this.SkipBuildEditor = GetParamValueIfNotSpecified(Command, SkipBuildEditor, this.SkipBuildEditor, "skipbuildeditor");
+			this.SkipBuildEditor = GetParamValueIfNotSpecified(Command, SkipBuildEditor, this.SkipBuildEditor, "skipbuildeditor", "nocompileeditor");
 			this.Run = GetParamValueIfNotSpecified(Command, Run, this.Run, "run");
 			this.Cook = GetParamValueIfNotSpecified(Command, Cook, this.Cook, "cook");
 			this.CreateReleaseVersionBasePath = ParseParamValueIfNotSpecified(Command, CreateReleaseVersionBasePath, "createreleaseversionroot", String.Empty);
