Gauntlet - Formalize telemetry support in Gauntlet
- add simple data structure to hold the data
- introduce database config and driver
- support sending telemetry data from UnrealTestNode

diff --git a/Engine/Source/Programs/AutomationTool/Gauntlet/Framework/Base/Gauntlet.Database.cs b/Engine/Source/Programs/AutomationTool/Gauntlet/Framework/Base/Gauntlet.Database.cs
new file mode 100644
--- /dev/null
+++ b/Engine/Source/Programs/AutomationTool/Gauntlet/Framework/Base/Gauntlet.Database.cs
@@ -0,0 +1,207 @@
+// Copyright Epic Games, Inc. All Rights Reserved.
+
+using System;
+using System.IO;
+using System.Linq;
+using System.Collections.Generic;
+using MySql.Data.MySqlClient;
+using AutomationTool;
+
+namespace Gauntlet
+{
+	/// <summary>
+	/// Interface to drive Database submission
+	/// </summary>
+	public interface IDatabaseDriver<Data> where Data : class
+	{
+		/// <summary>
+		/// Submit collection of object to target Database, use TestContext to complete data modeling.
+		/// </summary>
+		/// <param name="DataItems"></param>
+		/// <param name="Context"></param>
+		/// <returns></returns>
+		bool SubmitDataItems(IEnumerable<Data> DataItems, ITestContext Context);
+	}
+
+	/// <summary>
+	/// Interface for Database Configuration
+	/// </summary>
+	public interface IDatabaseConfig<Data> where Data : class 
+	{
+		void LoadConfig(string ConfigFilePath);
+		IDatabaseDriver<Data> GetDriver();
+	}
+
+	/// <summary>
+	/// Interface for data type MySQL Configuration
+	/// </summary>
+	public interface IMySQLConfig<Data> where Data : class
+	{
+		/// <summary>
+		/// Get Target Table name based on data type
+		/// </summary>
+		string GetTableName();
+
+		/// <summary>
+		/// Get Target Table columns based on data type
+		/// </summary>
+		IEnumerable<string> GetTableColumns();
+
+		/// Format the data for target table based on data type
+		/// </summary>
+		/// <returns></returns>
+		IEnumerable<string> FormatDataForTable(Data InData, ITestContext InContext);
+	}
+
+	public class MySQLDriver<Data> : IDatabaseDriver<Data> where Data : class 
+	{
+		protected MySQLConfig<Data> Config;
+
+		public MySQLDriver(MySQLConfig<Data> InConfig)
+		{
+			Config = InConfig;
+			if (string.IsNullOrEmpty(Config.ConfigString))
+			{
+				throw new AutomationException(string.Format("Database Driver '{0}' not configured.", this.GetType().FullName));
+			}
+			if (string.IsNullOrEmpty(Config.DatabaseName))
+			{
+				throw new AutomationException(string.Format("Database Driver '{0}' not configured properly, missing Database name.", this.GetType().FullName));
+			}
+		}
+
+		public override string ToString()
+		{
+			return Config.GetConfigValue("Server");
+		}
+
+		public object Insert(string Table, IEnumerable<string> Columns, IEnumerable<IEnumerable<string>> Rows)
+		{
+			string SqlQuery = string.Format(
+				"INSERT INTO `{0}`.{1} ({2}) VALUES {3}",
+				Config.DatabaseName, Table, string.Join(", ", Columns), string.Join(", ", Rows.Select(R => string.Format("({0})", string.Join(", ", R.Select(V => string.Format("'{0}'", V))))))
+			);
+			return MySqlHelper.ExecuteScalar(Config.ConfigString, SqlQuery);
+		}
+
+		public bool SubmitDataItems(IEnumerable<Data> DataRows, ITestContext TestContext)
+		{
+			if (Config is IMySQLConfig<Data> DataConfig)
+			{
+				object Value = Insert(
+					DataConfig.GetTableName(),
+					DataConfig.GetTableColumns(),
+					DataRows.Select(D => DataConfig.FormatDataForTable(D, TestContext))
+				);
+				return Value != null;
+			}
+			else
+			{
+				Log.Error("MySQL configuration '{0}' does not known how to handle {1}.", Config.GetType().FullName, typeof(Data).FullName);
+			}
+
+			return false;
+		}
+	}
+
+	public abstract class MySQLConfig<Data> : IDatabaseConfig<Data>, IMySQLConfig<Data> where Data : class
+	{
+		public string ConfigString { get; protected set; }
+		public string DatabaseName { get; protected set; }
+		protected Dictionary<string, string> KeyValuePairs = null;
+
+		public virtual void LoadConfig(string ConfigFilePath)
+		{
+			ConfigString = string.Empty;
+			DatabaseName = string.Empty;
+			KeyValuePairs = null;
+
+			if (File.Exists(ConfigFilePath))
+			{
+				using (StreamReader ConnectionReader = new StreamReader(ConfigFilePath))
+				{
+					ConfigString = ConnectionReader.ReadLine();
+				}
+
+				if (string.IsNullOrEmpty(ConfigString))
+				{
+					Log.Warning("Properly found config file, but couldn't read a valid connection string.");
+					return;
+				}
+				else
+				{
+					Log.Info("Found MySQL connection string from config file.");
+				}
+			}
+			else
+			{
+				Log.Error("Could not find connection string config file at '{0}'.", ConfigFilePath);
+				return;
+			}
+
+			DatabaseName = GetConfigValue("database");
+			if (string.IsNullOrEmpty(DatabaseName))
+			{
+				Log.Warning("Missing MySQL Database name in config file '{0}'.", ConfigFilePath);
+				return;
+			}
+		}
+
+		public string GetConfigValue(string Key)
+		{
+			if (string.IsNullOrEmpty(ConfigString))
+			{
+				return null;
+			}
+			if (KeyValuePairs == null)
+			{
+				KeyValuePairs = ConfigString.Split(';').Where(KeyValue => KeyValue.Contains('='))
+														.Select(KeyValue => KeyValue.Split('=', 2))
+														.ToDictionary(
+															KeyValue => KeyValue[0].Trim(),
+															KeyValue => KeyValue[1].Trim(),
+															StringComparer.InvariantCultureIgnoreCase
+														);
+			}
+			string FoundValue;
+			if (KeyValuePairs.TryGetValue(Key, out FoundValue))
+			{
+				return FoundValue;
+			}
+			return string.Empty;
+		}
+
+		public IDatabaseDriver<Data> GetDriver()
+		{
+			return new MySQLDriver<Data>(this);
+		}
+
+		/// <summary>
+		/// Get Target Table name based on data type
+		/// </summary>
+		public abstract string GetTableName();
+		/// <summary>
+		/// Get Target Table columns based on data type
+		/// </summary>
+		public abstract IEnumerable<string> GetTableColumns();
+		/// <summary>
+		/// Format the data for target table based on data type
+		/// </summary>
+		/// <returns></returns>
+		public abstract IEnumerable<string> FormatDataForTable(Data InData, ITestContext InContext);
+
+	}
+
+	public class DatabaseConfigManager<Data> where Data : class
+	{
+		protected static IEnumerable<IDatabaseConfig<Data>> Configs;
+		static DatabaseConfigManager()
+		{
+			Configs = Gauntlet.Utils.InterfaceHelpers.FindImplementations<IDatabaseConfig<Data>>();
+		}
+		public static IDatabaseConfig<Data> GetConfigByName(string Name)
+		{
+			return Configs.Where(C => string.Equals(C.GetType().Name, Name, StringComparison.OrdinalIgnoreCase)).FirstOrDefault();
+		}
+	}
+}
diff --git a/Engine/Source/Programs/AutomationTool/Gauntlet/Framework/Base/Gauntlet.TestReport.cs b/Engine/Source/Programs/AutomationTool/Gauntlet/Framework/Base/Gauntlet.TestReport.cs
--- a/Engine/Source/Programs/AutomationTool/Gauntlet/Framework/Base/Gauntlet.TestReport.cs
+++ b/Engine/Source/Programs/AutomationTool/Gauntlet/Framework/Base/Gauntlet.TestReport.cs
@@ -2,6 +2,7 @@
 
 using System;
 using System.Reflection;
+using System.Collections.Generic;
 
 
 namespace Gauntlet
@@ -41,15 +42,57 @@ namespace Gauntlet
 		bool AttachArtifact(string ArtifactPath, string Name = null);
 	}
 
+	/// <summary>
+	/// Hold telemetry data
+	/// </summary>
+	public class TelemetryData
+	{
+		public string TestName { get; private set; }
+		public string DataPoint { get; private set; }
+		public double Measurement { get; private set; }
+		public TelemetryData(string InTestName, string InDataPoint, double InMeasurement)
+		{
+			TestName = InTestName;
+			DataPoint = InDataPoint;
+			Measurement = InMeasurement;
+		}
+	}
+
+	/// <summary>
+	/// Interface for telemetry report
+	/// </summary>
+	public interface ITelemetryReport
+	{ 
+		/// <summary>
+		/// Attach Telemetry data to the ITestReport
+		/// </summary>
+		/// <param name="TestName"></param>
+		/// <param name="DataPoint"></param>
+		/// <param name="Measurement"></param>
+		/// <returns></returns>
+		void AddTelemetry(string TestName, string DataPoint, double Measurement);
+
+		/// <summary>
+		/// Return the telemetry data accumulated
+		/// </summary>
+		/// <returns></returns>
+		IEnumerable<TelemetryData> GetAllTelemetryData();
+	}
+
 	/// <summary>
 	/// Test Report Base class
 	/// </summary>
-	public class BaseTestReport : ITestReport
+	public abstract class BaseTestReport : ITestReport, ITelemetryReport
 	{
 		/// <summary>
 		/// Return report type
 		/// </summary>
 		public virtual string Type { get { return "Base Test Report"; } }
+	
+		/// <summary>
+		/// Hold the list of telemetry data accumulated
+		/// </summary>
+		protected List<TelemetryData> TelemetryDataList;
 
 		/// <summary>
 		/// Set a property of the test report type
@@ -70,10 +113,7 @@ namespace Gauntlet
 		/// <param name="Message"></param>
 		/// <param name="Context"></param>
 		/// <returns></returns>
-		public virtual void AddEvent(string Type, string Message, object Context = null)
-		{
-			throw new System.NotImplementedException("AddEvent not implemented");
-		}
+		public abstract void AddEvent(string Type, string Message, object Context = null);
 
 		public void AddError(string Message, object Context = null)
 		{
@@ -94,9 +134,32 @@ namespace Gauntlet
 		/// <param name="ArtifactPath"></param>
 		/// <param name="Name"></param>
 		/// <returns>return true if the file was successfully attached</returns>
-		public virtual bool AttachArtifact(string ArtifactPath, string Name = null)
+		public abstract bool AttachArtifact(string ArtifactPath, string Name = null);
+
+		/// <summary>
+		/// Attach Telemetry data to the ITestReport
+		/// </summary>
+		/// <param name="TestName"></param>
+		/// <param name="DataPoint"></param>
+		/// <param name="Measurement"></param>
+		/// <returns></returns>
+		public virtual void AddTelemetry(string TestName, string DataPoint, double Measurement)
+		{
+			if (TelemetryDataList is null)
+			{
+				TelemetryDataList = new List<TelemetryData>();
+			}
+
+			TelemetryDataList.Add(new TelemetryData(TestName, DataPoint, Measurement));
+		}
+
+		/// <summary>
+		/// Return all the telemetry data stored
+		/// </summary>
+		/// <returns></returns>
+		public virtual IEnumerable<TelemetryData> GetAllTelemetryData()
 		{
-			throw new System.NotImplementedException("AttachArtifact not implemented");
+			return TelemetryDataList;
 		}
 	}
 }
\ No newline at end of file
diff --git a/Engine/Source/Programs/AutomationTool/Gauntlet/Framework/Utils/Gauntlet.HordeReport.cs b/Engine/Source/Programs/AutomationTool/Gauntlet/Framework/Utils/Gauntlet.HordeReport.cs
--- a/Engine/Source/Programs/AutomationTool/Gauntlet/Framework/Utils/Gauntlet.HordeReport.cs
+++ b/Engine/Source/Programs/AutomationTool/Gauntlet/Framework/Utils/Gauntlet.HordeReport.cs
@@ -41,7 +41,7 @@ namespace Gauntlet
 			}
 		}
 
-		public class BaseHordeReport : BaseTestReport
+		public abstract class BaseHordeReport : BaseTestReport
 		{
 			protected string OutputArtifactPath;
 
@@ -264,6 +264,11 @@ namespace Gauntlet
 				return NewTestResult;
 			}
 
+			public override void AddEvent(string Type, string Message, object Context = null)
+			{
+				throw new System.NotImplementedException("AddEvent not implemented");
+			}
+
 			/// <summary>
 			/// Convert UnrealAutomatedTestPassResults to Horde data model
 			/// </summary>
diff --git a/Engine/Source/Programs/AutomationTool/Gauntlet/Gauntlet.Automation.csproj b/Engine/Source/Programs/AutomationTool/Gauntlet/Gauntlet.Automation.csproj
--- a/Engine/Source/Programs/AutomationTool/Gauntlet/Gauntlet.Automation.csproj
+++ b/Engine/Source/Programs/AutomationTool/Gauntlet/Gauntlet.Automation.csproj
@@ -63,6 +63,7 @@
   </ItemGroup>
   <ItemGroup>
     <PackageReference Include="Microsoft.CSharp" Version="4.7.0" />
+    <PackageReference Include="MySql.Data" Version="6.10.9" />
     <PackageReference Include="System.Data.DataSetExtensions" Version="4.5.0" />
     <PackageReference Include="System.Net.Http" Version="4.3.4" />
 
diff --git a/Engine/Source/Programs/AutomationTool/Gauntlet/Unreal/Base/Gauntlet.UnrealTestConfiguration.cs b/Engine/Source/Programs/AutomationTool/Gauntlet/Unreal/Base/Gauntlet.UnrealTestConfiguration.cs
--- a/Engine/Source/Programs/AutomationTool/Gauntlet/Unreal/Base/Gauntlet.UnrealTestConfiguration.cs
+++ b/Engine/Source/Programs/AutomationTool/Gauntlet/Unreal/Base/Gauntlet.UnrealTestConfiguration.cs
@@ -536,6 +536,18 @@ namespace Gauntlet
 		[AutoParam]
 		public string HordeArtifactPath = "";
 
+		/// <summary>
+		/// Telemetry Database config to use
+		/// </summary>
+		[AutoParam]
+		public string PublishTelemetryTo = "";
+
+		/// <summary>
+		/// Path to Database config file
+		/// </summary>
+		[AutoParam]
+		public string DatabaseConfigPath = "";
+
 		/// <summary>
 		/// What the test result should be treated as if we reach max duration.
 		/// </summary>
diff --git a/Engine/Source/Programs/AutomationTool/Gauntlet/Unreal/Base/Gauntlet.UnrealTestNode.cs b/Engine/Source/Programs/AutomationTool/Gauntlet/Unreal/Base/Gauntlet.UnrealTestNode.cs
--- a/Engine/Source/Programs/AutomationTool/Gauntlet/Unreal/Base/Gauntlet.UnrealTestNode.cs
+++ b/Engine/Source/Programs/AutomationTool/Gauntlet/Unreal/Base/Gauntlet.UnrealTestNode.cs
@@ -1146,6 +1146,25 @@ namespace Gauntlet
 				HordeTestDataCollection.AddNewTestReport(HordeTestDataKey, Report);
 				HordeTestDataCollection.WriteToJson(HordeTestDataFilePath, true);
 			}
+			if (!string.IsNullOrEmpty(GetConfiguration().PublishTelemetryTo) && Report is ITelemetryReport Telemetry)
+			{
+				IEnumerable<TelemetryData> DataRows = Telemetry.GetAllTelemetryData();
+				if (DataRows != null)
+				{
+					IDatabaseConfig<TelemetryData> DBConfig = DatabaseConfigManager<TelemetryData>.GetConfigByName(GetConfiguration().PublishTelemetryTo);
+					if (DBConfig != null)
+					{
+						DBConfig.LoadConfig(GetConfiguration().DatabaseConfigPath);
+						IDatabaseDriver<TelemetryData> DB = DBConfig.GetDriver();
+						Log.Verbose("Submitting telemetry data to {0}", DB.ToString());
+						DB.SubmitDataItems(DataRows, Context);
+					}
+					else
+					{
+						Log.Warning("Got telemetry data, but database configuration is unknown '{0}'.", GetConfiguration().PublishTelemetryTo);
+					}
+				}
+			}
 		}
 
 		/// <summary>
diff --git a/Engine/Source/Programs/AutomationTool/Gauntlet/Unreal/Config/Gauntlet.EpicConfig.cs b/Engine/Source/Programs/AutomationTool/Gauntlet/Unreal/Config/Gauntlet.EpicConfig.cs
new file mode 100644
--- /dev/null
+++ b/Engine/Source/Programs/AutomationTool/Gauntlet/Unreal/Config/Gauntlet.EpicConfig.cs
@@ -0,0 +1,81 @@
+// Copyright Epic Games, Inc. All Rights Reserved.
+
+using AutomationTool;
+using Gauntlet;
+using System;
+using System.Collections.Generic;
+using UnrealBuildTool;
+
+namespace EpicConfig
+{
+	public class UETelemetry : MySQLConfig<Gauntlet.TelemetryData>
+	{
+		public override void LoadConfig(string ConfigFilePath)
+		{
+			if (string.IsNullOrEmpty(ConfigFilePath))
+			{
+				if (HostPlatform.Current.HostEditorPlatform == UnrealTargetPlatform.Win64)
+				{
+					ConfigFilePath = @"\\epicgames.net\root\Builds\Automation\UE5\Config\ue-automation-telemetry-mysql-config.txt";
+				}
+				else
+				{
+					ConfigFilePath = "/Volumes/Root/Builds/Automation/UE5/Config/ue-automation-telemetry-mysql-config.txt";
+				}
+			}
+			base.LoadConfig(ConfigFilePath);
+		}
+		/// <summary>
+		/// Get Target Table name based on data type
+		/// </summary>
+		public override string GetTableName()
+		{
+			return "test_records";
+		}
+		/// <summary>
+		/// Get Target Table columns based on data type
+		/// </summary>
+		public override IEnumerable<string> GetTableColumns()
+		{
+			return new List<string>() {
+				"changelist",
+				"system",
+				"stream",
+				"dateTime",
+				"project",
+				"platform",
+				"config",
+				"testName",
+				"dataPoint",
+				"recordedValue"
+			};
+		}
+		/// <summary>
+		/// Format the data for target table based on data type
+		/// </summary>
+		/// <returns></returns>
+		public override IEnumerable<string> FormatDataForTable(Gauntlet.TelemetryData Data, ITestContext TestContext)
+		{
+			if (TestContext is UnrealTestContext Context)
+			{
+				UnrealTestRoleContext Role = Context.RoleContext[UnrealTargetRole.Client];
+				return new List<string>() {
+					Context.BuildInfo.Changelist.ToString(),
+					Environment.MachineName,
+					Context.BuildInfo.Branch,
+					DateTime.Now.ToString("yyyy-MM-dd HH:mm:ss"),
+					Context.BuildInfo.ProjectName,
+					Role.Platform.ToString(),
+					Role.Configuration.ToString(),
+					Data.TestName,
+					Data.DataPoint,
+					Data.Measurement.ToString()
+				};
+			}
+			else
+			{
+				throw new AutomationException(string.Format("The Test Context is not UnrealTestContext; '{0}'.", TestContext.GetType().FullName));
+			}
+		}
+	}
+}
\ No newline at end of file
