A few GetCookPlatform() implementations were not checking bIsClientOnly correctly. As a result they were staging Client and Game versions to the same output location.

#rnx

diff --git a/Engine/Source/Programs/AutomationTool/HoloLens/HoloLensPlatform.Automation.cs b/Engine/Source/Programs/AutomationTool/HoloLens/HoloLensPlatform.Automation.cs
--- a/Engine/Source/Programs/AutomationTool/HoloLens/HoloLensPlatform.Automation.cs
+++ b/Engine/Source/Programs/AutomationTool/HoloLens/HoloLensPlatform.Automation.cs
@@ -616,7 +616,8 @@ namespace HoloLens.Automation
 
 		public override string GetCookPlatform(bool bDedicatedServer, bool bIsClientOnly)
 		{
-			return PlatformType.ToString();
+		    string Return = PlatformType.ToString();
+            return bIsClientOnly ? Return + "Client" : Return;
 		}
 
 		static void FillMapfile(DirectoryReference DirectoryName, string SearchPath, string Mask, StringBuilder AppXRecipeBuiltFiles)
diff --git a/Engine/Source/Programs/AutomationTool/TVOS/TVOSPlatform.Automation.cs b/Engine/Source/Programs/AutomationTool/TVOS/TVOSPlatform.Automation.cs
--- a/Engine/Source/Programs/AutomationTool/TVOS/TVOSPlatform.Automation.cs
+++ b/Engine/Source/Programs/AutomationTool/TVOS/TVOSPlatform.Automation.cs
@@ -54,7 +54,7 @@ public class TVOSPlatform : IOSPlatform
 
     public override string GetCookPlatform(bool bDedicatedServer, bool bIsClientOnly)
 	{
-		return "TVOS";
+		return bIsClientOnly ? "TVOSClient" : "TVOS";
 	}
 
     public override void GetFilesToDeployOrStage(ProjectParams Params, DeploymentContext SC)
