Gauntlet - Fixing an issue on Mac that prevent updating file attributes after copying a file.

#rnx

diff --git a/Engine/Source/Programs/AutomationTool/Gauntlet/Framework/Gauntlet.Utils.cs b/Engine/Source/Programs/AutomationTool/Gauntlet/Framework/Gauntlet.Utils.cs
--- a/Engine/Source/Programs/AutomationTool/Gauntlet/Framework/Gauntlet.Utils.cs
+++ b/Engine/Source/Programs/AutomationTool/Gauntlet/Framework/Gauntlet.Utils.cs
@@ -1129,7 +1129,7 @@ namespace Gauntlet
 								Log.Verbose("Copying to {0}", DestInfo.FullName);
 							}
 
-							SrcInfo.CopyTo(DestInfo.FullName, true);
+							DestInfo = SrcInfo.CopyTo(DestInfo.FullName, true);
 
 							// Clear attributes and set last write time
 							DestInfo.Attributes = FileAttributes.Normal;
