Fix issue with Gauntlet and content-only projects where "UE4" was still being used in build source checks

diff --git a/Engine/Source/Programs/AutomationTool/Gauntlet/Unreal/Utils/Gauntlet.UnrealHelpers.cs b/Engine/Source/Programs/AutomationTool/Gauntlet/Unreal/Utils/Gauntlet.UnrealHelpers.cs
--- a/Engine/Source/Programs/AutomationTool/Gauntlet/Unreal/Utils/Gauntlet.UnrealHelpers.cs
+++ b/Engine/Source/Programs/AutomationTool/Gauntlet/Unreal/Utils/Gauntlet.UnrealHelpers.cs
@@ -254,7 +254,7 @@ namespace Gauntlet
 
 			if (InName.StartsWith("UnrealGame", StringComparison.OrdinalIgnoreCase))
 			{
-				ShortName = "UE4";
+				ShortName = "Unreal";
 			}
 
 			string AppName = Path.GetFileNameWithoutExtension(InName);
