Fix for missing file permissions on Mac.

Switches the default zip utility to be the built in NET ZipArchive class. Adds additional support for storing permission bits by P/Invoking libSystem.dylib and adding the flags from stat64() to the ZipArchiveEntry. Also restores the permissions on unzip. One less dependency on Ionic.Zip. :)

diff --git a/Engine/Source/Programs/AutomationTool/AutomationUtils/CommandUtils.cs b/Engine/Source/Programs/AutomationTool/AutomationUtils/CommandUtils.cs
--- a/Engine/Source/Programs/AutomationTool/AutomationUtils/CommandUtils.cs
+++ b/Engine/Source/Programs/AutomationTool/AutomationUtils/CommandUtils.cs
@@ -18,6 +18,7 @@ using System.Threading.Tasks;
 using EpicGames.Core;
 using System.Xml;
 using System.Xml.Serialization;
+using System.IO.Compression;
 
 namespace AutomationTool
 {
@@ -2169,19 +2170,21 @@ namespace AutomationTool
 		/// <param name="ZipFile">Filename for the zip</param>
 		/// <param name="BaseDirectory">Base directory to store relative paths in the zip file to</param>
 		/// <param name="Files">Files to include in the archive</param>
-		public static void ZipFiles(FileReference ZipFile, DirectoryReference BaseDirectory, IEnumerable<FileReference> Files)
+		public static void ZipFiles(FileReference OutputFile, DirectoryReference BaseDirectory, IEnumerable<FileReference> Files)
 		{
-			using(Ionic.Zip.ZipFile Zip = new Ionic.Zip.ZipFile(Encoding.UTF8))
+			using (ZipArchive ZipArchive = ZipFile.Open(OutputFile.FullName, ZipArchiveMode.Create))
 			{
-				Zip.UseZip64WhenSaving = Ionic.Zip.Zip64Option.AsNecessary;
-					foreach(FileReference File in Files)
+				foreach (FileReference File in Files)
+				{
+					string Name = File.MakeRelativeTo(BaseDirectory);
+					if (Path.DirectorySeparatorChar != '/')
 					{
-						Zip.AddFile(File.FullName, Path.GetDirectoryName(File.MakeRelativeTo(BaseDirectory)));
+						Name = Name.Replace(Path.DirectorySeparatorChar, '/');
 					}
-					CommandUtils.CreateDirectory(ZipFile.Directory);
-					Zip.Save(ZipFile.FullName);
+					ZipArchive.CreateEntryFromFile_CrossPlatform(File.FullName, Name, CompressionLevel.Fastest);
 				}
 			}
+		}
 
 		/// <summary>
 		/// Extracts the contents of a zip file
@@ -2191,63 +2194,29 @@ namespace AutomationTool
 		/// <returns>List of files written</returns>
 		public static IEnumerable<string> UnzipFiles(string ZipFileName, string BaseDirectory)
 		{
-			List<string> OutputFileNames = new List<string>();
-			if (Utils.IsRunningOnMono)
-			{
-				CommandUtils.CreateDirectory(BaseDirectory);
-
-				// Use system unzip tool as there have been instances of Ionic not being able to open zips created with Mac zip tool
-				string Output = CommandUtils.RunAndLog("unzip", "\"" + ZipFileName + "\" -d \"" + BaseDirectory + "\"", Options: ERunOptions.Default | ERunOptions.SpewIsVerbose);
-
-				// Split log output into lines
-				string[] Lines = Output.Split(new char[] { '\n', '\r' });
-
-				foreach (string LogLine in Lines)
-				{
-					CommandUtils.LogInformation(LogLine);
+			return UnzipFiles(new FileReference(ZipFileName), new DirectoryReference(BaseDirectory)).Select(x => x.FullName);
+		}
 
-					// Split each line into two by whitespace
-					string[] SplitLine = LogLine.Split(new char[] { ' ', '\t' }, 2, StringSplitOptions.RemoveEmptyEntries);
-					if (SplitLine.Length == 2)
-					{
-						// Second part of line should be a path
-						string FilePath = SplitLine[1].Trim();
-						CommandUtils.LogInformation(FilePath);
-						if (File.Exists(FilePath) && !OutputFileNames.Contains(FilePath) && FilePath != ZipFileName)
-						{
-							if (CommandUtils.IsProbablyAMacOrIOSExe(FilePath))
-							{
-								FixUnixFilePermissions(FilePath);
-							}
-							OutputFileNames.Add(FilePath);
-						}
-					}
-				}
-				if (OutputFileNames.Count == 0)
-				{
-					CommandUtils.LogWarning("Unable to parse unzipped files from {0}", ZipFileName);
-				}
-			}
-			else
+		/// <summary>
+		/// Extracts the contents of a zip file
+		/// </summary>
+		/// <param name="ZipFileName">Name of the zip file</param>
+		/// <param name="BaseDirectory">Output directory</param>
+		/// <returns>List of files written</returns>
+		public static IEnumerable<FileReference> UnzipFiles(FileReference ZipFileName, DirectoryReference BaseDirectory)
+		{
+			List<FileReference> OutputFiles = new List<FileReference>();
+			using (ZipArchive ZipArchive = ZipFile.Open(ZipFileName.FullName, ZipArchiveMode.Read))
 			{
-				// manually extract the files. There was a problem with the Ionic.Zip library that required this on non-PC at one point,
-				// but that problem is now fixed. Leaving this code as is as we need to return the list of created files anyway.
-				using (Ionic.Zip.ZipFile Zip = new Ionic.Zip.ZipFile(ZipFileName))
+				foreach (ZipArchiveEntry Entry in ZipArchive.Entries)
 				{
-					
-					foreach (Ionic.Zip.ZipEntry Entry in Zip.Entries.Where(x => !x.IsDirectory))
-					{
-						string OutputFileName = Path.Combine(BaseDirectory, Entry.FileName);
-						Directory.CreateDirectory(Path.GetDirectoryName(OutputFileName));
-						using (FileStream OutputStream = new FileStream(OutputFileName, FileMode.Create, FileAccess.Write))
-						{
-							Entry.Extract(OutputStream);
-						}
-						OutputFileNames.Add(OutputFileName);
-					}
+					FileReference OutputFile = FileReference.Combine(BaseDirectory, Entry.FullName);
+					DirectoryReference.CreateDirectory(OutputFile.Directory);
+					Entry.ExtractToFile_CrossPlatform(OutputFile.FullName, true);
+					OutputFiles.Add(OutputFile);
 				}
 			}
-			return OutputFileNames;
+			return OutputFiles;
 		}
 
 		/// <summary>
diff --git a/Engine/Source/Programs/AutomationTool/BuildGraph/Tasks/UnzipTask.cs b/Engine/Source/Programs/AutomationTool/BuildGraph/Tasks/UnzipTask.cs
--- a/Engine/Source/Programs/AutomationTool/BuildGraph/Tasks/UnzipTask.cs
+++ b/Engine/Source/Programs/AutomationTool/BuildGraph/Tasks/UnzipTask.cs
@@ -73,7 +73,7 @@ namespace BuildGraph.Tasks
 			HashSet<FileReference> OutputFiles = new HashSet<FileReference>();
 			foreach(FileReference ZipFile in ZipFiles)
 			{
-				OutputFiles.UnionWith(CommandUtils.UnzipFiles(ZipFile.FullName, ToDir.FullName).Select(x => new FileReference(x)));
+				OutputFiles.UnionWith(CommandUtils.UnzipFiles(ZipFile, ToDir));
 			}
 
 			// Apply the optional tag to the produced archive
diff --git a/Engine/Source/Programs/Shared/EpicGames.Core/FileUtils.cs b/Engine/Source/Programs/Shared/EpicGames.Core/FileUtils.cs
--- a/Engine/Source/Programs/Shared/EpicGames.Core/FileUtils.cs
+++ b/Engine/Source/Programs/Shared/EpicGames.Core/FileUtils.cs
@@ -318,6 +318,28 @@ namespace EpicGames.Core
 			}
 		}
 
+		/// <summary>
+		/// Gets the file mode on Mac
+		/// </summary>
+		/// <param name="FileName"></param>
+		/// <returns></returns>
+		public static int GetFileMode_Mac(string FileName)
+		{
+			stat64_t stat = new stat64_t();
+			int Result = stat64(FileName, stat);
+			return (Result >= 0)? stat.st_mode : -1;
+		}
+
+		/// <summary>
+		/// Sets the file mode on Mac
+		/// </summary>
+		/// <param name="FileName"></param>
+		/// <param name="Mode"></param>
+		public static void SetFileMode_Mac(string FileName, ushort Mode)
+		{
+			chmod(FileName, Mode);
+		}
+
 		#region Win32 Native File Methods
 
 		[StructLayout(LayoutKind.Sequential, CharSet = CharSet.Unicode)]
@@ -692,5 +714,46 @@ namespace EpicGames.Core
 		}
 
 		#endregion
+
+		#region Mac Native File Methods
+#pragma warning disable CS0649
+		struct timespec_t
+		{
+			public ulong tv_sec;
+			public ulong tv_nsec;
+		}
+
+		[StructLayout(LayoutKind.Sequential)]
+		class stat64_t
+		{
+			public uint st_dev;
+			public ushort st_mode;
+			public ushort st_nlink;
+			public ulong st_ino;
+			public uint st_uid;
+			public uint st_gid;
+			public uint st_rdev;
+			public timespec_t st_atimespec;
+			public timespec_t st_mtimespec;
+			public timespec_t st_ctimespec;
+			public timespec_t st_birthtimespec;
+			public ulong st_size;
+			public ulong st_blocks;
+			public uint st_blksize;
+			public uint st_flags;
+			public uint st_gen;
+			public uint st_lspare;
+			public ulong st_qspare1;
+			public ulong st_qspare2;
+		}
+
+		[DllImport("libSystem.dylib")]
+		static extern int stat64(string pathname, stat64_t stat);
+
+		[DllImport("libSystem.dylib")]
+		static extern int chmod(string path, ushort mode);
+
+#pragma warning restore CS0649
+		#endregion
 	}
 }
diff --git a/Engine/Source/Programs/Shared/EpicGames.Core/ZipArchiveExtensions.cs b/Engine/Source/Programs/Shared/EpicGames.Core/ZipArchiveExtensions.cs
new file mode 100644
--- /dev/null
+++ b/Engine/Source/Programs/Shared/EpicGames.Core/ZipArchiveExtensions.cs
@@ -0,0 +1,51 @@
+using System;
+using System.Collections.Generic;
+using System.Text;
+using System.Runtime.InteropServices;
+using System.IO.Compression;
+
+namespace EpicGames.Core
+{
+	/// <summary>
+	/// Additional functionality around <see cref="System.IO.Compression.ZipArchive"/> to support non-Windows filesystems
+	/// </summary>
+	public static class ZipArchiveExtensions
+	{
+		/// <summary>
+		/// Create a zip archive entry, preserving platform mode bits
+		/// </summary>
+		/// <param name="Destination"></param>
+		/// <param name="SourceFileName"></param>
+		/// <param name="EntryName"></param>
+		/// <param name="CompressionLevel"></param>
+		/// <returns></returns>
+		public static ZipArchiveEntry CreateEntryFromFile_CrossPlatform(this ZipArchive Destination, string SourceFileName, string EntryName, CompressionLevel CompressionLevel)
+		{
+			ZipArchiveEntry Entry = ZipFileExtensions.CreateEntryFromFile(Destination, SourceFileName, EntryName, CompressionLevel);
+			if (RuntimeInformation.IsOSPlatform(OSPlatform.OSX))
+			{
+				int Result = FileUtils.GetFileMode_Mac(SourceFileName);
+				if(Result >= 0)
+				{
+					Entry.ExternalAttributes = (int)Result << 16;
+				}
+			}
+			return Entry;
+		}
+
+		/// <summary>
+		/// Extract a zip archive entry, preserving platform mode bits
+		/// </summary>
+		/// <param name="Entry"></param>
+		/// <param name="TargetFileName"></param>
+		/// <param name="Overwrite"></param>
+		public static void ExtractToFile_CrossPlatform(this ZipArchiveEntry Entry, string TargetFileName, bool Overwrite)
+		{
+			ZipFileExtensions.ExtractToFile(Entry, TargetFileName, Overwrite);
+			if (RuntimeInformation.IsOSPlatform(OSPlatform.OSX))
+			{
+				FileUtils.SetFileMode_Mac(TargetFileName, (ushort)(Entry.ExternalAttributes >> 16));
+			}
+		}
+	}
+}
