Make sure to create the output directory when zipping.

#rnx

diff --git a/Engine/Source/Programs/AutomationTool/AutomationUtils/CommandUtils.cs b/Engine/Source/Programs/AutomationTool/AutomationUtils/CommandUtils.cs
--- a/Engine/Source/Programs/AutomationTool/AutomationUtils/CommandUtils.cs
+++ b/Engine/Source/Programs/AutomationTool/AutomationUtils/CommandUtils.cs
@@ -2172,6 +2172,11 @@ namespace AutomationTool
 		/// <param name="Files">Files to include in the archive</param>
 		public static void ZipFiles(FileReference OutputFile, DirectoryReference BaseDirectory, IEnumerable<FileReference> Files)
 		{
+			if (!DirectoryReference.Exists(OutputFile.Directory))
+			{
+				DirectoryReference.CreateDirectory(OutputFile.Directory);
+			}
+
 			using (ZipArchive ZipArchive = ZipFile.Open(OutputFile.FullName, ZipArchiveMode.Create))
 			{
 				foreach (FileReference File in Files)
