Add the legacy unzip functionality back into UAT. BPT currently stored in a zip that isn't compatible with NET ZipArchive class.

#rnx

diff --git a/Engine/Source/Programs/AutomationTool/AutomationUtils/CommandUtils.cs b/Engine/Source/Programs/AutomationTool/AutomationUtils/CommandUtils.cs
--- a/Engine/Source/Programs/AutomationTool/AutomationUtils/CommandUtils.cs
+++ b/Engine/Source/Programs/AutomationTool/AutomationUtils/CommandUtils.cs
@@ -2224,6 +2224,73 @@ namespace AutomationTool
 			return OutputFiles;
 		}
 
+		/// <summary>
+		/// Extracts the contents of a zip file
+		/// </summary>
+		/// <param name="ZipFileName">Name of the zip file</param>
+		/// <param name="BaseDirectory">Output directory</param>
+		/// <returns>List of files written</returns>
+		public static IEnumerable<string> LegacyUnzipFiles(string ZipFileName, string BaseDirectory)
+		{
+			List<string> OutputFileNames = new List<string>();
+			if (Utils.IsRunningOnMono)
+			{
+				CommandUtils.CreateDirectory(BaseDirectory);
+
+				// Use system unzip tool as there have been instances of Ionic not being able to open zips created with Mac zip tool
+				string Output = CommandUtils.RunAndLog("unzip", "\"" + ZipFileName + "\" -d \"" + BaseDirectory + "\"", Options: ERunOptions.Default | ERunOptions.SpewIsVerbose);
+
+				// Split log output into lines
+				string[] Lines = Output.Split(new char[] { '\n', '\r' });
+
+				foreach (string LogLine in Lines)
+				{
+					CommandUtils.LogInformation(LogLine);
+
+					// Split each line into two by whitespace
+					string[] SplitLine = LogLine.Split(new char[] { ' ', '\t' }, 2, StringSplitOptions.RemoveEmptyEntries);
+					if (SplitLine.Length == 2)
+					{
+						// Second part of line should be a path
+						string FilePath = SplitLine[1].Trim();
+						CommandUtils.LogInformation(FilePath);
+						if (File.Exists(FilePath) && !OutputFileNames.Contains(FilePath) && FilePath != ZipFileName)
+						{
+							if (CommandUtils.IsProbablyAMacOrIOSExe(FilePath))
+							{
+								FixUnixFilePermissions(FilePath);
+							}
+							OutputFileNames.Add(FilePath);
+						}
+					}
+				}
+				if (OutputFileNames.Count == 0)
+				{
+					CommandUtils.LogWarning("Unable to parse unzipped files from {0}", ZipFileName);
+				}
+			}
+			else
+			{
+				// manually extract the files. There was a problem with the Ionic.Zip library that required this on non-PC at one point,
+				// but that problem is now fixed. Leaving this code as is as we need to return the list of created files anyway.
+				using (Ionic.Zip.ZipFile Zip = new Ionic.Zip.ZipFile(ZipFileName))
+				{
+
+					foreach (Ionic.Zip.ZipEntry Entry in Zip.Entries.Where(x => !x.IsDirectory))
+					{
+						string OutputFileName = Path.Combine(BaseDirectory, Entry.FileName);
+						Directory.CreateDirectory(Path.GetDirectoryName(OutputFileName));
+						using (FileStream OutputStream = new FileStream(OutputFileName, FileMode.Create, FileAccess.Write))
+						{
+							Entry.Extract(OutputStream);
+						}
+						OutputFileNames.Add(OutputFileName);
+					}
+				}
+			}
+			return OutputFileNames;
+		}
+
 		/// <summary>
 		/// Resolve an arbitrary file specification against a directory. May contain any number of p4 wildcard operators (?, *, ...).
 		/// </summary>
