Enable IoStore and the new loader by default for all projects.
Enable UPS by default for all non-dedicated server targets for all projects.
When staging, remove the requirement to specify the -iostore command line argument by reading the value "bUseIoStore" from the platform game ini section "/Script/UnrealEd.ProjectPackagingSettings".
The command line arguments (-iostore and -skipiostore) can still be used to override the ini file setting.

#rnx

diff --git a/Engine/Config/BaseEngine.ini b/Engine/Config/BaseEngine.ini
--- a/Engine/Config/BaseEngine.ini
+++ b/Engine/Config/BaseEngine.ini
@@ -1017,7 +1017,7 @@ ZeroEngineVersionWarning=True
 UseStrictEngineVersioning=True
 CanStripEditorOnlyExportsAndImports=True
 CanSkipEditorReferencedPackagesWhenCooking=False
-CanUseUnversionedPropertySerialization=False
+CanUseUnversionedPropertySerialization=True
 TestUnversionedPropertySerializationWhenCooking=False
 DetailedCallstacksInNonMonolithicBuilds=True
 UseSeperateBulkDataFiles=False
@@ -1039,7 +1039,7 @@ MinPayloadLength=0
 
 [/Script/Engine.StreamingSettings]
 s.MinBulkDataSizeForAsyncLoading=131072
-s.AsyncLoadingThreadEnabled=False
+s.AsyncLoadingThreadEnabled=True
 s.EventDrivenLoaderEnabled=True
 s.WarnIfTimeLimitExceeded=False
 s.TimeLimitExceededMultiplier=1.5
diff --git a/Engine/Config/BaseGame.ini b/Engine/Config/BaseGame.ini
--- a/Engine/Config/BaseGame.ini
+++ b/Engine/Config/BaseGame.ini
@@ -85,6 +85,7 @@ MinWindowHeight=16
 BuildConfiguration=PPBC_Development
 FullRebuild=False
 UsePakFile=True
+bUseIoStore=True
 bGenerateChunks=False
 bChunkHardReferencesOnly=False
 IncludePrerequisites=True
diff --git a/Engine/Source/Programs/AutomationTool/Scripts/CopyBuildToStagingDirectory.Automation.cs b/Engine/Source/Programs/AutomationTool/Scripts/CopyBuildToStagingDirectory.Automation.cs
--- a/Engine/Source/Programs/AutomationTool/Scripts/CopyBuildToStagingDirectory.Automation.cs
+++ b/Engine/Source/Programs/AutomationTool/Scripts/CopyBuildToStagingDirectory.Automation.cs
@@ -2515,7 +2515,7 @@ public partial class Project : CommandUtils
 					AdditionalArgs += " -platform=" + ConfigHierarchy.GetIniPlatformName(SC.StageTargetPlatform.IniPlatformType);
 
 					Dictionary<string, string> UnrealPakResponseFile = PakParams.UnrealPakResponseFile;
-					if (ShouldCreateIoStoreContainerFiles(Params, SC))
+					if (ShouldCreateIoStoreContainerFiles(Params, SC.StageTargetPlatform))
 					{
 						bool bAllowBulkDataInIoStore = true;
 						if(!PlatformEngineConfig.GetBool("Core.System", "AllowBulkDataInIoStore", out bAllowBulkDataInIoStore))
@@ -2676,7 +2676,7 @@ public partial class Project : CommandUtils
 
 				InternalUtils.SafeCreateDirectory(Path.GetDirectoryName(ReleaseVersionPath));
 				InternalUtils.SafeCopyFile(OutputLocation.FullName, ReleaseVersionPath);
-				if (ShouldCreateIoStoreContainerFiles(Params, SC))
+				if (ShouldCreateIoStoreContainerFiles(Params, SC.StageTargetPlatform))
 				{
 					InternalUtils.SafeCopyFile(Path.ChangeExtension(OutputLocation.FullName, ".utoc"), Path.ChangeExtension(ReleaseVersionPath, ".utoc"));
 					InternalUtils.SafeCopyFile(Path.ChangeExtension(OutputLocation.FullName, ".ucas"), Path.ChangeExtension(ReleaseVersionPath, ".ucas"));
@@ -2786,7 +2786,7 @@ public partial class Project : CommandUtils
 					{
 						HashSet<string> IncludedExtensions = new HashSet<string>();
 						IncludedExtensions.Add(OutputFilenameExtension);
-						if (ShouldCreateIoStoreContainerFiles(Params, SC))
+						if (ShouldCreateIoStoreContainerFiles(Params, SC.StageTargetPlatform))
 						{
 							IncludedExtensions.Add(".ucas");
 							IncludedExtensions.Add(".utoc");
@@ -3413,14 +3413,32 @@ public partial class Project : CommandUtils
 		return TmpPackagingPath;
 	}
 
-	private static bool ShouldCreateIoStoreContainerFiles(ProjectParams Params, DeploymentContext SC)
+	private static bool ShouldCreateIoStoreContainerFiles(ProjectParams Params, Platform StageTargetPlatform)
 	{
 		if (Params.CookOnTheFly)
 		{
 			return false;
 		}
 
-		return Params.IoStore && !Params.SkipIoStore;
+		if (Params.SkipIoStore)
+		{
+			return false;
+		}
+
+		if (Params.IoStore)
+		{
+			return true;
+		}
+
+		if (Params.Stage && !Params.SkipStage)
+		{
+			ConfigHierarchy PlatformGameConfig = ConfigCache.ReadHierarchy(ConfigHierarchyType.Game, DirectoryReference.FromFile(Params.RawProjectPath), StageTargetPlatform.IniPlatformType);
+			bool bUseIoStore = false;
+			PlatformGameConfig.GetBool("/Script/UnrealEd.ProjectPackagingSettings", "bUseIoStore", out bUseIoStore);
+			return bUseIoStore;
+		}
+
+		return false;
 	}
 
 	private static bool ShouldCreatePak(ProjectParams Params, DeploymentContext SC)
