Use the cross platform zip methods to store files in temp storage.

#rnx

diff --git a/Engine/Source/Programs/AutomationTool/BuildGraph/TempStorage.cs b/Engine/Source/Programs/AutomationTool/BuildGraph/TempStorage.cs
--- a/Engine/Source/Programs/AutomationTool/BuildGraph/TempStorage.cs
+++ b/Engine/Source/Programs/AutomationTool/BuildGraph/TempStorage.cs
@@ -1024,7 +1024,7 @@ namespace AutomationTool
 									// This is in cases of a super hot cache, so the operation was largely CPU bound.
 									// Also, sadly, mono appears to have a bug where nothing you can do will properly set the LastWriteTime on the created entry,
 									// so we have to ignore timestamps on files extracted from a zip, since it may have been created on a Mac.
-									ZipFileExtensions.CreateEntryFromFile(ZipArchive, File.FullName, CommandUtils.ConvertSeparators(PathSeparator.Slash, File.MakeRelativeTo(RootDir)), System.IO.Compression.CompressionLevel.Fastest);
+									ZipArchiveExtensions.CreateEntryFromFile_CrossPlatform(ZipArchive, File.FullName, CommandUtils.ConvertSeparators(PathSeparator.Slash, File.MakeRelativeTo(RootDir)), System.IO.Compression.CompressionLevel.Fastest);
 								} while (FilesToZip.TryDequeue(out File));
 							}
 							// if we are using a staging dir, copy to the final location and delete the staged copy.
@@ -1110,7 +1110,7 @@ namespace AutomationTool
 									{
 										try
 										{
-											Entry.ExtractToFile(ExtractedFilename, true);
+											Entry.ExtractToFile_CrossPlatform(ExtractedFilename, true);
 											break;
 										}
 										catch (IOException IOEx)
