Add support for shunting Unzip tasks to the legacy unzip logic.

#rnx

diff --git a/Engine/Source/Programs/AutomationTool/BuildGraph/Tasks/UnzipTask.cs b/Engine/Source/Programs/AutomationTool/BuildGraph/Tasks/UnzipTask.cs
--- a/Engine/Source/Programs/AutomationTool/BuildGraph/Tasks/UnzipTask.cs
+++ b/Engine/Source/Programs/AutomationTool/BuildGraph/Tasks/UnzipTask.cs
@@ -29,6 +29,12 @@ namespace BuildGraph.Tasks
 		[TaskParameter]
 		public DirectoryReference ToDir;
 
+		/// <summary>
+		/// Whether or not to use the legacy unzip code.
+		/// </summary>
+		[TaskParameter(Optional = true)]
+		public bool UseLegacyUnzip = false;
+
 		/// <summary>
 		/// Tag to be applied to the extracted files.
 		/// </summary>
@@ -73,7 +79,14 @@ namespace BuildGraph.Tasks
 			HashSet<FileReference> OutputFiles = new HashSet<FileReference>();
 			foreach(FileReference ZipFile in ZipFiles)
 			{
-				OutputFiles.UnionWith(CommandUtils.UnzipFiles(ZipFile, ToDir));
+				if (Parameters.UseLegacyUnzip)
+				{
+					OutputFiles.UnionWith(CommandUtils.LegacyUnzipFiles(ZipFile.FullName, ToDir.FullName).Select(x => new FileReference(x)));
+				}
+				else
+				{
+					OutputFiles.UnionWith(CommandUtils.UnzipFiles(ZipFile, ToDir));
+				}
 			}
 
 			// Apply the optional tag to the produced archive
