Gauntlet - Fix build installation issue with multiple clients test
by storing installation destination and returning the path to installation.

#rnx

diff --git a/Engine/Source/Programs/AutomationTool/Gauntlet/Platform/Windows/Gauntlet.TargetDeviceWindows.cs b/Engine/Source/Programs/AutomationTool/Gauntlet/Platform/Windows/Gauntlet.TargetDeviceWindows.cs
--- a/Engine/Source/Programs/AutomationTool/Gauntlet/Platform/Windows/Gauntlet.TargetDeviceWindows.cs
+++ b/Engine/Source/Programs/AutomationTool/Gauntlet/Platform/Windows/Gauntlet.TargetDeviceWindows.cs
@@ -362,7 +362,7 @@ namespace Gauntlet
 
 				if (!SkipDeploy)
 				{
-					StagedBuild.InstallBuildParallel(AppConfig, InBuild, BuildPath, DestPath, ToString());
+					DestPath = StagedBuild.InstallBuildParallel(AppConfig, InBuild, BuildPath, DestPath, ToString());
 				}
 				else
 				{
diff --git a/Engine/Source/Programs/AutomationTool/Gauntlet/Unreal/BuildSource/Gauntlet.StagedBuildSource.cs b/Engine/Source/Programs/AutomationTool/Gauntlet/Unreal/BuildSource/Gauntlet.StagedBuildSource.cs
--- a/Engine/Source/Programs/AutomationTool/Gauntlet/Unreal/BuildSource/Gauntlet.StagedBuildSource.cs
+++ b/Engine/Source/Programs/AutomationTool/Gauntlet/Unreal/BuildSource/Gauntlet.StagedBuildSource.cs
@@ -89,13 +89,25 @@ namespace Gauntlet
 			Installed
 		}
 
-		static Dictionary<StagedBuild, InstallStatus> LocalInstalls = new Dictionary<StagedBuild, InstallStatus>();
+		class Install
+		{
+			public InstallStatus Status;
+			public string Path;
+
+			public Install(InstallStatus InStatus, string InPath)
+			{
+				Status = InStatus;
+				Path = InPath;
+			}
+		}
+
+		static Dictionary<StagedBuild, Install> LocalInstalls = new Dictionary<StagedBuild, Install>();
 
 		/// <summary>
 		/// When running parallel tests using staged builds, local desktop devices only need one copy of client/server 
 		/// this method is used to coordinate the copy across multiple local devices
 		/// </summary>
-		public static void InstallBuildParallel(UnrealAppConfig AppConfig, StagedBuild InBuild, string BuildPath, string DestPath, string Desc)
+		public static string InstallBuildParallel(UnrealAppConfig AppConfig, StagedBuild InBuild, string BuildPath, string DestPath, string Desc)
 		{
 			// In parallel tests, we only want to copy the client/server once
 			bool Install = false;
@@ -106,11 +118,13 @@ namespace Gauntlet
 				if (!LocalInstalls.ContainsKey(InBuild))
 				{
 					Install = true;
-					LocalInstalls[InBuild] = Status = InstallStatus.Installing;
+					Status = InstallStatus.Installing;
+					LocalInstalls[InBuild] = new Install(Status, DestPath);
 				}
 				else
 				{
-					Status = LocalInstalls[InBuild];
+					Status = LocalInstalls[InBuild].Status;
+					DestPath = LocalInstalls[InBuild].Path;
 				}
 			}
 
@@ -132,7 +146,7 @@ namespace Gauntlet
 				{
 					lock (Globals.MainLock)
 					{
-						LocalInstalls[InBuild] = InstallStatus.Error;
+						LocalInstalls[InBuild].Status = InstallStatus.Error;
 					}
 
 					throw Ex;
@@ -140,7 +154,7 @@ namespace Gauntlet
 
 				lock (Globals.MainLock)
 				{
-					LocalInstalls[InBuild] = InstallStatus.Installed;
+					LocalInstalls[InBuild].Status = InstallStatus.Installed;
 				}
 
 			}
@@ -158,7 +172,7 @@ namespace Gauntlet
 
 					lock (Globals.MainLock)
 					{
-						Status = LocalInstalls[InBuild];
+						Status = LocalInstalls[InBuild].Status;
 					}
 
 					// install process failed in other thread, disk full, etc
@@ -175,6 +189,7 @@ namespace Gauntlet
 				}
 			}
 
+			return DestPath;
 		}
 
 		public static IEnumerable<T> CreateFromPath<T>(UnrealTargetPlatform InPlatform, string InProjectName, string InPath, string InExecutableExtension)
