Gauntlet - Add Context property to TelemetryData

#rnx

diff --git a/Engine/Source/Programs/AutomationTool/Gauntlet/Framework/Base/Gauntlet.TestReport.cs b/Engine/Source/Programs/AutomationTool/Gauntlet/Framework/Base/Gauntlet.TestReport.cs
--- a/Engine/Source/Programs/AutomationTool/Gauntlet/Framework/Base/Gauntlet.TestReport.cs
+++ b/Engine/Source/Programs/AutomationTool/Gauntlet/Framework/Base/Gauntlet.TestReport.cs
@@ -50,11 +50,13 @@ namespace Gauntlet
 		public string TestName { get; private set; }
 		public string DataPoint { get; private set; }
 		public double Measurement { get; private set; }
-		public TelemetryData(string InTestName, string InDataPoint, double InMeasurement)
+		public string Context { get; private set; }
+		public TelemetryData(string InTestName, string InDataPoint, double InMeasurement, string InContext = "")
 		{
 			TestName = InTestName;
 			DataPoint = InDataPoint;
 			Measurement = InMeasurement;
+			Context = InContext;
 		}
 	}
 
@@ -62,15 +64,16 @@ namespace Gauntlet
 	/// Interface for telemetry report
 	/// </summary>
 	public interface ITelemetryReport
-	{ 
+	{
 		/// <summary>
 		/// Attach Telemetry data to the ITestReport
 		/// </summary>
 		/// <param name="TestName"></param>
 		/// <param name="DataPoint"></param>
 		/// <param name="Measurement"></param>
+		/// <param name="Context"></param>
 		/// <returns></returns>
-		void AddTelemetry(string TestName, string DataPoint, double Measurement);
+		void AddTelemetry(string TestName, string DataPoint, double Measurement, string Context = "");
 
 		/// <summary>
 		/// Return the telemetry data accumulated
@@ -142,15 +145,16 @@ namespace Gauntlet
 		/// <param name="TestName"></param>
 		/// <param name="DataPoint"></param>
 		/// <param name="Measurement"></param>
+		/// <param name="Context"></param>
 		/// <returns></returns>
-		public virtual void AddTelemetry(string TestName, string DataPoint, double Measurement)
+		public virtual void AddTelemetry(string TestName, string DataPoint, double Measurement, string Context = "")
 		{
 			if (TelemetryDataList is null)
 			{
 				TelemetryDataList = new List<TelemetryData>();
 			}
 
-			TelemetryDataList.Add(new TelemetryData(TestName, DataPoint, Measurement));
+			TelemetryDataList.Add(new TelemetryData(TestName, DataPoint, Measurement, Context));
 		}
 
 		/// <summary>
diff --git a/Engine/Source/Programs/AutomationTool/Gauntlet/Unreal/Config/Gauntlet.EpicConfig.cs b/Engine/Source/Programs/AutomationTool/Gauntlet/Unreal/Config/Gauntlet.EpicConfig.cs
--- a/Engine/Source/Programs/AutomationTool/Gauntlet/Unreal/Config/Gauntlet.EpicConfig.cs
+++ b/Engine/Source/Programs/AutomationTool/Gauntlet/Unreal/Config/Gauntlet.EpicConfig.cs
@@ -46,6 +46,7 @@ namespace EpicConfig
 				"platform",
 				"config",
 				"testName",
+				"context",
 				"dataPoint",
 				"recordedValue"
 			};
@@ -68,6 +69,7 @@ namespace EpicConfig
 					Role.Platform.ToString(),
 					Role.Configuration.ToString(),
 					Data.TestName,
+					Data.Context,
 					Data.DataPoint,
 					Data.Measurement.ToString()
 				};
