Fix more quoting issues

diff --git a/Engine/Source/Programs/AutomationTool/Linux/LinuxPlatform.Automation.cs b/Engine/Source/Programs/AutomationTool/Linux/LinuxPlatform.Automation.cs
--- a/Engine/Source/Programs/AutomationTool/Linux/LinuxPlatform.Automation.cs
+++ b/Engine/Source/Programs/AutomationTool/Linux/LinuxPlatform.Automation.cs
@@ -135,7 +135,7 @@ public abstract class BaseLinuxPlatform : Platform
 
 		if (Utils.IsRunningOnMono)
 		{
-			var Result = CommandUtils.Run("sh", string.Format("-c 'chmod +x \"{0}\"'", IntermediateFile.ToString().Replace("'", "'\"'\"'")));
+			var Result = CommandUtils.Run("env", string.Format("-- \"chmod\" \"+x\" \"{0}\"", IntermediateFile.ToString().Replace("'", "'\"'\"'")));
 			if (Result.ExitCode != 0)
 			{
 				throw new AutomationException(string.Format("Failed to chmod \"{0}\"", IntermediateFile));
