Remove legacy Mono unzip path.

diff --git a/Engine/Source/Programs/AutomationTool/BuildGraph/TempStorage.cs b/Engine/Source/Programs/AutomationTool/BuildGraph/TempStorage.cs
--- a/Engine/Source/Programs/AutomationTool/BuildGraph/TempStorage.cs
+++ b/Engine/Source/Programs/AutomationTool/BuildGraph/TempStorage.cs
@@ -1062,72 +1062,55 @@ namespace AutomationTool
 		/// </remarks>
 		private static void ParallelUnzipFiles(FileInfo[] ZipFiles, DirectoryReference RootDir)
 		{
-			// Sadly, mono implemention of System.IO.Compression is really poor (as of 2015/Aug), causing OOM when parallel zipping a large set of files.
-			// However, Ionic is MUCH slower than .NET's native implementation (2x+ slower in our build farm), so we stick to the faster solution on PC.
-			// The code duplication in the threadprocs is unfortunate here, and hopefully we can settle on .NET's implementation on both platforms eventually.
-			if (Utils.IsRunningOnMono)
-			{
-				Parallel.ForEach(ZipFiles,
-					(ZipFile) =>
-					{
-						using (var ZipArchive = Ionic.Zip.ZipFile.Read(ZipFile.FullName))
-						{
-							ZipArchive.ExtractAll(RootDir.FullName, Ionic.Zip.ExtractExistingFileAction.OverwriteSilently);
-						}
-					});
-			}
-			else
-			{
-				Parallel.ForEach(ZipFiles,
-					(ZipFile) =>
+			Parallel.ForEach(ZipFiles,
+				(ZipFile) =>
+				{
+					// unzip the files manually instead of caling ZipFile.ExtractToDirectory() because we need to overwrite readonly files. Because of this, creating the directories is up to us as well.
+					using (ZipArchive ZipArchive = System.IO.Compression.ZipFile.OpenRead(ZipFile.FullName))
 					{
-						// unzip the files manually instead of caling ZipFile.ExtractToDirectory() because we need to overwrite readonly files. Because of this, creating the directories is up to us as well.
-						using (var ZipArchive = System.IO.Compression.ZipFile.OpenRead(ZipFile.FullName))
+						foreach (ZipArchiveEntry Entry in ZipArchive.Entries)
 						{
-							foreach (var Entry in ZipArchive.Entries)
+							// Use CommandUtils.CombinePaths to ensure directory separators get converted correctly. On mono on *nix, if the path has backslashes it will not convert it.
+							var ExtractedFilename = CommandUtils.CombinePaths(RootDir.FullName, Entry.FullName);
+							// Zips can contain empty dirs. Ours usually don't have them, but we should support it.
+							if (Path.GetFileName(ExtractedFilename).Length == 0)
+							{
+								Directory.CreateDirectory(ExtractedFilename);
+							}
+							else
 							{
-								// Use CommandUtils.CombinePaths to ensure directory separators get converted correctly. On mono on *nix, if the path has backslashes it will not convert it.
-								var ExtractedFilename = CommandUtils.CombinePaths(RootDir.FullName, Entry.FullName);
-								// Zips can contain empty dirs. Ours usually don't have them, but we should support it.
-								if (Path.GetFileName(ExtractedFilename).Length == 0)
+								// We must delete any existing file, even if it's readonly. .Net does not do this by default.
+								if (File.Exists(ExtractedFilename))
 								{
-									Directory.CreateDirectory(ExtractedFilename);
+									InternalUtils.SafeDeleteFile(ExtractedFilename, true);
 								}
 								else
 								{
-									// We must delete any existing file, even if it's readonly. .Net does not do this by default.
-									if (File.Exists(ExtractedFilename))
+									Directory.CreateDirectory(Path.GetDirectoryName(ExtractedFilename));
+								}
+
+								int UnzipAttempts = 3;
+								while (UnzipAttempts-- > 0)
+								{
+									try
 									{
-										InternalUtils.SafeDeleteFile(ExtractedFilename, true);
+										Entry.ExtractToFile_CrossPlatform(ExtractedFilename, true);
+										break;
 									}
-									else
+									catch (IOException IOEx)
 									{
-										Directory.CreateDirectory(Path.GetDirectoryName(ExtractedFilename));
-									}
-
-									int UnzipAttempts = 3;
-									while (UnzipAttempts-- > 0)
-									{
-										try
+										if (UnzipAttempts == 0)
 										{
-											Entry.ExtractToFile_CrossPlatform(ExtractedFilename, true);
-											break;
+											throw;
 										}
-										catch (IOException IOEx)
-										{
-											if (UnzipAttempts == 0)
-											{
-												throw;
-											}
 
-											Log.TraceWarning("Failed to unzip '{0}' from '{1}' to '{2}', retrying.. (Error: {3})", Entry.FullName, ZipFile.FullName, ExtractedFilename, IOEx.Message);
-										}
+										Log.TraceWarning("Failed to unzip '{0}' from '{1}' to '{2}', retrying.. (Error: {3})", Entry.FullName, ZipFile.FullName, ExtractedFilename, IOEx.Message);
 									}
 								}
 							}
 						}
-					});
-			}
+					}
+				});
 		}
 
 		/// <summary>
