Report an error if we fail to delete an empty folder but continue working.

#rnx

diff --git a/Engine/Source/Programs/AutomationTool/BuildGraph/TempStorage.cs b/Engine/Source/Programs/AutomationTool/BuildGraph/TempStorage.cs
--- a/Engine/Source/Programs/AutomationTool/BuildGraph/TempStorage.cs
+++ b/Engine/Source/Programs/AutomationTool/BuildGraph/TempStorage.cs
@@ -1402,21 +1402,28 @@ namespace AutomationTool
 			// Try to delete any empty branch folders
 			foreach (DirectoryInfo StreamDirectory in new DirectoryInfo(TempStorageDir).EnumerateDirectories())
 			{
-				if(StreamDirectory.EnumerateDirectories().Count() == 0 && StreamDirectory.EnumerateFiles().Count() == 0)
+				try
 				{
-					try
-					{
-						StreamDirectory.Delete();
-					}
-					catch (IOException)
+					if (StreamDirectory.EnumerateDirectories().Count() == 0 && StreamDirectory.EnumerateFiles().Count() == 0)
 					{
-						// only catch "directory is not empty type exceptions, if possible. Best we can do is check for IOException.
-					}
-					catch (Exception Ex)
-					{
-						CommandUtils.LogWarning("Unexpected failure trying to delete (potentially empty) stream directory {0}: {1}", StreamDirectory.FullName, Ex);
+						try
+						{
+							StreamDirectory.Delete();
+						}
+						catch (IOException)
+						{
+							// only catch "directory is not empty type exceptions, if possible. Best we can do is check for IOException.
+						}
+						catch (Exception Ex)
+						{
+							CommandUtils.LogWarning("Unexpected failure trying to delete (potentially empty) stream directory {0}: {1}", StreamDirectory.FullName, Ex);
+						}
 					}
 				}
+				catch (Exception Ex)
+				{
+					Log.TraceError("Exception while trying to delete {0}: {1}", StreamDirectory, Ex);
+				}
 			}
 		}
 	}
