Gauntlet - pickup telemetry outputs from UE and attached them to Gauntlet telemetry report in order to publish to database
+ fix issue with CIS label 'Boot ET (platform) Editor' being always green in EngineTest

#rnx

diff --git a/Engine/Source/Programs/AutomationTool/Gauntlet/Framework/Utils/Gauntlet.CSVParser.cs b/Engine/Source/Programs/AutomationTool/Gauntlet/Framework/Utils/Gauntlet.CSVParser.cs
new file mode 100644
--- /dev/null
+++ b/Engine/Source/Programs/AutomationTool/Gauntlet/Framework/Utils/Gauntlet.CSVParser.cs
@@ -0,0 +1,40 @@
+// Copyright Epic Games, Inc. All Rights Reserved.
+
+using System.IO;
+using System.Collections.Generic;
+
+
+namespace Gauntlet
+{
+	public static class CSVParser
+	{
+		public static List<Dictionary<string, string>> Load(string File)
+		{
+			List<Dictionary<string, string>> Rows = new List<Dictionary<string, string>>();
+			// Parse csv file
+			using (StreamReader Reader = new StreamReader(File))
+			{
+				// Parse Header
+				string FirstLine = Reader.ReadLine();
+				string[] HeaderValues = FirstLine.Split(',');
+				// Parse the rest of the data
+				while (!Reader.EndOfStream)
+				{
+					Dictionary<string, string> OneRow = new Dictionary<string, string>();
+					string Line = Reader.ReadLine();
+					string[] Values = Line.Split(',');
+					if (Values.Length == HeaderValues.Length)
+					{
+						for (int i = 0; i < HeaderValues.Length; i++)
+						{
+							OneRow[HeaderValues[i]] = Values[i];
+						}
+						Rows.Add(OneRow);
+					}
+				}
+			}
+
+			return Rows;
+		}
+	}
+}
\ No newline at end of file
diff --git a/Engine/Source/Programs/AutomationTool/Gauntlet/Unreal/Automation/UE.Automation.cs b/Engine/Source/Programs/AutomationTool/Gauntlet/Unreal/Automation/UE.Automation.cs
--- a/Engine/Source/Programs/AutomationTool/Gauntlet/Unreal/Automation/UE.Automation.cs
+++ b/Engine/Source/Programs/AutomationTool/Gauntlet/Unreal/Automation/UE.Automation.cs
@@ -72,6 +72,12 @@ namespace UE
 		[AutoParam]
 		public string ReportURL = "";
 
+		/// <summary>
+		/// Absolute or project relative directory path to write automation telemetry outputs.
+		/// </summary>
+		[AutoParam]
+		public string TelemetryDirectory = "";
+
 		/// <summary>
 		/// Use Simple Horde Report instead of Unreal Automated Tests
 		/// </summary>
@@ -130,6 +136,21 @@ namespace UE
 				}
 			}
 
+			// Setup commandline for telemetry outputs
+			if (!string.IsNullOrEmpty(PublishTelemetryTo) || !string.IsNullOrEmpty(TelemetryDirectory))
+			{
+				if (string.IsNullOrEmpty(TelemetryDirectory))
+				{
+					TelemetryDirectory = Path.Combine(Globals.TempDir, "Telemetry");
+				}
+				if (Directory.Exists(TelemetryDirectory))
+				{
+					// clean any previous data
+					Directory.Delete(TelemetryDirectory, true);
+				}
+				AppConfig.CommandLine += string.Format("-TelemetryDirectory=\"{0}\"", TelemetryDirectory);
+			}
+
 			// Arguments for writing out the report and providing a URL where it can be viewed
 			string ReportArgs = string.IsNullOrEmpty(ReportExportPath) ? "" : string.Format("-ReportExportPath=\"{0}\"", ReportExportPath);
 
@@ -410,25 +431,60 @@ namespace UE
 		/// <returns>ITestReport</returns>
 		public override ITestReport CreateReport(TestResult Result)
 		{
+			ITestReport Report = null;
 			// Save test result data for Horde build system
 			bool WriteTestResultsForHorde = GetConfiguration().WriteTestResultsForHorde;
 			if (WriteTestResultsForHorde)
 			{
 				if (GetConfiguration().SimpleHordeReport)
 				{
-					return base.CreateReport(Result);
+					Report = base.CreateReport(Result);
 				}
 				else
 				{
 					string ReportPath = GetConfiguration().ReportExportPath;
 					if (!string.IsNullOrEmpty(ReportPath))
 					{
-						return CreateUnrealEngineTestPassReport(ReportPath, GetConfiguration().ReportURL);
+						Report = CreateUnrealEngineTestPassReport(ReportPath, GetConfiguration().ReportURL);
 					}
 				}
 			}
 
-			return null;
+			string TelemetryDirectory = GetConfiguration().TelemetryDirectory;
+			if (!string.IsNullOrEmpty(TelemetryDirectory))
+			{
+				if (Report is ITelemetryReport Telemetry)
+				{
+					// Scan for csv files
+					foreach (string File in Directory.GetFiles(TelemetryDirectory))
+					{
+						if (Path.GetExtension(File) != ".csv")
+						{
+							continue;
+						}
+						List<Dictionary<string, string>> CSVData = CSVParser.Load(File);
+						// Populate telemetry report
+						foreach(Dictionary<string, string> Row in CSVData)
+						{
+							string TestName;
+							Row.TryGetValue("TestName", out TestName);
+							string DataPoint;
+							Row.TryGetValue("DataPoint", out DataPoint);
+							string Measurement;
+							Row.TryGetValue("Measurement", out Measurement);
+							string Context;
+							Row.TryGetValue("Context", out Context);
+							Telemetry.AddTelemetry(TestName, DataPoint, double.Parse(Measurement), Context);
+						}
+					}
+				}
+				else
+				{
+					Log.Warning("Publishing Telemetry is requested but '{0}' does not support telemetry input.", Report.GetType().FullName);
+				}
+			}
+
+			return Report;
 		}
 
 		/// <summary>
