Catch and log exceptions when enumerating files/folders to delete.

#rnx

diff --git a/Engine/Source/Programs/AutomationTool/BuildGraph/TempStorage.cs b/Engine/Source/Programs/AutomationTool/BuildGraph/TempStorage.cs
--- a/Engine/Source/Programs/AutomationTool/BuildGraph/TempStorage.cs
+++ b/Engine/Source/Programs/AutomationTool/BuildGraph/TempStorage.cs
@@ -1375,11 +1375,18 @@ namespace AutomationTool
 				CommandUtils.LogInformation("Scanning {0}...", StreamDirectory.FullName);
 				foreach (DirectoryInfo BuildDirectory in StreamDirectory.EnumerateDirectories())
 				{
-					if(!BuildDirectory.EnumerateFiles("*", SearchOption.AllDirectories).Any(x => x.LastWriteTimeUtc > RetainTime))
+					try
 					{
-						BuildsToDelete.Add(BuildDirectory);
+						if (!BuildDirectory.EnumerateFiles("*", SearchOption.AllDirectories).Any(x => x.LastWriteTimeUtc > RetainTime))
+						{
+							BuildsToDelete.Add(BuildDirectory);
+						}
+						NumBuilds++;
+					}
+					catch
+					{
+						Log.TraceError("Exception while trying to delete {0}: {1}", StreamDirectory, Ex);
 					}
-					NumBuilds++;
 				}
 			}
 			CommandUtils.LogInformation("Found {0} builds; {1} to delete.", NumBuilds, BuildsToDelete.Count);
