Fixing compile break.

#rnx

diff --git a/Engine/Source/Programs/AutomationTool/BuildGraph/TempStorage.cs b/Engine/Source/Programs/AutomationTool/BuildGraph/TempStorage.cs
--- a/Engine/Source/Programs/AutomationTool/BuildGraph/TempStorage.cs
+++ b/Engine/Source/Programs/AutomationTool/BuildGraph/TempStorage.cs
@@ -1383,7 +1383,7 @@ namespace AutomationTool
 						}
 						NumBuilds++;
 					}
-					catch
+					catch (Exception Ex)
 					{
 						Log.TraceError("Exception while trying to delete {0}: {1}", StreamDirectory, Ex);
 					}
