Gauntlet - fix issue not finding telemetry output directory
+ improve catching potential issue with telemetry output

#rnx

diff --git a/Engine/Source/Programs/AutomationTool/Gauntlet/Unreal/Automation/UE.Automation.cs b/Engine/Source/Programs/AutomationTool/Gauntlet/Unreal/Automation/UE.Automation.cs
--- a/Engine/Source/Programs/AutomationTool/Gauntlet/Unreal/Automation/UE.Automation.cs
+++ b/Engine/Source/Programs/AutomationTool/Gauntlet/Unreal/Automation/UE.Automation.cs
@@ -451,7 +451,7 @@ namespace UE
 			}
 
 			string TelemetryDirectory = GetConfiguration().TelemetryDirectory;
-			if (!string.IsNullOrEmpty(TelemetryDirectory))
+			if (!string.IsNullOrEmpty(TelemetryDirectory) && Directory.Exists(TelemetryDirectory))
 			{
 				if (Report is ITelemetryReport Telemetry)
 				{
@@ -462,9 +462,23 @@ namespace UE
 						{
 							continue;
 						}
-						List<Dictionary<string, string>> CSVData = CSVParser.Load(File);
+
+						List<Dictionary<string, string>> CSVData = null;
+						try
+						{
+							CSVData = CSVParser.Load(File);
+						}
+						catch (Exception Ex)
+						{
+							Log.Error("Telemetry - Failed to read CSV file '{0}'. {1}", File, Ex);
+						}
+
+						if (CSVData == null)
+						{
+							continue;
+						}
 						// Populate telemetry report
-						foreach(Dictionary<string, string> Row in CSVData)
+						foreach (Dictionary<string, string> Row in CSVData)
 						{
 							string TestName;
 							Row.TryGetValue("TestName", out TestName);
@@ -474,7 +488,19 @@ namespace UE
 							Row.TryGetValue("Measurement", out Measurement);
 							string Context;
 							Row.TryGetValue("Context", out Context);
-							Telemetry.AddTelemetry(TestName, DataPoint, double.Parse(Measurement), Context);
+							if(string.IsNullOrEmpty(TestName) || string.IsNullOrEmpty(DataPoint) || string.IsNullOrEmpty(Measurement))
+							{
+								Log.Warning("Telemetry - Missing data in CSV file '{0}':\n TestName='{1}', DataPoint='{2}', Measurement='{3}", File, TestName, DataPoint, Measurement);
+								continue;
+							}
+							try
+							{
+								Telemetry.AddTelemetry(TestName, DataPoint, double.Parse(Measurement), Context);
+							}
+							catch(FormatException)
+							{
+								Log.Error("Telemetry - Failed to parse Measurement value('{0}') in CSV file '{1}'.", Measurement, File);
+							}
 						}
 					}
 				}
