Gauntlet - Avoid using deprecated json parser, use System.Text.Json instead

#rnx

diff --git a/Engine/Source/Programs/AutomationTool/Gauntlet/Framework/Utils/Gauntlet.HordeReport.cs b/Engine/Source/Programs/AutomationTool/Gauntlet/Framework/Utils/Gauntlet.HordeReport.cs
--- a/Engine/Source/Programs/AutomationTool/Gauntlet/Framework/Utils/Gauntlet.HordeReport.cs
+++ b/Engine/Source/Programs/AutomationTool/Gauntlet/Framework/Utils/Gauntlet.HordeReport.cs
@@ -3,7 +3,7 @@
 using System;
 using System.Collections.Generic;
 using System.IO;
-using EpicGames.Core;
+using System.Text.Json;
 using AutomationTool;
 
 namespace Gauntlet
@@ -98,9 +98,9 @@ namespace Gauntlet
 		/// </summary>
 		public class ComparisonFiles
 		{
-			public string Difference;
-			public string Approved;
-			public string Unapproved;
+			public string Difference { get; set; }
+			public string Approved { get; set; }
+			public string Unapproved { get; set; }
 		}
 		/// <summary>
 		/// Contains information about test artifact
@@ -112,20 +112,20 @@ namespace Gauntlet
 				Files = new ComparisonFiles();
 			}
 
-			public string Id;
-			public string Name;
-			public string Type;
-			public ComparisonFiles Files;
+			public string Id { get; set; }
+			public string Name { get; set; }
+			public string Type { get; set; }
+			public ComparisonFiles Files { get; set; }
 		}
 		/// <summary>
 		/// Contains information about test entry event
 		/// </summary>
 		public class Event
 		{
-			public string Type;
-			public string Message;
-			public string Context;
-			public string Artifact;
+			public string Type { get; set; }
+			public string Message { get; set; }
+			public string Context { get; set; }
+			public string Artifact { get; set; }
 		}
 		/// <summary>
 		/// Contains information about test entry
@@ -137,10 +137,10 @@ namespace Gauntlet
 				Event = new Event();
 			}
 
-			public Event Event;
-			public string Filename;
-			public int LineNumber;
-			public string Timestamp;
+			public Event Event { get; set; }
+			public string Filename { get; set; }
+			public int LineNumber { get; set; }
+			public string Timestamp { get; set; }
 		}
 		/// <summary>
 		/// Contains detailed information about test result. This is to what TestPassResult refere to for each test result. 
@@ -153,13 +153,13 @@ namespace Gauntlet
 				Entries = new List<Entry>();
 			}
 
-			public string TestDisplayName;
-			public string FullTestPath;
-			public string State;
-			public int Warnings;
-			public int Errors;
-			public List<Artifact> Artifacts;
-			public List<Entry> Entries;
+			public string TestDisplayName { get; set; }
+			public string FullTestPath { get; set; }
+			public string State { get; set; }
+			public int Warnings { get; set; }
+			public int Errors { get; set; }
+			public List<Artifact> Artifacts { get; set; }
+			public List<Entry> Entries { get; set; }
 
 			/// <summary>
 			/// Add a new Artifact to the test result and return it 
@@ -208,9 +208,9 @@ namespace Gauntlet
 				get { return TestDetailed.State; }
 				set { TestDetailed.State = value; }
 			}
-			public string ArtifactName;
+			public string ArtifactName { get; set; }
 
-			private TestResultDetailed TestDetailed;
+			private TestResultDetailed TestDetailed { get; set; }
 
 			/// <summary>
 			/// Return the underlying TestResultDetailed 
@@ -243,15 +243,15 @@ namespace Gauntlet
 				Tests = new List<TestResult>();
 			}
 
-			public string ClientDescriptor;
-			public string ReportCreatedOn;
-			public string ReportURL;
-			public int SucceededCount;
-			public int SucceededWithWarningsCount;
-			public int FailedCount;
-			public int NotRunCount;
-			public float TotalDurationSeconds;
-			public List<TestResult> Tests;
+			public string ClientDescriptor { get; set; }
+			public string ReportCreatedOn { get; set; }
+			public string ReportURL { get; set; }
+			public int SucceededCount { get; set; }
+			public int SucceededWithWarningsCount { get; set; }
+			public int FailedCount { get; set; }
+			public int NotRunCount { get; set; }
+			public float TotalDurationSeconds { get; set; }
+			public List<TestResult> Tests { get; set; }
 
 			/// <summary>
 			/// Add a new TestResult to the pass results and return it 
@@ -278,51 +278,51 @@ namespace Gauntlet
 			public static UnrealEngineTestPassResults FromUnrealAutomatedTests(UnrealAutomatedTestPassResults InTestPassResults, string ReportPath, string ReportURL)
 			{
 				UnrealEngineTestPassResults OutTestPassResults = new UnrealEngineTestPassResults();
-				OutTestPassResults.ClientDescriptor = InTestPassResults.clientDescriptor;
-				OutTestPassResults.ReportCreatedOn = InTestPassResults.reportCreatedOn;
+				OutTestPassResults.ClientDescriptor = InTestPassResults.ClientDescriptor;
+				OutTestPassResults.ReportCreatedOn = InTestPassResults.ReportCreatedOn;
 				OutTestPassResults.ReportURL = ReportURL;
-				OutTestPassResults.SucceededCount = InTestPassResults.succeeded;
-				OutTestPassResults.SucceededWithWarningsCount = InTestPassResults.succeededWithWarnings;
-				OutTestPassResults.FailedCount = InTestPassResults.failed;
-				OutTestPassResults.NotRunCount = InTestPassResults.notRun;
-				OutTestPassResults.TotalDurationSeconds = InTestPassResults.totalDuration;
-				if (InTestPassResults.tests != null)
+				OutTestPassResults.SucceededCount = InTestPassResults.Succeeded;
+				OutTestPassResults.SucceededWithWarningsCount = InTestPassResults.SucceededWithWarnings;
+				OutTestPassResults.FailedCount = InTestPassResults.Failed;
+				OutTestPassResults.NotRunCount = InTestPassResults.NotRun;
+				OutTestPassResults.TotalDurationSeconds = InTestPassResults.TotalDuration;
+				if (InTestPassResults.Tests != null)
 				{
-					foreach (UnrealAutomatedTestResult InTestResult in InTestPassResults.tests)
+					foreach (UnrealAutomatedTestResult InTestResult in InTestPassResults.Tests)
 					{
 						TestResult ConvertedTestResult = OutTestPassResults.AddNewTestResult();
-						ConvertedTestResult.TestDisplayName = InTestResult.testDisplayName;
-						ConvertedTestResult.FullTestPath = InTestResult.fullTestPath;
-						ConvertedTestResult.State = InTestResult.state;
+						ConvertedTestResult.TestDisplayName = InTestResult.TestDisplayName;
+						ConvertedTestResult.FullTestPath = InTestResult.FullTestPath;
+						ConvertedTestResult.State = InTestResult.State;
 						Guid TestGuid = Guid.NewGuid();
 						ConvertedTestResult.ArtifactName = TestGuid + ".json";
-						InTestResult.artifactName = ConvertedTestResult.ArtifactName;
+						InTestResult.ArtifactName = ConvertedTestResult.ArtifactName;
 						// Copy Test Result Detail
 						TestResultDetailed ConvertedTestResultDetailed = ConvertedTestResult.GetTestResultDetailed();
-						ConvertedTestResultDetailed.Errors = InTestResult.errors;
-						ConvertedTestResultDetailed.Warnings = InTestResult.warnings;
-						foreach (UnrealAutomationArtifact InTestArtifact in InTestResult.artifacts)
+						ConvertedTestResultDetailed.Errors = InTestResult.Errors;
+						ConvertedTestResultDetailed.Warnings = InTestResult.Warnings;
+						foreach (UnrealAutomationArtifact InTestArtifact in InTestResult.Artifacts)
 						{
 							Artifact NewArtifact = ConvertedTestResultDetailed.AddNewArtifact();
-							NewArtifact.Id = InTestArtifact.id;
-							NewArtifact.Name = InTestArtifact.name;
-							NewArtifact.Type = InTestArtifact.type;
+							NewArtifact.Id = InTestArtifact.Id;
+							NewArtifact.Name = InTestArtifact.Name;
+							NewArtifact.Type = InTestArtifact.Type;
 							ComparisonFiles ArtifactFiles = NewArtifact.Files;
-							ArtifactFiles.Difference = !string.IsNullOrEmpty(InTestArtifact.files.difference)?Path.Combine(ReportPath, InTestArtifact.files.difference):null;
-							ArtifactFiles.Approved = !string.IsNullOrEmpty(InTestArtifact.files.approved)?Path.Combine(ReportPath, InTestArtifact.files.approved):null;
-							ArtifactFiles.Unapproved = !string.IsNullOrEmpty(InTestArtifact.files.unapproved)?Path.Combine(ReportPath, InTestArtifact.files.unapproved):null;
+							ArtifactFiles.Difference = !string.IsNullOrEmpty(InTestArtifact.Files.Difference)?Path.Combine(ReportPath, InTestArtifact.Files.Difference):null;
+							ArtifactFiles.Approved = !string.IsNullOrEmpty(InTestArtifact.Files.Approved)?Path.Combine(ReportPath, InTestArtifact.Files.Approved):null;
+							ArtifactFiles.Unapproved = !string.IsNullOrEmpty(InTestArtifact.Files.Unapproved)?Path.Combine(ReportPath, InTestArtifact.Files.Unapproved):null;
 						}
-						foreach (UnrealAutomationEntry InTestEntry in InTestResult.entries)
+						foreach (UnrealAutomationEntry InTestEntry in InTestResult.Entries)
 						{
 							Entry NewEntry = ConvertedTestResultDetailed.AddNewEntry();
-							NewEntry.Filename = InTestEntry.filename;
-							NewEntry.LineNumber = InTestEntry.lineNumber;
-							NewEntry.Timestamp = InTestEntry.timestamp;
+							NewEntry.Filename = InTestEntry.Filename;
+							NewEntry.LineNumber = InTestEntry.LineNumber;
+							NewEntry.Timestamp = InTestEntry.Timestamp;
 							Event EntryEvent = NewEntry.Event;
-							EntryEvent.Artifact = [REDACTED_EMAIL];
-							EntryEvent.Context = [REDACTED_EMAIL];
-							EntryEvent.Message = [REDACTED_EMAIL];
-							EntryEvent.Type = [REDACTED_EMAIL];
+							EntryEvent.Artifact = InTestEntry.Event.Artifact;
+							EntryEvent.Context = InTestEntry.Event.Context;
+							EntryEvent.Message = InTestEntry.Event.Message;
+							EntryEvent.Type = InTestEntry.Event.Type;
 						}
 					}
 				}
@@ -356,7 +356,7 @@ namespace Gauntlet
 					string TestResultFilePath = Path.Combine(OutputArtifactPath, OutputTestResult.ArtifactName);
 					try
 					{
-						Json_DEPRECATED.Save(new FileReference(TestResultFilePath), OutputTestResultDetailed);
+						File.WriteAllText(TestResultFilePath, JsonSerializer.Serialize(OutputTestResultDetailed, new JsonSerializerOptions { WriteIndented = true }));
 						ArtifactsCount++;
 					}
 					catch (Exception Ex)
@@ -384,15 +384,15 @@ namespace Gauntlet
 				Warnings = new List<String>();
 			}
 
-			public string Description;
-			public string ReportCreatedOn;
-			public float TotalDurationSeconds;
-			public bool HasSucceeded;
-			public string Status;
-			public string URLLink;
-			public List<String> Logs;
-			public List<String> Errors;
-			public List<String> Warnings;
+			public string Description { get; set; }
+			public string ReportCreatedOn { get; set; }
+			public float TotalDurationSeconds { get; set; }
+			public bool HasSucceeded { get; set; }
+			public string Status { get; set; }
+			public string URLLink { get; set; }
+			public List<String> Logs { get; set; }
+			public List<String> Errors { get; set; }
+			public List<String> Warnings { get; set; }
 
 			public override void AddEvent(string Type, string Message, object Context = null)
 			{
@@ -425,8 +425,8 @@ namespace Gauntlet
 		{
 			public class DataItem
 			{
-				public string Key;
-				public object Data;
+				public string Key { get; set; }
+				public object Data { get; set; }
 			}
 			public TestDataCollection()
 			{
@@ -446,7 +446,7 @@ namespace Gauntlet
 				return NewDataItem;
 			}
 
-			public List<DataItem> Items;
+			public List<DataItem> Items { get; set; }
 
 			/// <summary>
 			/// Write Test Data Collection to json
@@ -475,7 +475,7 @@ namespace Gauntlet
 				Log.Verbose("Writing Test Data Collection for Horde at {0}", OutputTestDataFilePath);
 				try
 				{
-					Json_DEPRECATED.Save(new FileReference(OutputTestDataFilePath), this);
+					File.WriteAllText(OutputTestDataFilePath, JsonSerializer.Serialize(this, new JsonSerializerOptions { WriteIndented = true }));
 				}
 				catch (Exception Ex)
 				{
diff --git a/Engine/Source/Programs/AutomationTool/Gauntlet/Unreal/Utils/Gauntlet.UnrealTestResults.cs b/Engine/Source/Programs/AutomationTool/Gauntlet/Unreal/Utils/Gauntlet.UnrealTestResults.cs
--- a/Engine/Source/Programs/AutomationTool/Gauntlet/Unreal/Utils/Gauntlet.UnrealTestResults.cs
+++ b/Engine/Source/Programs/AutomationTool/Gauntlet/Unreal/Utils/Gauntlet.UnrealTestResults.cs
@@ -1,65 +1,61 @@
 // Copyright Epic Games, Inc. All Rights Reserved.
 
-using AutomationTool;
+using System.IO;
 using System.Collections.Generic;
-using EpicGames.Core;
+using System.Text.Json;
 
 namespace Gauntlet
 {
-	// The following classes are used for json deserialization of the Unreal Automation Test Results.
-	// Because of a bug in the json deserialization used here(EpicGames.Core.Json),
-	// the names of the properties need to match the case of the json properties.
-
 	public class UnrealAutomationComparisonFiles
 	{
-		public string difference;
-		public string approved;
-		public string unapproved;
+		public string Difference { get; set; }
+		public string Approved { get; set; }
+		public string Unapproved { get; set; }
 	}
 	public class UnrealAutomationArtifact
 	{
-		public string id;
-		public string name;
-		public string type;
-		public UnrealAutomationComparisonFiles files;
+		public string Id { get; set; }
+		public string Name { get; set; }
+		public string Type { get; set; }
+		public UnrealAutomationComparisonFiles Files { get; set; }
 	}
 	public class UnrealAutomationEvent
 	{
-		public string type;
-		public string message;
-		public string context;
-		public string artifact;
+		public string Type { get; set; }
+		public string Message { get; set; }
+		public string Context { get; set; }
+		public string Artifact { get; set; }
 	}
 	public class UnrealAutomationEntry
 	{
-		public UnrealAutomationEvent @event;
-		public string filename;
-		public int lineNumber;
-		public string timestamp;
+		public UnrealAutomationEvent Event { get; set; }
+		public string Filename { get; set; }
+		public int LineNumber { get; set; }
+		public string Timestamp { get; set; }
 	}
 	public class UnrealAutomatedTestResult
 	{
-		public string testDisplayName;
-		public string fullTestPath;
-		public string artifactName;
-		public string state;
-		public int warnings;
-		public int errors;
-		public List<UnrealAutomationArtifact> artifacts;
-		public List<UnrealAutomationEntry> entries;
+		public string TestDisplayName { get; set; }
+		public string FullTestPath { get; set; }
+		public string ArtifactName { get; set; }
+		public string State { get; set; }
+		public int Warnings { get; set; }
+		public int Errors { get; set; }
+		public List<UnrealAutomationArtifact> Artifacts { get; set; }
+		public List<UnrealAutomationEntry> Entries { get; set; }
 	}
 	public class UnrealAutomatedTestPassResults
 	{
-		public string clientDescriptor;
-		public string reportCreatedOn;
-		public int succeeded;
-		public int succeededWithWarnings;
-		public int failed;
-		public int notRun;
-		public float totalDuration;
-		public bool comparisonExported;
-		public string comparisonExportDirectory;
-		public List<UnrealAutomatedTestResult> tests;
+		public string ClientDescriptor { get; set; }
+		public string ReportCreatedOn { get; set; }
+		public int Succeeded { get; set; }
+		public int SucceededWithWarnings { get; set; }
+		public int Failed { get; set; }
+		public int NotRun { get; set; }
+		public float TotalDuration { get; set; }
+		public bool ComparisonExported { get; set; }
+		public string ComparisonExportDirectory { get; set; }
+		public List<UnrealAutomatedTestResult> Tests { get; set; }
 
 		/// <summary>
 		/// Load Unreal Automated Test Results from json report
@@ -67,7 +63,12 @@ namespace Gauntlet
 		/// <param name="FilePath"></param>
 		public static UnrealAutomatedTestPassResults LoadFromJson(string FilePath)
 		{
-			UnrealAutomatedTestPassResults JsonTestPassResults = Json_DEPRECATED.Load<UnrealAutomatedTestPassResults>(new FileReference(FilePath));
+			JsonSerializerOptions Options = new JsonSerializerOptions
+			{
+				PropertyNameCaseInsensitive = true
+			};
+			string JsonString = File.ReadAllText(FilePath);
+			UnrealAutomatedTestPassResults JsonTestPassResults = JsonSerializer.Deserialize<UnrealAutomatedTestPassResults>(JsonString, Options);
 			return JsonTestPassResults;
 		}
 	}
