Remove unused GetTitleSettings() function from PlatformFeatures.
Add "-originalreleaseversion" command line switch in UAT. This is required by some platforms for patching / updates.

diff --git a/Engine/Source/Programs/AutomationTool/AutomationUtils/ProjectParams.cs b/Engine/Source/Programs/AutomationTool/AutomationUtils/ProjectParams.cs
--- a/Engine/Source/Programs/AutomationTool/AutomationUtils/ProjectParams.cs
+++ b/Engine/Source/Programs/AutomationTool/AutomationUtils/ProjectParams.cs
@@ -240,13 +240,13 @@ namespace AutomationTool
             this.DDCGraph = InParams.DDCGraph;
             this.InternationalizationPreset = InParams.InternationalizationPreset;
             this.CulturesToCook = InParams.CulturesToCook;
-            this.BasedOnReleaseVersion = InParams.BasedOnReleaseVersion;
+			this.OriginalReleaseVersion = InParams.OriginalReleaseVersion;
+			this.BasedOnReleaseVersion = InParams.BasedOnReleaseVersion;
             this.CreateReleaseVersion = InParams.CreateReleaseVersion;
             this.GeneratePatch = InParams.GeneratePatch;
 			this.AddPatchLevel = InParams.AddPatchLevel;
 			this.StageBaseReleasePaks = InParams.StageBaseReleasePaks;
 			this.DLCFile = InParams.DLCFile;
-			this.GenerateRemaster = InParams.GenerateRemaster;
 			this.DiscVersion = InParams.DiscVersion;
             this.DLCIncludeEngineContent = InParams.DLCIncludeEngineContent;
 			this.DLCActLikePatch = InParams.DLCActLikePatch;
@@ -430,14 +430,14 @@ namespace AutomationTool
 			bool? SkipCookingEditorContent = null,
             int? NumCookersToSpawn = null,
             string AdditionalCookerOptions = null,
-            string BasedOnReleaseVersion = null,
+			string OriginalReleaseVersion = null,
+			string BasedOnReleaseVersion = null,
             string CreateReleaseVersion = null,
 			string CreateReleaseVersionBasePath = null,
 			string BasedOnReleaseVersionBasePath = null,
             bool? GeneratePatch = null,
 			bool? AddPatchLevel = null,
 			bool? StageBaseReleasePaks = null,
-			bool? GenerateRemaster = null,
             string DiscVersion = null,
             string DLCName = null,
             string DiffCookedContentPath = null,
@@ -592,12 +592,12 @@ namespace AutomationTool
 			this.Cook = GetParamValueIfNotSpecified(Command, Cook, this.Cook, "cook");
 			this.CreateReleaseVersionBasePath = ParseParamValueIfNotSpecified(Command, CreateReleaseVersionBasePath, "createreleaseversionroot", String.Empty);
 			this.BasedOnReleaseVersionBasePath = ParseParamValueIfNotSpecified(Command, BasedOnReleaseVersionBasePath, "basedonreleaseversionroot", String.Empty);
-            this.CreateReleaseVersion = ParseParamValueIfNotSpecified(Command, CreateReleaseVersion, "createreleaseversion", String.Empty);
+			this.OriginalReleaseVersion = ParseParamValueIfNotSpecified(Command, OriginalReleaseVersion, "originalreleaseversion", String.Empty);
+			this.CreateReleaseVersion = ParseParamValueIfNotSpecified(Command, CreateReleaseVersion, "createreleaseversion", String.Empty);
             this.BasedOnReleaseVersion = ParseParamValueIfNotSpecified(Command, BasedOnReleaseVersion, "basedonreleaseversion", String.Empty);
             this.GeneratePatch = GetParamValueIfNotSpecified(Command, GeneratePatch, this.GeneratePatch, "GeneratePatch");
             this.AddPatchLevel = GetParamValueIfNotSpecified(Command, AddPatchLevel, this.AddPatchLevel, "AddPatchLevel");
 			this.StageBaseReleasePaks = GetParamValueIfNotSpecified(Command, StageBaseReleasePaks, this.StageBaseReleasePaks, "StageBaseReleasePaks");
-			this.GenerateRemaster = GetParamValueIfNotSpecified(Command, GenerateRemaster, this.GenerateRemaster, "GenerateRemaster");
 			this.DiscVersion = ParseParamValueIfNotSpecified(Command, DiscVersion, "DiscVersion", String.Empty);
 			this.AdditionalCookerOptions = ParseParamValueIfNotSpecified(Command, AdditionalCookerOptions, "AdditionalCookerOptions", String.Empty);
 
@@ -1533,11 +1533,6 @@ namespace AutomationTool
         /// </summary>
         public string InternationalizationPreset;
 
-        /// <summary>
-        /// Cook: Create a cooked release version.  Also, the version. e.g. 1.0
-        /// </summary>
-        public string CreateReleaseVersion;
-
         /// <summary>
 		/// Cook: While cooking clean up packages as we go along rather then cleaning everything (and potentially having to reload some of it) when we run out of space
 		/// </summary>
@@ -1555,14 +1550,24 @@ namespace AutomationTool
 		/// </summary>
 		public string CookOutputDir;
 
+		/// <summary>
+		/// Cook: Create a cooked release version.  Also, the version. e.g. 1.0
+		/// </summary>
+		public string CreateReleaseVersion;
+
 		/// <summary>
 		/// Cook: Base this cook of a already released version of the cooked data
 		/// </summary>
 		public string BasedOnReleaseVersion;
 
 		/// <summary>
-        /// Cook: Path to the root of the directory where we store released versions of the game for a given version
-        /// </summary>
+		/// The version of the originally released build. This is required by some platforms when generating patches.
+		/// </summary>
+		public string OriginalReleaseVersion;
+
+		/// <summary>
+		/// Cook: Path to the root of the directory where we store released versions of the game for a given version
+		/// </summary>
 		public string BasedOnReleaseVersionBasePath;
 
 		/// <summary>
@@ -1570,19 +1575,12 @@ namespace AutomationTool
 		/// </summary>
 		public string CreateReleaseVersionBasePath;
 
-        /// <summary>
-        /// Are we generating a patch, generate a patch from a previously released version of the game (use CreateReleaseVersion to create a release). 
-        /// this requires BasedOnReleaseVersion
-        /// see also CreateReleaseVersion, BasedOnReleaseVersion
-        /// </summary>
-        public bool GeneratePatch;
-
-        /// <summary>
-		/// Are we generating a remaster, generate a patch from a previously released version of the game (use CreateReleaseVersion to create a release). 
+		/// <summary>
+		/// Are we generating a patch, generate a patch from a previously released version of the game (use CreateReleaseVersion to create a release). 
 		/// this requires BasedOnReleaseVersion
 		/// see also CreateReleaseVersion, BasedOnReleaseVersion
 		/// </summary>
-		public bool GenerateRemaster;
+		public bool GeneratePatch;
 
 		/// <summary>
 		/// Required when building remaster package
@@ -2393,12 +2391,22 @@ namespace AutomationTool
             get { return !String.IsNullOrEmpty(InternationalizationPreset); }
         }
 
+		public bool HasCreateReleaseVersion
+        {
+            get { return !String.IsNullOrEmpty(CreateReleaseVersion); }
+        }
+
         public bool HasBasedOnReleaseVersion
         {
             get { return !String.IsNullOrEmpty(BasedOnReleaseVersion); }
         }
 
-        public bool HasAdditionalCookerOptions
+		public bool HasOriginalReleaseVersion
+		{
+			get { return !String.IsNullOrEmpty(OriginalReleaseVersion); }
+		}
+
+		public bool HasAdditionalCookerOptions
         {
             get { return !String.IsNullOrEmpty(AdditionalCookerOptions); }
         }
@@ -2413,11 +2421,6 @@ namespace AutomationTool
             get { return !String.IsNullOrEmpty(DiffCookedContentPath); }
         }
 
-        public bool HasCreateReleaseVersion
-        {
-            get { return !String.IsNullOrEmpty(CreateReleaseVersion); }
-        }
-
         public bool HasCulturesToCook
         {
             get { return CulturesToCook != null; }
@@ -2517,19 +2520,12 @@ namespace AutomationTool
 				BasePath = CommandUtils.CombinePaths(BasePath, BasedOnReleaseVersion, Platform);
 			}
 
-            /*if ( TitleID != null && TitleID.Count == 1 )
-            {
-                BasePath = CommandUtils.CombinePaths( BasePath, TitleID[0]);
-            }*/
-
 			return BasePath;
 		}
 
 		/// <summary>
 		/// Get the path to the target directory for creating a new release version
 		/// </summary>
-		/// <param name="SC"></param>
-		/// <returns></returns>
 		public String GetCreateReleaseVersionPath(DeploymentContext SC, bool bIsClientOnly)
 		{
 			String BasePath = CreateReleaseVersionBasePath;
@@ -2543,18 +2539,33 @@ namespace AutomationTool
 				BasePath = CommandUtils.CombinePaths(BasePath, CreateReleaseVersion, Platform);
 			}
 
-            /*if (TitleID != null && TitleID.Count == 1)
-            {
-                BasePath = CommandUtils.CombinePaths(BasePath, TitleID[0]);
-            }*/
+			return BasePath;
+		}
+
+		/// <summary>
+		/// Get the path to the directory of the originally released version we're using to generate a patch.
+		/// Only required by some platforms.
+		/// </summary>				
+		public String GetOriginalReleaseVersionPath(DeploymentContext SC, bool bIsClientOnly)
+		{
+			String BasePath = BasedOnReleaseVersionBasePath;
+			String Platform = SC.StageTargetPlatform.GetCookPlatform(SC.DedicatedServer, bIsClientOnly);
+			if (String.IsNullOrEmpty(BasePath))
+			{
+				BasePath = CommandUtils.CombinePaths(SC.ProjectRoot.FullName, "Releases", OriginalReleaseVersion, Platform);
+			}
+			else
+			{
+				BasePath = CommandUtils.CombinePaths(BasePath, OriginalReleaseVersion, Platform);
+			}
 
 			return BasePath;
 		}
 
-        /// <summary>
-        /// True if we are generating a patch
-        /// </summary>
-        public bool IsGeneratingPatch
+		/// <summary>
+		/// True if we are generating a patch
+		/// </summary>
+		public bool IsGeneratingPatch
         {
             get { return GeneratePatch; }
         }
@@ -2575,14 +2586,6 @@ namespace AutomationTool
             get { return StageBaseReleasePaks; }
         }
 
-		/// <summary>
-		/// True if we are generating a patch
-		/// </summary>
-		public bool IsGeneratingRemaster
-		{
-			get { return GenerateRemaster; }
-		}
-
 		public List<Platform> ClientTargetPlatformInstances
 		{
 			get
@@ -2800,11 +2803,6 @@ namespace AutomationTool
 			{
 				throw new AutomationException("-createchunkinstall must specify the chunk install data version string with -chunkinstallversion=");
 			}
-
-			/*if(IsGeneratingRemaster && string.IsNullOrEmpty(DiscVersion))
-			{
-				throw new AutomationException("DiscVersion is required for generating remaster package.");
-			}*/
 		}
 
 		protected bool bLogged = false;
@@ -2842,11 +2840,11 @@ namespace AutomationTool
                 CommandUtils.LogLog("GeneratePatch={0}", GeneratePatch);
 				CommandUtils.LogLog("AddPatchLevel={0}", AddPatchLevel);
 				CommandUtils.LogLog("StageBaseReleasePaks={0}", StageBaseReleasePaks);
-				CommandUtils.LogLog("GenerateRemaster={0}", GenerateRemaster);
 				CommandUtils.LogLog("DiscVersion={0}", DiscVersion);
 				CommandUtils.LogLog("CreateReleaseVersion={0}", CreateReleaseVersion);
                 CommandUtils.LogLog("BasedOnReleaseVersion={0}", BasedOnReleaseVersion);
-                CommandUtils.LogLog("DLCFile={0}", DLCFile);
+				CommandUtils.LogLog("OriginalReleaseVersion={0}", OriginalReleaseVersion);
+				CommandUtils.LogLog("DLCFile={0}", DLCFile);
                 CommandUtils.LogLog("DLCIncludeEngineContent={0}", DLCIncludeEngineContent);
 				CommandUtils.LogLog("DLCPakPluginFile={0}", DLCPakPluginFile);
                 CommandUtils.LogLog("DiffCookedContentPath={0}", DiffCookedContentPath);
diff --git a/Engine/Source/Runtime/Engine/Private/PlatformFeatures.cpp b/Engine/Source/Runtime/Engine/Private/PlatformFeatures.cpp
--- a/Engine/Source/Runtime/Engine/Private/PlatformFeatures.cpp
+++ b/Engine/Source/Runtime/Engine/Private/PlatformFeatures.cpp
@@ -19,11 +19,6 @@ IDVRStreamingSystem* IPlatformFeaturesModule::GetStreamingSystem()
 	return &GenericStreamingSystem;
 }
 
-TSharedPtr<const class FJsonObject> IPlatformFeaturesModule::GetTitleSettings()
-{
-	return nullptr;
-}
-
 FString IPlatformFeaturesModule::GetUniqueAppId()
 {
 	return FString();
diff --git a/Engine/Source/Runtime/Engine/Public/PlatformFeatures.h b/Engine/Source/Runtime/Engine/Public/PlatformFeatures.h
--- a/Engine/Source/Runtime/Engine/Public/PlatformFeatures.h
+++ b/Engine/Source/Runtime/Engine/Public/PlatformFeatures.h
@@ -42,8 +42,6 @@ public:
 
 	virtual class IDVRStreamingSystem* GetStreamingSystem();
 
-	virtual TSharedPtr<const class FJsonObject> GetTitleSettings();
-
 	virtual FString GetUniqueAppId();
 
 	virtual IVideoRecordingSystem* GetVideoRecordingSystem();
