Fix permission bits being reset by heurisitc that retrieves files from temp storage.

diff --git a/Engine/Source/Programs/AutomationTool/BuildGraph/TempStorage.cs b/Engine/Source/Programs/AutomationTool/BuildGraph/TempStorage.cs
--- a/Engine/Source/Programs/AutomationTool/BuildGraph/TempStorage.cs
+++ b/Engine/Source/Programs/AutomationTool/BuildGraph/TempStorage.cs
@@ -883,14 +883,10 @@ namespace AutomationTool
 					FileInfo[] ZipFiles = Manifest.ZipFiles.Select(x => new FileInfo(FileReference.Combine(SharedNodeDir, x.Name).FullName)).ToArray();
 					ParallelUnzipFiles(ZipFiles, RootDir);
 
-					// Fix any Unix permissions/chmod issues, and update the timestamps to match the manifest. Zip files only use local time, and there's no guarantee it matches the local clock.
+					// Update the timestamps to match the manifest. Zip files only use local time, and there's no guarantee it matches the local clock.
 					foreach(TempStorageFile ManifestFile in Manifest.Files)
 					{
 						FileReference File = ManifestFile.ToFileReference(RootDir);
-						if (Utils.IsRunningOnMono)
-						{
-							CommandUtils.FixUnixFilePermissions(File.FullName);
-						}
 						System.IO.File.SetLastWriteTimeUtc(File.FullName, new DateTime(ManifestFile.LastWriteTimeUtcTicks, DateTimeKind.Utc));
 					}
 
