Include a hash of the server address when creating a workspace for submitting binaries. Should fix issues with different edge servers being selected.

diff --git a/Engine/Source/Programs/AutomationTool/BuildGraph/Tasks/SubmitTask.cs b/Engine/Source/Programs/AutomationTool/BuildGraph/Tasks/SubmitTask.cs
--- a/Engine/Source/Programs/AutomationTool/BuildGraph/Tasks/SubmitTask.cs
+++ b/Engine/Source/Programs/AutomationTool/BuildGraph/Tasks/SubmitTask.cs
@@ -117,7 +117,7 @@ namespace AutomationTool.Tasks
 					// Create a brand new workspace
 					P4ClientInfo Client = new P4ClientInfo();
 					Client.Owner = CommandUtils.P4Env.User;
-					Client.Host = Environment.MachineName;
+					Client.Host = Environment.MachineName + "_" + ContentHash.MD5((CommandUtils.P4Env.ServerAndPort ?? "").ToUpperInvariant()).ToString();
 					Client.Stream = Parameters.Stream ?? CommandUtils.P4Env.Branch;
 					Client.RootPath = Parameters.RootDir.FullName ?? CommandUtils.RootDirectory.FullName;
 					Client.Name = Parameters.Workspace;
