Fix incorrect hostname for PCB submit workspace.

#rnx

diff --git a/Engine/Source/Programs/AutomationTool/BuildGraph/Tasks/SubmitTask.cs b/Engine/Source/Programs/AutomationTool/BuildGraph/Tasks/SubmitTask.cs
--- a/Engine/Source/Programs/AutomationTool/BuildGraph/Tasks/SubmitTask.cs
+++ b/Engine/Source/Programs/AutomationTool/BuildGraph/Tasks/SubmitTask.cs
@@ -117,10 +117,10 @@ namespace AutomationTool.Tasks
 					// Create a brand new workspace
 					P4ClientInfo Client = new P4ClientInfo();
 					Client.Owner = CommandUtils.P4Env.User;
-					Client.Host = Environment.MachineName + "_" + ContentHash.MD5((CommandUtils.P4Env.ServerAndPort ?? "").ToUpperInvariant()).ToString();
+					Client.Host = Environment.MachineName;
 					Client.Stream = Parameters.Stream ?? CommandUtils.P4Env.Branch;
 					Client.RootPath = Parameters.RootDir.FullName ?? CommandUtils.RootDirectory.FullName;
-					Client.Name = Parameters.Workspace;
+					Client.Name = Parameters.Workspace + "_" + ContentHash.MD5((CommandUtils.P4Env.ServerAndPort ?? "").ToUpperInvariant()).ToString();
 					Client.Options = P4ClientOption.NoAllWrite | P4ClientOption.Clobber | P4ClientOption.NoCompress | P4ClientOption.Unlocked | P4ClientOption.NoModTime | P4ClientOption.RmDir;
 					Client.LineEnd = P4LineEnd.Local;
 					CommandUtils.P4.CreateClient(Client, AllowSpew: Parameters.P4Verbose);
