Catch exceptions when scanning folders for folders that need to be deleted.

#rnx

diff --git a/Engine/Source/Programs/AutomationTool/BuildGraph/TempStorage.cs b/Engine/Source/Programs/AutomationTool/BuildGraph/TempStorage.cs
--- a/Engine/Source/Programs/AutomationTool/BuildGraph/TempStorage.cs
+++ b/Engine/Source/Programs/AutomationTool/BuildGraph/TempStorage.cs
@@ -1369,21 +1369,28 @@ namespace AutomationTool
 			foreach (DirectoryInfo StreamDirectory in new DirectoryInfo(TempStorageDir).EnumerateDirectories().OrderBy(x => x.Name))
 			{
 				CommandUtils.LogInformation("Scanning {0}...", StreamDirectory.FullName);
-				foreach (DirectoryInfo BuildDirectory in StreamDirectory.EnumerateDirectories())
+				try
 				{
-					try
+					foreach (DirectoryInfo BuildDirectory in StreamDirectory.EnumerateDirectories())
 					{
-						if (!BuildDirectory.EnumerateFiles("*", SearchOption.AllDirectories).Any(x => x.LastWriteTimeUtc > RetainTime))
+						try
 						{
-							BuildsToDelete.Add(BuildDirectory);
+							if (!BuildDirectory.EnumerateFiles("*", SearchOption.AllDirectories).Any(x => x.LastWriteTimeUtc > RetainTime))
+							{
+								BuildsToDelete.Add(BuildDirectory);
+							}
+							NumBuilds++;
+						}
+						catch (Exception Ex)
+						{
+							Log.TraceError("Exception while trying to delete {0}: {1}", BuildDirectory, Ex);
 						}
-						NumBuilds++;
-					}
-					catch (Exception Ex)
-					{
-						Log.TraceError("Exception while trying to delete {0}: {1}", StreamDirectory, Ex);
 					}
 				}
+				catch (Exception Ex)
+				{
+					Log.TraceError("Exception while trying to scan {0}: {1}", StreamDirectory, Ex);
+				}
 			}
 			CommandUtils.LogInformation("Found {0} builds; {1} to delete.", NumBuilds, BuildsToDelete.Count);
 
