Move gathering telemetry output from UE to add them in Gauntlet report into Gauntlet/Unreal/Utils so it can be used on other Unreal test

#rnx

diff --git a/Engine/Source/Programs/AutomationTool/Gauntlet/Unreal/Automation/UE.Automation.cs b/Engine/Source/Programs/AutomationTool/Gauntlet/Unreal/Automation/UE.Automation.cs
--- a/Engine/Source/Programs/AutomationTool/Gauntlet/Unreal/Automation/UE.Automation.cs
+++ b/Engine/Source/Programs/AutomationTool/Gauntlet/Unreal/Automation/UE.Automation.cs
@@ -451,58 +451,11 @@ namespace UE
 			}
 
 			string TelemetryDirectory = GetConfiguration().TelemetryDirectory;
-			if (!string.IsNullOrEmpty(TelemetryDirectory) && Directory.Exists(TelemetryDirectory))
+			if (Report != null && !string.IsNullOrEmpty(TelemetryDirectory) && Directory.Exists(TelemetryDirectory))
 			{
 				if (Report is ITelemetryReport Telemetry)
 				{
-					// Scan for csv files
-					foreach (string File in Directory.GetFiles(TelemetryDirectory))
-					{
-						if (Path.GetExtension(File) != ".csv")
-						{
-							continue;
-						}
-
-						List<Dictionary<string, string>> CSVData = null;
-						try
-						{
-							CSVData = CSVParser.Load(File);
-						}
-						catch (Exception Ex)
-						{
-							Log.Error("Telemetry - Failed to read CSV file '{0}'. {1}", File, Ex);
-						}
-
-						if (CSVData == null)
-						{
-							continue;
-						}
-						// Populate telemetry report
-						foreach (Dictionary<string, string> Row in CSVData)
-						{
-							string TestName;
-							Row.TryGetValue("TestName", out TestName);
-							string DataPoint;
-							Row.TryGetValue("DataPoint", out DataPoint);
-							string Measurement;
-							Row.TryGetValue("Measurement", out Measurement);
-							string Context;
-							Row.TryGetValue("Context", out Context);
-							if(string.IsNullOrEmpty(TestName) || string.IsNullOrEmpty(DataPoint) || string.IsNullOrEmpty(Measurement))
-							{
-								Log.Warning("Telemetry - Missing data in CSV file '{0}':\n TestName='{1}', DataPoint='{2}', Measurement='{3}", File, TestName, DataPoint, Measurement);
-								continue;
-							}
-							try
-							{
-								Telemetry.AddTelemetry(TestName, DataPoint, double.Parse(Measurement), Context);
-							}
-							catch(FormatException)
-							{
-								Log.Error("Telemetry - Failed to parse Measurement value('{0}') in CSV file '{1}'.", Measurement, File);
-							}
-						}
-					}
+					UnrealAutomationTelemetry.LoadOutputsIntoReport(TelemetryDirectory, Telemetry);
 				}
 				else
 				{
diff --git a/Engine/Source/Programs/AutomationTool/Gauntlet/Unreal/Utils/Gauntlet.UnrealTelemetry.cs b/Engine/Source/Programs/AutomationTool/Gauntlet/Unreal/Utils/Gauntlet.UnrealTelemetry.cs
new file mode 100644
--- /dev/null
+++ b/Engine/Source/Programs/AutomationTool/Gauntlet/Unreal/Utils/Gauntlet.UnrealTelemetry.cs
@@ -0,0 +1,72 @@
+// Copyright Epic Games, Inc. All Rights Reserved.
+
+using System;
+using System.IO;
+using System.Collections.Generic;
+
+
+namespace Gauntlet
+{
+	/// <summary>
+	/// Utility functions for Unreal Automation Telemetry
+	/// </summary>
+	static class UnrealAutomationTelemetry
+	{
+		/// <summary>
+		/// Gather csv telemetry outputs generated by UE automated tests and add them to a test report
+		/// </summary>
+		/// <param name="TelemetryDirectory"></param>
+		/// <param name="TelemetryReport"></param>
+		static public void LoadOutputsIntoReport(string TelemetryDirectory, ITelemetryReport TelemetryReport)
+		{
+			// Scan for csv files
+			foreach (string File in Directory.GetFiles(TelemetryDirectory))
+			{
+				if (Path.GetExtension(File) != ".csv")
+				{
+					continue;
+				}
+
+				List<Dictionary<string, string>> CSVData = null;
+				try
+				{
+					CSVData = CSVParser.Load(File);
+				}
+				catch (Exception Ex)
+				{
+					Log.Error("Telemetry - Failed to read CSV file '{0}'. {1}", File, Ex);
+				}
+
+				if (CSVData == null)
+				{
+					continue;
+				}
+				// Populate telemetry report
+				foreach (Dictionary<string, string> Row in CSVData)
+				{
+					string TestName;
+					Row.TryGetValue("TestName", out TestName);
+					string DataPoint;
+					Row.TryGetValue("DataPoint", out DataPoint);
+					string Measurement;
+					Row.TryGetValue("Measurement", out Measurement);
+					string Context;
+					Row.TryGetValue("Context", out Context);
+					if (string.IsNullOrEmpty(TestName) || string.IsNullOrEmpty(DataPoint) || string.IsNullOrEmpty(Measurement))
+					{
+						Log.Warning("Telemetry - Missing data in CSV file '{0}':\n TestName='{1}', DataPoint='{2}', Measurement='{3}", File, TestName, DataPoint, Measurement);
+						continue;
+					}
+					try
+					{
+						TelemetryReport.AddTelemetry(TestName, DataPoint, double.Parse(Measurement), Context);
+					}
+					catch (FormatException)
+					{
+						Log.Error("Telemetry - Failed to parse Measurement value('{0}') in CSV file '{1}'.", Measurement, File);
+					}
+				}
+			}
+		}
+	}
+}
\ No newline at end of file
