Editgrate cl 15716397
Fix Gauntlet App package deployement on Mac. Utils.SystemHelpers.IsNetworkPath returns always true on Mac. Because Path.GetRootPath gives the same root for any path on unix based system.
Revert to use the GetVolumeName and comparing the names.

#rnx

diff --git a/Engine/Source/Programs/AutomationTool/Gauntlet/Platform/Mac/Gauntlet.TargetDeviceMac.cs b/Engine/Source/Programs/AutomationTool/Gauntlet/Platform/Mac/Gauntlet.TargetDeviceMac.cs
--- a/Engine/Source/Programs/AutomationTool/Gauntlet/Platform/Mac/Gauntlet.TargetDeviceMac.cs
+++ b/Engine/Source/Programs/AutomationTool/Gauntlet/Platform/Mac/Gauntlet.TargetDeviceMac.cs
@@ -219,8 +219,11 @@ namespace Gauntlet
 
 			string OutBuildPath = InBuildPath;
 
+			string BuildVolume = GetVolumeName(OutBuildPath);
+			string LocalRoot = GetVolumeName(Environment.CurrentDirectory);
+
 			// Must be on our volume to run
-			if (Utils.SystemHelpers.IsNetworkPath(InBuildPath))
+			if (BuildVolume.Equals(LocalRoot, StringComparison.OrdinalIgnoreCase) == false)
 			{
 				string SubDir = string.IsNullOrEmpty(AppConfig.Sandbox) ? AppConfig.ProjectName : AppConfig.Sandbox;
 				string DestPath = Path.Combine(this.LocalCachePath, SubDir, AppConfig.ProcessType.ToString());
