Fix detection of Android Studio location

#ue5
#android

diff --git a/Engine/Source/Programs/AutomationTool/Android/AndroidPlatform.Automation.cs b/Engine/Source/Programs/AutomationTool/Android/AndroidPlatform.Automation.cs
--- a/Engine/Source/Programs/AutomationTool/Android/AndroidPlatform.Automation.cs
+++ b/Engine/Source/Programs/AutomationTool/Android/AndroidPlatform.Automation.cs
@@ -74,12 +74,17 @@ public class AndroidPlatform : Platform
 	private static string GetAndroidStudioExe()
 	{
 		string DefaultAndroidStudioInstallDir = Path.Combine(Environment.GetFolderPath(Environment.SpecialFolder.ProgramFiles), "Android", "Android Studio");
-		string RegValue = Microsoft.Win32.Registry.GetValue("HKEY_LOCAL_MACHINE\\SOFTWARE\\Android", "Path", null) as string;
+		string RegValue = Microsoft.Win32.Registry.GetValue("HKEY_LOCAL_MACHINE\\SOFTWARE\\Android Studio", "Path", null) as string;
 		string AndroidStudioInstallDir = RegValue == null ? DefaultAndroidStudioInstallDir : RegValue;
 		return Path.Combine(AndroidStudioInstallDir, "bin", "studio64.exe");
 	}
 	private static string GetSdkDir()
 	{
+		string AndroidHome = Environment.GetEnvironmentVariable("ANDROID_HOME");
+		if (!string.IsNullOrEmpty(AndroidHome) && Directory.Exists(AndroidHome))
+		{
+			return AndroidHome;
+		}	
 		string DefaultSdkDir = Path.Combine(Environment.GetFolderPath(Environment.SpecialFolder.LocalApplicationData), "Android", "Sdk");
 		string RegValue = Microsoft.Win32.Registry.GetValue("HKEY_LOCAL_MACHINE\\SOFTWARE\\Android", "SdkPath", null) as string;
 		return RegValue == null ? DefaultSdkDir : RegValue;
