Fix spaces in path in engine #rb jack.porter #jira UE-109630

diff --git a/Engine/Source/Programs/AutomationTool/IOS/IOSPlatform.Automation.cs b/Engine/Source/Programs/AutomationTool/IOS/IOSPlatform.Automation.cs
--- a/Engine/Source/Programs/AutomationTool/IOS/IOSPlatform.Automation.cs
+++ b/Engine/Source/Programs/AutomationTool/IOS/IOSPlatform.Automation.cs
@@ -1705,7 +1705,7 @@ public class IOSPlatform : Platform
 										+ " push " + CombinePaths(Params.BaseStageDirectory, PlatformName, SC.GetNonUFSDeployedManifestFileName(null));		
 				System.IO.File.WriteAllText(Directory.GetCurrentDirectory() + "\\CommandsToPush.txt", AllCommandsToPush);
 
-				string IdeviceFSArgs = "-b " + BundleIdentifier + " -x " + Directory.GetCurrentDirectory() + "\\CommandsToPush.txt -u " + Params.DeviceNames[0];
+				string IdeviceFSArgs = "-b " + "\"" + BundleIdentifier + " -x " + Directory.GetCurrentDirectory() + "\\CommandsToPush.txt -u " + "\"" + Params.DeviceNames[0];
 				IdeviceFSArgs = GetLibimobileDeviceNetworkedArgument(IdeviceFSArgs, Params.DeviceNames[0]);
 
 				Utils.RunLocalProcessAndReturnStdOut(DeviceFS, IdeviceFSArgs, out ExitCode, true);
@@ -1792,7 +1792,7 @@ public class IOSPlatform : Platform
 		System.IO.File.WriteAllText(Directory.GetCurrentDirectory() + "\\CommandsToPush.txt", AllCommandsToPush);
 		int ExitCode = 0;
 
-		string IdeviceFSArgs = "-u " + Params.DeviceNames[0] + " -b " + BundleIdentifier + " -x " + Directory.GetCurrentDirectory() + "\\CommandsToPush.txt";
+		string IdeviceFSArgs = "-u " + Params.DeviceNames[0] + " -b " + BundleIdentifier + " -x " + "\"" + Directory.GetCurrentDirectory() + "\\CommandsToPush.txt" + "\"";
 		IdeviceFSArgs = GetLibimobileDeviceNetworkedArgument(IdeviceFSArgs, Params.DeviceNames[0]);
 
 		Utils.RunLocalProcessAndReturnStdOut(DeviceFS, IdeviceFSArgs, out ExitCode, true);
