With new PLINK/PSCP, it seems PSCP requires a -P number, and doesnt default to 22 anymore. Update both just in case to assume 22

diff --git a/Engine/Source/Programs/AutomationTool/Linux/LinuxPlatform.Automation.cs b/Engine/Source/Programs/AutomationTool/Linux/LinuxPlatform.Automation.cs
--- a/Engine/Source/Programs/AutomationTool/Linux/LinuxPlatform.Automation.cs
+++ b/Engine/Source/Programs/AutomationTool/Linux/LinuxPlatform.Automation.cs
@@ -242,23 +242,23 @@ chmod +x {0}
 				if (!Params.IterativeDeploy)
 				{
 					// non-null input is essential, since RedirectStandardInput=true is needed for PLINK, see http://stackoverflow.com/questions/1910592/process-waitforexit-on-console-vs-windows-forms
-					RunAndLog(CmdEnv, PlinkPath, String.Format("-batch -l {0} -pw {1} {2} rm -rf {3} && mkdir -p {3}", Params.DeviceUsername, Params.DevicePassword, Params.DeviceNames[0], DestPath), Input: "");
+					RunAndLog(CmdEnv, PlinkPath, String.Format("-batch -P 22 -l {0} -pw {1} {2} rm -rf {3} && mkdir -p {3}", Params.DeviceUsername, Params.DevicePassword, Params.DeviceNames[0], DestPath), Input: "");
 				}
 				else
 				{
 					// still make sure that the path exists
-					RunAndLog(CmdEnv, PlinkPath, String.Format("-batch -l {0} -pw {1} {2} mkdir -p {3}", Params.DeviceUsername, Params.DevicePassword, Params.DeviceNames[0], DestPath), Input: "");
+					RunAndLog(CmdEnv, PlinkPath, String.Format("-batch -P 22 -l {0} -pw {1} {2} mkdir -p {3}", Params.DeviceUsername, Params.DevicePassword, Params.DeviceNames[0], DestPath), Input: "");
 				}
 
 				// copy the contents
-				RunAndLog(CmdEnv, PScpPath, String.Format("-batch -pw {0} -r \"{1}\" {2}", Params.DevicePassword, SourcePath, Params.DeviceUsername + "@" + DeviceAddress + ":" + DestPath));
+				RunAndLog(CmdEnv, PScpPath, String.Format("-batch -P 22 -pw {0} -r \"{1}\" {2}", Params.DevicePassword, SourcePath, Params.DeviceUsername + "@" + DeviceAddress + ":" + DestPath));
 
 				// copy the helper script
-				RunAndLog(CmdEnv, PScpPath, String.Format("-batch -pw {0} -r \"{1}\" {2}", Params.DevicePassword, ScriptFile, Params.DeviceUsername + "@" + DeviceAddress + ":" + DestPath));
+				RunAndLog(CmdEnv, PScpPath, String.Format("-batch -P 22 -pw {0} -r \"{1}\" {2}", Params.DevicePassword, ScriptFile, Params.DeviceUsername + "@" + DeviceAddress + ":" + DestPath));
 
 				string RemoteScriptFile = DestPath + "/" + LaunchOnHelperShellScriptName;
 				// non-null input is essential, since RedirectStandardInput=true is needed for PLINK, see http://stackoverflow.com/questions/1910592/process-waitforexit-on-console-vs-windows-forms
-				RunAndLog(CmdEnv, PlinkPath, String.Format(" -batch -l {0} -pw {1} {2} chmod +x {3}", Params.DeviceUsername, Params.DevicePassword, Params.DeviceNames[0], RemoteScriptFile), Input: "");
+				RunAndLog(CmdEnv, PlinkPath, String.Format(" -batch -P 22 -l {0} -pw {1} {2} chmod +x {3}", Params.DeviceUsername, Params.DevicePassword, Params.DeviceNames[0], RemoteScriptFile), Input: "");
 			}
 		}
 	}
@@ -329,7 +329,7 @@ chmod +x {0}
 			// try contacting the device(s) and cache the key(s)
 			foreach(string DeviceAddress in ProjParams.DeviceNames)
 			{
-				RunAndLog(CmdEnv, "cmd.exe", String.Format("/c \"echo y | {0} -ssh -t -l {1} -pw {2} {3} echo All Ok\"", PlinkPath, ProjParams.DeviceUsername, ProjParams.DevicePassword, DeviceAddress));
+				RunAndLog(CmdEnv, "cmd.exe", String.Format("/c \"echo y | {0} -P 22 -ssh -t -l {1} -pw {2} {3} echo All Ok\"", PlinkPath, ProjParams.DeviceUsername, ProjParams.DevicePassword, DeviceAddress));
 			}
 		}
 	}
@@ -344,7 +344,7 @@ chmod +x {0}
 				string RemoteBasePath = "./" + Params.ShortProjectName;
 				string RemotePathToBinary = RemoteBasePath + "/" + LaunchOnHelperShellScriptName;
 				// non-null input is essential, since RedirectStandardInput=true is needed for PLINK, see http://stackoverflow.com/questions/1910592/process-waitforexit-on-console-vs-windows-forms
-				Result = Run(PlinkPath, String.Format("-batch -ssh -t -l {0} -pw {1} {2}  {3} {4}", Params.DeviceUsername, Params.DevicePassword, Params.DeviceNames[0], RemotePathToBinary, ClientCmdLine), "");
+				Result = Run(PlinkPath, String.Format("-batch -P 22 -ssh -t -l {0} -pw {1} {2}  {3} {4}", Params.DeviceUsername, Params.DevicePassword, Params.DeviceNames[0], RemotePathToBinary, ClientCmdLine), "");
 			}
 			return Result;
 		}
