Add platform option to disable the Manifest_*_.txt files in staged builds
 - Some platforms don't make use of these files, so their generation is unnecessary.

diff --git a/Engine/Source/Programs/AutomationTool/AutomationUtils/Platform.cs b/Engine/Source/Programs/AutomationTool/AutomationUtils/Platform.cs
--- a/Engine/Source/Programs/AutomationTool/AutomationUtils/Platform.cs
+++ b/Engine/Source/Programs/AutomationTool/AutomationUtils/Platform.cs
@@ -656,6 +656,14 @@ namespace AutomationTool
 			get { return false; }
 		}
 
+		/// <summary>
+		/// Returns whether the platform requires the Manifest_*_.txt files to be copied to the staged directory.
+		/// </summary>
+		public virtual bool RequiresManifestFiles
+		{
+			get { return true; }
+		}
+
 		public virtual HashSet<StagedFileReference> GetFilesForCRCCheck()
 		{
 			string CmdLine = "UE4CommandLine.txt";
diff --git a/Engine/Source/Programs/AutomationTool/Scripts/CopyBuildToStagingDirectory.Automation.cs b/Engine/Source/Programs/AutomationTool/Scripts/CopyBuildToStagingDirectory.Automation.cs
--- a/Engine/Source/Programs/AutomationTool/Scripts/CopyBuildToStagingDirectory.Automation.cs
+++ b/Engine/Source/Programs/AutomationTool/Scripts/CopyBuildToStagingDirectory.Automation.cs
@@ -1580,34 +1580,39 @@ public partial class Project : CommandUtils
 		}
 	}
 
-	public static void CopyManifestFilesToStageDir(Dictionary<StagedFileReference, FileReference> Mapping, DirectoryReference StageDir, DirectoryReference ManifestDir, string ManifestName, HashSet<StagedFileReference> CRCFiles, string PlatformName, List<string> IniKeyBlacklist, List<string> IniSectionBlacklist)
+	private static void CopyManifestFilesToStageDir(DeploymentContext SC, Dictionary<StagedFileReference, FileReference> Mapping, string ManifestName)
 	{
-		LogInformation("Copying {0} to staging directory: {1}", ManifestName, StageDir);
-		FileReference ManifestPath = null;
-		string ManifestFile = "";
-		if (!String.IsNullOrEmpty(ManifestName))
-		{
-			ManifestFile = "Manifest_" + ManifestName + "_" + PlatformName + ".txt";
-			ManifestPath = FileReference.Combine(ManifestDir, ManifestFile);
-			DeleteFile(ManifestPath.FullName);
-		}
+		LogInformation("Copying {0} to staging directory: {1}", ManifestName, SC.StageDirectory);
 		foreach (KeyValuePair<StagedFileReference, FileReference> Pair in Mapping)
 		{
 			FileReference Src = Pair.Value;
-			FileReference Dest = FileReference.Combine(StageDir, Pair.Key.Name);
+			FileReference Dest = FileReference.Combine(SC.StageDirectory, Pair.Key.Name);
 			if (Src != Dest)  // special case for things created in the staging directory, like the pak file
 			{
-				CopyFileIncremental(Src, Dest, IniKeyBlacklist:IniKeyBlacklist, IniSectionBlacklist:IniSectionBlacklist);
+				CopyFileIncremental(Src, Dest, IniKeyBlacklist: SC.IniKeyBlacklist, IniSectionBlacklist: SC.IniSectionBlacklist);
 			}
 		}
-		if (ManifestPath != null && Mapping.Count > 0)
+
+		if (Mapping.Count > 0)
 		{
-			DumpTargetManifest(Mapping, ManifestPath, StageDir, CRCFiles);
-			if (!FileReference.Exists(ManifestPath))
+			var ManifestPathStaged = FileReference.Combine(SC.DebugStageDirectory, $"Manifest_{ManifestName}_{SC.StageTargetPlatform.PlatformType}.txt");
+			var ManifestPathLog = FileReference.Combine(new DirectoryReference(CmdEnv.LogFolder), ManifestPathStaged.GetFileName());
+
+			DeleteFile(ManifestPathStaged.FullName);
+			DeleteFile(ManifestPathLog.FullName);
+
+			DumpTargetManifest(Mapping, ManifestPathLog, SC.StageDirectory, SC.StageTargetPlatform.GetFilesForCRCCheck());
+
+			if (!FileReference.Exists(ManifestPathLog))
+			{
+				throw new AutomationException("Failed to write manifest {0}", ManifestPathLog);
+			}
+
+			if (SC.StageTargetPlatform.RequiresManifestFiles)
 			{
-				throw new AutomationException("Failed to write manifest {0}", ManifestPath);
+				// Also copy the manifest file to the staged directory
+				CopyFile(ManifestPathLog.FullName, ManifestPathStaged.FullName);
 			}
-			CopyFile(ManifestPath.FullName, CombinePaths(CmdEnv.LogFolder, ManifestFile));
 		}
 	}
 
@@ -1636,7 +1641,7 @@ public partial class Project : CommandUtils
 
 	public static void CopyUsingStagingManifest(ProjectParams Params, DeploymentContext SC)
 	{
-		CopyManifestFilesToStageDir(SC.FilesToStage.NonUFSFiles, SC.StageDirectory, SC.DebugStageDirectory, "NonUFSFiles", SC.StageTargetPlatform.GetFilesForCRCCheck(), SC.StageTargetPlatform.PlatformType.ToString(), SC.IniKeyBlacklist, SC.IniSectionBlacklist);
+		CopyManifestFilesToStageDir(SC, SC.FilesToStage.NonUFSFiles, "NonUFSFiles");
 
 		bool bStageUnrealFileSystemFiles = !Params.CookOnTheFly && !Params.UsePak(SC.StageTargetPlatform) && !Params.FileServer;
 		if (bStageUnrealFileSystemFiles)
@@ -1654,14 +1659,14 @@ public partial class Project : CommandUtils
 					throw new AutomationException("File '{0}' is set to be staged from '{1}' for project and '{2}' for crash reporter", Pair.Key, ExistingLocation, Pair.Value);
 				}
 			}
-			CopyManifestFilesToStageDir(UFSFiles, SC.StageDirectory, SC.DebugStageDirectory, "UFSFiles", SC.StageTargetPlatform.GetFilesForCRCCheck(), SC.StageTargetPlatform.PlatformType.ToString(), SC.IniKeyBlacklist, SC.IniSectionBlacklist);
+			CopyManifestFilesToStageDir(SC, UFSFiles, "UFSFiles");
 		}
 
 		// Copy debug files last
 		// they do not respect the DeployLowerCaseFilenames() setting, but if copied to a case-insensitive staging directory first they determine the casing for outer directories (like Engine/Content) 
 		if (!Params.NoDebugInfo)
 		{
-			CopyManifestFilesToStageDir(SC.FilesToStage.NonUFSDebugFiles, SC.DebugStageDirectory, SC.DebugStageDirectory, "DebugFiles", SC.StageTargetPlatform.GetFilesForCRCCheck(), SC.StageTargetPlatform.PlatformType.ToString(), SC.IniKeyBlacklist, SC.IniSectionBlacklist);
+			CopyManifestFilesToStageDir(SC, SC.FilesToStage.NonUFSDebugFiles, "DebugFiles");
 		}
 	}
 
@@ -4147,7 +4152,7 @@ public partial class Project : CommandUtils
 				{
 					ApplyStagingManifest(Params, SC);
 
-					if (Params.Deploy)
+					if (Params.Deploy && SC.StageTargetPlatform.RequiresManifestFiles)
 					{
 						List<string> UFSManifests;
 						List<string> NonUFSManifests;
