Disable shader distribution to fix timing out jobs.
We also don't want to distribute these workloads for data consistency reasons.

diff --git a/Engine/Source/Programs/AutomationTool/Gauntlet/Unreal/Editor/EditorTests.cs b/Engine/Source/Programs/AutomationTool/Gauntlet/Unreal/Editor/EditorTests.cs
--- a/Engine/Source/Programs/AutomationTool/Gauntlet/Unreal/Editor/EditorTests.cs
+++ b/Engine/Source/Programs/AutomationTool/Gauntlet/Unreal/Editor/EditorTests.cs
@@ -55,6 +55,12 @@ namespace UnrealEditor
 		[AutoParam]
 		public string GameInstanceLostTimerSeconds = string.Empty;
 
+		/// <summary>
+		/// Disable distribution of shader builds (but use worker processes still)
+		/// </summary>
+		[AutoParam]
+		public bool NoShaderDistrib = false;
+
 		/// <summary>
 		/// Applies these options to the provided app config
 		/// </summary>
@@ -90,6 +96,11 @@ namespace UnrealEditor
 			{
 				AppConfig.CommandLineParams.Add(string.Format("ini:Engine:[/Script/AutomationController.AutomationControllerSettings]:GameInstanceLostTimerSeconds={0}", GameInstanceLostTimerSeconds));
 			}
+
+			if (NoShaderDistrib)
+			{
+				AppConfig.CommandLineParams.Add("-noxgeshadercompile");
+			}
 		}
 	}
 
