Fix various shortcomings with archiving builds via BuildCookRun

- Fixed issue where -archive did not work with -skippackage for several platforms.

- Modified behaviour of -archive -archivedirectory so <path> will only become <path>/<CookedPlatform> *IF* <path> does not already contain CookedPlatform or PlatformName

Both of these issues allow greater flexibility in build scripts.

Specifically the first allows -package and -archive to be performed in separate steps which is necessary for scripts that want to "do stuff" before archiving or have different per-platform logic before archiving.

The second gives scripts flexibility in selecing directory structures without adding unncessary components.

E.g. a script can now archive development builds to <BuildDir>/<Project>/PS4/Development/<Project>.pkg without it becoming <BuildDir>/<Project>/PS4/Development/PS4/<Project>.pkg

(Checked in prior to feedback due to being a prerequisite for some Frosty tests)

diff --git a/Engine/Source/Programs/AutomationTool/AutomationUtils/DeploymentContext.cs b/Engine/Source/Programs/AutomationTool/AutomationUtils/DeploymentContext.cs
--- a/Engine/Source/Programs/AutomationTool/AutomationUtils/DeploymentContext.cs
+++ b/Engine/Source/Programs/AutomationTool/AutomationUtils/DeploymentContext.cs
@@ -408,9 +408,21 @@ public class DeploymentContext //: ProjectParams
 			DebugStageDirectory = InSeparateDebugStageDirectory? DirectoryReference.Combine(BaseStageDirectory, FinalCookPlatform + "Debug") : StageDirectory;
 		}
 
-		if(BaseArchiveDirectory != null)
+		if (BaseArchiveDirectory != null)
 		{
-			ArchiveDirectory = DirectoryReference.Combine(BaseArchiveDirectory, FinalCookPlatform);
+			// If the user specifies a path that contains the platform or cooked platform names, don't append to it.
+			// This allows more control for scripts respect to putting packages in things like <path>/PS4/Development/Pkg.pkg, <path>/PS4/Test/Pkg.pkg
+			string PlatformName = StageTargetPlatform.GetStagePlatforms().FirstOrDefault().ToString();
+			IEnumerable<string> PathComponents = new DirectoryInfo(BaseArchiveDirectory.FullName).FullName.Split(Path.DirectorySeparatorChar, StringSplitOptions.RemoveEmptyEntries);
+
+			if (!PathComponents.Any(C => C.StartsWith(FinalCookPlatform, StringComparison.OrdinalIgnoreCase) || C.StartsWith(PlatformName, StringComparison.OrdinalIgnoreCase)))
+			{
+				ArchiveDirectory = DirectoryReference.Combine(BaseArchiveDirectory, FinalCookPlatform);
+			}
+			else
+			{
+				ArchiveDirectory = BaseArchiveDirectory;
+			}
 		}
 
 		if (!FileReference.Exists(RawProjectPath))
