Add plugin directories as search dirs for dynamic dependency resolution

When AutomationTool is built by RunUAT (or otherwise outside of UE5.sln), plugin dependencies are not copied to the AutomationTool.exe output dir. This change will ensure plugin output dirs will be searched for dependent assemblies.

diff --git a/Engine/Source/Programs/AutomationTool/AutomationUtils/ScriptCompiler.cs b/Engine/Source/Programs/AutomationTool/AutomationUtils/ScriptCompiler.cs
--- a/Engine/Source/Programs/AutomationTool/AutomationUtils/ScriptCompiler.cs
+++ b/Engine/Source/Programs/AutomationTool/AutomationUtils/ScriptCompiler.cs
@@ -461,6 +461,8 @@ namespace AutomationTool
 				try
 				{
 					AssemblyUtils.AddFileToAssemblyCache(AssemblyLocation.FullName);
+					// Add a resolver for the Assembly directory, so that its dependencies may be found alongside it
+					AssemblyUtils.InstallRecursiveAssemblyResolver(AssemblyLocation.Directory.FullName);
 					Assembly Assembly = AppDomain.CurrentDomain.Load(AssemblyName.GetAssemblyName(AssemblyLocation.FullName));
 					Assemblies.Add(Assembly);
 				}
