Harden code for finding shared file manifests.

#rnx

diff --git a/Engine/Source/Programs/AutomationTool/BuildGraph/TempStorage.cs b/Engine/Source/Programs/AutomationTool/BuildGraph/TempStorage.cs
--- a/Engine/Source/Programs/AutomationTool/BuildGraph/TempStorage.cs
+++ b/Engine/Source/Programs/AutomationTool/BuildGraph/TempStorage.cs
@@ -719,7 +719,7 @@ namespace AutomationTool
 		/// <returns>The set of files</returns>
 		public TempStorageFileList ReadFileList(string NodeName, string TagName)
 		{
-			TempStorageFileList FileList;
+			TempStorageFileList FileList = null;
 
 			// Try to read the tag set from the local directory
 			FileReference LocalFileListLocation = GetTaggedFileListLocation(LocalDir, NodeName, TagName);
@@ -736,20 +736,45 @@ namespace AutomationTool
 					throw new AutomationException("Missing local file list - {0}", LocalFileListLocation.FullName);
 				}
 
-				// Make sure the manifest exists
-				FileReference SharedFileListLocation = GetTaggedFileListLocation(SharedDir, NodeName, TagName);
-				if(!FileReference.Exists(SharedFileListLocation))
+				// Make sure the manifest exists. Try up to 5 times with a 5s wait between to harden against network hiccups.
+				int Attempts = 5;
+				FileReference SharedFileListLocation;
+				SharedFileListLocation = GetTaggedFileListLocation(SharedDir, NodeName, TagName);
+				while (Attempts-- > 0)
 				{
-					throw new AutomationException("Missing local or shared file list - {0}", SharedFileListLocation.FullName);
-				}
+					if (!FileReference.Exists(SharedFileListLocation))
+					{
+						if (Attempts == 0)
+						{
+							throw new AutomationException("Missing local or shared file list - {0}", SharedFileListLocation.FullName);
+						}
+
+						Thread.Sleep(TimeSpan.FromSeconds(5));
+						continue;
+					}
+
+					try
+					{
+						// Read the shared manifest
+						CommandUtils.LogInformation("Copying shared tag set from {0} to {1}", SharedFileListLocation.FullName, LocalFileListLocation.FullName);
+						FileList = TempStorageFileList.Load(SharedFileListLocation);
+					}
+					catch (IOException)
+					{
+						if (Attempts == 0)
+						{
+							throw new AutomationException("Local or shared file list {0} was found but failed to be read", SharedFileListLocation.FullName);
+						}
+
+						continue;
+					}
 
-				// Read the shared manifest
-				CommandUtils.LogInformation("Copying shared tag set from {0} to {1}", SharedFileListLocation.FullName, LocalFileListLocation.FullName);
-				FileList = TempStorageFileList.Load(SharedFileListLocation);
+					break;
+				}
 
 				// Save the manifest locally
 				DirectoryReference.CreateDirectory(LocalFileListLocation.Directory);
-				FileList.Save(LocalFileListLocation);
+				FileList?.Save(LocalFileListLocation);
 			}
 			return FileList;
 		}
