Hololens: Package missing IoStore files

Plus a little cleanup. Technically, IoStore can run w/o pak.

diff --git a/Engine/Source/Programs/AutomationTool/HoloLens/HoloLensPlatform.Automation.cs b/Engine/Source/Programs/AutomationTool/HoloLens/HoloLensPlatform.Automation.cs
--- a/Engine/Source/Programs/AutomationTool/HoloLens/HoloLensPlatform.Automation.cs
+++ b/Engine/Source/Programs/AutomationTool/HoloLens/HoloLensPlatform.Automation.cs
@@ -696,6 +696,16 @@ namespace HoloLens.Automation
 			}
 		}
 
+		private List<string> GetContainerFileExtensions()
+		{
+			return new List<string> { "*.pak", "*.ucas", "*.utoc" };
+		}
+
+		private bool IsUsingContainers(ProjectParams Params)
+		{
+			return (Params.UsePak(this) || (Params.IoStore && !Params.SkipIoStore));
+		}
+
 		private void PackagePakFiles(ProjectParams Params, DeploymentContext SC, string OutputNameBase)
 		{
 			string IntermediateDirectory = Path.Combine(SC.ProjectRoot.FullName, "Intermediate", "Deploy", "neutral");
@@ -710,9 +720,12 @@ namespace HoloLens.Automation
 
 			string OutputAppX = Path.Combine(SC.StageDirectory.FullName, OutputName + Extension);
 
-			if(Params.UsePak(this))
+			if(IsUsingContainers(Params))
 			{
-				FillMapfile(SC.StageDirectory, SC.StageDirectory.FullName, "*.pak", AppXRecipeBuiltFiles);
+				foreach (string ContainerExtension in GetContainerFileExtensions())
+				{
+					FillMapfile(SC.StageDirectory, SC.StageDirectory.FullName, ContainerExtension, AppXRecipeBuiltFiles);
+				}
 			}
 			else
 			{
@@ -759,9 +772,12 @@ namespace HoloLens.Automation
 
 					string OutputAppX = Path.Combine(SC.StageDirectory.FullName, Product.Path.GetFileNameWithoutExtension() + Extension);
 
-					if (Params.UsePak(this))
+					if(IsUsingContainers(Params))
 					{
-						FillMapfile(SC.StageDirectory, SC.StageDirectory.FullName, "*.pak", AppXRecipeBuiltFiles);
+						foreach (string ContainerExtension in GetContainerFileExtensions())
+						{
+							FillMapfile(SC.StageDirectory, SC.StageDirectory.FullName, ContainerExtension, AppXRecipeBuiltFiles);
+						}
 					}
 					else
 					{
@@ -825,9 +841,12 @@ namespace HoloLens.Automation
 
 				string OutputAppX = Path.Combine(SC.StageDirectory.FullName, OutputName + Extension);
 
-				if (Params.UsePak(this) && !Params.Run)
+				if (IsUsingContainers(Params) && !Params.Run)
 				{
-					FillMapfile(SC.StageDirectory, SC.StageDirectory.FullName, "*.pak", Dict);
+					foreach (string ContainerExtension in GetContainerFileExtensions())
+					{
+						FillMapfile(SC.StageDirectory, SC.StageDirectory.FullName, ContainerExtension, Dict);
+					}
 				}
 				else
 				{
diff --git a/Engine/Source/Programs/UnrealBuildTool/Platform/HoloLens/HoloLensDeploy.cs b/Engine/Source/Programs/UnrealBuildTool/Platform/HoloLens/HoloLensDeploy.cs
--- a/Engine/Source/Programs/UnrealBuildTool/Platform/HoloLens/HoloLensDeploy.cs
+++ b/Engine/Source/Programs/UnrealBuildTool/Platform/HoloLens/HoloLensDeploy.cs
@@ -212,25 +212,6 @@ namespace UnrealBuildTool
 			return true;
 		}
 
-		private static void FillMapfile(string DirectoryName, string SearchPath, string Mask, StringBuilder AppXRecipeBuiltFiles)
-		{
-			string StageDir = DirectoryName;
-			if (!StageDir.EndsWith("\\") && !StageDir.EndsWith("/"))
-			{
-				StageDir += Path.DirectorySeparatorChar;
-			}
-			Uri StageDirUri = new Uri(StageDir);
-			foreach (var pf in Directory.GetFiles(SearchPath, Mask, SearchOption.AllDirectories))
-			{
-				var FileUri = new Uri(pf);
-				var relativeUri = StageDirUri.MakeRelativeUri(FileUri);
-				var relativePath = Uri.UnescapeDataString(relativeUri.ToString());
-				var newPath = relativePath.Replace('/', Path.DirectorySeparatorChar);
-
-				AppXRecipeBuiltFiles.AppendLine(String.Format("\"{0}\"\t\"{1}\"", pf, newPath));
-			}
-		}
-
 		private void MakePackage(TargetReceipt Receipt, TargetReceipt NewReceipt, WindowsArchitecture Architecture, List<string> UpdatedFiles)
 		{
 			string OutputName = String.Format("{0}_{1}_{2}_{3}", Receipt.TargetName, Receipt.Platform, Receipt.Configuration, WindowsExports.GetArchitectureSubpath(Architecture));
