BuildGraph: Ignore the PreflightChange argument, even if the BuildGraph script does not declare it as an option.

diff --git a/Engine/Source/Programs/AutomationTool/BuildGraph/Script.cs b/Engine/Source/Programs/AutomationTool/BuildGraph/Script.cs
--- a/Engine/Source/Programs/AutomationTool/BuildGraph/Script.cs
+++ b/Engine/Source/Programs/AutomationTool/BuildGraph/Script.cs
@@ -296,9 +296,15 @@ namespace AutomationTool
 			}
 
 			// Make sure all the arguments were valid
-			foreach(string InvalidArgumentName in Arguments.Keys.Except(Reader.Graph.Options.Select(x => x.Name), StringComparer.InvariantCultureIgnoreCase))
+			HashSet<string> ValidArgumentNames = new HashSet<string>(Reader.Graph.Options.Select(x => x.Name), StringComparer.OrdinalIgnoreCase);
+			ValidArgumentNames.Add("PreflightChange");
+
+			foreach(string ArgumentName in Arguments.Keys)
 			{
-				CommandUtils.LogWarning("Unknown argument '{0}' for '{1}'", InvalidArgumentName, File.FullName);
+				if (!ValidArgumentNames.Contains(ArgumentName))
+				{
+					CommandUtils.LogWarning("Unknown argument '{0}' for '{1}'", ArgumentName, File.FullName);
+				}
 			}
 
 			// Return the constructed graph
