Fix false positives for Gauntlet tests that use Gauntlet.UnrealTestNode.  This fixes tests passing even though they exited with a non passing errorcode.

diff --git a/Engine/Source/Programs/AutomationTool/Gauntlet/Unreal/Base/Gauntlet.UnrealTestNode.cs b/Engine/Source/Programs/AutomationTool/Gauntlet/Unreal/Base/Gauntlet.UnrealTestNode.cs
--- a/Engine/Source/Programs/AutomationTool/Gauntlet/Unreal/Base/Gauntlet.UnrealTestNode.cs
+++ b/Engine/Source/Programs/AutomationTool/Gauntlet/Unreal/Base/Gauntlet.UnrealTestNode.cs
@@ -1232,8 +1232,7 @@ namespace Gauntlet
 					ExitReason = string.Format("Tests exited with error code {0}", LogSummary.TestExitCode);
 				}
 
-				// tests failed but the process didn't
-				ExitCode = 0;
+				ExitCode = LogSummary.TestExitCode;
 			}
 			else if (LogSummary.EngineInitialized == false)
 			{
