Establish Android port forwarding for COTF network fileserver and Unreal Insights
Change COTF server order to specify localhost first, speeding up connections for platforms that use port forwarding

diff --git a/Engine/Source/Developer/Android/AndroidDeviceDetection/Private/AndroidDeviceDetectionModule.cpp b/Engine/Source/Developer/Android/AndroidDeviceDetection/Private/AndroidDeviceDetectionModule.cpp
--- a/Engine/Source/Developer/Android/AndroidDeviceDetection/Private/AndroidDeviceDetectionModule.cpp
+++ b/Engine/Source/Developer/Android/AndroidDeviceDetection/Private/AndroidDeviceDetectionModule.cpp
@@ -488,6 +488,20 @@ private:
 
 					TcpMessagingModule->AddOutgoingConnection(FString::Printf(TEXT("127.0.0.1:%d"), NewDeviceInfo.HostMessageBusPort));
 				}
+
+				// Add reverse port forwarding
+				uint16 ReversePortMappings[] = {
+					41899,	// Network file server, DEFAULT_TCP_FILE_SERVING_PORT in NetworkMessage.h
+					1980,	// Unreal Insights data collection, TraceInsightsModule.cpp
+					0		// end of list
+					};
+
+				for (int32 Idx=0; ReversePortMappings[Idx] > 0; Idx++)
+				{
+					FString DeviceCommand = FString::Printf(TEXT("-s %s reverse tcp:%d tcp:%d"), *NewDeviceInfo.SerialNumber, ReversePortMappings[Idx], ReversePortMappings[Idx]);
+					// It doesn't really matter if a mapping already exists. There is no listening local port so no contention between multiple editor instances
+					ExecuteAdbCommand(*DeviceCommand, nullptr, nullptr);
+				}
 			}
 
 			// add the device to the map
diff --git a/Engine/Source/Programs/AutomationTool/Scripts/CopyBuildToStagingDirectory.Automation.cs b/Engine/Source/Programs/AutomationTool/Scripts/CopyBuildToStagingDirectory.Automation.cs
--- a/Engine/Source/Programs/AutomationTool/Scripts/CopyBuildToStagingDirectory.Automation.cs
+++ b/Engine/Source/Programs/AutomationTool/Scripts/CopyBuildToStagingDirectory.Automation.cs
@@ -3641,7 +3641,37 @@ public partial class Project : CommandUtils
 			if (Params.CookOnTheFly || Params.FileServer)
 			{
 				FileHostParams += "-filehostip=";
-				bool FirstParam = true;
+				// add localhost first for platforms using redirection
+				const string LocalHost = "127.0.0.1";
+				if (!IsNullOrEmpty(Params.Port))
+				{
+					bool FirstParam = true;
+					foreach (var Port in Params.Port)
+					{
+						if (!FirstParam)
+						{
+							FileHostParams += "+";
+						}
+						FirstParam = false;
+						string[] PortProtocol = Port.Split(new char[] { ':' });
+						if (PortProtocol.Length > 1)
+						{
+							FileHostParams += String.Format("{0}://{1}:{2}", PortProtocol[0], LocalHost, PortProtocol[1]);
+						}
+						else
+						{
+							FileHostParams += LocalHost;
+							FileHostParams += ":";
+							FileHostParams += Params.Port;
+						}
+
+					}
+				}
+				else
+				{
+					// use default port
+					FileHostParams += LocalHost;
+				}
 				if (UnrealBuildTool.BuildHostPlatform.Current.Platform == UnrealTargetPlatform.Mac)
 				{
 					NetworkInterface[] Interfaces = NetworkInterface.GetAllNetworkInterfaces();
@@ -3658,11 +3688,7 @@ public partial class Project : CommandUtils
 									{
 										foreach (var Port in Params.Port)
 										{
-											if (!FirstParam)
-											{
-												FileHostParams += "+";
-											}
-											FirstParam = false;
+											FileHostParams += "+";
 											string[] PortProtocol = Port.Split(new char[] { ':' });
 											if (PortProtocol.Length > 1)
 											{
@@ -3679,12 +3705,7 @@ public partial class Project : CommandUtils
 									}
 									else
 									{
-										if (!FirstParam)
-										{
-											FileHostParams += "+";
-										}
-										FirstParam = false;
-
+										FileHostParams += "+";
 										// use default port
 										FileHostParams += IP.UnicastAddresses[Index].Address.ToString();
 									}
@@ -3710,11 +3731,7 @@ public partial class Project : CommandUtils
 									{
 										foreach (var Port in Params.Port)
 										{
-											if (!FirstParam)
-											{
-												FileHostParams += "+";
-											}
-											FirstParam = false;
+											FileHostParams += "+";
 											string[] PortProtocol = Port.Split(new char[] { ':' });
 											if (PortProtocol.Length > 1)
 											{
@@ -3730,56 +3747,15 @@ public partial class Project : CommandUtils
 									}
 									else
 									{
-										if (!FirstParam)
-										{
-											FileHostParams += "+";
-										}
-										FirstParam = false;
-
+										FileHostParams += "+";
 										// use default port
 										FileHostParams += IP.UnicastAddresses[Index].Address.ToString();
 									}
-
 								}
 							}
 						}
 					}
 				}
-				const string LocalHost = "127.0.0.1";
-				if (!IsNullOrEmpty(Params.Port))
-				{
-					foreach (var Port in Params.Port)
-					{
-						if (!FirstParam)
-						{
-							FileHostParams += "+";
-						}
-						FirstParam = false;
-						string[] PortProtocol = Port.Split(new char[] { ':' });
-						if (PortProtocol.Length > 1)
-						{
-							FileHostParams += String.Format("{0}://{1}:{2}", PortProtocol[0], LocalHost, PortProtocol[1]);
-						}
-						else
-						{
-							FileHostParams += LocalHost;
-							FileHostParams += ":";
-							FileHostParams += Params.Port;
-						}
-
-					}
-				}
-				else
-				{
-					if (!FirstParam)
-					{
-						FileHostParams += "+";
-					}
-					FirstParam = false;
-
-					// use default port
-					FileHostParams += LocalHost;
-				}
 				FileHostParams += " ";
 			}
 
diff --git a/Engine/Source/Programs/AutomationTool/Scripts/RunProjectCommand.Automation.cs b/Engine/Source/Programs/AutomationTool/Scripts/RunProjectCommand.Automation.cs
--- a/Engine/Source/Programs/AutomationTool/Scripts/RunProjectCommand.Automation.cs
+++ b/Engine/Source/Programs/AutomationTool/Scripts/RunProjectCommand.Automation.cs
@@ -626,7 +626,38 @@ public partial class Project : CommandUtils
 			if (Params.CookOnTheFly || Params.FileServer)
 			{
 				TempCmdLine += "-filehostip=";
-				bool FirstParam = true;
+				// add localhost first for platforms using redirection
+				const string LocalHost = "127.0.0.1";
+				if (!IsNullOrEmpty(Params.Port))
+				{
+					bool FirstParam = true;
+					foreach (var Port in Params.Port)
+					{
+						if (!FirstParam)
+						{
+							TempCmdLine += "+";
+						}
+						FirstParam = false;
+						string[] PortProtocol = Port.Split(new char[] { ':' });
+						if (PortProtocol.Length > 1)
+						{
+							TempCmdLine += String.Format("{0}://{1}:{2}", PortProtocol[0], LocalHost, PortProtocol[1]);
+						}
+						else
+						{
+							TempCmdLine += LocalHost;
+							TempCmdLine += ":";
+							TempCmdLine += Params.Port;
+						}
+
+					}
+				}
+				else
+				{
+					// use default port
+					TempCmdLine += LocalHost;
+				}
+
 				if (UnrealBuildTool.BuildHostPlatform.Current.Platform == UnrealTargetPlatform.Mac)
 				{
 					NetworkInterface[] Interfaces = NetworkInterface.GetAllNetworkInterfaces();
@@ -643,11 +674,7 @@ public partial class Project : CommandUtils
 									{
 										foreach (var Port in Params.Port)
 										{
-											if (!FirstParam)
-											{
-												TempCmdLine += "+";
-											}
-											FirstParam = false;
+											TempCmdLine += "+";
 											string[] PortProtocol = Port.Split(new char[] { ':' });
 											if (PortProtocol.Length > 1)
 											{
@@ -663,12 +690,7 @@ public partial class Project : CommandUtils
 									}
                                     else
                                     {
-										if (!FirstParam)
-										{
-											TempCmdLine += "+";
-										}
-										FirstParam = false;
-										
+										TempCmdLine += "+";
 										// use default port
                                         TempCmdLine += IP.UnicastAddresses[Index].Address.ToString();
                                     }
@@ -693,11 +715,7 @@ public partial class Project : CommandUtils
 									{
 										foreach (var Port in Params.Port)
 										{
-											if (!FirstParam)
-											{
-												TempCmdLine += "+";
-											}
-											FirstParam = false;
+											TempCmdLine += "+";
 											string[] PortProtocol = Port.Split(new char[] { ':' });
 											if (PortProtocol.Length > 1)
 											{
@@ -713,12 +731,7 @@ public partial class Project : CommandUtils
 									}
                                     else
                                     {
-										if (!FirstParam)
-										{
-											TempCmdLine += "+";
-										}
-										FirstParam = false;
-										
+										TempCmdLine += "+";
 										// use default port
                                         TempCmdLine += IP.UnicastAddresses[Index].Address.ToString();
                                     }
@@ -727,43 +740,6 @@ public partial class Project : CommandUtils
 						}
 					}
 				}
-
-				const string LocalHost = "127.0.0.1";
-
-				if (!IsNullOrEmpty(Params.Port))
-				{
-					foreach (var Port in Params.Port)
-					{
-						if (!FirstParam)
-						{
-							TempCmdLine += "+";
-						}
-						FirstParam = false;
-						string[] PortProtocol = Port.Split(new char[] { ':' });
-						if (PortProtocol.Length > 1)
-						{
-							TempCmdLine += String.Format("{0}://{1}:{2}", PortProtocol[0], LocalHost, PortProtocol[1]);
-						}
-						else
-						{
-							TempCmdLine += LocalHost;
-							TempCmdLine += ":";
-							TempCmdLine += Params.Port;
-						}
-
-					}
-				}
-                else
-                {
-					if (!FirstParam)
-					{
-						TempCmdLine += "+";
-					}
-					FirstParam = false;
-					
-					// use default port
-                    TempCmdLine += LocalHost;
-                }
 				TempCmdLine += " ";
 
 				if (Params.CookOnTheFlyStreaming)
