Add a command on some DDC ENV pointing out the values for Linux/Mac so its easier to grep/FIF if looking for them

diff --git a/Engine/Source/Programs/AutomationTool/Scripts/BenchmarkBuild/BenchmarkTaskBase.cs b/Engine/Source/Programs/AutomationTool/Scripts/BenchmarkBuild/BenchmarkTaskBase.cs
--- a/Engine/Source/Programs/AutomationTool/Scripts/BenchmarkBuild/BenchmarkTaskBase.cs
+++ b/Engine/Source/Programs/AutomationTool/Scripts/BenchmarkBuild/BenchmarkTaskBase.cs
@@ -246,6 +246,7 @@ namespace AutomationTool.Benchmark
 				// We put our temp DDC paths in here
 				DirectoryReference BasePath = DirectoryReference.Combine(CommandUtils.EngineDirectory, "BenchmarkDDC");
 
+				// For Linux and Mac the ENV vars will be UE_BootDataCachePath and UE_LocalDataCachePath
 				IEnumerable<string> DDCEnvVars = new string[] { GetXPlatformEnvironmentKey("UE-BootDataCachePath"), GetXPlatformEnvironmentKey("UE-LocalDataCachePath") };
 				
 				if (TaskOptions.HasFlag(DDCTaskOptions.KeepMemoryDDC))
