Adds support for DnsEligible

For Linux add a command line option to disable exclusive locks on files opened. This is due to a dotnet implemention
[REDACTED_PROFILE_URL]/runtime/issues/34126

diff --git a/Engine/Source/Programs/AutomationTool/AutomationUtils/Utils.cs b/Engine/Source/Programs/AutomationTool/AutomationUtils/Utils.cs
--- a/Engine/Source/Programs/AutomationTool/AutomationUtils/Utils.cs
+++ b/Engine/Source/Programs/AutomationTool/AutomationUtils/Utils.cs
@@ -913,6 +913,21 @@ namespace AutomationTool
 			}
 		}
 
+		/// <summary>
+		/// Indicates whether the Internet Protocol (IP) address is valid to appear in a Domain Name System (DNS) server database.
+		///
+		/// Addresses in the range 169.254.0.0 to 169.254.255.255 are not DNS eligible. These addresses are reserved for Automatic Private IP Addressing (APIPA).
+		/// https://docs.microsoft.com/en-us/dotnet/api/system.net.networkinformation.ipaddressinformation.isdnseligible?view=netcore-3.1
+		/// 
+		/// NET Core 3.1 does not include a usable implementation of UnicastIPAddressInformation.IsDnsEligible() for Linux or Mac
+		/// </summary>
+		/// <param name="AddressInformation">Information about the address to evaluate for DNS eligibility</param>
+		/// <returns></returns>
+		public static bool IsDnsEligible(System.Net.NetworkInformation.IPAddressInformation AddressInformation)
+		{
+			byte[] AddressBytes = AddressInformation.Address.GetAddressBytes();
+			return !(AddressBytes[0] == 169 && AddressBytes[1] == 254);
+		}
 	}
 
 
diff --git a/Engine/Source/Programs/AutomationTool/Scripts/CopyBuildToStagingDirectory.Automation.cs b/Engine/Source/Programs/AutomationTool/Scripts/CopyBuildToStagingDirectory.Automation.cs
--- a/Engine/Source/Programs/AutomationTool/Scripts/CopyBuildToStagingDirectory.Automation.cs
+++ b/Engine/Source/Programs/AutomationTool/Scripts/CopyBuildToStagingDirectory.Automation.cs
@@ -3790,7 +3790,7 @@ public partial class Project : CommandUtils
 							IPInterfaceProperties IP = adapter.GetIPProperties();
 							for (int Index = 0; Index < IP.UnicastAddresses.Count; ++Index)
 							{
-								if (IP.UnicastAddresses[Index].IsDnsEligible && IP.UnicastAddresses[Index].Address.AddressFamily == System.Net.Sockets.AddressFamily.InterNetwork)
+								if (InternalUtils.IsDnsEligible(IP.UnicastAddresses[Index]) && IP.UnicastAddresses[Index].Address.AddressFamily == System.Net.Sockets.AddressFamily.InterNetwork)
 								{
 									if (!IsNullOrEmpty(Params.Port))
 									{
@@ -3833,7 +3833,7 @@ public partial class Project : CommandUtils
 							IPInterfaceProperties IP = adapter.GetIPProperties();
 							for (int Index = 0; Index < IP.UnicastAddresses.Count; ++Index)
 							{
-								if (IP.UnicastAddresses[Index].IsDnsEligible)
+								if (InternalUtils.IsDnsEligible(IP.UnicastAddresses[Index]))
 								{
 									if (!IsNullOrEmpty(Params.Port))
 									{
diff --git a/Engine/Source/Programs/AutomationTool/Scripts/RunProjectCommand.Automation.cs b/Engine/Source/Programs/AutomationTool/Scripts/RunProjectCommand.Automation.cs
--- a/Engine/Source/Programs/AutomationTool/Scripts/RunProjectCommand.Automation.cs
+++ b/Engine/Source/Programs/AutomationTool/Scripts/RunProjectCommand.Automation.cs
@@ -517,7 +517,6 @@ public partial class Project : CommandUtils
 		}
 		else
 		{
-
 			LogFileReaderProcess(ServerLogFile, ServerProcess, (string Output) =>
 			{
 				bool bKeepReading = true;
@@ -668,7 +667,7 @@ public partial class Project : CommandUtils
 							IPInterfaceProperties IP = adapter.GetIPProperties();
 							for (int Index = 0; Index < IP.UnicastAddresses.Count; ++Index)
 							{
-								if (IP.UnicastAddresses[Index].IsDnsEligible && IP.UnicastAddresses[Index].Address.AddressFamily == System.Net.Sockets.AddressFamily.InterNetwork)
+								if (InternalUtils.IsDnsEligible(IP.UnicastAddresses[Index]) && IP.UnicastAddresses[Index].Address.AddressFamily == System.Net.Sockets.AddressFamily.InterNetwork)
 								{
 									if (!IsNullOrEmpty(Params.Port))
 									{
@@ -709,7 +708,7 @@ public partial class Project : CommandUtils
 							IPInterfaceProperties IP = adapter.GetIPProperties();
 							for (int Index = 0; Index < IP.UnicastAddresses.Count; ++Index)
 							{
-								if (IP.UnicastAddresses[Index].IsDnsEligible)
+								if (InternalUtils.IsDnsEligible(IP.UnicastAddresses[Index]))
 								{
 									if (!IsNullOrEmpty(Params.Port))
 									{
@@ -988,6 +987,13 @@ public partial class Project : CommandUtils
 			CommandUtils.MakePathSafeToUseWithCommandLine(ProjectName.FullName));
 		if (!String.IsNullOrEmpty(ServerLogFile))
 		{
+			// Issue with dotnet not allowing any files with an exclusive advisory lock to be opened for read-only or copied
+			// [REDACTED_PROFILE_URL]/runtime/issues/34126
+			if (TargetPlatform == "Linux")
+			{
+				Args += " -noexclusivelockonwrite";
+			}
+
 			Args += " -abslog=" + CommandUtils.MakePathSafeToUseWithCommandLine(ServerLogFile);
 		}
 		if (IsBuildMachine)
diff --git a/Engine/Source/Runtime/Core/Private/Unix/UnixPlatformFile.cpp b/Engine/Source/Runtime/Core/Private/Unix/UnixPlatformFile.cpp
--- a/Engine/Source/Runtime/Core/Private/Unix/UnixPlatformFile.cpp
+++ b/Engine/Source/Runtime/Core/Private/Unix/UnixPlatformFile.cpp
@@ -953,9 +953,13 @@ IFileHandle* FUnixPlatformFile::OpenWrite(const TCHAR* Filename, bool bAppend, b
 	int32 Handle = open(TCHAR_TO_UTF8(*NormalizeFilename(Filename, true)), Flags, S_IRUSR | S_IWUSR | S_IRGRP | S_IWGRP | S_IROTH);
 	if (Handle != -1)
 	{
+		// Due to dotnet not allowing any files marked as LOCK_EX to be opened for read only or copied, this allows us to
+		// to disable the locking mechanics. [REDACTED_PROFILE_URL]/runtime/issues/34126
+		extern bool GAllowExclusiveLockOnWrite;
+
 		// mimic Windows "exclusive write" behavior (we don't use FILE_SHARE_WRITE) by locking the file.
 		// note that the (non-mandatory) "lock" will be removed by itself when the last file descriptor is close()d
-		if (flock(Handle, LOCK_EX | LOCK_NB) == -1)
+		if (GAllowExclusiveLockOnWrite && (flock(Handle, LOCK_EX | LOCK_NB) == -1))
 		{
 			// if locked, consider operation a failure
 			if (EAGAIN == errno || EWOULDBLOCK == errno)
diff --git a/Engine/Source/Runtime/Core/Private/Unix/UnixPlatformMemory.cpp b/Engine/Source/Runtime/Core/Private/Unix/UnixPlatformMemory.cpp
--- a/Engine/Source/Runtime/Core/Private/Unix/UnixPlatformMemory.cpp
+++ b/Engine/Source/Runtime/Core/Private/Unix/UnixPlatformMemory.cpp
@@ -63,6 +63,11 @@ bool CORE_API GTimeEnsures = true;
 // Allows settings a specific signal to maintain its default handler rather then ignoring the signal
 int32 CORE_API GSignalToDefault = 0;
 
+// Due to dotnet not allowing any files marked as LOCK_EX to be opened for read only or copied, this allows us to
+// to disable the locking mechanics. [REDACTED_PROFILE_URL]/runtime/issues/34126
+// Default to true, can be disabled with -noexclusivelockonwrite
+bool GAllowExclusiveLockOnWrite = true;
+
 #if UE_SERVER
 // Scale factor for how much we would like to increase or decrease the memory pool size
 float CORE_API GPoolTableScale = 1.0f;
@@ -195,6 +200,11 @@ class FMalloc* FUnixPlatformMemory::BaseAllocator()
 					GTimeEnsures = false;
 				}
 
+				if (FCStringAnsi::Stricmp(Arg, "-noexclusivelockonwrite") == 0)
+				{
+					GAllowExclusiveLockOnWrite = false;
+				}
+
 				const char SignalToDefaultCmd[] = "-sigdfl=";
 				if (const char* Cmd = FCStringAnsi::Stristr(Arg, SignalToDefaultCmd))
 				{
diff --git a/Engine/Source/Runtime/Core/Private/Unix/UnixPlatformMisc.cpp b/Engine/Source/Runtime/Core/Private/Unix/UnixPlatformMisc.cpp
--- a/Engine/Source/Runtime/Core/Private/Unix/UnixPlatformMisc.cpp
+++ b/Engine/Source/Runtime/Core/Private/Unix/UnixPlatformMisc.cpp
@@ -223,6 +223,7 @@ void FUnixPlatformMisc::PlatformInit()
 	UE_LOG(LogInit, Log, TEXT(" -ksmmergeall - marks all mmap'd memory pages suitable for KSM (%s)"), GKSMMergeAllPages ? TEXT("ON") : TEXT("OFF"));
 	UE_LOG(LogInit, Log, TEXT(" -preloadmodulesymbols - Loads the main module symbols file into memory (%s)"), bPreloadedModuleSymbolFile ? TEXT("ON") : TEXT("OFF"));
 	UE_LOG(LogInit, Log, TEXT(" -sigdfl=SIGNAL - Allows a specific signal to be set to its default handler rather then ignoring the signal"));
+	UE_LOG(LogInit, Log, TEXT(" -noexclusivelockonwrite - disables marking files created by the engine as exclusive locked while the engine has them opened"));
 
 	// [RCL] FIXME: this should be printed in specific modules, if at all
 	UE_LOG(LogInit, Log, TEXT(" -httpproxy=ADDRESS:PORT - redirects HTTP requests to a proxy (only supported if compiled with libcurl)"));
