Change BootTest idle timeout tolerance until a solution is found with buildmachine lag

#rnx

diff --git a/Engine/Source/Programs/AutomationTool/Gauntlet/Unreal/Automation/UE.BootTest.cs b/Engine/Source/Programs/AutomationTool/Gauntlet/Unreal/Automation/UE.BootTest.cs
--- a/Engine/Source/Programs/AutomationTool/Gauntlet/Unreal/Automation/UE.BootTest.cs
+++ b/Engine/Source/Programs/AutomationTool/Gauntlet/Unreal/Automation/UE.BootTest.cs
@@ -89,7 +89,7 @@ namespace UE
 		/// </summary>
 		public override void TickTest()
 		{
-			const int kTimeOutDuration = 2;
+			const int kTimeOutDuration = 10;
 
 			// run the base class tick;
 			base.TickTest();
