AutomationUtils/ScriptCompiler.cs:

Removed no longer used script compilation code

diff --git a/Engine/Source/Programs/AutomationTool/AutomationUtils/AutomationUtils.Automation.csproj b/Engine/Source/Programs/AutomationTool/AutomationUtils/AutomationUtils.Automation.csproj
--- a/Engine/Source/Programs/AutomationTool/AutomationUtils/AutomationUtils.Automation.csproj
+++ b/Engine/Source/Programs/AutomationTool/AutomationUtils/AutomationUtils.Automation.csproj
@@ -41,7 +41,6 @@
     <ProjectReference Include="..\..\UnrealBuildTool\UnrealBuildTool.csproj" PrivateAssets="All"><Private>false</Private></ProjectReference>
   </ItemGroup>
   <ItemGroup>
-    <PackageReference Include="Microsoft.CSharp" Version="4.7.0" />
     <PackageReference Include="System.Data.DataSetExtensions" Version="4.5.0" />
     <PackageReference Include="System.Net.Http" Version="4.3.4" />
     <PackageReference Include="System.Drawing.Common" Version="4.7.0" />
diff --git a/Engine/Source/Programs/AutomationTool/AutomationUtils/ScriptCompiler.cs b/Engine/Source/Programs/AutomationTool/AutomationUtils/ScriptCompiler.cs
--- a/Engine/Source/Programs/AutomationTool/AutomationUtils/ScriptCompiler.cs
+++ b/Engine/Source/Programs/AutomationTool/AutomationUtils/ScriptCompiler.cs
@@ -2,11 +2,7 @@
 using System;
 using System.Collections.Generic;
 using System.Linq;
-using System.Text;
 using System.IO;
-using System.CodeDom.Compiler;
-using Microsoft.CSharp;
-using Microsoft.Win32;
 using System.Reflection;
 using System.Diagnostics;
 using UnrealBuildTool;
@@ -194,251 +190,6 @@ namespace AutomationTool
 			return ModulesToCompile;
 		}
 
-		/// <summary>
-		/// Creates a hash collection that represents the state of all modules in this list. If 
-		/// a module has no current output or is missing files then a null collection will be returned
-		/// </summary>
-		/// <param name="Modules"></param>
-		/// <returns></returns>
-		private static HashCollection HashModules(IEnumerable<string> Modules, bool WarnOnFailure=true)
-		{
-			HashCollection Hashes = new HashCollection();
-
-			foreach (string Module in Modules)
-			{
-				// this project is built by the AutomationTool project that the RunUAT script builds so 
-				// it will always be newer than was last built
-				if (Module.Contains("AutomationUtils.Automation"))
-				{
-					continue;
-				}
-
-				CsProjectInfo Proj;
-
-				Dictionary<string, string> Properties = new Dictionary<string, string>();
-				Properties.Add("Platform", "AnyCPU");
-				Properties.Add("Configuration", "Development");
-
-				FileReference ModuleFile = new FileReference(Module);
-
-				if (!CsProjectInfo.TryRead(ModuleFile, Properties, out Proj))
-				{
-					if (WarnOnFailure)
-					{
-						Log.TraceWarning("Failed to read file {0}", ModuleFile);
-					}
-					return null;
-				}
-
-				if (!Hashes.AddCsProjectInfo(Proj, HashCollection.HashType.MetaData))
-				{
-					if (WarnOnFailure)
-					{
-						Log.TraceWarning("Failed to hash file {0}", ModuleFile);
-					}
-					return null;
-				}
-			}
-
-			return Hashes;
-		}
-
-		/// <summary>
-		/// Pulls all dependencies from the specified module list, gathers the state of them, and
-		/// compares it to the state expressed in the provided file from a previous run
-		/// </summary>
-		/// <param name="ModuleList"></param>
-		/// <param name="DependencyFile"></param>
-		/// <returns></returns>
-		private static bool AreDependenciesUpToDate(IEnumerable<string> ModuleList, string DependencyFile)
-		{
-			bool UpToDate = false;
-
-			if (File.Exists(DependencyFile))
-			{
-				Log.TraceVerbose("Read dependency file at {0}", DependencyFile);
-
-				HashCollection OldHashes = HashCollection.CreateFromFile(DependencyFile);
-
-				if (OldHashes != null)
-				{
-					HashCollection CurrentHashes = HashModules(ModuleList, false);
-
-					if (OldHashes.Equals(CurrentHashes))
-					{
-						UpToDate = true;
-					}
-					else
-					{
-						if (Log.OutputLevel >= LogEventType.VeryVerbose)
-						{
-							CurrentHashes.LogDifferences(OldHashes);
-						}
-					}
-				}
-				else
-				{
-					Log.TraceInformation("Failed to read dependency info!");
-				}
-			}
-			else
-			{
-				Log.TraceVerbose("No dependency file exists at {0}. Will do full build.", DependencyFile);
-			}
-
-			return UpToDate;
-		}
-
-		/// <summary>
-		/// Converts a set of MSBuild properties into command line arguments
-		/// </summary>
-		/// <param name="Properties">The properties to set</param>
-		/// <returns>Command line arguments</returns>
-		private static string GetMsBuildPropertyArguments(Dictionary<string, string> Properties)
-		{
-			return String.Join(" ", Properties.Select(x => String.Format(" /p:{0}={1}", Utils.MakePathSafeToUseWithCommandLine(x.Key), Utils.MakePathSafeToUseWithCommandLine(x.Value))));
-		}
-
-		/// <summary>
-		/// Compiles all automation projects
-		/// </summary>
-		/// <param name="Projects">Projects to compile</param>
-		/// <param name="MsBuildProperties">Properties to set</param>
-		private static void CompileAutomationProjects(List<CsProjectInfo> Projects, Dictionary<string, string> MsBuildProperties)
-		{
-			Stopwatch Timer = Stopwatch.StartNew();
-
-			string DependencyFile = Path.Combine(CommandUtils.CmdEnv.EngineSavedFolder, "UATModuleHashes.xml");
-			if (AreDependenciesUpToDate(Projects.Select(x => x.ProjectPath.FullName), DependencyFile) && !GlobalCommandLine.IgnoreDependencies)
-			{
-				Log.TraceInformation("Dependencies are up to date ({0:0.000}s). Skipping compile.", Timer.Elapsed.TotalSeconds);
-				return;
-			}
-
-			Log.TraceInformation("Dependencies are out of date. Compiling scripts....");
-
-			// clean old assemblies
-			CleanupScriptsAssemblies();
-
-			string BuildTool = CommandUtils.CmdEnv.MsBuildExe;
-
-			// msbuild (standard on windows, in mono >=5.0 is preferred due to speed and parallel compilation)
-			bool UseParallelMsBuild = Path.GetFileNameWithoutExtension(BuildTool).ToLower() == "msbuild";
-			if (UseParallelMsBuild)
-			{
-				string ProjectsList = string.Join(";", Projects.Select(x => x.ProjectPath));
-
-				// Mono has an issue where arugments with semicolons or commas can't be passed through to
-				// as arguments so we need to manually construct a temp file with the list of modules
-				// see ([REDACTED_PROFILE_URL]/msbuild/issues/471)
-				string UATProjTemplate = CommandUtils.CombinePaths(CommandUtils.CmdEnv.LocalRoot, @"Engine\Source\Programs\AutomationTool\Scripts\UAT.proj");
-				string UATProjFile = Path.Combine(CommandUtils.CmdEnv.EngineSavedFolder, "UATTempProj.proj");
-
-				string ProjContents = File.ReadAllText(UATProjTemplate);
-				ProjContents = ProjContents.Replace("$(Modules)", ProjectsList);
-
-				Directory.CreateDirectory(Path.GetDirectoryName(UATProjFile));
-				File.WriteAllText(UATProjFile, ProjContents);
-				
-				string MsBuildVerbosity = Log.OutputLevel >= LogEventType.Verbose ? "minimal" : "quiet";
-
-				string CmdLine = String.Format("\"{0}\" {1} /verbosity:{2} /nologo", UATProjFile, GetMsBuildPropertyArguments(MsBuildProperties), MsBuildVerbosity);
-				// suppress the run command because it can be long and intimidating, making the logs around this code harder to read.
-				IProcessResult Result = CommandUtils.Run(BuildTool, CmdLine, Options: CommandUtils.ERunOptions.Default | CommandUtils.ERunOptions.NoLoggingOfRunCommand | CommandUtils.ERunOptions.LoggingOfRunDuration);
-				if (Result.ExitCode != 0)
-				{
-					throw new AutomationException(String.Format("Failed to build \"{0}\":{1}{2}", UATProjFile, Environment.NewLine, Result.Output));
-				}
-			}
-			else
-			{
-				// Make sure DefaultScriptsDLLName is compiled first
-				string DefaultScriptsProjName = Path.ChangeExtension(DefaultScriptsDLLName, "csproj");
-
-				// Primary modules must be built first
-				List<CsProjectInfo> PrimaryProjects = Projects.Where(M => M.ProjectPath.FullName.IndexOf(DefaultScriptsProjName, StringComparison.InvariantCultureIgnoreCase) >= 0).ToList();
-				foreach (CsProjectInfo PrimaryProject in PrimaryProjects)
-				{
-					Log.TraceInformation("Building script module: {0}", PrimaryProject.ProjectPath);
-					try
-					{
-						CompileAutomationProject(PrimaryProject.ProjectPath, MsBuildProperties);
-					}
-					catch (Exception Ex)
-					{
-						throw new AutomationException(Ex, "Failed to compile module {0}", PrimaryProject.ProjectPath);
-					}
-					break;
-				}
-
-				// Second pass, compile everything else
-				List<CsProjectInfo> SecondaryProjects = Projects.Where(M => !PrimaryProjects.Contains(M)).ToList();
-
-				// Non-parallel method
-				foreach (CsProjectInfo SecondaryProject in SecondaryProjects)
-				{
-					Log.TraceInformation("Building script module: {0}", SecondaryProject.ProjectPath);
-					try
-					{
-						CompileAutomationProject(SecondaryProject.ProjectPath, MsBuildProperties);
-					}
-					catch (Exception Ex)
-					{
-						throw new AutomationException(Ex, "Failed to compile module {0}", SecondaryProject.ProjectPath);
-					}
-				}
-			}			
-			
-			Log.TraceInformation("Compiled {0} modules in {1:0.000} secs", Projects.Count, Timer.Elapsed.TotalSeconds);
-
-			HashCollection NewHashes = HashModules(Projects.Select(x => x.ProjectPath.FullName));
-
-			if (NewHashes == null)
-			{
-				Log.TraceWarning("Failed to save dependency info!");
-			}
-			else
-			{
-				NewHashes.SaveToFile(DependencyFile);
-				Log.TraceVerbose("Wrote dependencies to {0}", DependencyFile);
-			}
-		}
-
-		/// <summary>
-		/// Starts compiling the provided project file and returns the process. Caller should check HasExited
-		/// or call WaitForExit() to get results
-		/// </summary>
-		/// <param name="ProjectFile"></param>
-		/// <returns></returns>
-		private static bool CompileAutomationProject(FileReference ProjectFile, Dictionary<string, string> Properties)
-		{
-			if (!ProjectFile.HasExtension(".csproj"))
-			{
-				throw new AutomationException(String.Format("Unable to build Project {0}. Not a valid .csproj file.", ProjectFile));
-			}
-			if (!FileReference.Exists(ProjectFile))
-			{
-				throw new AutomationException(String.Format("Unable to build Project {0}. Project file not found.", ProjectFile));
-			}
-
-			string CmdLine = String.Format("\"{0}\" /verbosity:quiet /nologo /target:Build {1} /p:TreatWarningsAsErrors=false /p:NoWarn=\"612,618,672,1591\" /p:BuildProjectReferences=true",
-				ProjectFile, GetMsBuildPropertyArguments(Properties));
-
-			// Compile the project
-			IProcessResult Result = CommandUtils.Run(CommandUtils.CmdEnv.MsBuildExe, CmdLine);
-
-			if (Result.ExitCode != 0)
-			{
-				throw new AutomationException(String.Format("Failed to build \"{0}\":{1}{2}", ProjectFile, Environment.NewLine, Result.Output));
-			}
-			else
-			{
-				// Remove .Automation.csproj and copy to target dir
-				Log.TraceVerbose("Successfully compiled {0}", ProjectFile);
-			}
-			return Result.ExitCode == 0;
-		}
-
 		/// <summary>
 		/// Loads all precompiled assemblies (DLLs that end with *Scripts.dll).
 		/// </summary>
@@ -474,31 +225,6 @@ namespace AutomationTool
 			return Assemblies;
 		}
 
-		private static void CleanupScriptsAssemblies()
-		{
-			if (!CommandUtils.IsEngineInstalled())
-			{
-				CommandUtils.LogVerbose("Cleaning up script DLL folder");
-				CommandUtils.DeleteDirectory(GetScriptAssemblyFolder());
-
-				// Bug in PortalPublishingTool caused these DLLs to be copied into Engine/Binaries/DotNET. Delete any files left over.
-				DirectoryReference BinariesDir = DirectoryReference.Combine(CommandUtils.RootDirectory, "Engine", "Binaries", "DotNET");
-				foreach (FileReference FileToDelete in DirectoryReference.EnumerateFiles(BinariesDir, "*.automation.dll"))
-				{
-					CommandUtils.DeleteFile(FileToDelete.FullName);
-				}
-				foreach (FileReference FileToDelete in DirectoryReference.EnumerateFiles(BinariesDir, "*.automation.pdb"))
-				{
-					CommandUtils.DeleteFile(FileToDelete.FullName);
-				}
-			}
-		}
-
-		private static DirectoryReference GetScriptAssemblyFolder()
-		{
-			return DirectoryReference.Combine(CommandUtils.EngineDirectory, "Binaries", "DotNET", "AutomationScripts");
-		}
-
 		public static Dictionary<string, Type> Commands
 		{
 			get { return ScriptCommands; }
