Support two msbuilds, for NET Framework and Core

This is a preparatory change ahead of adding the ability to build NET Core projects to the CsCompile task, and subsequent conversion of NET Framework projects to NET Core. On its own, it should have no substantive visible effect.

diff --git a/Engine/Source/Programs/AutomationTool/AutomationUtils/BuildUtils.cs b/Engine/Source/Programs/AutomationTool/AutomationUtils/BuildUtils.cs
--- a/Engine/Source/Programs/AutomationTool/AutomationUtils/BuildUtils.cs
+++ b/Engine/Source/Programs/AutomationTool/AutomationUtils/BuildUtils.cs
@@ -19,9 +19,9 @@ namespace AutomationTool
 		/// <param name="LogName">Optional logfile name.</param>
 		public static void MsBuild(CommandEnvironment Env, string Project, string Arguments, string LogName)
 		{
-			if (String.IsNullOrEmpty(Env.MsBuildExe))
+			if (String.IsNullOrEmpty(Env.FrameworkMsbuildPath))
 			{
-				throw new AutomationException("Unable to find msbuild.exe at: \"{0}\"", Env.MsBuildExe);
+				throw new AutomationException("Unable to find msbuild.exe at: \"{0}\"", Env.FrameworkMsbuildPath);
 			}
 			if (!FileExists(Project))
 			{
@@ -32,7 +32,8 @@ namespace AutomationTool
 			{
 				RunArguments += " " + Arguments;
 			}
-			RunAndLog(Env, Env.MsBuildExe, RunArguments, LogName);
+
+			RunAndLog(Env, Env.FrameworkMsbuildPath, RunArguments, LogName);
 		}
 
 		/// <summary>
@@ -50,14 +51,14 @@ namespace AutomationTool
 			{
 				throw new AutomationException(String.Format("Unabled to build Solution {0}. Solution file not found.", SolutionFile));
 			}
-			if (String.IsNullOrEmpty(Env.MsBuildExe))
+			if (String.IsNullOrEmpty(Env.FrameworkMsbuildPath))
 			{
-				throw new AutomationException("Unable to find msbuild.exe at: \"{0}\"", Env.MsBuildExe);
+				throw new AutomationException("Unable to find msbuild.exe at: \"{0}\"", Env.FrameworkMsbuildPath);
 			}
 			string CmdLine = String.Format("\"{0}\" /m /t:Build /p:Configuration=\"{1}\" /p:Platform=\"{2}\" /verbosity:minimal /nologo", SolutionFile, BuildConfig, BuildPlatform);
 			using (TelemetryStopwatch CompileStopwatch = new TelemetryStopwatch("Compile.{0}.{1}.{2}", Path.GetFileName(SolutionFile), "WinC#", BuildConfig))
 			{
-				RunAndLog(Env, Env.MsBuildExe, CmdLine, LogName);
+				RunAndLog(Env, Env.FrameworkMsbuildPath, CmdLine, LogName);
 			}
 		}
 
diff --git a/Engine/Source/Programs/AutomationTool/AutomationUtils/CommandEnvironment.cs b/Engine/Source/Programs/AutomationTool/AutomationUtils/CommandEnvironment.cs
--- a/Engine/Source/Programs/AutomationTool/AutomationUtils/CommandEnvironment.cs
+++ b/Engine/Source/Programs/AutomationTool/AutomationUtils/CommandEnvironment.cs
@@ -54,14 +54,15 @@ namespace AutomationTool
 		public string EngineSavedFolder { get; protected set; }
 		public string LogFolder { get; protected set; }
 		public string FinalLogFolder { get; protected set; }
-        public string CSVFile { get; protected set; }
+		public string CSVFile { get; protected set; }
 		public string RobocopyExe { get; protected set; }
 		public string MountExe { get; protected set; }
 		public string CmdExe { get; protected set; }
-		public string UATExe { get; protected set; }		
+		public string UATExe { get; protected set; }
 		public string TimestampAsString { get; protected set; }
 		public bool HasCapabilityToCompile { get; protected set; }
-		public string MsBuildExe { get; protected set; }
+		public string FrameworkMsbuildPath { get; protected set; }
+		public string DotnetMsbuildPath { get; protected set; }
 		public string MallocNanoZone { get; protected set; }
 		public bool IsChildInstance { get; protected set; }
 
@@ -183,7 +184,8 @@ namespace AutomationTool
 			Log.TraceVerbose("LocalRoot={0}", LocalRoot);
 			Log.TraceVerbose("LogFolder={0}", LogFolder);
 			Log.TraceVerbose("MountExe={0}", MountExe);
-			Log.TraceVerbose("MsBuildExe={0}", MsBuildExe);
+			Log.TraceVerbose("FrameworkMsbuildExe={0}", FrameworkMsbuildPath);
+			Log.TraceVerbose("DotnetMsbuildExe={0}", DotnetMsbuildPath);
 			Log.TraceVerbose("RobocopyExe={0}", RobocopyExe);
 			Log.TraceVerbose("TimestampAsString={0}", TimestampAsString);
 			Log.TraceVerbose("UATExe={0}", UATExe);			
@@ -197,23 +199,31 @@ namespace AutomationTool
 			// Assume we have the capability co compile.
 			HasCapabilityToCompile = true;
 
-			if (HasCapabilityToCompile)
+			try
 			{
-				try
-				{
-					MsBuildExe = HostPlatform.Current.GetMsBuildExe();
-				}
-				catch (Exception Ex)
-				{
-					Log.WriteLine(LogEventType.Warning, Ex.Message);
-					Log.WriteLine(LogEventType.Warning, "Assuming no compilation capability.");
-					HasCapabilityToCompile = false;
-					MsBuildExe = "";
-				}
+				FrameworkMsbuildPath = HostPlatform.Current.GetFrameworkMsbuildExe();
+			}
+			catch (Exception Ex)
+			{
+				Log.WriteLine(LogEventType.Warning, Ex.Message);
+				Log.WriteLine(LogEventType.Warning, "Assuming no compilation capability for NET Framework projects.");
+				HasCapabilityToCompile = false;
+				FrameworkMsbuildPath = "";
 			}
-
 			Log.TraceVerbose("CompilationEvironment.HasCapabilityToCompile={0}", HasCapabilityToCompile);
-			Log.TraceVerbose("CompilationEvironment.MsBuildExe={0}", MsBuildExe);
+			Log.TraceVerbose("CompilationEvironment.FrameworkMsbuildExe={0}", FrameworkMsbuildPath);
+
+			try
+			{
+				DotnetMsbuildPath = HostPlatform.Current.GetDotnetMsbuildExe();
+			}
+			catch (Exception Ex)
+			{
+				Log.WriteLine(LogEventType.Warning, Ex.Message);
+				Log.WriteLine(LogEventType.Warning, "Assuming no compilation capability for NET Core projects.");
+				FrameworkMsbuildPath = "";
+			}
+			Log.TraceVerbose("CompilationEvironment.DotnetMsbuildExe={0}", DotnetMsbuildPath);
 		}
 
 		/// <summary>
diff --git a/Engine/Source/Programs/AutomationTool/AutomationUtils/HostPlatform.cs b/Engine/Source/Programs/AutomationTool/AutomationUtils/HostPlatform.cs
--- a/Engine/Source/Programs/AutomationTool/AutomationUtils/HostPlatform.cs
+++ b/Engine/Source/Programs/AutomationTool/AutomationUtils/HostPlatform.cs
@@ -1,4 +1,4 @@
-﻿// Copyright Epic Games, Inc. All Rights Reserved.
+// Copyright Epic Games, Inc. All Rights Reserved.
 
 using System;
 using System.Collections.Generic;
@@ -66,10 +66,16 @@ namespace AutomationTool
 		}
 
 		/// <summary>
-		/// Gets the build executable filename.
+		/// Gets the build executable filename for NET Framework projects e.g. msbuild, or xbuild
 		/// </summary>
 		/// <returns></returns>
-		abstract public string GetMsBuildExe();
+		abstract public string GetFrameworkMsbuildExe();
+
+		/// <summary>
+		/// Gets the build executable filename for NET Core projects. Typically, the path to the bundled dotnet executable.
+		/// </summary>
+		/// <returns></returns>
+		abstract public string GetDotnetMsbuildExe();
 
 		/// <summary>
 		/// Folder under UE4/ to the platform's binaries.
diff --git a/Engine/Source/Programs/AutomationTool/AutomationUtils/LinuxHostPlatform.cs b/Engine/Source/Programs/AutomationTool/AutomationUtils/LinuxHostPlatform.cs
--- a/Engine/Source/Programs/AutomationTool/AutomationUtils/LinuxHostPlatform.cs
+++ b/Engine/Source/Programs/AutomationTool/AutomationUtils/LinuxHostPlatform.cs
@@ -13,12 +13,12 @@ namespace AutomationTool
 {
 	class LinuxHostPlatform : HostPlatform
 	{
-		static string CachedMsBuildTool = "";
+		static string CachedFrameworkMsbuildTool = "";
 
-		public override string GetMsBuildExe()
+		public override string GetFrameworkMsbuildExe()
 		{
 			// As of 5.0 mono comes with msbuild which performs better. If that's installed then use it
-			if (string.IsNullOrEmpty(CachedMsBuildTool))
+			if (string.IsNullOrEmpty(CachedFrameworkMsbuildTool))
 			{
 				int Value;
 				bool CanUseMsBuild = (int.TryParse(Environment.GetEnvironmentVariable("UE_USE_SYSTEM_MONO"), out Value) &&
@@ -27,15 +27,20 @@ namespace AutomationTool
 
 				if (CanUseMsBuild)
 				{
-					CachedMsBuildTool = "msbuild";
+					CachedFrameworkMsbuildTool = "msbuild";
 				}
 				else
 				{
-					CachedMsBuildTool = "xbuild";
+					CachedFrameworkMsbuildTool = "xbuild";
 				}
 			}
 
-			return CachedMsBuildTool;
+			return CachedFrameworkMsbuildTool;
+		}
+
+		public override string GetDotnetMsbuildExe()
+		{
+			return "../../ThirdParty/DotNet/Linux/dotnet";
 		}
 
 		public override string RelativeBinariesFolder
diff --git a/Engine/Source/Programs/AutomationTool/AutomationUtils/MacHostPlatform.cs b/Engine/Source/Programs/AutomationTool/AutomationUtils/MacHostPlatform.cs
--- a/Engine/Source/Programs/AutomationTool/AutomationUtils/MacHostPlatform.cs
+++ b/Engine/Source/Programs/AutomationTool/AutomationUtils/MacHostPlatform.cs
@@ -14,27 +14,32 @@ namespace AutomationTool
 {
 	class MacHostPlatform : HostPlatform
 	{
-		static string CachedMsBuildTool = "";
+		static string CachedFrameworkMsbuildExe = "";
 
-		public override string GetMsBuildExe()
+		public override string GetFrameworkMsbuildExe()
 		{
 			// As of 5.0 mono comes with msbuild which performs better. If that's installed then use it
-			if (string.IsNullOrEmpty(CachedMsBuildTool))
+			if (string.IsNullOrEmpty(CachedFrameworkMsbuildExe))
 			{
 				bool CanUseMsBuild = string.IsNullOrEmpty(CommandUtils.WhichApp("msbuild")) == false;
 
 				if (CanUseMsBuild)
 				{
-					CachedMsBuildTool = "msbuild";
+					CachedFrameworkMsbuildExe = "msbuild";
 				}
 				else
 				{
 					Log.TraceInformation("Using xbuild. Install Mono 5.0 or greater for faster builds!");
-					CachedMsBuildTool = "xbuild";
+					CachedFrameworkMsbuildExe = "xbuild";
 				}
 			}
 
-			return CachedMsBuildTool;
+			return CachedFrameworkMsbuildExe;
+		}
+
+		public override string GetDotnetMsbuildExe()
+		{
+			return "../../ThirdParty/DotNet/Mac/dotnet";
 		}
 
 		public override string RelativeBinariesFolder
@@ -137,6 +142,9 @@ namespace AutomationTool
 				}
 				// some of our C# applications are converted to dotnet core, do not run those via mono
 				// they are instead assumed to produce a host executable that can just be run
+
+				// this needs fixing: these applications should be launched through the bundled dotnet
+
 				else if (AppName.Contains("UnrealBuildTool") || AppName.Contains("AutomationTool"))
 				{
 					Options &= ~CommandUtils.ERunOptions.AppMustExist;
diff --git a/Engine/Source/Programs/AutomationTool/AutomationUtils/WindowsHostPlatform.cs b/Engine/Source/Programs/AutomationTool/AutomationUtils/WindowsHostPlatform.cs
--- a/Engine/Source/Programs/AutomationTool/AutomationUtils/WindowsHostPlatform.cs
+++ b/Engine/Source/Programs/AutomationTool/AutomationUtils/WindowsHostPlatform.cs
@@ -12,11 +12,16 @@ namespace AutomationTool
 {
 	class WindowsHostPlatform : HostPlatform
 	{
-		public override string GetMsBuildExe()
+		public override string GetFrameworkMsbuildExe()
 		{
 			return WindowsExports.GetMSBuildToolPath();
 		}
 
+		public override string GetDotnetMsbuildExe()
+		{
+			return @"..\..\ThirdParty\DotNet\Windows\dotnet.exe";
+		}
+
 		public override string RelativeBinariesFolder
 		{
 			get { return @"Engine/Binaries/Win64/"; }
diff --git a/Engine/Source/Programs/AutomationTool/BuildGraph/Tasks/CsCompileTask.cs b/Engine/Source/Programs/AutomationTool/BuildGraph/Tasks/CsCompileTask.cs
--- a/Engine/Source/Programs/AutomationTool/BuildGraph/Tasks/CsCompileTask.cs
+++ b/Engine/Source/Programs/AutomationTool/BuildGraph/Tasks/CsCompileTask.cs
@@ -159,7 +159,7 @@ namespace AutomationTool.Tasks
 				{
 					Arguments.Add(String.Format("/target:{0}", CommandUtils.MakePathSafeToUseWithCommandLine(Parameters.Target)));
 				}
-				if(!CommandUtils.CmdEnv.MsBuildExe.Equals("xbuild"))
+				if(!CommandUtils.CmdEnv.FrameworkMsbuildPath.Equals("xbuild"))
 				{
 					// not supported by xbuild
 					Arguments.Add("/restore");
