Remove extraneous rpath entries prior to code signing.

#UE-111538
#rnx

diff --git a/Engine/Source/Programs/AutomationTool/Mac/MacPlatform.Automation.cs b/Engine/Source/Programs/AutomationTool/Mac/MacPlatform.Automation.cs
--- a/Engine/Source/Programs/AutomationTool/Mac/MacPlatform.Automation.cs
+++ b/Engine/Source/Programs/AutomationTool/Mac/MacPlatform.Automation.cs
@@ -436,12 +436,6 @@ public class MacPlatform : Platform
 	{
 		// package up the program, potentially with an installer for Mac
 		PrintRunTime();
-
-		if (Params.Archive)
-		{
-			// Remove extra RPATHs if we will be archiving the project
-			RemoveExtraRPaths(Params, SC);
-		}
 	}
 
 	public override void ProcessArchivedProject(ProjectParams Params, DeploymentContext SC)
@@ -605,6 +599,13 @@ public class MacPlatform : Platform
 	{
 		if (UnrealBuildTool.BuildHostPlatform.Current.Platform == UnrealTargetPlatform.Mac)
 		{
+			if (Params.Archive)
+			{
+				// Remove extra RPATHs if we will be archiving the project
+				LogInformation("Removing extraneous rpath entries");
+				RemoveExtraRPaths(Params, SC);
+			}
+
 			// Sign everything we built
 			List<FileReference> FilesToSign = GetExecutableNames(SC);
 			LogInformation("RuntimeProjectRootDir: " + SC.RuntimeProjectRootDir);
@@ -622,6 +623,7 @@ public class MacPlatform : Platform
 					LogInformation("Starts with Engine/Binaries");
 					AppBundlePath = CombinePaths("Engine/Binaries", SC.PlatformDir, Path.GetFileNameWithoutExtension(Exe.FullName) + ".app");
 				}
+
 				LogInformation("Signing: " + AppBundlePath);
 				CodeSign.SignMacFileOrFolder(AppBundlePath);
 			}
