Gauntlet - Add build command to publish telemetry from CSV
+ avoid using ITestContext instead use a blackboard type object as context to store abitrary data

at the command line:
 Engine\Build\BatchFiles\RunUAT.bat PublishUnrealAutomationTelemetry -CSVFile=C:\temp\myfile.csv -project=Reverb -changelist=111 -branch=ue5-main

in buildgraph:
<Command Name="PublishUnrealAutomationTelemetry" Arguments="-CSVFile=C:\temp\myfile.csv -project=Reverb -changelist=111 -branch=ue5-main" />

Support csv re-mapping:
 -CSVMapping=DataPoint:SomeName,TestName:SomeKey

#rnx

diff --git a/Engine/Source/Programs/AutomationTool/Gauntlet/Framework/Base/Gauntlet.Database.cs b/Engine/Source/Programs/AutomationTool/Gauntlet/Framework/Base/Gauntlet.Database.cs
--- a/Engine/Source/Programs/AutomationTool/Gauntlet/Framework/Base/Gauntlet.Database.cs
+++ b/Engine/Source/Programs/AutomationTool/Gauntlet/Framework/Base/Gauntlet.Database.cs
@@ -9,6 +9,15 @@ using AutomationTool;
 
 namespace Gauntlet
 {
+
+	/// <summary>
+	/// Hold information about the telemetry context
+	/// </summary>
+	public interface ITelemetryContext
+	{
+		object GetProperty(string Name);
+	}
+
 	/// <summary>
 	/// Interface to drive Database submission
 	/// </summary>
@@ -20,7 +29,7 @@ namespace Gauntlet
 		/// <param name="DataItems"></param>
 		/// <param name="Context"></param>
 		/// <returns></returns>
-		bool SubmitDataItems(IEnumerable<Data> DataItems, ITestContext Context);
+		bool SubmitDataItems(IEnumerable<Data> DataItems, ITelemetryContext Context);
 	}
 
 	/// <summary>
@@ -50,7 +59,7 @@ namespace Gauntlet
 		/// Format the data for target table based on data type
 		/// </summary>
 		/// <returns></returns>
-		IEnumerable<string> FormatDataForTable(Data InData, ITestContext InContext);
+		IEnumerable<string> FormatDataForTable(Data InData, ITelemetryContext InContext);
 	}
 
 	public class MySQLDriver<Data> : IDatabaseDriver<Data> where Data : class 
@@ -84,7 +93,7 @@ namespace Gauntlet
 			return MySqlHelper.ExecuteScalar(Config.ConfigString, SqlQuery);
 		}
 
-		public bool SubmitDataItems(IEnumerable<Data> DataRows, ITestContext TestContext)
+		public bool SubmitDataItems(IEnumerable<Data> DataRows, ITelemetryContext TestContext)
 		{
 			if (Config is IMySQLConfig<Data> DataConfig)
 			{
@@ -188,7 +197,7 @@ namespace Gauntlet
 		/// Format the data for target table based on data type
 		/// </summary>
 		/// <returns></returns>
-		public abstract IEnumerable<string> FormatDataForTable(Data InData, ITestContext InContext);
+		public abstract IEnumerable<string> FormatDataForTable(Data InData, ITelemetryContext InContext);
 
 	}
 
diff --git a/Engine/Source/Programs/AutomationTool/Gauntlet/Unreal/Base/Gauntlet.UnrealTestConfiguration.cs b/Engine/Source/Programs/AutomationTool/Gauntlet/Unreal/Base/Gauntlet.UnrealTestConfiguration.cs
--- a/Engine/Source/Programs/AutomationTool/Gauntlet/Unreal/Base/Gauntlet.UnrealTestConfiguration.cs
+++ b/Engine/Source/Programs/AutomationTool/Gauntlet/Unreal/Base/Gauntlet.UnrealTestConfiguration.cs
@@ -722,6 +722,46 @@ namespace Gauntlet
 			return Roles;
 		}
 
+		/// <summary>
+		/// Return the list of required roles for the target role
+		/// </summary>
+		/// <param name="InRole"></param>
+		/// <returns></returns>
+		public IEnumerable<UnrealTestRole> GetRequiredRoles(UnrealTargetRole InRole)
+		{
+			if (RequiredRoles.ContainsKey(InRole))
+			{
+				return RequiredRoles[InRole];
+			}
+			return new List<UnrealTestRole>();
+		}
+
+		/// <summary>
+		/// Return the main required role to execute the test.
+		/// </summary>
+		/// <returns></returns>
+		public UnrealTestRole GetMainRequiredRole()
+		{
+			var ClientRoles = GetRequiredRoles(UnrealTargetRole.Client);
+			if (ClientRoles.Any())
+			{
+				return ClientRoles.First();
+			}
+			var ServerRoles = GetRequiredRoles(UnrealTargetRole.Server);
+			if (ServerRoles.Any())
+			{
+				return ServerRoles.First();
+			}
+			var EditorRoles = GetRequiredRoles(UnrealTargetRole.Editor);
+			if (EditorRoles.Any())
+			{
+				return EditorRoles.First();
+			}
+			var RoleEnumerator = RequiredRoles.Values.GetEnumerator();
+			RoleEnumerator.MoveNext();
+			return RoleEnumerator.Current.First();
+		}
+
 		/// <summary>
 		/// 
 		/// </summary>
diff --git a/Engine/Source/Programs/AutomationTool/Gauntlet/Unreal/Base/Gauntlet.UnrealTestNode.cs b/Engine/Source/Programs/AutomationTool/Gauntlet/Unreal/Base/Gauntlet.UnrealTestNode.cs
--- a/Engine/Source/Programs/AutomationTool/Gauntlet/Unreal/Base/Gauntlet.UnrealTestNode.cs
+++ b/Engine/Source/Programs/AutomationTool/Gauntlet/Unreal/Base/Gauntlet.UnrealTestNode.cs
@@ -1157,7 +1157,17 @@ namespace Gauntlet
 						DBConfig.LoadConfig(GetConfiguration().DatabaseConfigPath);
 						IDatabaseDriver<TelemetryData> DB = DBConfig.GetDriver();
 						Log.Verbose("Submitting telemetry data to {0}", DB.ToString());
-						DB.SubmitDataItems(DataRows, Context);
+
+						UnrealTelemetryContext TestContext = new UnrealTelemetryContext();
+						TestContext.SetProperty("ProjectName", Context.BuildInfo.ProjectName);
+						TestContext.SetProperty("Branch", Context.BuildInfo.Branch);
+						TestContext.SetProperty("Changelist", Context.BuildInfo.Changelist);
+						var RoleType = GetConfiguration().GetMainRequiredRole().Type;
+						var Role = Context.GetRoleContext(RoleType);
+						TestContext.SetProperty("Platform", Role.Platform);
+						TestContext.SetProperty("Configuration", string.Format("{0} {1}", RoleType, Role.Configuration));
+
+						DB.SubmitDataItems(DataRows, TestContext);
 					}
 					else
 					{
diff --git a/Engine/Source/Programs/AutomationTool/Gauntlet/Unreal/Config/Gauntlet.EpicConfig.cs b/Engine/Source/Programs/AutomationTool/Gauntlet/Unreal/Config/Gauntlet.EpicConfig.cs
--- a/Engine/Source/Programs/AutomationTool/Gauntlet/Unreal/Config/Gauntlet.EpicConfig.cs
+++ b/Engine/Source/Programs/AutomationTool/Gauntlet/Unreal/Config/Gauntlet.EpicConfig.cs
@@ -55,29 +55,21 @@ namespace EpicConfig
 		/// Format the data for target table based on data type
 		/// </summary>
 		/// <returns></returns>
-		public override IEnumerable<string> FormatDataForTable(Gauntlet.TelemetryData Data, ITestContext TestContext)
+		public override IEnumerable<string> FormatDataForTable(Gauntlet.TelemetryData Data, ITelemetryContext Context)
 		{
-			if (TestContext is UnrealTestContext Context)
-			{
-				UnrealTestRoleContext Role = Context.RoleContext[UnrealTargetRole.Client];
-				return new List<string>() {
-					Context.BuildInfo.Changelist.ToString(),
-					Environment.MachineName,
-					Context.BuildInfo.Branch,
-					DateTime.Now.ToString("yyyy-MM-dd HH:mm:ss"),
-					Context.BuildInfo.ProjectName,
-					Role.Platform.ToString(),
-					Role.Configuration.ToString(),
-					Data.TestName,
-					Data.Context,
-					Data.DataPoint,
-					Data.Measurement.ToString()
-				};
-			}
-			else
-			{
-				throw new AutomationException(string.Format("The Test Context is not UnrealTestContext; '{0}'.", TestContext.GetType().FullName));
-			}
+			return new List<string>() {
+				Context.GetProperty("Changelist").ToString(),
+				Environment.MachineName,
+				Context.GetProperty("Branch").ToString(),
+				DateTime.Now.ToString("yyyy-MM-dd HH:mm:ss"),
+				Context.GetProperty("ProjectName").ToString(),
+				Context.GetProperty("Platform").ToString(),
+				Context.GetProperty("Configuration").ToString(),
+				Data.TestName,
+				Data.Context,
+				Data.DataPoint,
+				Data.Measurement.ToString()
+			};
 		}
 	}
 }
\ No newline at end of file
diff --git a/Engine/Source/Programs/AutomationTool/Gauntlet/Unreal/Utils/Gauntlet.UnrealTelemetry.cs b/Engine/Source/Programs/AutomationTool/Gauntlet/Unreal/Utils/Gauntlet.UnrealTelemetry.cs
--- a/Engine/Source/Programs/AutomationTool/Gauntlet/Unreal/Utils/Gauntlet.UnrealTelemetry.cs
+++ b/Engine/Source/Programs/AutomationTool/Gauntlet/Unreal/Utils/Gauntlet.UnrealTelemetry.cs
@@ -3,6 +3,9 @@
 using System;
 using System.IO;
 using System.Collections.Generic;
+using AutomationTool;
+using UnrealBuildTool;
+using EpicGames.Core;
 
 
 namespace Gauntlet
@@ -17,7 +20,8 @@ namespace Gauntlet
 		/// </summary>
 		/// <param name="TelemetryDirectory"></param>
 		/// <param name="TelemetryReport"></param>
-		static public void LoadOutputsIntoReport(string TelemetryDirectory, ITelemetryReport TelemetryReport)
+		/// <param name="Mapping"></param>
+		static public void LoadOutputsIntoReport(string TelemetryDirectory, ITelemetryReport TelemetryReport, Dictionary<string, string> Mapping = null)
 		{
 			// Scan for csv files
 			foreach (string File in Directory.GetFiles(TelemetryDirectory))
@@ -27,46 +31,220 @@ namespace Gauntlet
 					continue;
 				}
 
-				List<Dictionary<string, string>> CSVData = null;
-				try
+				LoadCSVOutputIntoReport(File, TelemetryReport, Mapping);
+			}
+		}
+
+		/// <summary>
+		/// Gather telemetry data from csv file generated by UE automated tests and add them to a test report
+		/// </summary>
+		/// <param name="File"></param>
+		/// <param name="TelemetryReport"></param>
+		/// <param name="Mapping"></param>
+		static public void LoadCSVOutputIntoReport(string File, ITelemetryReport TelemetryReport, Dictionary<string, string> Mapping = null)
+		{ 
+			List<Dictionary<string, string>> CSVData = null;
+			try
+			{
+				CSVData = CSVParser.Load(File);
+			}
+			catch (Exception Ex)
+			{
+				Log.Error("Telemetry - Failed to read CSV file '{0}'. {1}", File, Ex);
+			}
+
+			if (CSVData == null)
+			{
+				return;
+			}
+			string TestNameKey = "TestName";
+			string DataPointKey = "DataPoint";
+			string MeasurementKey = "Measurement";
+			string ContextKey = "Context";
+			if (Mapping != null)
+			{
+				if (Mapping.ContainsKey("TestName"))
+					Mapping.TryGetValue("TestName", out TestNameKey);
+				if (Mapping.ContainsKey("DataPoint"))
+					Mapping.TryGetValue("DataPoint", out DataPointKey);
+				if (Mapping.ContainsKey("Measurement"))
+					Mapping.TryGetValue("Measurement", out MeasurementKey);
+				if (Mapping.ContainsKey("Context"))
+					Mapping.TryGetValue("Context", out ContextKey);
+			}
+			// Populate telemetry report
+			foreach (Dictionary<string, string> Row in CSVData)
+			{
+				string TestName;
+				Row.TryGetValue(TestNameKey, out TestName);
+				string DataPoint;
+				Row.TryGetValue(DataPointKey, out DataPoint);
+				string Measurement;
+				Row.TryGetValue(MeasurementKey, out Measurement);
+				string Context;
+				Row.TryGetValue(ContextKey, out Context);
+				if (string.IsNullOrEmpty(TestName) || string.IsNullOrEmpty(DataPoint) || string.IsNullOrEmpty(Measurement))
 				{
-					CSVData = CSVParser.Load(File);
+					Log.Warning("Telemetry - Missing data in CSV file '{0}':\n TestName='{1}', DataPoint='{2}', Measurement='{3}", File, TestName, DataPoint, Measurement);
+					continue;
 				}
-				catch (Exception Ex)
+				try
 				{
-					Log.Error("Telemetry - Failed to read CSV file '{0}'. {1}", File, Ex);
+					TelemetryReport.AddTelemetry(TestName, DataPoint, double.Parse(Measurement), Context);
 				}
-
-				if (CSVData == null)
+				catch (FormatException)
 				{
-					continue;
+					Log.Error("Telemetry - Failed to parse Measurement value('{0}') in CSV file '{1}'.", Measurement, File);
 				}
-				// Populate telemetry report
-				foreach (Dictionary<string, string> Row in CSVData)
+			}
+		}
+
+		static public List<string> TargetableKeys = new List<string>{ "TestName", "DataPoint", "Measurement", "Context" };
+
+	}
+
+	public class UnrealTelemetryContext : ITelemetryContext
+	{
+		public object GetProperty(string Name)
+		{
+			object Value;
+			Properties.TryGetValue(Name, out Value);
+			return Value;
+		}
+		Dictionary<string, object> Properties;
+
+		public UnrealTelemetryContext()
+		{
+			Properties = new Dictionary<string, object>();
+		}
+		public void SetProperty(string Name, object Value)
+		{
+			Properties[Name] = Value;
+		}
+	}
+
+	/// <summary>
+	/// Build command to read Unreal Automation telemetry data saved in csv format and then publish it to a database.
+	/// </summary>
+	class PublishUnrealAutomationTelemetry : BuildCommand
+	{
+		[Help("CSVFile", "Path to the csv file to parse.")]
+		[Help("CSVDirectory", "Path to a folder containing csv files to parse.")]
+		[Help("CSVMapping", "Optional CSV column mapping. Formate is: <target key>:<source key>,...")]
+		[Help("TelemetryConfig", "Telemetry configuration to use to publish to Database. Default: UETelemtr.y")]
+		[Help("DatabaseConfigPath", "Path to alternate Database config. Default is TelemetryConfig default.")]
+		[Help("Project", "Target Project name.")]
+		[Help("Platform", "Target platform name. Default: current environment platform.")]
+		[Help("Role", "Target Role name. Default: Editor.")]
+		[Help("Branch", "Target Branch name. Default: Unknown.")]
+		[Help("Changelist", "Target Changelist number. Default: 0.")]
+		[Help("Configuration", "Target Configuration name. Default: Development.")]
+
+		public override ExitCode Execute()
+		{
+			string CSVFile = ParseParamValue("CSVFile=", "");
+			string CSVDirectory = ParseParamValue("CSVDirectory=", "");
+			string Config = ParseParamValue("TelemetryConfig=", "UETelemetry");
+			string DatabaseConfigPath = ParseParamValue("DatabaseConfigPath=", "");
+			List<string> CSVMappingStrings = Globals.Params.ParseValues("CSVMapping", true);
+			string ProjectString = ParseParamValue("Project=", "");
+			string PlatformString = ParseParamValue("Platform=", "");
+			string BranchString = ParseParamValue("Branch=", "Unknown");
+			string ChangelistString = ParseParamValue("Changelist=", "0");
+			string RoleString = ParseParamValue("Role=", "Editor");
+			string ConfigurationString = ParseParamValue("Configuration=", "Development");
+
+			if (string.IsNullOrEmpty(CSVDirectory) && string.IsNullOrEmpty(CSVFile))
+			{
+				throw new AutomationException("No telemetry files specified. Use -CSVDirectory=<path> or -CSVFile=<path>.csv");
+			}
+
+			if (!string.IsNullOrEmpty(CSVDirectory) && !Directory.Exists(CSVDirectory))
+			{
+				throw new AutomationException(string.Format("CSVDirectory '{0}' is missing.", CSVDirectory));
+			}
+			if (!string.IsNullOrEmpty(CSVFile) && !File.Exists(CSVFile))
+			{
+				throw new AutomationException(string.Format("CSVFile '{0}' is missing.", CSVDirectory));
+			}
+
+			// CSV Mapping
+			Dictionary<string, string> CSVMapping = null;
+			if (CSVMappingStrings.Count > 0)
+			{
+				CSVMapping = new Dictionary<string, string>();
+				foreach(string StringLine in CSVMappingStrings)
 				{
-					string TestName;
-					Row.TryGetValue("TestName", out TestName);
-					string DataPoint;
-					Row.TryGetValue("DataPoint", out DataPoint);
-					string Measurement;
-					Row.TryGetValue("Measurement", out Measurement);
-					string Context;
-					Row.TryGetValue("Context", out Context);
-					if (string.IsNullOrEmpty(TestName) || string.IsNullOrEmpty(DataPoint) || string.IsNullOrEmpty(Measurement))
+					var SplitLine = StringLine.Split(":");
+					if (SplitLine.Length != 2)
 					{
-						Log.Warning("Telemetry - Missing data in CSV file '{0}':\n TestName='{1}', DataPoint='{2}', Measurement='{3}", File, TestName, DataPoint, Measurement);
-						continue;
+						throw new AutomationException("CSVMapping entry must be in the form <target key>:<source key>.");
 					}
-					try
+					string TargetKey = SplitLine[0];
+					string SourceKey = SplitLine[1];
+					if(!UnrealAutomationTelemetry.TargetableKeys.Contains(TargetKey))
 					{
-						TelemetryReport.AddTelemetry(TestName, DataPoint, double.Parse(Measurement), Context);
+						throw new AutomationException(string.Format("Unknown target key '{0}', CSVMapping target key must be one of the value: {1}.", TargetKey, string.Join(", ", UnrealAutomationTelemetry.TargetableKeys)));
 					}
-					catch (FormatException)
+					CSVMapping[TargetKey] = SourceKey;
+				}
+			}
+
+			// Get Context
+			UnrealTelemetryContext Context = new UnrealTelemetryContext();
+
+			if (string.IsNullOrEmpty(ProjectString))
+			{
+				throw new AutomationException("No project specified. Use -Project=ShooterGame etc");
+			}
+			Context.SetProperty("ProjectName", ProjectString);
+
+			RoleString = Enum.Parse(typeof(UnrealTargetRole), RoleString, true).ToString();
+			ConfigurationString = Enum.Parse(typeof(UnrealTargetConfiguration), ConfigurationString, true).ToString();
+			Context.SetProperty("Configuration", string.Format("{0} {1}", RoleString, ConfigurationString));
+
+			Context.SetProperty("Platform", string.IsNullOrEmpty(PlatformString) ? BuildHostPlatform.Current.Platform : UnrealTargetPlatform.Parse(PlatformString));
+
+			Context.SetProperty("Branch", BranchString);
+			Context.SetProperty("Changelist", ChangelistString);
+
+			// Create a report
+			HordeReport.SimpleTestReport Report = new HordeReport.SimpleTestReport();
+
+			if(Report is ITelemetryReport TelemetryReport)
+			{
+				// Parse CSV
+				if (!string.IsNullOrEmpty(CSVDirectory))
+				{
+					UnrealAutomationTelemetry.LoadOutputsIntoReport(CSVDirectory, TelemetryReport, CSVMapping);
+				}
+				else
+				{
+					UnrealAutomationTelemetry.LoadCSVOutputIntoReport(CSVFile, TelemetryReport, CSVMapping);
+				}
+
+				// Publish Telemetry Data
+				var DataRows = TelemetryReport.GetAllTelemetryData();
+
+				if (DataRows != null)
+				{
+					IDatabaseConfig<TelemetryData> DBConfig = DatabaseConfigManager<TelemetryData>.GetConfigByName(Config);
+					if (DBConfig != null)
 					{
-						Log.Error("Telemetry - Failed to parse Measurement value('{0}') in CSV file '{1}'.", Measurement, File);
+						DBConfig.LoadConfig(DatabaseConfigPath);
+						IDatabaseDriver<TelemetryData> DB = DBConfig.GetDriver();
+						Log.Verbose("Submitting telemetry data to {0}", DB.ToString());
+						DB.SubmitDataItems(DataRows, Context);
+					}
+					else
+					{
+						Log.Warning("Got telemetry data, but database configuration is unknown '{0}'.", Config);
 					}
 				}
+
 			}
+
+			return 0;
 		}
 	}
 }
\ No newline at end of file
