Fix incorrect comment

diff --git a/Engine/Source/Programs/AutomationTool/Gauntlet/Unreal/Editor/EditorTests.cs b/Engine/Source/Programs/AutomationTool/Gauntlet/Unreal/Editor/EditorTests.cs
--- a/Engine/Source/Programs/AutomationTool/Gauntlet/Unreal/Editor/EditorTests.cs
+++ b/Engine/Source/Programs/AutomationTool/Gauntlet/Unreal/Editor/EditorTests.cs
@@ -50,7 +50,7 @@ namespace UnrealEditor
 		public string TreatLogErrorsAsTestErrors = "true";
 
 		/// <summary>
-		/// Control for interpretation of log errors as test failures
+		/// Modify the game instance lost timeout interval
 		/// </summary>
 		[AutoParam]
 		public string GameInstanceLostTimerSeconds = string.Empty;
