Gauntlet Telemetry - use an enum instead of a list for CSV columns

#rnx

diff --git a/Engine/Source/Programs/AutomationTool/Gauntlet/Unreal/Utils/Gauntlet.UnrealTelemetry.cs b/Engine/Source/Programs/AutomationTool/Gauntlet/Unreal/Utils/Gauntlet.UnrealTelemetry.cs
--- a/Engine/Source/Programs/AutomationTool/Gauntlet/Unreal/Utils/Gauntlet.UnrealTelemetry.cs
+++ b/Engine/Source/Programs/AutomationTool/Gauntlet/Unreal/Utils/Gauntlet.UnrealTelemetry.cs
@@ -42,7 +42,7 @@ namespace Gauntlet
 		/// <param name="TelemetryReport"></param>
 		/// <param name="Mapping"></param>
 		static public void LoadCSVOutputIntoReport(string File, ITelemetryReport TelemetryReport, Dictionary<string, string> Mapping = null)
-		{ 
+		{
 			List<Dictionary<string, string>> CSVData = null;
 			try
 			{
@@ -57,20 +57,20 @@ namespace Gauntlet
 			{
 				return;
 			}
-			string TestNameKey = "TestName";
-			string DataPointKey = "DataPoint";
-			string MeasurementKey = "Measurement";
-			string ContextKey = "Context";
+			string TestNameKey = CSVColumns.TestName.ToString();
+			string DataPointKey = CSVColumns.DataPoint.ToString();
+			string MeasurementKey = CSVColumns.Measurement.ToString();
+			string ContextKey = CSVColumns.Context.ToString();
 			if (Mapping != null)
 			{
-				if (Mapping.ContainsKey("TestName"))
-					Mapping.TryGetValue("TestName", out TestNameKey);
-				if (Mapping.ContainsKey("DataPoint"))
-					Mapping.TryGetValue("DataPoint", out DataPointKey);
-				if (Mapping.ContainsKey("Measurement"))
-					Mapping.TryGetValue("Measurement", out MeasurementKey);
-				if (Mapping.ContainsKey("Context"))
-					Mapping.TryGetValue("Context", out ContextKey);
+				if (Mapping.ContainsKey(TestNameKey))
+					Mapping.TryGetValue(TestNameKey, out TestNameKey);
+				if (Mapping.ContainsKey(DataPointKey))
+					Mapping.TryGetValue(DataPointKey, out DataPointKey);
+				if (Mapping.ContainsKey(MeasurementKey))
+					Mapping.TryGetValue(MeasurementKey, out MeasurementKey);
+				if (Mapping.ContainsKey(ContextKey))
+					Mapping.TryGetValue(ContextKey, out ContextKey);
 			}
 			// Populate telemetry report
 			foreach (Dictionary<string, string> Row in CSVData)
@@ -85,7 +85,7 @@ namespace Gauntlet
 				Row.TryGetValue(ContextKey, out Context);
 				if (string.IsNullOrEmpty(TestName) || string.IsNullOrEmpty(DataPoint) || string.IsNullOrEmpty(Measurement))
 				{
-					Log.Warning("Telemetry - Missing data in CSV file '{0}':\n TestName='{1}', DataPoint='{2}', Measurement='{3}", File, TestName, DataPoint, Measurement);
+					Log.Warning("Telemetry - Missing data in CSV file '{0}':\n TestName='{1}', DataPoint='{2}', Measurement='{3}'", File, TestName, DataPoint, Measurement);
 					continue;
 				}
 				try
@@ -99,7 +99,13 @@ namespace Gauntlet
 			}
 		}
 
-		static public List<string> TargetableKeys = new List<string>{ "TestName", "DataPoint", "Measurement", "Context" };
+		public enum CSVColumns
+		{
+			TestName,
+			DataPoint,
+			Measurement,
+			Context
+		}
 
 	}
 
@@ -182,10 +188,13 @@ namespace Gauntlet
 					}
 					string TargetKey = SplitLine[0];
 					string SourceKey = SplitLine[1];
-					if(!UnrealAutomationTelemetry.TargetableKeys.Contains(TargetKey))
+					object Target;
+					if (!Enum.TryParse(typeof(UnrealAutomationTelemetry.CSVColumns), TargetKey, true, out Target))
 					{
-						throw new AutomationException(string.Format("Unknown target key '{0}', CSVMapping target key must be one of the value: {1}.", TargetKey, string.Join(", ", UnrealAutomationTelemetry.TargetableKeys)));
+						string AllKeys = string.Join(", ", Enum.GetNames(typeof(UnrealAutomationTelemetry.CSVColumns)));
+						throw new AutomationException(string.Format("Unknown target key '{0}', CSVMapping target key must be one of the values: {1}.", TargetKey, AllKeys));
 					}
+					TargetKey = ((UnrealAutomationTelemetry.CSVColumns)Target).ToString();
 					CSVMapping[TargetKey] = SourceKey;
 				}
 			}
@@ -199,8 +208,22 @@ namespace Gauntlet
 			}
 			Context.SetProperty("ProjectName", ProjectString);
 
-			RoleString = Enum.Parse(typeof(UnrealTargetRole), RoleString, true).ToString();
-			ConfigurationString = Enum.Parse(typeof(UnrealTargetConfiguration), ConfigurationString, true).ToString();
+			object Role;
+			if (!Enum.TryParse(typeof(UnrealTargetRole), RoleString, true, out Role))
+			{
+				string AllKeys = string.Join(", ", Enum.GetNames(typeof(UnrealTargetRole)));
+				throw new AutomationException(string.Format("Unknown Role '{0}', it must be one of the values: {1}.", RoleString, AllKeys));
+			}
+			RoleString = ((UnrealTargetRole)Role).ToString();
+
+			object Configuration;
+			if (!Enum.TryParse(typeof(UnrealTargetConfiguration), ConfigurationString, true, out Configuration))
+			{
+				string AllKeys = string.Join(", ", Enum.GetNames(typeof(UnrealTargetConfiguration)));
+				throw new AutomationException(string.Format("Unknown Configuration '{0}', it must be one of the values: {1}.", ConfigurationString, AllKeys));
+			}
+			ConfigurationString = ((UnrealTargetConfiguration)Configuration).ToString();
+
 			Context.SetProperty("Configuration", string.Format("{0} {1}", RoleString, ConfigurationString));
 
 			Context.SetProperty("Platform", string.IsNullOrEmpty(PlatformString) ? BuildHostPlatform.Current.Platform : UnrealTargetPlatform.Parse(PlatformString));
