Fix UBT executable path in Launcher builds

diff --git a/Engine/Source/Programs/AutomationTool/IOS/IOSPlatform.Automation.cs b/Engine/Source/Programs/AutomationTool/IOS/IOSPlatform.Automation.cs
--- a/Engine/Source/Programs/AutomationTool/IOS/IOSPlatform.Automation.cs
+++ b/Engine/Source/Programs/AutomationTool/IOS/IOSPlatform.Automation.cs
@@ -1060,7 +1060,7 @@ public class IOSPlatform : Platform
 			if (CommandUtils.IsEngineInstalled())
 			{
 				// Get the path to UBT
-				string InstalledUBT = CombinePaths(CmdEnv.LocalRoot, "Engine/Binaries/DotNET/UnrealBuildTool/UnrealBuildTool.exe");
+				string InstalledUBT = UE4Build.GetUBTExecutable();
 				Arguments = "-XcodeProjectFiles " + Arguments;
 				RunUBT(CmdEnv, InstalledUBT, Arguments);
 			}
