Gauntlet - Don't call deprecated CreateReport function signature if the new signature is implemented.

#rnx

diff --git a/Engine/Source/Programs/AutomationTool/Gauntlet/Unreal/Base/Gauntlet.UnrealTestNode.cs b/Engine/Source/Programs/AutomationTool/Gauntlet/Unreal/Base/Gauntlet.UnrealTestNode.cs
--- a/Engine/Source/Programs/AutomationTool/Gauntlet/Unreal/Base/Gauntlet.UnrealTestNode.cs
+++ b/Engine/Source/Programs/AutomationTool/Gauntlet/Unreal/Base/Gauntlet.UnrealTestNode.cs
@@ -881,11 +881,12 @@ namespace Gauntlet
 			ITestReport Report = null;
 
 			var DeprecatedSignature = new[] { typeof(TestResult), typeof(UnrealTestContext), typeof(UnrealBuildSource), typeof(IEnumerable<UnrealRoleArtifacts>), typeof(string) };
+			var NewSignature = new[] { typeof(TestResult) };
 
 			try
 			{
 				// Check if the deprecated signature is overriden, call it anyway if that the case, and trigger a warning.
-				if (Utils.InterfaceHelpers.HasOverriddenMethod(this.GetType(), "CreateReport", DeprecatedSignature))
+				if (!Utils.InterfaceHelpers.HasOverriddenMethod(this.GetType(), "CreateReport", NewSignature) && Utils.InterfaceHelpers.HasOverriddenMethod(this.GetType(), "CreateReport", DeprecatedSignature))
 				{
 					Log.Warning("This CreateReport signature is deprecated, please use the overloaded version.");
 					CreateReport(GetTestResult(), Context, Context.BuildInfo, SessionArtifacts, ArtifactPath);
