Fix ImageMagick library missing reference in Gauntlet

#rnx

diff --git a/Engine/Source/Programs/AutomationTool/Gauntlet/Gauntlet.Automation.csproj b/Engine/Source/Programs/AutomationTool/Gauntlet/Gauntlet.Automation.csproj
--- a/Engine/Source/Programs/AutomationTool/Gauntlet/Gauntlet.Automation.csproj
+++ b/Engine/Source/Programs/AutomationTool/Gauntlet/Gauntlet.Automation.csproj
@@ -31,6 +31,10 @@
     <Reference Include="fastJSON">
       <HintPath>..\..\..\..\Binaries\ThirdParty\fastJSON\netstandard2.0\fastJSON.dll</HintPath>
     </Reference>
+    <Reference Include="Magick.NET-Q16-HDRI-AnyCPU, Version=7.0.0.0, Culture=neutral, PublicKeyToken=2004825badfa91ec, processorArchitecture=MSIL">
+      <SpecificVersion>False</SpecificVersion>
+      <HintPath>ThirdParty\Magick.NET\Magick.NET-Q16-HDRI-AnyCPU.dll</HintPath>
+    </Reference>
     <Reference Include="Ionic.Zip.Reduced">
       <HintPath>..\..\..\..\Binaries\DotNET\Ionic.Zip.Reduced.dll</HintPath>
       <Private>False</Private>
@@ -66,7 +70,5 @@
     <PackageReference Include="MySql.Data" Version="6.10.9" />
     <PackageReference Include="System.Data.DataSetExtensions" Version="4.5.0" />
     <PackageReference Include="System.Net.Http" Version="4.3.4" />
-
-    <PackageReference Include="Magick.NET-Q16-HDRI-AnyCPU" Version="7.22.0" />
   </ItemGroup>
 </Project>
\ No newline at end of file
