Rename UE4CommandLine.txt to UECommandLine.txt and deal with all the related changes

diff --git a/Engine/Build/Android/Java/src/com/epicgames/unreal/GameActivity.java.template b/Engine/Build/Android/Java/src/com/epicgames/unreal/GameActivity.java.template
--- a/Engine/Build/Android/Java/src/com/epicgames/unreal/GameActivity.java.template
+++ b/Engine/Build/Android/Java/src/com/epicgames/unreal/GameActivity.java.template
@@ -703,7 +703,7 @@ public class GameActivity extends $${gameActivitySuperClass}$$ implements Surfac
 
 		BufferedReader reader = null;
 
-		// determine the proper place to look for UE4CommandLine.txt as override outside APK
+		// determine the proper place to look for UECommandLine.txt as override outside APK
 		String BaseDirectory = android.os.Environment.getExternalStorageDirectory().getAbsolutePath();
 		if (bUseExternalFilesDir)
 		{
@@ -715,11 +715,11 @@ public class GameActivity extends $${gameActivitySuperClass}$$ implements Surfac
 		}
 
 		// first look for an override in the project directory
-		String filename = BaseDirectory + "/UnrealGame/" + ProjectName + "/UE4CommandLine.txt";
+		String filename = BaseDirectory + "/UnrealGame/" + ProjectName + "/UECommandLine.txt";
 		reader = TryOpenFileReader(filename);
 		if (reader == null)
 		{
-			filename = BaseDirectory + "/UnrealGame/" + ProjectName + "/ue4commandline.txt";
+			filename = BaseDirectory + "/UnrealGame/" + ProjectName + "/uecommandline.txt";
 			reader = TryOpenFileReader(filename);
 		}
 
@@ -728,7 +728,7 @@ public class GameActivity extends $${gameActivitySuperClass}$$ implements Surfac
 		{
 			try
 			{
-				InputStream stream = AssetManagerReference.open("UE4CommandLine.txt");
+				InputStream stream = AssetManagerReference.open("UECommandLine.txt");
 				reader = new BufferedReader(new InputStreamReader(stream));
 				Log.debug("Using APK commandline");
 			}
diff --git a/Engine/Build/Android/UnrealGame/PushCommandLine.bat b/Engine/Build/Android/UnrealGame/PushCommandLine.bat
--- a/Engine/Build/Android/UnrealGame/PushCommandLine.bat
+++ b/Engine/Build/Android/UnrealGame/PushCommandLine.bat
@@ -1 +1 @@
-%ANDROID_HOME%\platform-tools\adb.exe push UE4CommandLine.txt /mnt/sdcard/UnrealGame/UnrealGame/UE4CommandLine.txt
\ No newline at end of file
+%ANDROID_HOME%\platform-tools\adb.exe push UECommandLine.txt /mnt/sdcard/UnrealGame/UnrealGame/UECommandLine.txt
\ No newline at end of file
diff --git a/Engine/Build/IOS/Resources/FrameworkWrapper/FrameworkWrapper/ue4commandline.txt b/Engine/Build/IOS/Resources/FrameworkWrapper/FrameworkWrapper/ue4commandline.txt
deleted file mode 100644
diff --git a/Engine/Build/IOS/Resources/FrameworkWrapper/FrameworkWrapper/uecommandline.txt b/Engine/Build/IOS/Resources/FrameworkWrapper/FrameworkWrapper/uecommandline.txt
new file mode 100644
--- /dev/null
+++ b/Engine/Build/IOS/Resources/FrameworkWrapper/FrameworkWrapper/uecommandline.txt
@@ -0,0 +1 @@
+../../../PROJECT_NAME/PROJECT_NAME.uproject
diff --git a/Engine/Build/IOS/UE4CommandLine.txt b/Engine/Build/IOS/UE4CommandLine.txt
deleted file mode 100644
diff --git a/Engine/Source/Developer/LauncherServices/Private/Launcher/LauncherWorker.cpp b/Engine/Source/Developer/LauncherServices/Private/Launcher/LauncherWorker.cpp
--- a/Engine/Source/Developer/LauncherServices/Private/Launcher/LauncherWorker.cpp
+++ b/Engine/Source/Developer/LauncherServices/Private/Launcher/LauncherWorker.cpp
@@ -488,7 +488,7 @@ FString FLauncherWorker::CreateUATCommand( const ILauncherProfileRef& InProfile,
 	}
 #endif	// WITH_EDITOR
 
-	// to reduce UE4CommandLine.txt churn (timestamp causing extra work), for LaunchOn (ie iterative deploy) we use a single session guid
+	// to reduce UECommandLine.txt churn (timestamp causing extra work), for LaunchOn (ie iterative deploy) we use a single session guid
 	if (InProfile->GetDeploymentMode() == ELauncherProfileDeploymentModes::CopyToDevice && Profile->IsDeployingIncrementally())
 	{
 		static FGuid StaticGuid(FGuid::NewGuid());
diff --git a/Engine/Source/Programs/AutomationTool/Android/AndroidPlatform.Automation.cs b/Engine/Source/Programs/AutomationTool/Android/AndroidPlatform.Automation.cs
--- a/Engine/Source/Programs/AutomationTool/Android/AndroidPlatform.Automation.cs
+++ b/Engine/Source/Programs/AutomationTool/Android/AndroidPlatform.Automation.cs
@@ -1095,7 +1095,7 @@ public class AndroidPlatform : Platform
                         "\techo",
 						"\techo Removing old data. Failures here are usually fine - indicating the files were not on the device.",
                         "\t$ADB $DEVICE shell 'rm -r $EXTERNAL_STORAGE/UnrealGame/" + Params.ShortProjectName + "'",
-						"\t$ADB $DEVICE shell 'rm -r $EXTERNAL_STORAGE/UnrealGame/UE4CommandLine.txt" + "'",
+						"\t$ADB $DEVICE shell 'rm -r $EXTERNAL_STORAGE/UnrealGame/UECommandLine.txt" + "'",
 						"\t$ADB $DEVICE shell 'rm -r $EXTERNAL_STORAGE/" + TargetAndroidLocation + PackageName + "'",
 						"\t$ADB $DEVICE shell 'rm -r $EXTERNAL_STORAGE/Android/" + TargetAndroidLocation + PackageName + "'",
 						"\t$ADB $DEVICE shell 'rm -r $EXTERNAL_STORAGE/Download/" + TargetAndroidLocation + PackageName + "'",
@@ -1150,7 +1150,7 @@ public class AndroidPlatform : Platform
 						"%ADB% %DEVICE% install " + Path.GetFileName(ApkName),
 						"@if \"%ERRORLEVEL%\" NEQ \"0\" goto Error",
                         "%ADB% %DEVICE% shell rm -r %STORAGE%/UnrealGame/" + Params.ShortProjectName,
-						"%ADB% %DEVICE% shell rm -r %STORAGE%/UnrealGame/UE4CommandLine.txt", // we need to delete the commandline in UnrealGame or it will mess up loading
+						"%ADB% %DEVICE% shell rm -r %STORAGE%/UnrealGame/UECommandLine.txt", // we need to delete the commandline in UnrealGame or it will mess up loading
 						"%ADB% %DEVICE% shell rm -r %STORAGE%/" + TargetAndroidLocation + PackageName,
 						"%ADB% %DEVICE% shell rm -r %STORAGE%/Android/" + TargetAndroidLocation + PackageName,
 						"%ADB% %DEVICE% shell rm -r %STORAGE%/Download/" + TargetAndroidLocation + PackageName,
@@ -1210,7 +1210,7 @@ public class AndroidPlatform : Platform
 						"echo",
 						"echo Removing old data. Failures here are usually fine - indicating the files were not on the device.",
 						"$ADB $DEVICE shell 'rm -r $EXTERNAL_STORAGE/UnrealGame/" + Params.ShortProjectName + "'",
-						"$ADB $DEVICE shell 'rm -r $EXTERNAL_STORAGE/UnrealGame/UE4CommandLine.txt" + "'",
+						"$ADB $DEVICE shell 'rm -r $EXTERNAL_STORAGE/UnrealGame/UECommandLine.txt" + "'",
 						"$ADB $DEVICE shell 'rm -r $EXTERNAL_STORAGE/" + TargetAndroidLocation + PackageName + "'",
 						"$ADB $DEVICE shell 'rm -r $EXTERNAL_STORAGE/Android/" + TargetAndroidLocation + PackageName + "'",
 						"echo",
@@ -1236,7 +1236,7 @@ public class AndroidPlatform : Platform
 						"@echo.",
 						"echo Removing old data. Failures here are usually fine - indicating the files were not on the device.",
 						"%ADB% %DEVICE% shell rm -r %STORAGE%/UnrealGame/" + Params.ShortProjectName,
-						"%ADB% %DEVICE% shell rm -r %STORAGE%/UnrealGame/UE4CommandLine.txt", // we need to delete the commandline in UnrealGame or it will mess up loading
+						"%ADB% %DEVICE% shell rm -r %STORAGE%/UnrealGame/UECommandLine.txt", // we need to delete the commandline in UnrealGame or it will mess up loading
 						"%ADB% %DEVICE% shell rm -r %STORAGE%/" + TargetAndroidLocation + PackageName,
 						"%ADB% %DEVICE% shell rm -r %STORAGE%/Android/" + TargetAndroidLocation + PackageName,
 						"@echo.",
@@ -1823,10 +1823,10 @@ public class AndroidPlatform : Platform
                 }
             }
 
-            // update the ue4commandline.txt
-            // update and deploy ue4commandline.txt
+            // update the uecommandline.txt
+            // update and deploy uecommandline.txt
             // always delete the existing commandline text file, so it doesn't reuse an old one
-            FileReference IntermediateCmdLineFile = FileReference.Combine(SC.StageDirectory, "UE4CommandLine.txt");
+            FileReference IntermediateCmdLineFile = FileReference.Combine(SC.StageDirectory, "UECommandLine.txt");
             Project.WriteStageCommandline(IntermediateCmdLineFile, Params, SC);
 
             // copy files to device if we were staging
@@ -1839,7 +1839,7 @@ public class AndroidPlatform : Platform
 
                 if (Params.IterativeDeploy)
                 {
-                    // always send UE4CommandLine.txt (it was written above after delta checks applied)
+                    // always send UECommandLine.txt (it was written above after delta checks applied)
                     EntriesToDeploy.Add(IntermediateCmdLineFile.FullName);
 
                     // Add non UFS files if any to deploy
@@ -2177,7 +2177,7 @@ public class AndroidPlatform : Platform
 
                 string FinalRemoteDir = RemoteDir;
                 /*
-			    // handle the special case of the UE4Commandline.txt when using content only game (UnrealGame)
+			    // handle the special case of the UECommandline.txt when using content only game (UnrealGame)
 			    if (!Params.IsCodeBasedProject)
 			    {
 				    FinalRemoteDir = "/mnt/sdcard/UnrealGame";
diff --git a/Engine/Source/Programs/AutomationTool/AutomationUtils/Platform.cs b/Engine/Source/Programs/AutomationTool/AutomationUtils/Platform.cs
--- a/Engine/Source/Programs/AutomationTool/AutomationUtils/Platform.cs
+++ b/Engine/Source/Programs/AutomationTool/AutomationUtils/Platform.cs
@@ -666,7 +666,7 @@ namespace AutomationTool
 
 		public virtual HashSet<StagedFileReference> GetFilesForCRCCheck()
 		{
-			string CmdLine = "UE4CommandLine.txt";
+			string CmdLine = "UECommandLine.txt";
 			// using SystemNonUFS because that is how it's staged in CreateStagingManifest
 			if (DeployLowerCaseFilenames(StagedFileType.SystemNonUFS))
 			{
@@ -738,7 +738,7 @@ namespace AutomationTool
 		}
 
 		/// <summary>
-		/// Determines whether we should stage a UE4CommandLine.txt for this platform
+		/// Determines whether we should stage a UECommandLine.txt for this platform
 		/// </summary>
 		public virtual bool ShouldStageCommandLine(ProjectParams Params, DeploymentContext SC)
 		{
diff --git a/Engine/Source/Programs/AutomationTool/AutomationUtils/ProjectParams.cs b/Engine/Source/Programs/AutomationTool/AutomationUtils/ProjectParams.cs
--- a/Engine/Source/Programs/AutomationTool/AutomationUtils/ProjectParams.cs
+++ b/Engine/Source/Programs/AutomationTool/AutomationUtils/ProjectParams.cs
@@ -1746,7 +1746,7 @@ namespace AutomationTool
 		/// <summary>
 		/// Stage: If non-empty, the contents will be put into the stage
 		/// </summary>
-		[Help("cmdline", "command line to put into the stage in UE4CommandLine.txt")]
+		[Help("cmdline", "command line to put into the stage in UECommandLine.txt")]
 		public string StageCommandline;
 
         /// <summary>
diff --git a/Engine/Source/Programs/AutomationTool/Gauntlet/Platform/Android/Gauntlet.TargetDeviceAndroid.cs b/Engine/Source/Programs/AutomationTool/Gauntlet/Platform/Android/Gauntlet.TargetDeviceAndroid.cs
--- a/Engine/Source/Programs/AutomationTool/Gauntlet/Platform/Android/Gauntlet.TargetDeviceAndroid.cs
+++ b/Engine/Source/Programs/AutomationTool/Gauntlet/Platform/Android/Gauntlet.TargetDeviceAndroid.cs
@@ -1018,7 +1018,7 @@ namespace Gauntlet
 				// create a tempfile, insert the command line, and push it over
 				string TmpFile = Path.GetTempFileName();
 
-				CommandLineFilePath = string.Format("{0}/UE4CommandLine.txt", RemoteDir);
+				CommandLineFilePath = string.Format("{0}/UECommandLine.txt", RemoteDir);
 
 				// I've seen a weird thing where adb push truncates by a byte, so add some padding...
 				File.WriteAllText(TmpFile, AppConfig.CommandLine + "    ");
diff --git a/Engine/Source/Programs/AutomationTool/HoloLens/HoloLensPlatform.Automation.cs b/Engine/Source/Programs/AutomationTool/HoloLens/HoloLensPlatform.Automation.cs
--- a/Engine/Source/Programs/AutomationTool/HoloLens/HoloLensPlatform.Automation.cs
+++ b/Engine/Source/Programs/AutomationTool/HoloLens/HoloLensPlatform.Automation.cs
@@ -1084,8 +1084,8 @@ namespace HoloLens.Automation
 
 					if (updateCommandLine)
 					{
-						// Update the ue4commandline.txt
-						FileReference IntermediateCmdLineFile = FileReference.Combine(SC.StageDirectory, "UE4CommandLine.txt");
+						// Update the uecommandline.txt
+						FileReference IntermediateCmdLineFile = FileReference.Combine(SC.StageDirectory, "UECommandLine.txt");
 						Log.TraceInformation("Writing cmd line to: " + IntermediateCmdLineFile.FullName);
 						Project.WriteStageCommandline(IntermediateCmdLineFile, Params, SC);
 					}
diff --git a/Engine/Source/Programs/AutomationTool/IOS/IOSPlatform.Automation.cs b/Engine/Source/Programs/AutomationTool/IOS/IOSPlatform.Automation.cs
--- a/Engine/Source/Programs/AutomationTool/IOS/IOSPlatform.Automation.cs
+++ b/Engine/Source/Programs/AutomationTool/IOS/IOSPlatform.Automation.cs
@@ -827,7 +827,7 @@ public class IOSPlatform : Platform
 
 			// delete the old cookeddata
 			InternalUtils.SafeDeleteDirectory(AppDirectory + "/cookeddata", true);
-			InternalUtils.SafeDeleteFile(AppDirectory + "/ue4commandline.txt", true);
+			InternalUtils.SafeDeleteFile(AppDirectory + "/uecommandline.txt", true);
 
 			SearchOption searchMethod;
 			if (!Params.IterativeDeploy)
diff --git a/Engine/Source/Programs/AutomationTool/Lumin/LuminPlatform.Automation.cs b/Engine/Source/Programs/AutomationTool/Lumin/LuminPlatform.Automation.cs
--- a/Engine/Source/Programs/AutomationTool/Lumin/LuminPlatform.Automation.cs
+++ b/Engine/Source/Programs/AutomationTool/Lumin/LuminPlatform.Automation.cs
@@ -530,9 +530,9 @@ public class LuminPlatform : Platform
 					bool NeedDeployNonUFSFiles = false;
 					foreach (string Line in Lines)
 					{
-						// We should not count ue4commandline.txt as a NonUFS file since it is deployed regardless (see Deploy() function).
+						// We should not count uecommandline.txt as a NonUFS file since it is deployed regardless (see Deploy() function).
 						// and counting it here as NonUFS causes an unnecesarry minimal package.
-						if (!(string.IsNullOrWhiteSpace(Line) || Line.Contains("ue4commandline.txt")))
+						if (!(string.IsNullOrWhiteSpace(Line) || Line.Contains("uecommandline.txt")))
 						{
 							NeedDeployNonUFSFiles = true;
 							break;
@@ -906,11 +906,11 @@ public class LuminPlatform : Platform
 
 	private void Deploy(ProjectParams Params, DeploymentContext SC, String DeviceName)
 	{
-		// update the ue4commandline.txt
-		// update and deploy ue4commandline.txt
+		// update the uecommandline.txt
+		// update and deploy uecommandline.txt
 		// always delete the existing commandline text file, so it doesn't reuse an old one
 		// only the file name should be lowercased, not the entire path
-		string IntermediateCmdLineFile = CombinePaths(SC.StageDirectory.FullName, "ue4commandline.txt");
+		string IntermediateCmdLineFile = CombinePaths(SC.StageDirectory.FullName, "uecommandline.txt");
 		Project.WriteStageCommandline(new FileReference(IntermediateCmdLineFile), Params, SC);
 
 		// Where we put the files on device.
@@ -925,7 +925,7 @@ public class LuminPlatform : Platform
 			if (Params.IterativeDeploy)
 			{
 				LogInformation("ITERATIVE DEPLOY..");
-				// always send ue4commandline.txt (it was written above after delta checks applied)
+				// always send uecommandline.txt (it was written above after delta checks applied)
 				EntriesToDeploy.Add(IntermediateCmdLineFile);
 
 				// Add non UFS files if any to deploy
diff --git a/Engine/Source/Programs/AutomationTool/Scripts/CopyBuildToStagingDirectory.Automation.cs b/Engine/Source/Programs/AutomationTool/Scripts/CopyBuildToStagingDirectory.Automation.cs
--- a/Engine/Source/Programs/AutomationTool/Scripts/CopyBuildToStagingDirectory.Automation.cs
+++ b/Engine/Source/Programs/AutomationTool/Scripts/CopyBuildToStagingDirectory.Automation.cs
@@ -760,13 +760,13 @@ public partial class Project : CommandUtils
 				}
 			}
 
-			// move the UE4Commandline.txt file to the root of the stage
+			// move the UECommandline.txt file to the root of the stage
 			// this file needs to be treated as a UFS file for casing, but NonUFS for being put into the .pak file
 			// @todo: Maybe there should be a new category - UFSNotForPak - and reevaluate all SystemNonUFS references
-			FileReference CommandLineFile = FileReference.Combine(GetIntermediateCommandlineDir(SC), "UE4CommandLine.txt");
+			FileReference CommandLineFile = FileReference.Combine(GetIntermediateCommandlineDir(SC), "UECommandLine.txt");
 			if (FileReference.Exists(CommandLineFile))
 			{
-				StagedFileReference StagedCommandLineFile = new StagedFileReference("UE4CommandLine.txt");
+				StagedFileReference StagedCommandLineFile = new StagedFileReference("UECommandLine.txt");
 				if (SC.StageTargetPlatform.DeployLowerCaseFilenames(StagedFileType.SystemNonUFS))
 				{
 					StagedCommandLineFile = StagedCommandLineFile.ToLowerInvariant();
@@ -3742,7 +3742,7 @@ public partial class Project : CommandUtils
 			return;
 		}
 
-		LogInformation("Creating UE4CommandLine.txt");
+		LogInformation("Creating UECommandLine.txt");
 		if (!string.IsNullOrEmpty(Params.StageCommandline) || !string.IsNullOrEmpty(Params.RunCommandline))
 		{
 			string FileHostParams = " ";
@@ -3895,7 +3895,7 @@ public partial class Project : CommandUtils
 	private static void WriteStageCommandline(ProjectParams Params, DeploymentContext SC)
 	{
 		// always delete the existing commandline text file, so it doesn't reuse an old one
-		FileReference IntermediateCmdLineFile = FileReference.Combine(GetIntermediateCommandlineDir(SC), "UE4CommandLine.txt");
+		FileReference IntermediateCmdLineFile = FileReference.Combine(GetIntermediateCommandlineDir(SC), "UECommandLine.txt");
 		WriteStageCommandline(IntermediateCmdLineFile, Params, SC);
 	}
 
diff --git a/Engine/Source/Programs/UnrealBuildTool/Platform/Android/UEDeployAndroid.cs b/Engine/Source/Programs/UnrealBuildTool/Platform/Android/UEDeployAndroid.cs
--- a/Engine/Source/Programs/UnrealBuildTool/Platform/Android/UEDeployAndroid.cs
+++ b/Engine/Source/Programs/UnrealBuildTool/Platform/Android/UEDeployAndroid.cs
@@ -3963,9 +3963,9 @@ namespace UnrealBuildTool
 					SafeDeleteFile(DestFileName);
 				}
 
-				// See if we need to stage a UE4CommandLine.txt file in assets
-				string CommandLineSourceFileName = Path.Combine(Path.GetDirectoryName(ObbFileLocation), Path.GetFileNameWithoutExtension(ObbFileLocation), "UE4CommandLine.txt");
-				string CommandLineDestFileName = Path.Combine(UnrealBuildPath, "assets", "UE4CommandLine.txt");
+				// See if we need to stage a UECommandLine.txt file in assets
+				string CommandLineSourceFileName = Path.Combine(Path.GetDirectoryName(ObbFileLocation), Path.GetFileNameWithoutExtension(ObbFileLocation), "UECommandLine.txt");
+				string CommandLineDestFileName = Path.Combine(UnrealBuildPath, "assets", "UECommandLine.txt");
 				if (File.Exists(CommandLineSourceFileName))
 				{
 					Directory.CreateDirectory(UnrealBuildPath);
diff --git a/Engine/Source/Programs/UnrealBuildTool/ProjectFiles/Xcode/XcodeFrameworkWrapperProject.cs b/Engine/Source/Programs/UnrealBuildTool/ProjectFiles/Xcode/XcodeFrameworkWrapperProject.cs
--- a/Engine/Source/Programs/UnrealBuildTool/ProjectFiles/Xcode/XcodeFrameworkWrapperProject.cs
+++ b/Engine/Source/Programs/UnrealBuildTool/ProjectFiles/Xcode/XcodeFrameworkWrapperProject.cs
@@ -20,7 +20,7 @@ namespace UnrealBuildTool.ProjectFiles.Xcode
 		private readonly static string TEMPLATE_NAME = "FrameworkWrapper";
 		private readonly static string FRAMEWORK_WRAPPER_TEMPLATE_DIRECTORY = Path.Combine(UnrealBuildTool.EngineDirectory.ToNormalizedPath(), "Build", "IOS", "Resources", TEMPLATE_NAME);
 		private readonly static string TEMPLATE_PROJECT_NAME = "PROJECT_NAME";
-		private readonly static string COMMANDLINE_FILENAME = "ue4commandline.txt";
+		private readonly static string COMMANDLINE_FILENAME = "uecommandline.txt";
 
 		/// <summary>
 		/// Recursively copies all of the files and directories that are inside <paramref name="SourceDirectory"/> into <paramref name="DestinationDirectory"/>.
diff --git a/Engine/Source/Runtime/Core/Public/IOS/IOSCommandLineHelper.h b/Engine/Source/Runtime/Core/Public/IOS/IOSCommandLineHelper.h
--- a/Engine/Source/Runtime/Core/Public/IOS/IOSCommandLineHelper.h
+++ b/Engine/Source/Runtime/Core/Public/IOS/IOSCommandLineHelper.h
@@ -90,10 +90,10 @@ class FIOSCommandLineHelper
 			FCommandLine::Set(TEXT(""));
 
 			// command line text file included in the bundle
-			FString BundleCommandLineFilePath = FString([[NSBundle mainBundle] bundlePath]) + TEXT("/ue4commandline.txt");
+			FString BundleCommandLineFilePath = FString([[NSBundle mainBundle] bundlePath]) + TEXT("/uecommandline.txt");
 
 			// command line text file pushed to the documents folder
-			FString DocumentsCommandLineFilePath = FString([NSSearchPathForDirectoriesInDomains(NSDocumentDirectory, NSUserDomainMask, YES) objectAtIndex:0]) + TEXT("/ue4commandline.txt");
+			FString DocumentsCommandLineFilePath = FString([NSSearchPathForDirectoriesInDomains(NSDocumentDirectory, NSUserDomainMask, YES) objectAtIndex:0]) + TEXT("/uecommandline.txt");
 			
 #if UE_BUILD_SHIPPING
 			bool bIsTestFlightOrInternal = [[[[NSBundle mainBundle] appStoreReceiptURL] lastPathComponent] isEqualToString: @"sandboxReceipt"];
diff --git a/Engine/Source/Runtime/Launch/Private/Android/LaunchAndroid.cpp b/Engine/Source/Runtime/Launch/Private/Android/LaunchAndroid.cpp
--- a/Engine/Source/Runtime/Launch/Private/Android/LaunchAndroid.cpp
+++ b/Engine/Source/Runtime/Launch/Private/Android/LaunchAndroid.cpp
@@ -182,8 +182,8 @@ static void* AndroidEventThreadWorker(void* param);
 // How often to process (read & dispatch) events, in seconds.
 static const float EventRefreshRate = 1.0f / 20.0f;
 
-// Name of the UE4 commandline append setprop
-static constexpr char UE4CommandLineSetprop[] = "debug.ue4.commandline";
+// Name of the UE commandline append setprop
+static constexpr char UECommandLineSetprop[] = "debug.ue.commandline";
 
 //Android event callback functions
 static int32_t HandleInputCB(struct android_app* app, AInputEvent* event); //Touch and key input events
@@ -280,7 +280,7 @@ static void InitCommandLine()
 	FCommandLine::Set(TEXT(""));
 
 	AAssetManager* AssetMgr = AndroidThunkCpp_GetAssetManager();
-	AAsset* asset = AAssetManager_open(AssetMgr, TCHAR_TO_UTF8(TEXT("UE4CommandLine.txt")), AASSET_MODE_BUFFER);
+	AAsset* asset = AAssetManager_open(AssetMgr, TCHAR_TO_UTF8(TEXT("UECommandLine.txt")), AASSET_MODE_BUFFER);
 	if (nullptr != asset)
 	{
 		const void* FileContents = AAsset_getBuffer(asset);
@@ -304,12 +304,12 @@ static void InitCommandLine()
 	}
 
 	// read in the command line text file from the sdcard if it exists
-	FString CommandLineFilePath = GFilePathBase + FString("/UnrealGame/") + (!FApp::IsProjectNameEmpty() ? FApp::GetProjectName() : FPlatformProcess::ExecutableName()) + FString("/UE4CommandLine.txt");
+	FString CommandLineFilePath = GFilePathBase + FString("/UnrealGame/") + (!FApp::IsProjectNameEmpty() ? FApp::GetProjectName() : FPlatformProcess::ExecutableName()) + FString("/UECommandLine.txt");
 	FILE* CommandLineFile = fopen(TCHAR_TO_UTF8(*CommandLineFilePath), "r");
 	if(CommandLineFile == NULL)
 	{
 		// if that failed, try the lowercase version
-		CommandLineFilePath = CommandLineFilePath.Replace(TEXT("UE4CommandLine.txt"), TEXT("ue4commandline.txt"));
+		CommandLineFilePath = CommandLineFilePath.Replace(TEXT("UECommandLine.txt"), TEXT("uecommandline.txt"));
 		CommandLineFile = fopen(TCHAR_TO_UTF8(*CommandLineFilePath), "r");
 	}
 
@@ -341,11 +341,11 @@ static void InitCommandLine()
 	}
 
 	char CommandLineSetpropAppend[CMD_LINE_MAX];
-	if (__system_property_get(UE4CommandLineSetprop, CommandLineSetpropAppend) > 0)
+	if (__system_property_get(UECommandLineSetprop, CommandLineSetpropAppend) > 0)
 	{
 		FCommandLine::Append(UTF8_TO_TCHAR(" "));
 		FCommandLine::Append(UTF8_TO_TCHAR(CommandLineSetpropAppend));
-		FPlatformMisc::LowLevelOutputDebugStringf(TEXT("UE4 setprop appended: %s"), UTF8_TO_TCHAR(CommandLineSetpropAppend));
+		FPlatformMisc::LowLevelOutputDebugStringf(TEXT("UE setprop appended: %s"), UTF8_TO_TCHAR(CommandLineSetpropAppend));
 	}
 #endif
 }
diff --git a/Engine/Source/Runtime/Launch/Private/HoloLens/HoloLensLaunch.cpp b/Engine/Source/Runtime/Launch/Private/HoloLens/HoloLensLaunch.cpp
--- a/Engine/Source/Runtime/Launch/Private/HoloLens/HoloLensLaunch.cpp
+++ b/Engine/Source/Runtime/Launch/Private/HoloLens/HoloLensLaunch.cpp
@@ -628,13 +628,13 @@ void ViewProvider::Initialize(Windows::ApplicationModel::Core::CoreApplicationVi
 
 static void InitCommandLine(Platform::String^ args)
 {
-	// check for a ue4commandline.txt.
+	// check for a uecommandline.txt.
 	FCommandLine::Set(L"");
 
 	// FPaths should be good at this point (FPaths::EngineDir() is used during static init for instance)
 	// so be specific about the location we want, otherwise working directory changes can mess this up.
 	// Also this means the regular file interface is in place, so let's use that.
-	FString FileName = FString(FPaths::RootDir()) / TEXT("UE4CommandLine.txt");
+	FString FileName = FString(FPaths::RootDir()) / TEXT("UECommandLine.txt");
 
 	IFileHandle* CmdLineFileHandle = IPlatformFile::GetPlatformPhysical().OpenRead(*FileName);
 	if (CmdLineFileHandle != nullptr)
diff --git a/Engine/Source/Runtime/Launch/Private/Lumin/LaunchLumin.cpp b/Engine/Source/Runtime/Launch/Private/Lumin/LaunchLumin.cpp
--- a/Engine/Source/Runtime/Launch/Private/Lumin/LaunchLumin.cpp
+++ b/Engine/Source/Runtime/Launch/Private/Lumin/LaunchLumin.cpp
@@ -38,16 +38,16 @@ static void InitCommandLine()
 	// Adds command line args coming from lifecycle app init args.
 	FLuminPlatformMisc::InitLifecycle();	
 
-	// append space as UE4CommandLine.txt may or may not have a space at the start.
+	// append space as UECommandLine.txt may or may not have a space at the start.
 	FCommandLine::Append(TEXT(" "));
 
 	// read in the command line text file from the sdcard if it exists
-	FString CommandLineFilePath = FString(FPlatformProcess::BaseDir()) + FString("UE4CommandLine.txt");
+	FString CommandLineFilePath = FString(FPlatformProcess::BaseDir()) + FString("UECommandLine.txt");
 	FILE* CommandLineFile = fopen(TCHAR_TO_UTF8(*CommandLineFilePath), "r");
 	if (CommandLineFile == NULL)
 	{
 		// if that failed, try the lowercase version
-		CommandLineFilePath = CommandLineFilePath.Replace(TEXT("UE4CommandLine.txt"), TEXT("ue4commandline.txt"));
+		CommandLineFilePath = CommandLineFilePath.Replace(TEXT("UECommandLine.txt"), TEXT("uecommandline.txt"));
 		CommandLineFile = fopen(TCHAR_TO_UTF8(*CommandLineFilePath), "r");
 	}
 
