Add [STAThread] to AutomationTool.
This fixes a deadlock on shutdown when displaying UI.

diff --git a/Engine/Source/Programs/AutomationTool/Program.cs b/Engine/Source/Programs/AutomationTool/Program.cs
--- a/Engine/Source/Programs/AutomationTool/Program.cs
+++ b/Engine/Source/Programs/AutomationTool/Program.cs
@@ -22,6 +22,7 @@ namespace AutomationTool
 		/// </summary>
 		static ProcessManager.CtrlHandlerDelegate CtrlHandlerDelegateInstance = CtrlHandler;
 
+		[STAThread]
 		public static int Main(string[] Arguments)
 		{
             // Ensure UTF8Output flag is respected, since we are initializing logging early in the program.
