Fix quotes preventing fastlane from being executed in Turnkey on MacOS
Disable fastlane certification check on iOS for EA Release

diff --git a/Engine/Source/Programs/AutomationTool/IOS/IOSPlatform.Automation.cs b/Engine/Source/Programs/AutomationTool/IOS/IOSPlatform.Automation.cs
--- a/Engine/Source/Programs/AutomationTool/IOS/IOSPlatform.Automation.cs
+++ b/Engine/Source/Programs/AutomationTool/IOS/IOSPlatform.Automation.cs
@@ -361,7 +361,7 @@ public class IOSPlatform : Platform
 
 			if (!string.IsNullOrEmpty(CodeSigningIdentity))
 			{
-				Params += string.Format(" --identity '{0}'", CodeSigningIdentity);
+				Params += string.Format(" --identity \"{0}\"", CodeSigningIdentity);
 			}
 
 			if (!string.IsNullOrEmpty(Account))
@@ -441,6 +441,8 @@ public class IOSPlatform : Platform
 		return UnrealBuildTool.Utils.RunLocalProcessAndReturnStdOut("sh", FindCommand);
 	}
 
+	//Disabling for 5.0 early access as this code was not executing and has not been tested.
+	/*
 	public override bool UpdateHostPrerequisites(BuildCommand Command, ITurnkeyContext TurnkeyContext, bool bVerifyOnly)
 	{
 		int ExitCode;
@@ -514,7 +516,7 @@ public class IOSPlatform : Platform
 		// now look for a provision that can be used with a (maybe newly) instally cert
 		return Settings.RunRubyCommand(bVerifyOnly, DeviceName);
 	}
-
+*/
 
 	public override DeviceInfo[] GetDevices()
 	{
