Add InsufficientDevices as a specific TestResult type for Gauntlet.

diff --git a/Engine/Source/Programs/AutomationTool/Gauntlet/Framework/Base/Gauntlet.TestNode.cs b/Engine/Source/Programs/AutomationTool/Gauntlet/Framework/Base/Gauntlet.TestNode.cs
--- a/Engine/Source/Programs/AutomationTool/Gauntlet/Framework/Base/Gauntlet.TestNode.cs
+++ b/Engine/Source/Programs/AutomationTool/Gauntlet/Framework/Base/Gauntlet.TestNode.cs
@@ -24,6 +24,7 @@ namespace Gauntlet
 		WantRetry,
 		Cancelled,
 		TimedOut,
+		InsufficientDevices,
 	};
 
 	/// <summary>
