BuildGraph: Add tasks for running Git and Docker, and add support for deleting directories from <Delete> tasks.

diff --git a/Engine/Source/Programs/AutomationTool/AutomationUtils/CommandUtils.cs b/Engine/Source/Programs/AutomationTool/AutomationUtils/CommandUtils.cs
--- a/Engine/Source/Programs/AutomationTool/AutomationUtils/CommandUtils.cs
+++ b/Engine/Source/Programs/AutomationTool/AutomationUtils/CommandUtils.cs
@@ -2466,6 +2466,38 @@ namespace AutomationTool
 				Callback();
 			}
 		}
+
+		public static FileReference FindToolInPath(string ToolName)
+		{
+			string PathVariable = Environment.GetEnvironmentVariable("PATH");
+			foreach (string PathEntry in PathVariable.Split(Path.PathSeparator))
+			{
+				try
+				{
+					DirectoryReference PathDir = new DirectoryReference(PathEntry);
+					if (HostPlatform.Current.HostEditorPlatform == UnrealTargetPlatform.Win64)
+					{
+						FileReference ToolFile = FileReference.Combine(PathDir, $"{ToolName}.exe");
+						if (FileReference.Exists(ToolFile))
+						{
+							return ToolFile;
+						}
+					}
+					else
+					{
+						FileReference ToolFile = FileReference.Combine(PathDir, ToolName);
+						if (FileReference.Exists(ToolFile))
+						{
+							return ToolFile;
+						}
+					}
+				}
+				catch
+				{
+				}
+			}
+			return null;
+		}
 	}
 
 	/// <summary>
@@ -3127,7 +3159,7 @@ namespace AutomationTool
 				{
 					FinalFiles.Add(new FileReference(TargetFileInfo));
 				}
-			}			
+			}
 			CodeSignWindows.Sign(FinalFiles, CodeSignWindows.SignatureType.SHA1);
 			CodeSignWindows.Sign(FinalFiles.Where(x => !x.HasExtension(".msi")).ToList(), CodeSignWindows.SignatureType.SHA256); // MSI files can only have one signature; prefer SHA1 for compatibility
 		}
diff --git a/Engine/Source/Programs/AutomationTool/AutomationUtils/ProcessUtils.cs b/Engine/Source/Programs/AutomationTool/AutomationUtils/ProcessUtils.cs
--- a/Engine/Source/Programs/AutomationTool/AutomationUtils/ProcessUtils.cs
+++ b/Engine/Source/Programs/AutomationTool/AutomationUtils/ProcessUtils.cs
@@ -54,7 +54,7 @@ namespace AutomationTool
 		/// Creates a new process and adds it to the tracking list.
 		/// </summary>
 		/// <returns>New Process objects</returns>
-		public static IProcessResult CreateProcess(string AppName, bool bAllowSpew, bool bCaptureSpew, Dictionary<string, string> Env = null, LogEventType SpewVerbosity = LogEventType.Console, ProcessResult.SpewFilterCallbackType SpewFilterCallback = null)
+		public static IProcessResult CreateProcess(string AppName, bool bAllowSpew, bool bCaptureSpew, Dictionary<string, string> Env = null, LogEventType SpewVerbosity = LogEventType.Console, ProcessResult.SpewFilterCallbackType SpewFilterCallback = null, string WorkingDir = null)
 		{
 			var NewProcess = HostPlatform.Current.CreateProcess(AppName);
 			if (Env != null)
@@ -71,6 +71,11 @@ namespace AutomationTool
 					}
 				}
 			}
+			if (WorkingDir != null)
+			{
+				NewProcess.StartInfo.WorkingDirectory = WorkingDir;
+			}
+
 			var Result = new ProcessResult(AppName, NewProcess, bAllowSpew, bCaptureSpew, SpewVerbosity: SpewVerbosity, InSpewFilterCallback: SpewFilterCallback);
 			AddProcess(Result);
 			return Result;
@@ -810,7 +815,7 @@ namespace AutomationTool
 		/// <param name="Env">Environment to pass to program.</param>
 		/// <param name="FilterCallback">Callback to filter log spew before output.</param>
 		/// <returns>Object containing the exit code of the program as well as it's stdout output.</returns>
-		public static IProcessResult Run(string App, string CommandLine = null, string Input = null, ERunOptions Options = ERunOptions.Default, Dictionary<string, string> Env = null, ProcessResult.SpewFilterCallbackType SpewFilterCallback = null, string Identifier = null)
+		public static IProcessResult Run(string App, string CommandLine = null, string Input = null, ERunOptions Options = ERunOptions.Default, Dictionary<string, string> Env = null, ProcessResult.SpewFilterCallbackType SpewFilterCallback = null, string Identifier = null, string WorkingDir = null)
 		{
 			App = ConvertSeparators(PathSeparator.Default, App);
 
@@ -840,7 +845,7 @@ namespace AutomationTool
                 LogWithVerbosity(SpewVerbosity,"Running: " + App + " " + (String.IsNullOrEmpty(CommandLine) ? "" : CommandLine));
             }
 
-			IProcessResult Result = ProcessManager.CreateProcess(App, Options.HasFlag(ERunOptions.AllowSpew), !Options.HasFlag(ERunOptions.NoStdOutCapture), Env, SpewVerbosity: SpewVerbosity, SpewFilterCallback: SpewFilterCallback);
+			IProcessResult Result = ProcessManager.CreateProcess(App, Options.HasFlag(ERunOptions.AllowSpew), !Options.HasFlag(ERunOptions.NoStdOutCapture), Env, SpewVerbosity: SpewVerbosity, SpewFilterCallback: SpewFilterCallback, WorkingDir: WorkingDir);
 			using (LogIndentScope Scope = Options.HasFlag(ERunOptions.AllowSpew) ? new LogIndentScope("  ") : null)
 			{
 				Process Proc = Result.ProcessObject;
diff --git a/Engine/Source/Programs/AutomationTool/BuildGraph/Tasks/DeleteTask.cs b/Engine/Source/Programs/AutomationTool/BuildGraph/Tasks/DeleteTask.cs
--- a/Engine/Source/Programs/AutomationTool/BuildGraph/Tasks/DeleteTask.cs
+++ b/Engine/Source/Programs/AutomationTool/BuildGraph/Tasks/DeleteTask.cs
@@ -21,9 +21,15 @@ namespace BuildGraph.Tasks
 		/// <summary>
 		/// List of file specifications separated by semicolons (for example, *.cpp;Engine/.../*.bat), or the name of a tag set
 		/// </summary>
-		[TaskParameter(ValidationType = TaskParameterValidationType.FileSpec)]
+		[TaskParameter(Optional = true, ValidationType = TaskParameterValidationType.FileSpec)]
 		public string Files;
 
+		/// <summary>
+		/// List of directory names
+		/// </summary>
+		[TaskParameter(Optional = true)]
+		public string Directories;
+
 		/// <summary>
 		/// Whether to delete empty directories after deleting the files. Defaults to true.
 		/// </summary>
@@ -59,38 +65,52 @@ namespace BuildGraph.Tasks
 		/// <param name="TagNameToFileSet">Mapping from tag names to the set of files they include</param>
 		public override void Execute(JobContext Job, HashSet<FileReference> BuildProducts, Dictionary<string, HashSet<FileReference>> TagNameToFileSet)
 		{
-			// Find all the referenced files and delete them
-			HashSet<FileReference> Files = ResolveFilespec(CommandUtils.RootDirectory, Parameters.Files, TagNameToFileSet);
-			foreach(FileReference File in Files)
+			if (Parameters.Files != null)
 			{
-				if (!InternalUtils.SafeDeleteFile(File.FullName))
+				// Find all the referenced files and delete them
+				HashSet<FileReference> Files = ResolveFilespec(CommandUtils.RootDirectory, Parameters.Files, TagNameToFileSet);
+				foreach (FileReference File in Files)
 				{
-					CommandUtils.LogWarning("Couldn't delete file {0}", File.FullName);
+					if (!InternalUtils.SafeDeleteFile(File.FullName))
+					{
+						CommandUtils.LogWarning("Couldn't delete file {0}", File.FullName);
+					}
 				}
-			}
 
-			// Try to delete all the parent directories. Keep track of the directories we've already deleted to avoid hitting the disk.
-			if(Parameters.DeleteEmptyDirectories)
-			{
-				// Find all the directories that we're touching
-				HashSet<DirectoryReference> ParentDirectories = new HashSet<DirectoryReference>();
-				foreach(FileReference File in Files)
+				// Try to delete all the parent directories. Keep track of the directories we've already deleted to avoid hitting the disk.
+				if (Parameters.DeleteEmptyDirectories)
 				{
-					ParentDirectories.Add(File.Directory);
-				}
+					// Find all the directories that we're touching
+					HashSet<DirectoryReference> ParentDirectories = new HashSet<DirectoryReference>();
+					foreach (FileReference File in Files)
+					{
+						ParentDirectories.Add(File.Directory);
+					}
 
-				// Recurse back up from each of those directories to the root folder
-				foreach(DirectoryReference ParentDirectory in ParentDirectories)
-				{
-					for(DirectoryReference CurrentDirectory = ParentDirectory; CurrentDirectory != CommandUtils.RootDirectory; CurrentDirectory = CurrentDirectory.ParentDirectory)
+					// Recurse back up from each of those directories to the root folder
+					foreach (DirectoryReference ParentDirectory in ParentDirectories)
 					{
-						if(!TryDeleteEmptyDirectory(CurrentDirectory))
+						for (DirectoryReference CurrentDirectory = ParentDirectory; CurrentDirectory != CommandUtils.RootDirectory; CurrentDirectory = CurrentDirectory.ParentDirectory)
 						{
-							break;
+							if (!TryDeleteEmptyDirectory(CurrentDirectory))
+							{
+								break;
+							}
 						}
 					}
 				}
 			}
+			if (Parameters.Directories != null)
+			{
+				foreach (string Directory in Parameters.Directories.Split(';'))
+				{
+					if (!String.IsNullOrEmpty(Directory))
+					{
+						DirectoryReference FullDir = new DirectoryReference(Directory);
+						FileUtils.ForceDeleteDirectory(FullDir);
+					}
+				}
+			}
 		}
 
 		/// <summary>
diff --git a/Engine/Source/Programs/AutomationTool/BuildGraph/Tasks/DockerTask.cs b/Engine/Source/Programs/AutomationTool/BuildGraph/Tasks/DockerTask.cs
new file mode 100644
--- /dev/null
+++ b/Engine/Source/Programs/AutomationTool/BuildGraph/Tasks/DockerTask.cs
@@ -0,0 +1,102 @@
+// Copyright Epic Games, Inc. All Rights Reserved.
+
+using EpicGames.Core;
+using System;
+using System.Collections.Generic;
+using System.Text;
+using System.Xml;
+
+namespace AutomationTool.Tasks
+{
+	/// <summary>
+	/// Parameters for a Docker task
+	/// </summary>
+	public class DockerTaskParameters
+	{
+		/// <summary>
+		/// Docker command line arguments
+		/// </summary>
+		[TaskParameter(Optional = true)]
+		public string Arguments;
+
+		/// <summary>
+		/// Base directory for running the command
+		/// </summary>
+		[TaskParameter(Optional = true)]
+		public string BaseDir;
+
+		/// <summary>
+		/// The minimum exit code, which is treated as an error.
+		/// </summary>
+		[TaskParameter(Optional = true)]
+		public int ErrorLevel = 1;
+	}
+
+	/// <summary>
+	/// Spawns Docker and waits for it to complete.
+	/// </summary>
+	[TaskElement("Docker", typeof(DockerTaskParameters))]
+	public class DockerTask : CustomTask
+	{
+		/// <summary>
+		/// Parameters for this task
+		/// </summary>
+		DockerTaskParameters Parameters;
+
+		/// <summary>
+		/// Construct a Docker task
+		/// </summary>
+		/// <param name="InParameters">Parameters for the task</param>
+		public DockerTask(DockerTaskParameters InParameters)
+		{
+			Parameters = InParameters;
+		}
+
+		/// <summary>
+		/// Execute the task.
+		/// </summary>
+		/// <param name="Job">Information about the current job</param>
+		/// <param name="BuildProducts">Set of build products produced by this node.</param>
+		/// <param name="TagNameToFileSet">Mapping from tag names to the set of files they include</param>
+		public override void Execute(JobContext Job, HashSet<FileReference> BuildProducts, Dictionary<string, HashSet<FileReference>> TagNameToFileSet)
+		{
+			FileReference ToolFile = CommandUtils.FindToolInPath("docker");
+			if(ToolFile == null)
+			{
+				throw new AutomationException("Unable to find path to Docker. Check you have it installed, and it is on your PATH.");
+			}
+
+			IProcessResult Result = CommandUtils.Run(ToolFile.FullName, Parameters.Arguments, WorkingDir: Parameters.BaseDir);
+			if (Result.ExitCode < 0 || Result.ExitCode >= Parameters.ErrorLevel)
+			{
+				throw new AutomationException("Docker terminated with an exit code indicating an error ({0})", Result.ExitCode);
+			}
+		}
+
+		/// <summary>
+		/// Output this task out to an XML writer.
+		/// </summary>
+		public override void Write(XmlWriter Writer)
+		{
+			Write(Writer, Parameters);
+		}
+
+		/// <summary>
+		/// Find all the tags which are used as inputs to this task
+		/// </summary>
+		/// <returns>The tag names which are read by this task</returns>
+		public override IEnumerable<string> FindConsumedTagNames()
+		{
+			yield break;
+		}
+
+		/// <summary>
+		/// Find all the tags which are modified by this task
+		/// </summary>
+		/// <returns>The tag names which are modified by this task</returns>
+		public override IEnumerable<string> FindProducedTagNames()
+		{
+			yield break;
+		}
+	}
+}
diff --git a/Engine/Source/Programs/AutomationTool/BuildGraph/Tasks/GitTask.cs b/Engine/Source/Programs/AutomationTool/BuildGraph/Tasks/GitTask.cs
new file mode 100644
--- /dev/null
+++ b/Engine/Source/Programs/AutomationTool/BuildGraph/Tasks/GitTask.cs
@@ -0,0 +1,102 @@
+// Copyright Epic Games, Inc. All Rights Reserved.
+
+using EpicGames.Core;
+using System;
+using System.Collections.Generic;
+using System.Text;
+using System.Xml;
+
+namespace AutomationTool.Tasks
+{
+	/// <summary>
+	/// Parameters for a Git task
+	/// </summary>
+	public class GitTaskParameters
+	{
+		/// <summary>
+		/// Git command line arguments
+		/// </summary>
+		[TaskParameter(Optional = true)]
+		public string Arguments;
+
+		/// <summary>
+		/// Base directory for running the command
+		/// </summary>
+		[TaskParameter(Optional = true)]
+		public string BaseDir;
+
+		/// <summary>
+		/// The minimum exit code, which is treated as an error.
+		/// </summary>
+		[TaskParameter(Optional = true)]
+		public int ErrorLevel = 1;
+	}
+
+	/// <summary>
+	/// Spawns Git and waits for it to complete.
+	/// </summary>
+	[TaskElement("Git", typeof(GitTaskParameters))]
+	public class GitTask : CustomTask
+	{
+		/// <summary>
+		/// Parameters for this task
+		/// </summary>
+		GitTaskParameters Parameters;
+
+		/// <summary>
+		/// Construct a Git task
+		/// </summary>
+		/// <param name="InParameters">Parameters for the task</param>
+		public GitTask(GitTaskParameters InParameters)
+		{
+			Parameters = InParameters;
+		}
+
+		/// <summary>
+		/// Execute the task.
+		/// </summary>
+		/// <param name="Job">Information about the current job</param>
+		/// <param name="BuildProducts">Set of build products produced by this node.</param>
+		/// <param name="TagNameToFileSet">Mapping from tag names to the set of files they include</param>
+		public override void Execute(JobContext Job, HashSet<FileReference> BuildProducts, Dictionary<string, HashSet<FileReference>> TagNameToFileSet)
+		{
+			FileReference ToolFile = CommandUtils.FindToolInPath("git");
+			if(ToolFile == null)
+			{
+				throw new AutomationException("Unable to find path to Git. Check you have it installed, and it is on your PATH.");
+			}
+
+			IProcessResult Result = CommandUtils.Run(ToolFile.FullName, Parameters.Arguments, WorkingDir: Parameters.BaseDir);
+			if (Result.ExitCode < 0 || Result.ExitCode >= Parameters.ErrorLevel)
+			{
+				throw new AutomationException("Git terminated with an exit code indicating an error ({0})", Result.ExitCode);
+			}
+		}
+
+		/// <summary>
+		/// Output this task out to an XML writer.
+		/// </summary>
+		public override void Write(XmlWriter Writer)
+		{
+			Write(Writer, Parameters);
+		}
+
+		/// <summary>
+		/// Find all the tags which are used as inputs to this task
+		/// </summary>
+		/// <returns>The tag names which are read by this task</returns>
+		public override IEnumerable<string> FindConsumedTagNames()
+		{
+			yield break;
+		}
+
+		/// <summary>
+		/// Find all the tags which are modified by this task
+		/// </summary>
+		/// <returns>The tag names which are modified by this task</returns>
+		public override IEnumerable<string> FindProducedTagNames()
+		{
+			yield break;
+		}
+	}
+}
