Remove [STAThread] attribute from UAT main entry point.
 - STAThread causes deadlocks with platform COM automation APIs.
 - Reverts change from CL 16094080.

diff --git a/Engine/Source/Programs/AutomationTool/Program.cs b/Engine/Source/Programs/AutomationTool/Program.cs
--- a/Engine/Source/Programs/AutomationTool/Program.cs
+++ b/Engine/Source/Programs/AutomationTool/Program.cs
@@ -22,7 +22,7 @@ namespace AutomationTool
 		/// </summary>
 		static ProcessManager.CtrlHandlerDelegate CtrlHandlerDelegateInstance = CtrlHandler;
 
-		[STAThread]
+		// Do not add [STAThread] here. It will cause deadlocks in platform automation code.
 		public static int Main(string[] Arguments)
 		{
             // Ensure UTF8Output flag is respected, since we are initializing logging early in the program.
