Move Turnkey editor dialog to an ApartmentState.STA thread so it doesn't deadlock on shutdown

diff --git a/Engine/Source/Programs/AutomationTool/Turnkey/Providers/HybridIOProvider.cs b/Engine/Source/Programs/AutomationTool/Turnkey/Providers/HybridIOProvider.cs
--- a/Engine/Source/Programs/AutomationTool/Turnkey/Providers/HybridIOProvider.cs
+++ b/Engine/Source/Programs/AutomationTool/Turnkey/Providers/HybridIOProvider.cs
@@ -149,9 +149,17 @@ namespace Turnkey
 		{
 // disable unreachable code warning as this fails on linux builds due to this method not being implemented
 #pragma warning disable 0162
-			string Result;
+			string Result = "";
 #if WINDOWS
-			Result = ShowFormsDialogs(Prompt, Default);
+
+			System.Threading.Thread t = new System.Threading.Thread(x =>
+			{
+				Result = ShowFormsDialogs(Prompt, Default);
+			});
+
+			t.SetApartmentState(System.Threading.ApartmentState.STA);
+			t.Start();
+			t.Join();
 #elif OSX
 			Result = ShowMacDialog(Prompt, Default);
 #else
