Add multihomehttp commandline argument to smoke test command so that overridden ip is written into matchmaking session information.

diff --git a/Engine/Source/Programs/AutomationTool/Gauntlet/Unreal/Game/EpicGame.TestConfig.cs b/Engine/Source/Programs/AutomationTool/Gauntlet/Unreal/Game/EpicGame.TestConfig.cs
--- a/Engine/Source/Programs/AutomationTool/Gauntlet/Unreal/Game/EpicGame.TestConfig.cs
+++ b/Engine/Source/Programs/AutomationTool/Gauntlet/Unreal/Game/EpicGame.TestConfig.cs
@@ -176,7 +176,7 @@ namespace EpicGame
 				// Do we need to add the -multihome argument to bind to specific IP?
 				if (ConfigRole.RoleType.IsServer() && (MultipleInterfaces || !string.IsNullOrEmpty(RequestedServerIP)))
 				{
-					AppConfig.CommandLine += string.Format(" -multihome={0}", ServerIP);
+					AppConfig.CommandLine += string.Format(" -multihome={0} -multihomehttp={1}", ServerIP, ServerIP);
 				}
 
 				// client too, but only desktop platforms
@@ -184,7 +184,7 @@ namespace EpicGame
 				{
 					if (ConfigRole.Platform == UnrealTargetPlatform.Win64 || ConfigRole.Platform == UnrealTargetPlatform.Mac)
 					{
-						AppConfig.CommandLine += string.Format(" -multihome={0}", ClientIP);
+						AppConfig.CommandLine += string.Format(" -multihome={0} -multihomehttp={1}", ClientIP, ClientIP);
 					}
 				}
 
