CopyTask - optionally throw an error if files not found

diff --git a/Engine/Source/Programs/AutomationTool/BuildGraph/Tasks/CopyTask.cs b/Engine/Source/Programs/AutomationTool/BuildGraph/Tasks/CopyTask.cs
--- a/Engine/Source/Programs/AutomationTool/BuildGraph/Tasks/CopyTask.cs
+++ b/Engine/Source/Programs/AutomationTool/BuildGraph/Tasks/CopyTask.cs
@@ -47,6 +47,12 @@ namespace BuildGraph.Tasks
 		/// </summary>
 		[TaskParameter(Optional = true, ValidationType = TaskParameterValidationType.TagList)]
 		public string Tag;
+
+		/// <summary>
+		/// Whether or not to throw an error if no files were found to copy
+		/// </summary>
+		[TaskParameter(Optional = true)]
+		public bool ErrorIfNotFound = false;
 	}
 
 	/// <summary>
@@ -111,7 +117,14 @@ namespace BuildGraph.Tasks
 			// Check we got some files
 			if(TargetFileToSourceFile.Count == 0)
 			{
-				CommandUtils.LogInformation("No files found matching '{0}'", SourcePattern);
+				if (Parameters.ErrorIfNotFound)
+				{
+					CommandUtils.LogError("No files found matching '{0}'", SourcePattern);
+				}
+				else
+				{
+					CommandUtils.LogInformation("No files found matching '{0}'", SourcePattern);
+				}
 				return;
 			}
 
