Only try to delete folders if they exist.

#rnx

diff --git a/Engine/Source/Programs/AutomationTool/BuildGraph/Tasks/DeleteTask.cs b/Engine/Source/Programs/AutomationTool/BuildGraph/Tasks/DeleteTask.cs
--- a/Engine/Source/Programs/AutomationTool/BuildGraph/Tasks/DeleteTask.cs
+++ b/Engine/Source/Programs/AutomationTool/BuildGraph/Tasks/DeleteTask.cs
@@ -107,7 +107,10 @@ namespace BuildGraph.Tasks
 					if (!String.IsNullOrEmpty(Directory))
 					{
 						DirectoryReference FullDir = new DirectoryReference(Directory);
-						FileUtils.ForceDeleteDirectory(FullDir);
+						if (DirectoryReference.Exists(FullDir))
+						{
+							FileUtils.ForceDeleteDirectory(FullDir);
+						}
 					}
 				}
 			}
