LuminPlatform.Automation:

Remove some references to Mono.

diff --git a/Engine/Source/Programs/AutomationTool/Lumin/LuminPlatform.Automation.cs b/Engine/Source/Programs/AutomationTool/Lumin/LuminPlatform.Automation.cs
--- a/Engine/Source/Programs/AutomationTool/Lumin/LuminPlatform.Automation.cs
+++ b/Engine/Source/Programs/AutomationTool/Lumin/LuminPlatform.Automation.cs
@@ -781,7 +781,7 @@ public class LuminPlatform : Platform
 		File.Copy(UninstallBatchName, SC.CookSourceRuntimeRootDir + "\\" + Path.GetFileName(UninstallBatchName).ToString(), true);
 
 		// If needed, make the batch files able to execute
-		if (Utils.IsRunningOnMono)
+		if (!Utils.IsRunningOnWindows)
 		{
 			CommandUtils.FixUnixFilePermissions(BatchName);
 			CommandUtils.FixUnixFilePermissions(UninstallBatchName);
@@ -1007,6 +1007,7 @@ public class LuminPlatform : Platform
 					EntriesToDeploy.Add(ExeTimestampFileName);
 
 					// mono is bugging out making stderr pipes
+					// Note: as of UE5, mono is no longer used. bGoSlow could be removed, pending testing.
 					bool bGoSlow = HostPlatform.Current.HostEditorPlatform != UnrealTargetPlatform.Win64;
 
 					// We now have a minimal set of file & dir entries we need
