HoloLensPlatform.Automation:

Remove some references to Mono

#trivial

diff --git a/Engine/Source/Programs/AutomationTool/HoloLens/HoloLensPlatform.Automation.cs b/Engine/Source/Programs/AutomationTool/HoloLens/HoloLensPlatform.Automation.cs
--- a/Engine/Source/Programs/AutomationTool/HoloLens/HoloLensPlatform.Automation.cs
+++ b/Engine/Source/Programs/AutomationTool/HoloLens/HoloLensPlatform.Automation.cs
@@ -12,14 +12,11 @@ using System.Diagnostics;
 using EpicGames.Core;
 using System.Xml;
 using System.Text.RegularExpressions;
-
-#if !__MonoCS__
 using Windows.Management.Deployment;
 using Windows.Foundation;
 #if false
 using Windows.ApplicationModel.Core;
 #endif
-#endif
 using System.Threading;
 using System.Security.Cryptography.X509Certificates;
 using Microsoft.Tools.WindowsDevicePortal;
@@ -196,7 +193,6 @@ namespace HoloLens.Automation
 		}
 	}
 
-#if !__MonoCS__
 	class HoloLensDevicePortalCreatedProcess : IProcessResult
 	{
 		object StateLock;
@@ -349,7 +345,6 @@ namespace HoloLens.Automation
 			}
 		}
 	}
-#endif
 
 	public class HoloLensPlatform : Platform
     {
@@ -1244,7 +1239,6 @@ namespace HoloLens.Automation
 			return !result;
 		}
 
-#if !__MonoCS__
 		private void WaitFor(IAsyncOperationWithProgress<DeploymentResult, DeploymentProgress> deploymentOperation)
 		{
 			// This event is signaled when the operation completes
@@ -1278,12 +1272,10 @@ namespace HoloLens.Automation
 				//LogInformation("Operation status unknown");
 			}
 		}
-#endif
 
 		private void DeployToLocalDevice(ProjectParams Params, DeploymentContext SC)
 		{
-#if !__MonoCS__
-            if (Utils.IsRunningOnMono)
+            if (!Utils.IsRunningOnWindows)
             {
                 return;
             }
@@ -1335,12 +1327,15 @@ namespace HoloLens.Automation
 			{
 				LogWarning("Failed to apply a loopback exemption to the deployed app.  Connection to a local cook server will fail.");
 			}
-#endif
 		}
 
 		private WindowsArchitecture RemoteDeviceArchitecture(string DeviceAddress, ProjectParams Params)
 		{
-#if !__MonoCS__
+			if (!Utils.IsRunningOnWindows)
+			{
+				return WindowsArchitecture.x64;
+			}
+
 			string OsVersionString = "";
 			try
 			{
@@ -1392,15 +1387,11 @@ namespace HoloLens.Automation
 				LogError(String.Format("Unsupported OS architecture {0} from {1} device", osArchBlock, DeviceAddress));
 				throw new AutomationException(ExitCode.Error_AppInstallFailed, "Unsupported OS architecture");
 			}
-#else
-			return WindowsArchitecture.x64;
-#endif
 		}
 
 		private void DeployToRemoteDevice(string DeviceAddress, ProjectParams Params, DeploymentContext SC)
 		{
-#if !__MonoCS__
-			if (Utils.IsRunningOnMono)
+			if (!Utils.IsRunningOnWindows)
             {
                 return;
             }
@@ -1483,7 +1474,6 @@ namespace HoloLens.Automation
 			{
 				throw new AutomationException(ExitCode.Error_AppInstallFailed, "Remote deployment of unpackaged apps is not supported.");
 			}
-#endif
 		}
 
 
@@ -1501,7 +1491,11 @@ namespace HoloLens.Automation
 #if false
 		private async Task<HoloLensLauncherCreatedProcess> RunUsingLauncherToolAsync(ProjectParams Params, ERunOptions ClientRunFlags)
 		{
-#if !__MonoCS__
+			if (!Utils.IsRunningOnWindows)
+            {
+                return;
+            }
+
 			string Name;
 			string Publisher;
 			GetPackageInfo(Params, out Name, out Publisher);
@@ -1567,9 +1561,6 @@ namespace HoloLens.Automation
 				UwpProcessResult.DisposeProcess();
 			}
 			return UwpProcessResult;
-#else
-			return null;
-#endif
 		}
 #else
 		private Task<HoloLensLauncherCreatedProcess> RunUsingLauncherToolAsync(ProjectParams Params, ERunOptions ClientRunFlags)
@@ -1580,8 +1571,7 @@ namespace HoloLens.Automation
 
 		private IProcessResult RunUsingDevicePortal(string DeviceAddress, ERunOptions ClientRunFlags, string ClientApp, string ClientCmdLine, ProjectParams Params)
 		{
-#if !__MonoCS__
-            if (Utils.IsRunningOnMono)
+            if (!Utils.IsRunningOnWindows)
             {
                 return null;
             }
@@ -1647,9 +1637,6 @@ namespace HoloLens.Automation
 			{
 				throw new AutomationException(ExitCode.Error_LauncherFailed, e, e.Message);
 			}
-#else
-			return null;
-#endif
 		}
 
 		private void GetPackageInfo(ProjectParams Params, out string Name, out string Publisher)
@@ -1684,8 +1671,7 @@ namespace HoloLens.Automation
 
 		private bool ShouldAcceptCertificate(System.Security.Cryptography.X509Certificates.X509Certificate2 Certificate, bool Unattended)
 		{
-#if !__MonoCS__
-            if (Utils.IsRunningOnMono)
+            if (!Utils.IsRunningOnWindows)
             {
                 return false;
             }
@@ -1713,12 +1699,8 @@ namespace HoloLens.Automation
 			}
 
 			throw new AutomationException(ExitCode.Error_CertificateNotFound, "Cannot connect to remote device: certificate is untrusted and user declined to accept.");
-#else
-			return false;
-#endif
 		}
 
-#if !__MonoCS__
 		private void Portal_AppInstallStatus(Microsoft.Tools.WindowsDevicePortal.DevicePortal sender, Microsoft.Tools.WindowsDevicePortal.ApplicationInstallStatusEventArgs args)
 		{
 			if (args.Status == Microsoft.Tools.WindowsDevicePortal.ApplicationInstallStatus.Failed)
@@ -1742,7 +1724,6 @@ namespace HoloLens.Automation
 				LogLog(args.Message);
 			}
 		}
-#endif
 
         private string[] GetPathToVCLibsPackages(bool UseDebugCrt, WindowsCompiler Compiler)
 		{
