TempStorage.cs:

Remove some mono references.

#trivial

diff --git a/Engine/Source/Programs/AutomationTool/BuildGraph/TempStorage.cs b/Engine/Source/Programs/AutomationTool/BuildGraph/TempStorage.cs
--- a/Engine/Source/Programs/AutomationTool/BuildGraph/TempStorage.cs
+++ b/Engine/Source/Programs/AutomationTool/BuildGraph/TempStorage.cs
@@ -965,101 +965,48 @@ namespace AutomationTool
 			// We want to end with the smaller files to more effectively fill in the gaps
 			var FilesToZip = new ConcurrentQueue<FileReference>(FilesInfo.OrderByDescending(FileInfo => FileInfo.FileSize).Select(FileInfo => FileInfo.File));
 
+			ConcurrentBag<FileInfo> ZipFiles = new ConcurrentBag<FileInfo>();
+
+			DirectoryReference ZipDir = StagingDir ?? OutputDir;
+
 			// We deliberately avoid Parallel.ForEach here because profiles have shown that dynamic partitioning creates
 			// too many zip files, and they can be of too varying size, creating uneven work when unzipping later,
 			// as ZipFile cannot unzip files in parallel from a single archive.
 			// We can safely assume the build system will not be doing more important things at the same time, so we simply use all our logical cores,
 			// which has shown to be optimal via profiling, and limits the number of resulting zip files to the number of logical cores.
-			// 
-			// Sadly, mono implementation of System.IO.Compression is really poor (as of 2015/Aug), causing OOM when parallel zipping a large set of files.
-			// However, Ionic is MUCH slower than .NET's native implementation (2x+ slower in our build farm), so we stick to the faster solution on PC.
-			// The code duplication in the threadprocs is unfortunate here, and hopefully we can settle on .NET's implementation on both platforms eventually.
-			List<Thread> ZipThreads;
-
-			ConcurrentBag<FileInfo> ZipFiles = new ConcurrentBag<FileInfo>();
-
-			DirectoryReference ZipDir = StagingDir ?? OutputDir;
-			// Disabled for net core as using ionic seems to crash when run on build machines (causes a process panic)
-			/*if (Utils.IsRunningOnMono)
-			{
-				ZipThreads = (
-					from CoreNum in Enumerable.Range(0, bZipInParallel ? Environment.ProcessorCount : 1)
-					select new Thread((object indexObject) =>
+			List<Thread> ZipThreads = (
+				from CoreNum in Enumerable.Range(0, bZipInParallel ? Environment.ProcessorCount : 1)
+				select new Thread((object indexObject) =>
+				{
+					int index = (int)indexObject;
+					var ZipFileName = FileReference.Combine(ZipDir, string.Format("{0}{1}.zip", ZipBaseName, bZipInParallel ? "-" + index.ToString("00") : ""));
+					// don't create the zip unless we have at least one file to add
+					FileReference File;
+					if (FilesToZip.TryDequeue(out File))
 					{
-						int index = (int) indexObject;
-						var ZipFileName = FileReference.Combine(ZipDir, string.Format("{0}{1}.zip", ZipBaseName, bZipInParallel ? "-" + index.ToString("00") : ""));
-
-						// don't create the zip unless we have at least one file to add
-						FileReference File;
-						if (FilesToZip.TryDequeue(out File))
+						// Create one zip per thread using the given basename
+						using (var ZipArchive = System.IO.Compression.ZipFile.Open(ZipFileName.FullName, System.IO.Compression.ZipArchiveMode.Create))
 						{
-							// Create one zip per thread using the given basename
-							using (var ZipArchive = new Ionic.Zip.ZipFile(ZipFileName.FullName) { CompressionLevel = Ionic.Zlib.CompressionLevel.BestSpeed })
-							{
-								ZipArchive.UseZip64WhenSaving = Ionic.Zip.Zip64Option.AsNecessary;
-
-								// pull from the queue until we are out of files.
-								do
-								{
-									// use fastest compression. In our best case we are CPU bound, so this is a good tradeoff,
-									// cutting overall time by 2/3 while only modestly increasing the compression ratio (22.7% -> 23.8% for RootEditor PDBs).
-									// This is in cases of a super hot cache, so the operation was largely CPU bound.
-									ZipArchive.AddFile(File.FullName, CommandUtils.ConvertSeparators(PathSeparator.Slash, File.Directory.MakeRelativeTo(RootDir)));
-								} while (FilesToZip.TryDequeue(out File));
-								ZipArchive.Save();
-							}
-							// if we are using a staging dir, copy to the final location and delete the staged copy.
-							FileInfo ZipFile = new FileInfo(ZipFileName.FullName);
-							if (StagingDir != null)
+							// pull from the queue until we are out of files.
+							do
 							{
-								FileInfo NewZipFile = ZipFile.CopyTo(CommandUtils.MakeRerootedFilePath(ZipFile.FullName, StagingDir.FullName, OutputDir.FullName));
-								ZipFile.Delete();
-								ZipFile = NewZipFile;
-							}
-							ZipFiles.Add(ZipFile);
+								// use fastest compression. In our best case we are CPU bound, so this is a good tradeoff,
+								// cutting overall time by 2/3 while only modestly increasing the compression ratio (22.7% -> 23.8% for RootEditor PDBs).
+								// This is in cases of a super hot cache, so the operation was largely CPU bound.
+								ZipArchiveExtensions.CreateEntryFromFile_CrossPlatform(ZipArchive, File.FullName, CommandUtils.ConvertSeparators(PathSeparator.Slash, File.MakeRelativeTo(RootDir)), System.IO.Compression.CompressionLevel.Fastest);
+							} while (FilesToZip.TryDequeue(out File));
 						}
-					})).ToList();
-			}
-			else*/
-
-			{
-				ZipThreads = (
-					from CoreNum in Enumerable.Range(0, bZipInParallel ? Environment.ProcessorCount : 1)
-					select new Thread((object indexObject) =>
-					{
-						int index = (int) indexObject;
-						var ZipFileName = FileReference.Combine(ZipDir, string.Format("{0}{1}.zip", ZipBaseName, bZipInParallel ? "-" + index.ToString("00") : ""));
-						// don't create the zip unless we have at least one file to add
-						FileReference File;
-						if (FilesToZip.TryDequeue(out File))
+						// if we are using a staging dir, copy to the final location and delete the staged copy.
+						FileInfo ZipFile = new FileInfo(ZipFileName.FullName);
+						if (StagingDir != null)
 						{
-							// Create one zip per thread using the given basename
-							using (var ZipArchive = System.IO.Compression.ZipFile.Open(ZipFileName.FullName, System.IO.Compression.ZipArchiveMode.Create))
-							{
-
-								// pull from the queue until we are out of files.
-								do
-								{
-									// use fastest compression. In our best case we are CPU bound, so this is a good tradeoff,
-									// cutting overall time by 2/3 while only modestly increasing the compression ratio (22.7% -> 23.8% for RootEditor PDBs).
-									// This is in cases of a super hot cache, so the operation was largely CPU bound.
-									// Also, sadly, mono appears to have a bug where nothing you can do will properly set the LastWriteTime on the created entry,
-									// so we have to ignore timestamps on files extracted from a zip, since it may have been created on a Mac.
-									ZipArchiveExtensions.CreateEntryFromFile_CrossPlatform(ZipArchive, File.FullName, CommandUtils.ConvertSeparators(PathSeparator.Slash, File.MakeRelativeTo(RootDir)), System.IO.Compression.CompressionLevel.Fastest);
-								} while (FilesToZip.TryDequeue(out File));
-							}
-							// if we are using a staging dir, copy to the final location and delete the staged copy.
-							FileInfo ZipFile = new FileInfo(ZipFileName.FullName);
-							if (StagingDir != null)
-							{
-								FileInfo NewZipFile = ZipFile.CopyTo(CommandUtils.MakeRerootedFilePath(ZipFile.FullName, StagingDir.FullName, OutputDir.FullName));
-								ZipFile.Delete();
-								ZipFile = NewZipFile;
-							}
-							ZipFiles.Add(ZipFile);
+							FileInfo NewZipFile = ZipFile.CopyTo(CommandUtils.MakeRerootedFilePath(ZipFile.FullName, StagingDir.FullName, OutputDir.FullName));
+							ZipFile.Delete();
+							ZipFile = NewZipFile;
 						}
-					})).ToList();
-			}
+						ZipFiles.Add(ZipFile);
+					}
+				})).ToList();
 
 			for (int index = 0; index < ZipThreads.Count; index++)
 			{
@@ -1091,7 +1038,7 @@ namespace AutomationTool
 					{
 						foreach (ZipArchiveEntry Entry in ZipArchive.Entries)
 						{
-							// Use CommandUtils.CombinePaths to ensure directory separators get converted correctly. On mono on *nix, if the path has backslashes it will not convert it.
+							// Use CommandUtils.CombinePaths to ensure directory separators get converted correctly.
 							var ExtractedFilename = CommandUtils.CombinePaths(RootDir.FullName, Entry.FullName);
 							// Zips can contain empty dirs. Ours usually don't have them, but we should support it.
 							if (Path.GetFileName(ExtractedFilename).Length == 0)
