AutomationTool/Scripts:

IsRunningOnMono -> !IsRunningOnWindows

#trivial

diff --git a/Engine/Source/Programs/AutomationTool/Scripts/BuildHlslcc.Automation.cs b/Engine/Source/Programs/AutomationTool/Scripts/BuildHlslcc.Automation.cs
--- a/Engine/Source/Programs/AutomationTool/Scripts/BuildHlslcc.Automation.cs
+++ b/Engine/Source/Programs/AutomationTool/Scripts/BuildHlslcc.Automation.cs
@@ -524,7 +524,7 @@ class BuildHlslcc : BuildCommand
 	}
 	private static void SetupStaticBuildEnvironment()
 	{
-		if (!Utils.IsRunningOnMono)
+		if (Utils.IsRunningOnWindows)
 		{
 			string VS2015Path = GetMsDevExe(WindowsCompiler.VisualStudio2015_DEPRECATED);
 			if (VS2015Path != null)
diff --git a/Engine/Source/Programs/AutomationTool/Scripts/Tests.Automation.cs b/Engine/Source/Programs/AutomationTool/Scripts/Tests.Automation.cs
--- a/Engine/Source/Programs/AutomationTool/Scripts/Tests.Automation.cs
+++ b/Engine/Source/Programs/AutomationTool/Scripts/Tests.Automation.cs
@@ -464,7 +464,7 @@ class TestMacZip : BuildCommand
 	{
 		LogInformation("TestMacZip *********************");
 
-		if (UnrealBuildTool.Utils.IsRunningOnMono)
+		if (!Utils.IsRunningOnWindows)
 		{
 			PushDir(CombinePaths(CmdEnv.LocalRoot, "Rocket/QFE"));
 			RunAndLog(CommandUtils.CmdEnv, "zip", "-r TestZip .");
