UnrealBuildTool:

Remove IsRunningOnMono

For remaining references, IsRunningOnMono -> !IsRunningOnWindows

#trivial

diff --git a/Engine/Source/Programs/AutomationTool/HoloLens/HoloLensPlatform.Automation.cs b/Engine/Source/Programs/AutomationTool/HoloLens/HoloLensPlatform.Automation.cs
--- a/Engine/Source/Programs/AutomationTool/HoloLens/HoloLensPlatform.Automation.cs
+++ b/Engine/Source/Programs/AutomationTool/HoloLens/HoloLensPlatform.Automation.cs
@@ -1283,7 +1283,7 @@ namespace HoloLens.Automation
 		private void DeployToLocalDevice(ProjectParams Params, DeploymentContext SC)
 		{
 #if !__MonoCS__
-            if (Utils.IsRunningOnMono)
+            if (!Utils.IsRunningOnWindows)
             {
                 return;
             }
@@ -1400,7 +1400,7 @@ namespace HoloLens.Automation
 		private void DeployToRemoteDevice(string DeviceAddress, ProjectParams Params, DeploymentContext SC)
 		{
 #if !__MonoCS__
-			if (Utils.IsRunningOnMono)
+			if (!Utils.IsRunningOnWindows)
             {
                 return;
             }
@@ -1489,7 +1489,7 @@ namespace HoloLens.Automation
 
 		private IProcessResult RunUsingLauncherTool(string DeviceAddress, ERunOptions ClientRunFlags, string ClientApp, string ClientCmdLine, ProjectParams Params)
 		{
-            if (Utils.IsRunningOnMono)
+            if (!Utils.IsRunningOnWindows)
             {
                 return null;
             }
@@ -1581,7 +1581,7 @@ namespace HoloLens.Automation
 		private IProcessResult RunUsingDevicePortal(string DeviceAddress, ERunOptions ClientRunFlags, string ClientApp, string ClientCmdLine, ProjectParams Params)
 		{
 #if !__MonoCS__
-            if (Utils.IsRunningOnMono)
+            if (!Utils.IsRunningOnWindows)
             {
                 return null;
             }
@@ -1685,7 +1685,7 @@ namespace HoloLens.Automation
 		private bool ShouldAcceptCertificate(System.Security.Cryptography.X509Certificates.X509Certificate2 Certificate, bool Unattended)
 		{
 #if !__MonoCS__
-            if (Utils.IsRunningOnMono)
+            if (!Utils.IsRunningOnWindows)
             {
                 return false;
             }
diff --git a/Engine/Source/Programs/AutomationTool/Turnkey/Turnkey.Automation.cs b/Engine/Source/Programs/AutomationTool/Turnkey/Turnkey.Automation.cs
--- a/Engine/Source/Programs/AutomationTool/Turnkey/Turnkey.Automation.cs
+++ b/Engine/Source/Programs/AutomationTool/Turnkey/Turnkey.Automation.cs
@@ -103,7 +103,7 @@ namespace Turnkey
 		private static void SetupVisuals()
 		{
 			// make the form look good on modern displays!
-			if (!UnrealBuildTool.Utils.IsRunningOnMono && Environment.OSVersion.Version.Major >= 6)
+			if (Utils.IsRunningOnWindows && Environment.OSVersion.Version.Major >= 6)
 			{
 				SetProcessDPIAware();
 			}
diff --git a/Engine/Source/Programs/UnrealBuildTool/System/Utils.cs b/Engine/Source/Programs/UnrealBuildTool/System/Utils.cs
--- a/Engine/Source/Programs/UnrealBuildTool/System/Utils.cs
+++ b/Engine/Source/Programs/UnrealBuildTool/System/Utils.cs
@@ -20,15 +20,11 @@ namespace UnrealBuildTool
 	/// </summary>
 	public static class Utils
 	{
-		/// <summary>
-		/// Whether we are currently running on Mono platform.  We cache this statically because it is a bit slow to check.
-		/// </summary>
-		public static readonly bool IsRunningOnMono = !RuntimeInformation.IsOSPlatform(OSPlatform.Windows); // TODO: This check seems to mainly be used to determine if this is running on windows or not. Should be refactored.
-
 		/// <summary>
 		/// Whether we are currently running on Linux.
 		/// </summary>
 		public static readonly bool IsRunningOnLinux = RuntimeInformation.IsOSPlatform(OSPlatform.Linux);
+
 		/// <summary>
 		/// Whether we are currently running on a MacOS platform.
 		/// </summary>
