Add a verbose option to DeleteTask to see what files and folders are being deleted.

#rnx

diff --git a/Engine/Source/Programs/AutomationTool/BuildGraph/Tasks/DeleteTask.cs b/Engine/Source/Programs/AutomationTool/BuildGraph/Tasks/DeleteTask.cs
--- a/Engine/Source/Programs/AutomationTool/BuildGraph/Tasks/DeleteTask.cs
+++ b/Engine/Source/Programs/AutomationTool/BuildGraph/Tasks/DeleteTask.cs
@@ -35,6 +35,12 @@ namespace BuildGraph.Tasks
 		/// </summary>
 		[TaskParameter(Optional = true)]
 		public bool DeleteEmptyDirectories = true;
+
+		/// <summary>
+		/// Whether or not to use verbose logging.
+		/// </summary>
+		[TaskParameter(Optional = true)]
+		public bool Verbose = false;
 	}
 
 	/// <summary>
@@ -71,6 +77,10 @@ namespace BuildGraph.Tasks
 				HashSet<FileReference> Files = ResolveFilespec(CommandUtils.RootDirectory, Parameters.Files, TagNameToFileSet);
 				foreach (FileReference File in Files)
 				{
+					if (Parameters.Verbose)
+					{
+						Log.TraceInformation("Deleting {0}", File.FullName);
+					}
 					if (!InternalUtils.SafeDeleteFile(File.FullName))
 					{
 						CommandUtils.LogWarning("Couldn't delete file {0}", File.FullName);
@@ -106,6 +116,10 @@ namespace BuildGraph.Tasks
 				{
 					if (!String.IsNullOrEmpty(Directory))
 					{
+						if (Parameters.Verbose)
+						{
+							Log.TraceInformation("Deleting {0}", Directory);
+						}
 						DirectoryReference FullDir = new DirectoryReference(Directory);
 						if (DirectoryReference.Exists(FullDir))
 						{
