Gauntlet telemetry - make the csv parser more resilient to non compliant formating

#rnx

diff --git a/Engine/Source/Programs/AutomationTool/Gauntlet/Framework/Utils/Gauntlet.CSVParser.cs b/Engine/Source/Programs/AutomationTool/Gauntlet/Framework/Utils/Gauntlet.CSVParser.cs
--- a/Engine/Source/Programs/AutomationTool/Gauntlet/Framework/Utils/Gauntlet.CSVParser.cs
+++ b/Engine/Source/Programs/AutomationTool/Gauntlet/Framework/Utils/Gauntlet.CSVParser.cs
@@ -2,6 +2,7 @@
 
 using System.IO;
 using System.Collections.Generic;
+using AutomationTool;
 
 
 namespace Gauntlet
@@ -16,14 +17,16 @@ namespace Gauntlet
 			{
 				// Parse Header
 				string FirstLine = Reader.ReadLine();
+				int LineNumber = 1;
 				string[] HeaderValues = FirstLine.Split(',');
 				// Parse the rest of the data
 				while (!Reader.EndOfStream)
 				{
 					Dictionary<string, string> OneRow = new Dictionary<string, string>();
 					string Line = Reader.ReadLine();
+					LineNumber++;
 					string[] Values = Line.Split(',');
-					if (Values.Length == HeaderValues.Length)
+					if (Values.Length >= HeaderValues.Length)
 					{
 						for (int i = 0; i < HeaderValues.Length; i++)
 						{
@@ -31,6 +34,10 @@ namespace Gauntlet
 						}
 						Rows.Add(OneRow);
 					}
+					else
+					{
+						throw new AutomationException(string.Format("Not enough number of columns at line {0}.", LineNumber));
+					}
 				}
 			}
 
diff --git a/Engine/Source/Programs/AutomationTool/Gauntlet/Unreal/Utils/Gauntlet.UnrealTelemetry.cs b/Engine/Source/Programs/AutomationTool/Gauntlet/Unreal/Utils/Gauntlet.UnrealTelemetry.cs
--- a/Engine/Source/Programs/AutomationTool/Gauntlet/Unreal/Utils/Gauntlet.UnrealTelemetry.cs
+++ b/Engine/Source/Programs/AutomationTool/Gauntlet/Unreal/Utils/Gauntlet.UnrealTelemetry.cs
@@ -5,7 +5,6 @@ using System.IO;
 using System.Collections.Generic;
 using AutomationTool;
 using UnrealBuildTool;
-using EpicGames.Core;
 
 
 namespace Gauntlet
@@ -136,8 +135,8 @@ namespace Gauntlet
 	{
 		[Help("CSVFile", "Path to the csv file to parse.")]
 		[Help("CSVDirectory", "Path to a folder containing csv files to parse.")]
-		[Help("CSVMapping", "Optional CSV column mapping. Formate is: <target key>:<source key>,...")]
-		[Help("TelemetryConfig", "Telemetry configuration to use to publish to Database. Default: UETelemtr.y")]
+		[Help("CSVMapping", "Optional CSV column mapping. Format is: <target key>:<source key>,...")]
+		[Help("TelemetryConfig", "Telemetry configuration to use to publish to Database. Default: UETelemtry.")]
 		[Help("DatabaseConfigPath", "Path to alternate Database config. Default is TelemetryConfig default.")]
 		[Help("Project", "Target Project name.")]
 		[Help("Platform", "Target platform name. Default: current environment platform.")]
@@ -256,7 +255,7 @@ namespace Gauntlet
 					{
 						DBConfig.LoadConfig(DatabaseConfigPath);
 						IDatabaseDriver<TelemetryData> DB = DBConfig.GetDriver();
-						Log.Verbose("Submitting telemetry data to {0}", DB.ToString());
+						Log.Info("Submitting telemetry data to {0}", DB.ToString());
 						DB.SubmitDataItems(DataRows, Context);
 					}
 					else
@@ -264,6 +263,10 @@ namespace Gauntlet
 						Log.Warning("Got telemetry data, but database configuration is unknown '{0}'.", Config);
 					}
 				}
+				else
+				{
+					Log.Error("Failed to retrieve any data from telemetry report.");
+				}
 
 			}
 
