Added requiredflags to the Gauntlet message when a valid build type can't be found for easier future triage

diff --git a/Engine/Source/Programs/AutomationTool/Gauntlet/Unreal/BuildSource/Gauntlet.UnrealBuildSource.cs b/Engine/Source/Programs/AutomationTool/Gauntlet/Unreal/BuildSource/Gauntlet.UnrealBuildSource.cs
--- a/Engine/Source/Programs/AutomationTool/Gauntlet/Unreal/BuildSource/Gauntlet.UnrealBuildSource.cs
+++ b/Engine/Source/Programs/AutomationTool/Gauntlet/Unreal/BuildSource/Gauntlet.UnrealBuildSource.cs
@@ -452,7 +452,7 @@ namespace Gauntlet
 				}
 			}
 
-			Reasons.Add(string.Format("No build at {0} that matches {1}", string.Join(",", BuildPaths), Role.ToString()));
+			Reasons.Add(string.Format("No build at {0} that matches {1} (RequiredFlags={2})", string.Join(",", BuildPaths), Role.ToString(), Role.RequiredBuildFlags.ToString()));
 
 			return false;
 		}
