Build GitSync on Horde.

#rnx

diff --git a/Engine/Source/Programs/AutomationTool/BuildGraph/Tasks/AwsTask.cs b/Engine/Source/Programs/AutomationTool/BuildGraph/Tasks/AwsTask.cs
new file mode 100644
--- /dev/null
+++ b/Engine/Source/Programs/AutomationTool/BuildGraph/Tasks/AwsTask.cs
@@ -0,0 +1,144 @@
+// Copyright Epic Games, Inc. All Rights Reserved.
+
+using EpicGames.Core;
+using System;
+using System.Collections.Generic;
+using System.IO;
+using System.Text;
+using System.Xml;
+
+namespace AutomationTool.Tasks
+{
+	/// <summary>
+	/// Parameters for an AWS CLI task
+	/// </summary>
+	public class AwsTaskParameters
+	{
+		/// <summary>
+		/// AWS command line arguments
+		/// </summary>
+		[TaskParameter(Optional = true)]
+		public string Arguments;
+
+		/// <summary>
+		/// Base directory for running the command
+		/// </summary>
+		[TaskParameter(Optional = true)]
+		public string BaseDir;
+
+		/// <summary>
+		/// The minimum exit code, which is treated as an error.
+		/// </summary>
+		[TaskParameter(Optional = true)]
+		public int ErrorLevel = 1;
+
+		/// <summary>
+		/// Whether or not to show the console output from the AWS CLI command.
+		/// </summary>
+		[TaskParameter(Optional = true)]
+		public bool Verbose = false;
+
+		/// <summary>
+		/// Path to a file containing the credentials needed to run the AWS command
+		/// </summary>
+		[TaskParameter(Optional = true)]
+		public string CredentialsFile;
+
+		/// <summary>
+		/// Path to a file to store any JSON output by the command.
+		/// </summary>
+		[TaskParameter(Optional = true)]
+		public string OutputFile;
+	}
+
+	/// <summary>
+	/// Spawns AWS CLI and waits for it to complete.
+	/// </summary>
+	[TaskElement("Aws", typeof(AwsTaskParameters))]
+	public class AwsTask : CustomTask
+	{
+		/// <summary>
+		/// Parameters for this task
+		/// </summary>
+		AwsTaskParameters Parameters;
+
+		/// <summary>
+		/// Construct an AWS CLI task
+		/// </summary>
+		/// <param name="InParameters">Parameters for the task</param>
+		public AwsTask(AwsTaskParameters InParameters)
+		{
+			Parameters = InParameters;
+		}
+
+		/// <summary>
+		/// Execute the task.
+		/// </summary>
+		/// <param name="Job">Information about the current job</param>
+		/// <param name="BuildProducts">Set of build products produced by this node.</param>
+		/// <param name="TagNameToFileSet">Mapping from tag names to the set of files they include</param>
+		public override void Execute(JobContext Job, HashSet<FileReference> BuildProducts, Dictionary<string, HashSet<FileReference>> TagNameToFileSet)
+		{
+			FileReference ToolFile = CommandUtils.FindToolInPath("aws");
+			if(ToolFile == null)
+			{
+				throw new AutomationException("Unable to find path to AWS CLI. Check you have it installed, and it is on your PATH.");
+			}
+
+			string AWSArgs = Parameters.Arguments;
+			Dictionary<string, string> EnvVars = new Dictionary<string, string>();
+			if (!string.IsNullOrWhiteSpace(Parameters.CredentialsFile))
+			{
+				if (!File.Exists(Parameters.CredentialsFile))
+				{
+					throw new AutomationException("Credentials file {0} could not be found.", Parameters.CredentialsFile);
+				}
+
+				string[] CredentialTokens = File.ReadAllText(Parameters.CredentialsFile).Split(':');
+				EnvVars.Add("AWS_ACCESS_KEY_ID", CredentialTokens[0]);
+				EnvVars.Add("AWS_SECRET_ACCESS_KEY", CredentialTokens[1]);
+			}
+
+			CommandUtils.ERunOptions Options = CommandUtils.ERunOptions.AppMustExist;
+			if (Parameters.Verbose)
+			{
+				Options |= CommandUtils.ERunOptions.AllowSpew;
+			}
+			IProcessResult Result = CommandUtils.Run(ToolFile.FullName, AWSArgs, WorkingDir: Parameters.BaseDir, Env: EnvVars, Options: Options);
+			if (!string.IsNullOrWhiteSpace(Parameters.OutputFile))
+			{
+				File.WriteAllText(Parameters.OutputFile, Result.Output);
+			}
+			if (Result.ExitCode < 0 || Result.ExitCode >= Parameters.ErrorLevel)
+			{
+				throw new AutomationException("AWS terminated with an exit code indicating an error ({0})", Result.ExitCode);
+			}
+		}
+
+		/// <summary>
+		/// Output this task out to an XML writer.
+		/// </summary>
+		public override void Write(XmlWriter Writer)
+		{
+			Write(Writer, Parameters);
+		}
+
+		/// <summary>
+		/// Find all the tags which are used as inputs to this task
+		/// </summary>
+		/// <returns>The tag names which are read by this task</returns>
+		public override IEnumerable<string> FindConsumedTagNames()
+		{
+			yield break;
+		}
+
+		/// <summary>
+		/// Find all the tags which are modified by this task
+		/// </summary>
+		/// <returns>The tag names which are modified by this task</returns>
+		public override IEnumerable<string> FindProducedTagNames()
+		{
+			yield break;
+		}
+	}
+}
diff --git a/Engine/Source/Programs/AutomationTool/BuildGraph/Tasks/DockerTask.cs b/Engine/Source/Programs/AutomationTool/BuildGraph/Tasks/DockerTask.cs
--- a/Engine/Source/Programs/AutomationTool/BuildGraph/Tasks/DockerTask.cs
+++ b/Engine/Source/Programs/AutomationTool/BuildGraph/Tasks/DockerTask.cs
@@ -3,6 +3,7 @@
 using EpicGames.Core;
 using System;
 using System.Collections.Generic;
+using System.IO;
 using System.Text;
 using System.Xml;
 
@@ -30,6 +31,18 @@ namespace AutomationTool.Tasks
 		/// </summary>
 		[TaskParameter(Optional = true)]
 		public int ErrorLevel = 1;
+
+		/// <summary>
+		/// Whether or not to show the console output from the Docker command.
+		/// </summary>
+		[TaskParameter(Optional = true)]
+		public bool Verbose = false;
+
+		/// <summary>
+		/// Path to a file containing the credentials needed to run the Docker command
+		/// </summary>
+		[TaskParameter(Optional = true)]
+		public string CredentialsFile;
 	}
 
 	/// <summary>
@@ -66,7 +79,36 @@ namespace AutomationTool.Tasks
 				throw new AutomationException("Unable to find path to Docker. Check you have it installed, and it is on your PATH.");
 			}
 
-			IProcessResult Result = CommandUtils.Run(ToolFile.FullName, Parameters.Arguments, WorkingDir: Parameters.BaseDir);
+			string DockerArgs = Parameters.Arguments;
+			string CmdString = $"{ToolFile.FullName} {DockerArgs}";
+			if (!string.IsNullOrWhiteSpace(Parameters.CredentialsFile))
+			{
+				if (!File.Exists(Parameters.CredentialsFile))
+				{
+					throw new AutomationException("Credentials file {0} could not be found.", Parameters.CredentialsFile);
+				}
+
+				// First see if it is JSON.
+				string CredentialString = "";
+				string[] CredentialTokens = File.ReadAllText(Parameters.CredentialsFile).Split(':');
+				if (CredentialTokens.Length == 1)
+				{
+					CredentialString = $"--password {CredentialTokens[0]}";
+				}
+				else
+				{
+					CredentialString = $"--username {CredentialTokens[0]} --password {CredentialTokens[1]}";
+				}
+				
+				DockerArgs = $"{DockerArgs} {CredentialString}";
+			}
+
+			CommandUtils.ERunOptions Options = CommandUtils.ERunOptions.AppMustExist;
+			if (Parameters.Verbose)
+			{
+				Options |= CommandUtils.ERunOptions.AllowSpew;
+			}
+			IProcessResult Result = CommandUtils.Run(ToolFile.FullName, DockerArgs, WorkingDir: Parameters.BaseDir, Options: Options);
 			if (Result.ExitCode < 0 || Result.ExitCode >= Parameters.ErrorLevel)
 			{
 				throw new AutomationException("Docker terminated with an exit code indicating an error ({0})", Result.ExitCode);
