Prevent Docker password being logged to stdout.

diff --git a/Engine/Source/Programs/AutomationTool/BuildGraph/Tasks/DockerTask.cs b/Engine/Source/Programs/AutomationTool/BuildGraph/Tasks/DockerTask.cs
--- a/Engine/Source/Programs/AutomationTool/BuildGraph/Tasks/DockerTask.cs
+++ b/Engine/Source/Programs/AutomationTool/BuildGraph/Tasks/DockerTask.cs
@@ -81,6 +81,8 @@ namespace AutomationTool.Tasks
 
 			string DockerArgs = Parameters.Arguments;
 			string CmdString = $"{ToolFile.FullName} {DockerArgs}";
+			string DockerInput = null;
+
 			if (!string.IsNullOrWhiteSpace(Parameters.CredentialsFile))
 			{
 				if (!File.Exists(Parameters.CredentialsFile))
@@ -88,19 +90,21 @@ namespace AutomationTool.Tasks
 					throw new AutomationException("Credentials file {0} could not be found.", Parameters.CredentialsFile);
 				}
 
-				// First see if it is JSON.
-				string CredentialString = "";
 				string[] CredentialTokens = File.ReadAllText(Parameters.CredentialsFile).Split(':');
-				if (CredentialTokens.Length == 1)
+				if (CredentialTokens.Length == 2)
+				{
+					DockerArgs += $" --username {CredentialTokens[0]} --password-stdin";
+					DockerInput = CredentialTokens[1];
+				}
+				else if (CredentialTokens.Length == 1)
 				{
-					CredentialString = $"--password {CredentialTokens[0]}";
+					DockerArgs += $" --password-stdin";
+					DockerInput = CredentialTokens[0];
 				}
 				else
 				{
-					CredentialString = $"--username {CredentialTokens[0]} --password {CredentialTokens[1]}";
+					throw new AutomationException("Expected password or username:password in credentials file ({0})", Parameters.CredentialsFile);
 				}
-				
-				DockerArgs = $"{DockerArgs} {CredentialString}";
 			}
 
 			CommandUtils.ERunOptions Options = CommandUtils.ERunOptions.AppMustExist;
@@ -108,7 +112,7 @@ namespace AutomationTool.Tasks
 			{
 				Options |= CommandUtils.ERunOptions.AllowSpew;
 			}
-			IProcessResult Result = CommandUtils.Run(ToolFile.FullName, DockerArgs, WorkingDir: Parameters.BaseDir, Options: Options);
+			IProcessResult Result = CommandUtils.Run(ToolFile.FullName, DockerArgs, DockerInput, WorkingDir: Parameters.BaseDir, Options: Options);
 			if (Result.ExitCode < 0 || Result.ExitCode >= Parameters.ErrorLevel)
 			{
 				throw new AutomationException("Docker terminated with an exit code indicating an error ({0})", Result.ExitCode);
