Enable verbose output from Docker tasks by default.

diff --git a/Engine/Source/Programs/AutomationTool/BuildGraph/Tasks/DockerTask.cs b/Engine/Source/Programs/AutomationTool/BuildGraph/Tasks/DockerTask.cs
--- a/Engine/Source/Programs/AutomationTool/BuildGraph/Tasks/DockerTask.cs
+++ b/Engine/Source/Programs/AutomationTool/BuildGraph/Tasks/DockerTask.cs
@@ -36,7 +36,7 @@ namespace AutomationTool.Tasks
 		/// Whether or not to show the console output from the Docker command.
 		/// </summary>
 		[TaskParameter(Optional = true)]
-		public bool Verbose = false;
+		public bool Verbose = true;
 
 		/// <summary>
 		/// Path to a file containing the credentials needed to run the Docker command
