Add a Helm task to allow updating releases.

diff --git a/Engine/Source/Programs/AutomationTool/BuildGraph/Tasks/HelmTask.cs b/Engine/Source/Programs/AutomationTool/BuildGraph/Tasks/HelmTask.cs
new file mode 100644
--- /dev/null
+++ b/Engine/Source/Programs/AutomationTool/BuildGraph/Tasks/HelmTask.cs
@@ -0,0 +1,173 @@
+// Copyright Epic Games, Inc. All Rights Reserved.
+
+using EpicGames.Core;
+using System;
+using System.Collections.Generic;
+using System.IO;
+using System.Text;
+using System.Xml;
+
+namespace AutomationTool.Tasks
+{
+	/// <summary>
+	/// Parameters for a Helm task
+	/// </summary>
+	public class HelmTaskParameters
+	{
+		/// <summary>
+		/// Helm command line arguments
+		/// </summary>
+		[TaskParameter]
+		public string Chart;
+
+		/// <summary>
+		/// Name of the release
+		/// </summary>
+		[TaskParameter]
+		public string Release;
+
+		/// <summary>
+		/// The Kubernetes namespace
+		/// </summary>
+		[TaskParameter(Optional = true)]
+		public string Namespace;
+
+		/// <summary>
+		/// The kubectl context
+		/// </summary>
+		[TaskParameter(Optional = true)]
+		public string KubeContext;
+
+		/// <summary>
+		/// Values to set for running the chart
+		/// </summary>
+		[TaskParameter(Optional = true)]
+		public List<string> Values;
+
+		/// <summary>
+		/// File to parse environment variables from
+		/// </summary>
+		[TaskParameter(Optional = true)]
+		public string EnvironmentFile;
+
+		/// <summary>
+		/// Additional arguments
+		/// </summary>
+		[TaskParameter(Optional = true)]
+		public string Arguments;
+
+		/// <summary>
+		/// Base directory for running the command
+		/// </summary>
+		[TaskParameter(Optional = true)]
+		public string BaseDir;
+	}
+
+	/// <summary>
+	/// Spawns Helm and waits for it to complete.
+	/// </summary>
+	[TaskElement("Helm", typeof(HelmTaskParameters))]
+	public class HelmTask : CustomTask
+	{
+		/// <summary>
+		/// Parameters for this task
+		/// </summary>
+		HelmTaskParameters Parameters;
+
+		/// <summary>
+		/// Construct a Helm task
+		/// </summary>
+		/// <param name="InParameters">Parameters for the task</param>
+		public HelmTask(HelmTaskParameters InParameters)
+		{
+			Parameters = InParameters;
+		}
+
+		/// <summary>
+		/// Execute the task.
+		/// </summary>
+		/// <param name="Job">Information about the current job</param>
+		/// <param name="BuildProducts">Set of build products produced by this node.</param>
+		/// <param name="TagNameToFileSet">Mapping from tag names to the set of files they include</param>
+		public override void Execute(JobContext Job, HashSet<FileReference> BuildProducts, Dictionary<string, HashSet<FileReference>> TagNameToFileSet)
+		{
+			FileReference HelmExe = CommandUtils.FindToolInPath("helm");
+			if(HelmExe == null)
+			{
+				throw new AutomationException("Unable to find path to Helm. Check you have it installed, and it is on your PATH.");
+			}
+
+			Dictionary<string, string> Environment = new Dictionary<string, string>();
+			if (Parameters.EnvironmentFile != null)
+			{
+			}
+
+			// Switch Kubernetes config
+			if (Parameters.KubeContext != null)
+			{
+				FileReference KubectlExe = CommandUtils.FindToolInPath("kubectl");
+				if (KubectlExe == null)
+				{
+					throw new AutomationException("Unable to find path to Kubectl. Check you have it installed, and it is on your PATH.");
+				}
+
+				IProcessResult KubectlResult = CommandUtils.Run(KubectlExe.FullName, $"config use-context {Parameters.KubeContext}", null, WorkingDir: Parameters.BaseDir);
+				if (KubectlResult.ExitCode != 0)
+				{
+					throw new AutomationException("Kubectl terminated with an exit code indicating an error ({0})", KubectlResult.ExitCode);
+				}
+			}
+
+			// Build the argument list
+			List<string> Arguments = new List<string>();
+			Arguments.Add("upgrade");
+			Arguments.Add(Parameters.Release);
+			Arguments.Add(Parameters.Chart);
+			Arguments.Add("--install");
+			Arguments.Add("--reset-values");
+			if(Parameters.Namespace != null)
+			{
+				Arguments.Add("--namespace");
+				Arguments.Add(Parameters.Namespace);
+			}
+			foreach (string Value in Parameters.Values)
+			{
+				Arguments.Add("--set");
+				Arguments.Add(Value);
+			}
+			Arguments.Add("--dry-run");
+
+			IProcessResult Result = CommandUtils.Run(HelmExe.FullName, CommandLineArguments.Join(Arguments), WorkingDir: Parameters.BaseDir);
+			if (Result.ExitCode != 0)
+			{
+				throw new AutomationException("Helm terminated with an exit code indicating an error ({0})", Result.ExitCode);
+			}
+		}
+
+		/// <summary>
+		/// Output this task out to an XML writer.
+		/// </summary>
+		public override void Write(XmlWriter Writer)
+		{
+			Write(Writer, Parameters);
+		}
+
+		/// <summary>
+		/// Find all the tags which are used as inputs to this task
+		/// </summary>
+		/// <returns>The tag names which are read by this task</returns>
+		public override IEnumerable<string> FindConsumedTagNames()
+		{
+			yield break;
+		}
+
+		/// <summary>
+		/// Find all the tags which are modified by this task
+		/// </summary>
+		/// <returns>The tag names which are modified by this task</returns>
+		public override IEnumerable<string> FindProducedTagNames()
+		{
+			yield break;
+		}
+	}
+}
