Add a Kubectl task.

diff --git a/Engine/Source/Programs/AutomationTool/BuildGraph/Tasks/KubectlTask.cs b/Engine/Source/Programs/AutomationTool/BuildGraph/Tasks/KubectlTask.cs
new file mode 100644
--- /dev/null
+++ b/Engine/Source/Programs/AutomationTool/BuildGraph/Tasks/KubectlTask.cs
@@ -0,0 +1,97 @@
+// Copyright Epic Games, Inc. All Rights Reserved.
+
+using EpicGames.Core;
+using System;
+using System.Collections.Generic;
+using System.IO;
+using System.Text;
+using System.Xml;
+
+namespace AutomationTool.Tasks
+{
+	/// <summary>
+	/// Parameters for a Kubectl task
+	/// </summary>
+	public class KubectlTaskParameters
+	{
+		/// <summary>
+		/// Command line arguments
+		/// </summary>
+		[TaskParameter]
+		public string Arguments;
+
+		/// <summary>
+		/// Base directory for running the command
+		/// </summary>
+		[TaskParameter(Optional = true)]
+		public string BaseDir;
+	}
+
+	/// <summary>
+	/// Spawns Kubectland waits for it to complete.
+	/// </summary>
+	[TaskElement("Kubectl", typeof(KubectlTaskParameters))]
+	public class KubectlTask : CustomTask
+	{
+		/// <summary>
+		/// Parameters for this task
+		/// </summary>
+		KubectlTaskParameters Parameters;
+
+		/// <summary>
+		/// Construct a Kubectl task
+		/// </summary>
+		/// <param name="InParameters">Parameters for the task</param>
+		public KubectlTask(KubectlTaskParameters InParameters)
+		{
+			Parameters = InParameters;
+		}
+
+		/// <summary>
+		/// Execute the task.
+		/// </summary>
+		/// <param name="Job">Information about the current job</param>
+		/// <param name="BuildProducts">Set of build products produced by this node.</param>
+		/// <param name="TagNameToFileSet">Mapping from tag names to the set of files they include</param>
+		public override void Execute(JobContext Job, HashSet<FileReference> BuildProducts, Dictionary<string, HashSet<FileReference>> TagNameToFileSet)
+		{
+			FileReference KubectlExe = CommandUtils.FindToolInPath("kubectl");
+			if (KubectlExe == null)
+			{
+				throw new AutomationException("Unable to find path to Kubectl. Check you have it installed, and it is on your PATH.");
+			}
+
+			IProcessResult Result = CommandUtils.Run(KubectlExe.FullName, Parameters.Arguments, null, WorkingDir: Parameters.BaseDir);
+			if (Result.ExitCode != 0)
+			{
+				throw new AutomationException("Kubectl terminated with an exit code indicating an error ({0})", Result.ExitCode);
+			}
+		}
+
+		/// <summary>
+		/// Output this task out to an XML writer.
+		/// </summary>
+		public override void Write(XmlWriter Writer)
+		{
+			Write(Writer, Parameters);
+		}
+
+		/// <summary>
+		/// Find all the tags which are used as inputs to this task
+		/// </summary>
+		/// <returns>The tag names which are read by this task</returns>
+		public override IEnumerable<string> FindConsumedTagNames()
+		{
+			yield break;
+		}
+
+		/// <summary>
+		/// Find all the tags which are modified by this task
+		/// </summary>
+		/// <returns>The tag names which are modified by this task</returns>
+		public override IEnumerable<string> FindProducedTagNames()
+		{
+			yield break;
+		}
+	}
+}
