BuildGraph: Add a task for executing DotNet commands.

diff --git a/Engine/Source/Programs/AutomationTool/BuildGraph/Tasks/DotNetTask.cs b/Engine/Source/Programs/AutomationTool/BuildGraph/Tasks/DotNetTask.cs
new file mode 100644
--- /dev/null
+++ b/Engine/Source/Programs/AutomationTool/BuildGraph/Tasks/DotNetTask.cs
@@ -0,0 +1,103 @@
+// Copyright Epic Games, Inc. All Rights Reserved.
+
+using EpicGames.Core;
+using System;
+using System.Collections.Generic;
+using System.IO;
+using System.Text;
+using System.Xml;
+
+namespace AutomationTool.Tasks
+{
+	/// <summary>
+	/// Parameters for a DotNet task
+	/// </summary>
+	public class DotNetTaskParameters
+	{
+		/// <summary>
+		/// Docker command line arguments
+		/// </summary>
+		[TaskParameter(Optional = true)]
+		public string Arguments;
+
+		/// <summary>
+		/// Base directory for running the command
+		/// </summary>
+		[TaskParameter(Optional = true)]
+		public string BaseDir;
+
+		/// <summary>
+		/// The minimum exit code, which is treated as an error.
+		/// </summary>
+		[TaskParameter(Optional = true)]
+		public int ErrorLevel = 1;
+	}
+
+	/// <summary>
+	/// Spawns Docker and waits for it to complete.
+	/// </summary>
+	[TaskElement("DotNet", typeof(DotNetTaskParameters))]
+	public class DotNetTask : CustomTask
+	{
+		/// <summary>
+		/// Parameters for this task
+		/// </summary>
+		DotNetTaskParameters Parameters;
+
+		/// <summary>
+		/// Construct a Docker task
+		/// </summary>
+		/// <param name="InParameters">Parameters for the task</param>
+		public DotNetTask(DotNetTaskParameters InParameters)
+		{
+			Parameters = InParameters;
+		}
+
+		/// <summary>
+		/// Execute the task.
+		/// </summary>
+		/// <param name="Job">Information about the current job</param>
+		/// <param name="BuildProducts">Set of build products produced by this node.</param>
+		/// <param name="TagNameToFileSet">Mapping from tag names to the set of files they include</param>
+		public override void Execute(JobContext Job, HashSet<FileReference> BuildProducts, Dictionary<string, HashSet<FileReference>> TagNameToFileSet)
+		{
+			FileReference DotNetFile = new FileReference(HostPlatform.Current.GetDotnetMsbuildExe());
+			if(!FileReference.Exists(DotNetFile))
+			{
+				throw new AutomationException("DotNet is missing from {0}", DotNetFile);
+			}
+
+			IProcessResult Result = CommandUtils.Run(DotNetFile.FullName, Parameters.Arguments, WorkingDir: Parameters.BaseDir);
+			if (Result.ExitCode < 0 || Result.ExitCode >= Parameters.ErrorLevel)
+			{
+				throw new AutomationException("Docker terminated with an exit code indicating an error ({0})", Result.ExitCode);
+			}
+		}
+
+		/// <summary>
+		/// Output this task out to an XML writer.
+		/// </summary>
+		public override void Write(XmlWriter Writer)
+		{
+			Write(Writer, Parameters);
+		}
+
+		/// <summary>
+		/// Find all the tags which are used as inputs to this task
+		/// </summary>
+		/// <returns>The tag names which are read by this task</returns>
+		public override IEnumerable<string> FindConsumedTagNames()
+		{
+			yield break;
+		}
+
+		/// <summary>
+		/// Find all the tags which are modified by this task
+		/// </summary>
+		/// <returns>The tag names which are modified by this task</returns>
+		public override IEnumerable<string> FindProducedTagNames()
+		{
+			yield break;
+		}
+	}
+}
