Support executable in root directory, but prefer executable in Project\Binaries\Platform or Engine\Binaries\Platform.

diff --git a/Engine/Source/Programs/AutomationTool/Gauntlet/Unreal/BuildSource/Gauntlet.StagedBuildSource.cs b/Engine/Source/Programs/AutomationTool/Gauntlet/Unreal/BuildSource/Gauntlet.StagedBuildSource.cs
--- a/Engine/Source/Programs/AutomationTool/Gauntlet/Unreal/BuildSource/Gauntlet.StagedBuildSource.cs
+++ b/Engine/Source/Programs/AutomationTool/Gauntlet/Unreal/BuildSource/Gauntlet.StagedBuildSource.cs
@@ -224,8 +224,8 @@ namespace Gauntlet
 				string[] ExecutablePaths = new string[]
 				{
 					Path.Combine(InPath, InProjectName, "Binaries", InPlatform.ToString()),
-					//Path.Combine(InPath),
 					Path.Combine(InPath, "Engine", "Binaries", InPlatform.ToString()),
+					InPath,
 				};
 
 				List<FileSystemInfo> Binaries = new List<FileSystemInfo>();
@@ -250,7 +250,7 @@ namespace Gauntlet
 					UnrealTargetConfiguration Config = UnrealHelpers.GetConfigurationFromExecutableName(InProjectName, App.Name);
 					UnrealTargetRole Role = UnrealHelpers.GetRoleFromExecutableName(InProjectName, App.Name);
 
-					if (Config != UnrealTargetConfiguration.Unknown && Role != UnrealTargetRole.Unknown)
+					if (Config != UnrealTargetConfiguration.Unknown && Role != UnrealTargetRole.Unknown && !DiscoveredBuilds.Any(B => B.Configuration == Config))
 					{
 						// store the exe path as relative to the staged dir path
 						T NewBuild = Activator.CreateInstance(typeof(T), new object[] { InPlatform, Config, Role, InPath, Utils.SystemHelpers.MakePathRelative(App.FullName, InPath) }) as T;
