Fix path to DotNET executable.

diff --git a/Engine/Source/Programs/AutomationTool/BuildGraph/Tasks/DotNetTask.cs b/Engine/Source/Programs/AutomationTool/BuildGraph/Tasks/DotNetTask.cs
--- a/Engine/Source/Programs/AutomationTool/BuildGraph/Tasks/DotNetTask.cs
+++ b/Engine/Source/Programs/AutomationTool/BuildGraph/Tasks/DotNetTask.cs
@@ -61,7 +61,7 @@ namespace AutomationTool.Tasks
 		/// <param name="TagNameToFileSet">Mapping from tag names to the set of files they include</param>
 		public override void Execute(JobContext Job, HashSet<FileReference> BuildProducts, Dictionary<string, HashSet<FileReference>> TagNameToFileSet)
 		{
-			FileReference DotNetFile = new FileReference(HostPlatform.Current.GetDotnetMsbuildExe());
+			FileReference DotNetFile = FileReference.Combine(CommandUtils.RootDirectory, HostPlatform.Current.RelativeBinariesFolder, HostPlatform.Current.GetDotnetMsbuildExe());
 			if(!FileReference.Exists(DotNetFile))
 			{
 				throw new AutomationException("DotNet is missing from {0}", DotNetFile);
