Add an explicit accessor for getting the path to the DotNet host.

diff --git a/Engine/Source/Programs/AutomationTool/AutomationUtils/HostPlatform.cs b/Engine/Source/Programs/AutomationTool/AutomationUtils/HostPlatform.cs
--- a/Engine/Source/Programs/AutomationTool/AutomationUtils/HostPlatform.cs
+++ b/Engine/Source/Programs/AutomationTool/AutomationUtils/HostPlatform.cs
@@ -6,6 +6,7 @@ using System.IO;
 using System.Linq;
 using System.Text;
 using System.Diagnostics;
+using EpicGames.Core;
 
 namespace AutomationTool
 {
@@ -71,6 +72,12 @@ namespace AutomationTool
 		/// <returns></returns>
 		abstract public string GetFrameworkMsbuildExe();
 
+		/// <summary>
+		/// Gets the path to dotnet
+		/// </summary>
+		/// <returns></returns>
+		abstract public FileReference GetDotnetExe();
+
 		/// <summary>
 		/// Gets the build executable filename for NET Core projects. Typically, the path to the bundled dotnet executable.
 		/// </summary>
diff --git a/Engine/Source/Programs/AutomationTool/AutomationUtils/LinuxHostPlatform.cs b/Engine/Source/Programs/AutomationTool/AutomationUtils/LinuxHostPlatform.cs
--- a/Engine/Source/Programs/AutomationTool/AutomationUtils/LinuxHostPlatform.cs
+++ b/Engine/Source/Programs/AutomationTool/AutomationUtils/LinuxHostPlatform.cs
@@ -38,6 +38,11 @@ namespace AutomationTool
 			return CachedFrameworkMsbuildTool;
 		}
 
+		public override FileReference GetDotnetExe()
+		{
+			return FileReference.Combine(CommandUtils.EngineDirectory, @"Binaries\ThirdParty\DotNet\Linux\dotnet");
+		}
+
 		public override string GetDotnetMsbuildExe()
 		{
 			return "../../ThirdParty/DotNet/Linux/dotnet";
diff --git a/Engine/Source/Programs/AutomationTool/AutomationUtils/MacHostPlatform.cs b/Engine/Source/Programs/AutomationTool/AutomationUtils/MacHostPlatform.cs
--- a/Engine/Source/Programs/AutomationTool/AutomationUtils/MacHostPlatform.cs
+++ b/Engine/Source/Programs/AutomationTool/AutomationUtils/MacHostPlatform.cs
@@ -37,6 +37,11 @@ namespace AutomationTool
 			return CachedFrameworkMsbuildExe;
 		}
 
+		public override FileReference GetDotnetExe()
+		{
+			return FileReference.Combine(CommandUtils.EngineDirectory, @"Binaries\ThirdParty\DotNet\Mac\dotnet");
+		}
+
 		public override string GetDotnetMsbuildExe()
 		{
 			return "../../ThirdParty/DotNet/Mac/dotnet";
diff --git a/Engine/Source/Programs/AutomationTool/AutomationUtils/WindowsHostPlatform.cs b/Engine/Source/Programs/AutomationTool/AutomationUtils/WindowsHostPlatform.cs
--- a/Engine/Source/Programs/AutomationTool/AutomationUtils/WindowsHostPlatform.cs
+++ b/Engine/Source/Programs/AutomationTool/AutomationUtils/WindowsHostPlatform.cs
@@ -7,6 +7,7 @@ using System.Text;
 using System.IO;
 using System.Diagnostics;
 using UnrealBuildTool;
+using EpicGames.Core;
 
 namespace AutomationTool
 {
@@ -17,6 +18,11 @@ namespace AutomationTool
 			return WindowsExports.GetMSBuildToolPath();
 		}
 
+		public override FileReference GetDotnetExe()
+		{
+			return FileReference.Combine(CommandUtils.EngineDirectory, @"Binaries\ThirdParty\DotNet\Windows\dotnet.exe");
+		}
+
 		public override string GetDotnetMsbuildExe()
 		{
 			return @"..\..\ThirdParty\DotNet\Windows\dotnet.exe";
diff --git a/Engine/Source/Programs/AutomationTool/BuildGraph/Tasks/DotNetTask.cs b/Engine/Source/Programs/AutomationTool/BuildGraph/Tasks/DotNetTask.cs
--- a/Engine/Source/Programs/AutomationTool/BuildGraph/Tasks/DotNetTask.cs
+++ b/Engine/Source/Programs/AutomationTool/BuildGraph/Tasks/DotNetTask.cs
@@ -61,7 +61,7 @@ namespace AutomationTool.Tasks
 		/// <param name="TagNameToFileSet">Mapping from tag names to the set of files they include</param>
 		public override void Execute(JobContext Job, HashSet<FileReference> BuildProducts, Dictionary<string, HashSet<FileReference>> TagNameToFileSet)
 		{
-			FileReference DotNetFile = FileReference.Combine(CommandUtils.RootDirectory, HostPlatform.Current.RelativeBinariesFolder, HostPlatform.Current.GetDotnetMsbuildExe());
+			FileReference DotNetFile = HostPlatform.Current.GetDotnetExe();
 			if(!FileReference.Exists(DotNetFile))
 			{
 				throw new AutomationException("DotNet is missing from {0}", DotNetFile);
