Fix Android setup for Linux, add TODO for Mac

diff --git a/Engine/Extras/Android/SetupAndroid.sh b/Engine/Extras/Android/SetupAndroid.sh
--- a/Engine/Extras/Android/SetupAndroid.sh
+++ b/Engine/Extras/Android/SetupAndroid.sh
@@ -25,7 +25,7 @@ else
     PAUSE="pausefunc"
 fi
 
-rem hardcoded versions for compatibility with non-Turnkey manual running
+# hardcoded versions for compatibility with non-Turnkey manual running
 if [[ -z "${PLATFORMS_VERSION}" ]]; then
     PLATFORMS_VERSION="android-28"
 fi
@@ -50,10 +50,7 @@ fi
 echo Android Studio Path: $STUDIO_PATH
 
 if [ "$STUDIO_SDK_PATH" == "" ]; then
-	STUDIO_SDK_PATH=$HOME/Android/Sdk
-fi
-if [ "$1" != "" ]; then
-	STUDIO_SDK_PATH=$1
+	STUDIO_SDK_PATH="$HOME/Android/Sdk"
 fi
 if [ ! -d "$STUDIO_SDK_PATH" ]; then
 	echo Android SDK not found at: $STUDIO_SDK_PATH
diff --git a/Engine/Source/Programs/AutomationTool/Android/AndroidPlatform.Automation.cs b/Engine/Source/Programs/AutomationTool/Android/AndroidPlatform.Automation.cs
--- a/Engine/Source/Programs/AutomationTool/Android/AndroidPlatform.Automation.cs
+++ b/Engine/Source/Programs/AutomationTool/Android/AndroidPlatform.Automation.cs
@@ -74,6 +74,20 @@ public class AndroidPlatform : Platform
 
 	private static string GetAndroidStudioExe()
 	{
+		if (HostPlatform.Current.HostEditorPlatform == UnrealTargetPlatform.Linux)
+		{
+			string UserHome = Environment.GetEnvironmentVariable("HOME");
+			string AndroidStudioExe = Path.Combine(UserHome, "android-studio", "bin", "studio.sh");
+
+			return AndroidStudioExe;
+		}
+		else if (HostPlatform.Current.HostEditorPlatform == UnrealTargetPlatform.Mac)
+		{
+			// TODO
+			return "";
+		}
+
+		// Win64
 		string DefaultAndroidStudioInstallDir = Path.Combine(Environment.GetFolderPath(Environment.SpecialFolder.ProgramFiles), "Android", "Android Studio");
 		string RegValue = Microsoft.Win32.Registry.GetValue("HKEY_LOCAL_MACHINE\\SOFTWARE\\Android Studio", "Path", null) as string;
 		string AndroidStudioInstallDir = RegValue == null ? DefaultAndroidStudioInstallDir : RegValue;
@@ -85,7 +99,22 @@ public class AndroidPlatform : Platform
 		if (!string.IsNullOrEmpty(AndroidHome) && Directory.Exists(AndroidHome))
 		{
 			return AndroidHome;
-		}	
+		}
+
+		if (HostPlatform.Current.HostEditorPlatform == UnrealTargetPlatform.Linux)
+		{
+			string UserHome = Environment.GetEnvironmentVariable("HOME");
+			string AndroidSdkPath = Path.Combine(UserHome, "Android", "Sdk");
+
+			return AndroidSdkPath;
+		}
+		else if (HostPlatform.Current.HostEditorPlatform == UnrealTargetPlatform.Mac)
+		{
+			// TODO
+			return "";
+		}
+
+		// Win64
 		string DefaultSdkDir = Path.Combine(Environment.GetFolderPath(Environment.SpecialFolder.LocalApplicationData), "Android", "Sdk");
 		string RegValue = Microsoft.Win32.Registry.GetValue("HKEY_LOCAL_MACHINE\\SOFTWARE\\Android", "SdkPath", null) as string;
 		return RegValue == null ? DefaultSdkDir : RegValue;
@@ -93,8 +122,8 @@ public class AndroidPlatform : Platform
 
 	public override bool UpdateHostPrerequisites(BuildCommand Command, ITurnkeyContext TurnkeyContext, bool bVerifyOnly)
 	{
-		// @todo turnkey: Handle Mac/Linux
-		if (HostPlatform.Current.HostEditorPlatform != UnrealTargetPlatform.Win64)
+		// @todo turnkey: Handle Mac
+		if (HostPlatform.Current.HostEditorPlatform == UnrealTargetPlatform.Mac)
 		{
 			return false;
 		}
@@ -142,14 +171,42 @@ public class AndroidPlatform : Platform
 			}
 			else
 			{
-				TurnkeyContext.PauseForUser("Running the Android Studio installer, and then Android Studio for first-time setup!\n\nChoose all default options unless you know what you are doing.");
-				int ExitCode = TurnkeyContext.RunExternalCommand(OutputPath, "/S", false, true, true);
-				//				Utils.RunLocalProcessAndReturnStdOut(OutputPath, "", out ExitCode, true);
+				if (HostPlatform.Current.HostEditorPlatform == UnrealTargetPlatform.Linux)
+				{
+					// TODO finish GUI support for Linux here, otherwise this will throw. Using zenity
+					// TurnkeyContext.PauseForUser("Running the Android Studio installer, and then Android Studio for first-time setup!\n\nChoose all default options unless you know what you are doing.");
+
+					string UserHome = Environment.GetEnvironmentVariable("HOME");
+					string Args = string.Format("-xf {0} -C {1}", OutputPath, UserHome);
+
+					int ExitCode = TurnkeyContext.RunExternalCommand("/bin/tar", Args, false, true, true);
+
+					if (ExitCode != 0)
+					{
+						TurnkeyContext.ReportError($"Android Studio installer failed. ExitCode = {ExitCode}");
+						return false;
+					}
+				}
+				else if (HostPlatform.Current.HostEditorPlatform == UnrealTargetPlatform.Mac)
+				{
+					// TODO
+				}
+				else if (HostPlatform.Current.HostEditorPlatform == UnrealTargetPlatform.Win64)
+				{
+					TurnkeyContext.PauseForUser("Running the Android Studio installer, and then Android Studio for first-time setup!\n\nChoose all default options unless you know what you are doing.");
+					int ExitCode = TurnkeyContext.RunExternalCommand(OutputPath, "/S", false, true, true);
+					//				Utils.RunLocalProcessAndReturnStdOut(OutputPath, "", out ExitCode, true);
 
-				// AS installer returns 1223 even on success ("user canceled" even tho there's no UI to cancel it...) when running with /S
-				if (ExitCode != 0 && ExitCode != 1223)
+					// AS installer returns 1223 even on success ("user canceled" even tho there's no UI to cancel it...) when running with /S
+					if (ExitCode != 0 && ExitCode != 1223)
+					{
+						TurnkeyContext.ReportError($"Android Studio installer failed. ExitCode = {ExitCode}");
+						return false;
+					}
+				}
+				else
 				{
-					TurnkeyContext.ReportError($"Android Studio installer failed. ExitCode = {ExitCode}");
+					TurnkeyContext.ReportError($"Invalid host platform");
 					return false;
 				}
 
@@ -2772,4 +2829,4 @@ public class AndroidPlatformASTC : AndroidPlatform
     {
         return new TargetPlatformDescriptor(TargetPlatformType, "ASTC");
     }
-}
\ No newline at end of file
+}
