Changing logic of BuildPlugin slightly such that it correctly outputs all of the Target Platforms that will be built, accounting for SupportedTargetPlatforms.  Still works the same, but now with correct output.

diff --git a/Engine/Source/Programs/AutomationTool/Scripts/BuildPluginCommand.Automation.cs b/Engine/Source/Programs/AutomationTool/Scripts/BuildPluginCommand.Automation.cs
--- a/Engine/Source/Programs/AutomationTool/Scripts/BuildPluginCommand.Automation.cs
+++ b/Engine/Source/Programs/AutomationTool/Scripts/BuildPluginCommand.Automation.cs
@@ -152,18 +152,16 @@ class BuildPlugin : BuildCommand
 			}
 		}
 
-		// Add the game targets
+		// Add the supported game targets
 		if(TargetPlatforms.Count > 0)
 		{
-			CommandUtils.LogInformation("Building plugin for target platforms: {0}", String.Join(", ", TargetPlatforms));
-			foreach (UnrealTargetPlatform TargetPlatform in TargetPlatforms)
+			List<UnrealTargetPlatform> SupportedTargetPlatforms = TargetPlatforms.FindAll(Plugin.SupportsTargetPlatform);
+			CommandUtils.LogInformation("Building plugin for target platforms: {0}", String.Join(", ", SupportedTargetPlatforms));
+			foreach (UnrealTargetPlatform TargetPlatform in SupportedTargetPlatforms)
 			{
-				if(Plugin.SupportsTargetPlatform(TargetPlatform))
-				{
-					string AdditionalTargetArgs = AdditionalArgs;
-					CompilePluginWithUBT(HostProjectFile, HostProjectPluginFile, Plugin, "UnrealGame", TargetType.Game, TargetPlatform, UnrealTargetConfiguration.Development, ManifestFileNames, AdditionalTargetArgs);
-					CompilePluginWithUBT(HostProjectFile, HostProjectPluginFile, Plugin, "UnrealGame", TargetType.Game, TargetPlatform, UnrealTargetConfiguration.Shipping, ManifestFileNames, AdditionalTargetArgs);
-				}
+				string AdditionalTargetArgs = AdditionalArgs;
+				CompilePluginWithUBT(HostProjectFile, HostProjectPluginFile, Plugin, "UnrealGame", TargetType.Game, TargetPlatform, UnrealTargetConfiguration.Development, ManifestFileNames, AdditionalTargetArgs);
+				CompilePluginWithUBT(HostProjectFile, HostProjectPluginFile, Plugin, "UnrealGame", TargetType.Game, TargetPlatform, UnrealTargetConfiguration.Shipping, ManifestFileNames, AdditionalTargetArgs);
 			}
 		}
 
