Gauntlet - listen to additional log channels when running tests (LogStaticMesh and LogSkeletalMech) to avoid idle timeout

#rnx

diff --git a/Engine/Source/Programs/AutomationTool/Gauntlet/Unreal/Automation/UE.Automation.cs b/Engine/Source/Programs/AutomationTool/Gauntlet/Unreal/Automation/UE.Automation.cs
--- a/Engine/Source/Programs/AutomationTool/Gauntlet/Unreal/Automation/UE.Automation.cs
+++ b/Engine/Source/Programs/AutomationTool/Gauntlet/Unreal/Automation/UE.Automation.cs
@@ -355,7 +355,7 @@ namespace UE
 			var AppInstance = TestInstance.EditorApp;
 
 			UnrealLogParser Parser = new UnrealLogParser(AppInstance.StdOut);
-			ChannelEntries.AddRange(Parser.GetLogChannels(new string[] { "Automation", "FunctionalTest", "Material", "DerivedDataCache", "ShaderCompilers" }, false));
+			ChannelEntries.AddRange(Parser.GetLogChannels(new string[] { "Automation", "FunctionalTest", "Material", "DerivedDataCache", "ShaderCompilers", "SkeletalMesh", "StaticMesh", "Texture" }, false));
 
 			// Any new entries?
 			if (ChannelEntries.Count > LastAutomationEntryCount)
diff --git a/Engine/Source/Programs/AutomationTool/Gauntlet/Unreal/Engine/UnrealTests.cs b/Engine/Source/Programs/AutomationTool/Gauntlet/Unreal/Engine/UnrealTests.cs
--- a/Engine/Source/Programs/AutomationTool/Gauntlet/Unreal/Engine/UnrealTests.cs
+++ b/Engine/Source/Programs/AutomationTool/Gauntlet/Unreal/Engine/UnrealTests.cs
@@ -142,7 +142,7 @@ namespace Gauntlet
 				var AppInstance = TestInstance.EditorApp;
 
 				UnrealLogParser Parser = new UnrealLogParser(AppInstance.StdOut);
-				ChannelEntries.AddRange(Parser.GetLogChannels(new string[] { "Automation", "FunctionalTest", "Material", "DerivedDataCache", "ShaderCompilers" }, false));
+				ChannelEntries.AddRange(Parser.GetLogChannels(new string[] { "Automation", "FunctionalTest", "Material", "DerivedDataCache", "ShaderCompilers", "SkeletalMesh", "StaticMesh", "Texture" }, false));
 
 				if (ChannelEntries.Count > LastAutomationEntryCount)
 				{
diff --git a/Engine/Source/Programs/AutomationTool/Gauntlet/Unreal/Utils/Gauntlet.UnrealLogParser.cs b/Engine/Source/Programs/AutomationTool/Gauntlet/Unreal/Utils/Gauntlet.UnrealLogParser.cs
--- a/Engine/Source/Programs/AutomationTool/Gauntlet/Unreal/Utils/Gauntlet.UnrealLogParser.cs
+++ b/Engine/Source/Programs/AutomationTool/Gauntlet/Unreal/Utils/Gauntlet.UnrealLogParser.cs
@@ -353,7 +353,7 @@ namespace Gauntlet
 		/// <returns></returns>
 		public IEnumerable<string> GetEditorBusyChannels ()
 		{
-			return GetLogChannels(new string[] { "Automation", "FunctionalTest", "Material", "DerivedDataCache", "ShaderCompilers", "Texture" }, false);
+			return GetLogChannels(new string[] { "Automation", "FunctionalTest", "Material", "DerivedDataCache", "ShaderCompilers", "Texture", "SkeletalMesh", "StaticMesh" }, false);
 		}
 
 		/// <summary>
