Gauntlet - centralize the busy log channel list

#rnx

diff --git a/Engine/Source/Programs/AutomationTool/Gauntlet/Unreal/Automation/UE.Automation.cs b/Engine/Source/Programs/AutomationTool/Gauntlet/Unreal/Automation/UE.Automation.cs
--- a/Engine/Source/Programs/AutomationTool/Gauntlet/Unreal/Automation/UE.Automation.cs
+++ b/Engine/Source/Programs/AutomationTool/Gauntlet/Unreal/Automation/UE.Automation.cs
@@ -355,7 +355,7 @@ namespace UE
 			var AppInstance = TestInstance.EditorApp;
 
 			UnrealLogParser Parser = new UnrealLogParser(AppInstance.StdOut);
-			ChannelEntries.AddRange(Parser.GetLogChannels(new string[] { "Automation", "FunctionalTest", "Material", "DerivedDataCache", "ShaderCompilers", "SkeletalMesh", "StaticMesh", "Texture" }, false));
+			ChannelEntries.AddRange(Parser.GetEditorBusyChannels());
 
 			// Any new entries?
 			if (ChannelEntries.Count > LastAutomationEntryCount)
diff --git a/Engine/Source/Programs/AutomationTool/Gauntlet/Unreal/Engine/UnrealTests.cs b/Engine/Source/Programs/AutomationTool/Gauntlet/Unreal/Engine/UnrealTests.cs
--- a/Engine/Source/Programs/AutomationTool/Gauntlet/Unreal/Engine/UnrealTests.cs
+++ b/Engine/Source/Programs/AutomationTool/Gauntlet/Unreal/Engine/UnrealTests.cs
@@ -142,7 +142,7 @@ namespace Gauntlet
 				var AppInstance = TestInstance.EditorApp;
 
 				UnrealLogParser Parser = new UnrealLogParser(AppInstance.StdOut);
-				ChannelEntries.AddRange(Parser.GetLogChannels(new string[] { "Automation", "FunctionalTest", "Material", "DerivedDataCache", "ShaderCompilers", "SkeletalMesh", "StaticMesh", "Texture" }, false));
+				ChannelEntries.AddRange(Parser.GetEditorBusyChannels());
 
 				if (ChannelEntries.Count > LastAutomationEntryCount)
 				{
