BuildGraph: Add a dedicated Git-Clone task for BuildGraph, to save some boilerplate code in scripts.

diff --git a/Engine/Source/Programs/AutomationTool/BuildGraph/Tasks/GitCloneTask.cs b/Engine/Source/Programs/AutomationTool/BuildGraph/Tasks/GitCloneTask.cs
new file mode 100644
--- /dev/null
+++ b/Engine/Source/Programs/AutomationTool/BuildGraph/Tasks/GitCloneTask.cs
@@ -0,0 +1,142 @@
+// Copyright Epic Games, Inc. All Rights Reserved.
+
+using EpicGames.Core;
+using System;
+using System.Collections.Generic;
+using System.Text;
+using System.Xml;
+
+namespace AutomationTool.Tasks
+{
+	/// <summary>
+	/// Parameters for a Git-Checkout task
+	/// </summary>
+	public class GitCloneTaskParameters
+	{
+		/// <summary>
+		/// Directory for the repository
+		/// </summary>
+		[TaskParameter]
+		public string Dir;
+
+		/// <summary>
+		/// The remote to add
+		/// </summary>
+		[TaskParameter(Optional = true)]
+		public string Remote;
+
+		/// <summary>
+		/// The branch to check out on the remote
+		/// </summary>
+		[TaskParameter]
+		public string Branch;
+
+		/// <summary>
+		/// Configuration file for the repo. This can be used to set up a remote to be fetched and/or provide credentials.
+		/// </summary>
+		[TaskParameter(Optional = true)]
+		public string ConfigFile;
+	}
+
+	/// <summary>
+	/// Spawns Git and waits for it to complete.
+	/// </summary>
+	[TaskElement("Git-Clone", typeof(GitCloneTaskParameters))]
+	public class GitCloneTask : CustomTask
+	{
+		/// <summary>
+		/// Parameters for this task
+		/// </summary>
+		GitCloneTaskParameters Parameters;
+
+		/// <summary>
+		/// Construct a Git task
+		/// </summary>
+		/// <param name="InParameters">Parameters for the task</param>
+		public GitCloneTask(GitCloneTaskParameters InParameters)
+		{
+			Parameters = InParameters;
+		}
+
+		/// <summary>
+		/// Execute the task.
+		/// </summary>
+		/// <param name="Job">Information about the current job</param>
+		/// <param name="BuildProducts">Set of build products produced by this node.</param>
+		/// <param name="TagNameToFileSet">Mapping from tag names to the set of files they include</param>
+		public override void Execute(JobContext Job, HashSet<FileReference> BuildProducts, Dictionary<string, HashSet<FileReference>> TagNameToFileSet)
+		{
+			FileReference GitExe = CommandUtils.FindToolInPath("git");
+			if(GitExe == null)
+			{
+				throw new AutomationException("Unable to find path to Git. Check you have it installed, and it is on your PATH.");
+			}
+
+			DirectoryReference Dir = ResolveDirectory(Parameters.Dir);
+			Log.TraceInformation("Cloning Git repository into {0}", Parameters.Dir);
+			using (LogIndentScope Scope = new LogIndentScope("  "))
+			{
+				DirectoryReference GitDir = DirectoryReference.Combine(Dir, ".git");
+				if (!FileReference.Exists(FileReference.Combine(GitDir, "HEAD")))
+				{
+					RunGit(GitExe, $"init \"{Dir}\"", CommandUtils.RootDirectory);
+				}
+
+				if (Parameters.ConfigFile != null)
+				{
+					CommandUtils.CopyFile(Parameters.ConfigFile, FileReference.Combine(GitDir, "config").FullName);
+				}
+
+				if (Parameters.Remote != null)
+				{
+					RunGit(GitExe, $"remote add origin {Parameters.Remote}", Dir);
+				}
+
+				RunGit(GitExe, "clean -dxf", Dir);
+				RunGit(GitExe, "fetch --all", Dir);
+				RunGit(GitExe, $"reset --hard {Parameters.Branch}", Dir);
+			}
+		}
+
+		/// <summary>
+		/// Runs a git command
+		/// </summary>
+		/// <param name="ToolFile"></param>
+		/// <param name="Arguments"></param>
+		/// <param name="WorkingDir"></param>
+		void RunGit(FileReference ToolFile, string Arguments, DirectoryReference WorkingDir)
+		{
+			IProcessResult Result = CommandUtils.Run(ToolFile.FullName, Arguments, WorkingDir: WorkingDir.FullName);
+			if (Result.ExitCode != 0)
+			{
+				throw new AutomationException("Git terminated with an exit code indicating an error ({0})", Result.ExitCode);
+			}
+		}
+
+		/// <summary>
+		/// Output this task out to an XML writer.
+		/// </summary>
+		public override void Write(XmlWriter Writer)
+		{
+			Write(Writer, Parameters);
+		}
+
+		/// <summary>
+		/// Find all the tags which are used as inputs to this task
+		/// </summary>
+		/// <returns>The tag names which are read by this task</returns>
+		public override IEnumerable<string> FindConsumedTagNames()
+		{
+			yield break;
+		}
+
+		/// <summary>
+		/// Find all the tags which are modified by this task
+		/// </summary>
+		/// <returns>The tag names which are modified by this task</returns>
+		public override IEnumerable<string> FindProducedTagNames()
+		{
+			yield break;
+		}
+	}
+}
