BuildGraph: Add specific tasks for building and pushing images with Docker.

diff --git a/Engine/Source/Programs/AutomationTool/AutomationUtils/CommandUtils.cs b/Engine/Source/Programs/AutomationTool/AutomationUtils/CommandUtils.cs
--- a/Engine/Source/Programs/AutomationTool/AutomationUtils/CommandUtils.cs
+++ b/Engine/Source/Programs/AutomationTool/AutomationUtils/CommandUtils.cs
@@ -1668,7 +1668,40 @@ namespace AutomationTool
 
 		/// <summary>
 		/// Copies files using multiple threads
-        /// </summary>
+		/// </summary>
+		/// <param name="SourceFiles">List of source files to copy</param>
+		/// <param name="TargetFiles">List of target files</param>
+		public static void ThreadedCopyFiles(List<FileReference> SourceFiles, List<FileReference> TargetFiles)
+		{
+			ThreadedCopyFiles(SourceFiles.ConvertAll(x => x.FullName), TargetFiles.ConvertAll(x => x.FullName));
+		}
+
+		/// <summary>
+		/// Copies files using multiple threads
+		/// </summary>
+		/// <param name="SourceFiles">List of source files to copy</param>
+		/// <param name="SourceDir">The source directory</param>
+		/// <param name="TargetDir">The target directory</param>
+		public static void ThreadedCopyFiles(List<FileReference> SourceFiles, DirectoryReference SourceDir, DirectoryReference TargetDir)
+		{
+			List<FileReference> TargetFiles = new List<FileReference>();
+			foreach (FileReference SourceFile in SourceFiles)
+			{
+				if (!SourceFile.IsUnderDirectory(SourceDir))
+				{
+					throw new AutomationException($"Source file '{SourceFile}' is not under source directory '{SourceDir}'");
+				}
+				else
+				{
+					TargetFiles.Add(FileReference.Combine(TargetDir, SourceFile.MakeRelativeTo(SourceDir)));
+				}
+			}
+			ThreadedCopyFiles(SourceFiles, TargetFiles);
+		}
+
+		/// <summary>
+		/// Copies files using multiple threads
+		/// </summary>
 		/// <param name="Source"></param>
 		/// <param name="Dest"></param>
 		/// <param name="MaxThreads"></param>
diff --git a/Engine/Source/Programs/AutomationTool/BuildGraph/Tasks/DockerBuildTask.cs b/Engine/Source/Programs/AutomationTool/BuildGraph/Tasks/DockerBuildTask.cs
new file mode 100644
--- /dev/null
+++ b/Engine/Source/Programs/AutomationTool/BuildGraph/Tasks/DockerBuildTask.cs
@@ -0,0 +1,164 @@
+// Copyright Epic Games, Inc. All Rights Reserved.
+
+using EpicGames.Core;
+using System;
+using System.Collections.Generic;
+using System.IO;
+using System.Linq;
+using System.Text;
+using System.Xml;
+
+namespace AutomationTool.Tasks
+{
+	/// <summary>
+	/// Parameters for a Docker-Build task
+	/// </summary>
+	public class DockerBuildTaskParameters
+	{
+		/// <summary>
+		/// Base directory for the build
+		/// </summary>
+		[TaskParameter]
+		public string BaseDir;
+
+		/// <summary>
+		/// Files to be staged before building the image
+		/// </summary>
+		[TaskParameter]
+		public string Files;
+
+		/// <summary>
+		/// Path to the Dockerfile. Uses the root of basedir if not specified.
+		/// </summary>
+		[TaskParameter(Optional = true)]
+		public string DockerFile;
+
+		/// <summary>
+		/// Tag for the image
+		/// </summary>
+		[TaskParameter(Optional = true)]
+		public string Tag;
+
+		/// <summary>
+		/// Optional arguments
+		/// </summary>
+		[TaskParameter(Optional = true)]
+		public string Arguments;
+
+		/// <summary>
+		/// List of additional directories to overlay into the staged input files. Allows credentials to be staged, etc...
+		/// </summary>
+		[TaskParameter(Optional = true)]
+		public string OverlayDirs;
+	}
+
+	/// <summary>
+	/// Spawns Docker and waits for it to complete.
+	/// </summary>
+	[TaskElement("Docker-Build", typeof(DockerBuildTaskParameters))]
+	public class DockerBuildTask : CustomTask
+	{
+		/// <summary>
+		/// Parameters for this task
+		/// </summary>
+		DockerBuildTaskParameters Parameters;
+
+		/// <summary>
+		/// Construct a Docker task
+		/// </summary>
+		/// <param name="InParameters">Parameters for the task</param>
+		public DockerBuildTask(DockerBuildTaskParameters InParameters)
+		{
+			Parameters = InParameters;
+		}
+
+		/// <summary>
+		/// Execute the task.
+		/// </summary>
+		/// <param name="Job">Information about the current job</param>
+		/// <param name="BuildProducts">Set of build products produced by this node.</param>
+		/// <param name="TagNameToFileSet">Mapping from tag names to the set of files they include</param>
+		public override void Execute(JobContext Job, HashSet<FileReference> BuildProducts, Dictionary<string, HashSet<FileReference>> TagNameToFileSet)
+		{
+			FileReference DockerExe = CommandUtils.FindToolInPath("docker");
+			if(DockerExe == null)
+			{
+				throw new AutomationException("Unable to find path to Docker. Check you have it installed, and it is on your PATH.");
+			}
+
+			Log.TraceInformation("Building Docker image");
+			using (LogIndentScope Scope = new LogIndentScope("  "))
+			{
+				DirectoryReference BaseDir = ResolveDirectory(Parameters.BaseDir);
+				List<FileReference> SourceFiles = ResolveFilespec(BaseDir, Parameters.Files, TagNameToFileSet).ToList();
+
+				DirectoryReference StagingDir = DirectoryReference.Combine(CommandUtils.EngineDirectory, "Intermediate", "Docker");
+
+				List<FileReference> TargetFiles = SourceFiles.ConvertAll(x => FileReference.Combine(StagingDir, x.MakeRelativeTo(BaseDir)));
+				CommandUtils.ThreadedCopyFiles(SourceFiles, BaseDir, StagingDir);
+
+				if (!String.IsNullOrEmpty(Parameters.OverlayDirs))
+				{
+					foreach (string OverlayDir in Parameters.OverlayDirs.Split(';'))
+					{
+						CommandUtils.ThreadedCopyFiles(ResolveDirectory(OverlayDir), StagingDir);
+					}
+				}
+
+				StringBuilder Arguments = new StringBuilder("build .");
+				if (Parameters.Tag != null)
+				{
+					Arguments.Append($" -t {Parameters.Tag}");
+				}
+				if (Parameters.DockerFile != null)
+				{
+					FileReference DockerFile = ResolveFile(Parameters.DockerFile);
+					if (!DockerFile.IsUnderDirectory(BaseDir))
+					{
+						throw new AutomationException($"Dockerfile '{DockerFile}' is not under base directory ({BaseDir})");
+					}
+					Arguments.Append($" -f {DockerFile.MakeRelativeTo(BaseDir).QuoteArgument()}");
+				}
+				if (Parameters.Arguments != null)
+				{
+					Arguments.Append($" {Parameters.Arguments}");
+				}
+
+				IProcessResult Result = CommandUtils.Run(DockerExe.FullName, Arguments.ToString(), WorkingDir: StagingDir.FullName, Options: CommandUtils.ERunOptions.AllowSpew);
+				if (Result.ExitCode != 0)
+				{
+					throw new AutomationException("Docker terminated with an exit code indicating an error ({0})", Result.ExitCode);
+				}
+			}
+		}
+
+		/// <summary>
+		/// Output this task out to an XML writer.
+		/// </summary>
+		public override void Write(XmlWriter Writer)
+		{
+			Write(Writer, Parameters);
+		}
+
+		/// <summary>
+		/// Find all the tags which are used as inputs to this task
+		/// </summary>
+		/// <returns>The tag names which are read by this task</returns>
+		public override IEnumerable<string> FindConsumedTagNames()
+		{
+			List<string> TagNames = new List<string>();
+			TagNames.AddRange(FindTagNamesFromFilespec(Parameters.DockerFile));
+			TagNames.AddRange(FindTagNamesFromFilespec(Parameters.Files));
+			return TagNames;
+		}
+
+		/// <summary>
+		/// Find all the tags which are modified by this task
+		/// </summary>
+		/// <returns>The tag names which are modified by this task</returns>
+		public override IEnumerable<string> FindProducedTagNames()
+		{
+			yield break;
+		}
+	}
+}
diff --git a/Engine/Source/Programs/AutomationTool/BuildGraph/Tasks/DockerPushTask.cs b/Engine/Source/Programs/AutomationTool/BuildGraph/Tasks/DockerPushTask.cs
new file mode 100644
--- /dev/null
+++ b/Engine/Source/Programs/AutomationTool/BuildGraph/Tasks/DockerPushTask.cs
@@ -0,0 +1,210 @@
+// Copyright Epic Games, Inc. All Rights Reserved.
+
+using EpicGames.Core;
+using System;
+using System.Collections.Generic;
+using System.IO;
+using System.Linq;
+using System.Text;
+using System.Text.Json;
+using System.Xml;
+
+namespace AutomationTool.Tasks
+{
+	/// <summary>
+	/// Parameters for a Docker-Build task
+	/// </summary>
+	public class DockerPushTaskParameters
+	{
+		/// <summary>
+		/// Repository
+		/// </summary>
+		[TaskParameter]
+		public string Repository;
+
+		/// <summary>
+		/// Source image to push
+		/// </summary>
+		[TaskParameter]
+		public string Image;
+
+		/// <summary>
+		/// Name of the target image
+		/// </summary>
+		[TaskParameter(Optional = true)]
+		public string TargetImage;
+
+		/// <summary>
+		/// Path to a credentials file
+		/// </summary>
+		[TaskParameter(Optional = true)]
+		public string CredentialsFile;
+	}
+
+	/// <summary>
+	/// Spawns Docker and waits for it to complete.
+	/// </summary>
+	[TaskElement("Docker-Push", typeof(DockerPushTaskParameters))]
+	public class DockerPushTask : CustomTask
+	{
+		/// <summary>
+		/// Base class for credentials that can be read from a JSON file
+		/// </summary>
+		[JsonKnownTypes(typeof(BasicDockerCredentials), typeof(AwsDockerCredentials))]
+		abstract class DockerCredentials
+		{
+			public abstract (string UserName, string Password) Resolve();
+		}
+
+		/// <summary>
+		/// Basic credentials
+		/// </summary>
+		[JsonDiscriminator("Basic")]
+		class BasicDockerCredentials : DockerCredentials
+		{
+			public string UserName { get; set; }
+			public string Password { get; set; }
+
+			public override (string, string) Resolve() => (UserName, Password);
+		}
+
+		/// <summary>
+		/// AWS credentials
+		/// </summary>
+		[JsonDiscriminator("AWS")]
+		class AwsDockerCredentials : DockerCredentials
+		{
+			public string Region { get; set; }
+			public string AccessKey { get; set; }
+			public string SecretKey { get; set; }
+
+			public override (string, string) Resolve()
+			{
+				Dictionary<string, string> EnvVars = new Dictionary<string, string>();
+				if (AccessKey != null)
+				{
+					EnvVars.Add("AWS_ACCESS_KEY_ID", AccessKey);
+				}
+				if(SecretKey != null)
+				{
+					EnvVars.Add("AWS_SECRET_ACCESS_KEY", SecretKey);
+				}
+
+				FileReference AwsExe = new FileReference(@"C:\Program Files\Amazon\AWSCLIV2\aws.exe");// CommandUtils.FindToolInPath("aws");
+				if (AwsExe == null)
+				{
+					throw new AutomationException("Unable to find path to AWSCLI. Check you have it installed, and it is on your PATH.");
+				}
+
+				StringBuilder Arguments = new StringBuilder("ecr get-login-password");
+				if(Region != null)
+				{
+					Arguments.Append($" --region {Region}");
+				}
+
+				IProcessResult Result = CommandUtils.Run(AwsExe.FullName, Arguments.ToString(), Env: EnvVars, Options: CommandUtils.ERunOptions.None);
+				if (Result.ExitCode != 0)
+				{
+					Log.TraceInformation(Result.Output);
+					throw new AutomationException("AWSCLI terminated with an exit code indicating an error ({0})", Result.ExitCode);
+				}
+
+				return ("AWS", Result.Output);
+			}
+		}
+
+		/// <summary>
+		/// Parameters for this task
+		/// </summary>
+		DockerPushTaskParameters Parameters;
+
+		/// <summary>
+		/// Construct a Docker task
+		/// </summary>
+		/// <param name="InParameters">Parameters for the task</param>
+		public DockerPushTask(DockerPushTaskParameters InParameters)
+		{
+			Parameters = InParameters;
+		}
+
+		/// <summary>
+		/// Execute the task.
+		/// </summary>
+		/// <param name="Job">Information about the current job</param>
+		/// <param name="BuildProducts">Set of build products produced by this node.</param>
+		/// <param name="TagNameToFileSet">Mapping from tag names to the set of files they include</param>
+		public override void Execute(JobContext Job, HashSet<FileReference> BuildProducts, Dictionary<string, HashSet<FileReference>> TagNameToFileSet)
+		{
+			FileReference DockerExe = CommandUtils.FindToolInPath("docker");
+			if(DockerExe == null)
+			{
+				throw new AutomationException("Unable to find path to Docker. Check you have it installed, and it is on your PATH.");
+			}
+
+			Log.TraceInformation("Pushing Docker image");
+			using (LogIndentScope Scope = new LogIndentScope("  "))
+			{
+				if (Parameters.CredentialsFile != null)
+				{
+					FileReference CredentialsFile = ResolveFile(Parameters.CredentialsFile);
+					byte[] Data = FileReference.ReadAllBytes(CredentialsFile);
+
+					JsonSerializerOptions Options = new JsonSerializerOptions();
+					Options.PropertyNamingPolicy = JsonNamingPolicy.CamelCase;
+					Options.PropertyNameCaseInsensitive = true;
+					Options.Converters.Add(new JsonKnownTypesConverterFactory());
+
+					DockerCredentials Credentials = JsonSerializer.Deserialize<DockerCredentials>(Data, Options);
+					(string UserName, string Password) = Credentials.Resolve();
+
+					RunDocker(DockerExe, $"login {Parameters.Repository} --username {UserName} --password-stdin", Password, CommandUtils.RootDirectory);
+				}
+
+				string TargetImage = Parameters.TargetImage ?? Parameters.Image;
+				RunDocker(DockerExe, $"tag {Parameters.Image} {Parameters.Repository}/{TargetImage}", null, CommandUtils.RootDirectory);
+				RunDocker(DockerExe, $"push {Parameters.Repository}/{TargetImage}", null, CommandUtils.RootDirectory);
+			}
+		}
+
+		/// <summary>
+		/// Runs Docker
+		/// </summary>
+		/// <param name="DockerExe"></param>
+		/// <param name="Arguments"></param>
+		/// <param name="StagingDir"></param>
+		public void RunDocker(FileReference DockerExe, string Arguments, string Input, DirectoryReference StagingDir)
+		{
+			IProcessResult Result = CommandUtils.Run(DockerExe.FullName, Arguments, Input: Input, WorkingDir: StagingDir.FullName, Options: CommandUtils.ERunOptions.AllowSpew);
+			if (Result.ExitCode != 0)
+			{
+				throw new AutomationException("Docker terminated with an exit code indicating an error ({0})", Result.ExitCode);
+			}
+		}
+
+		/// <summary>
+		/// Output this task out to an XML writer.
+		/// </summary>
+		public override void Write(XmlWriter Writer)
+		{
+			Write(Writer, Parameters);
+		}
+
+		/// <summary>
+		/// Find all the tags which are used as inputs to this task
+		/// </summary>
+		/// <returns>The tag names which are read by this task</returns>
+		public override IEnumerable<string> FindConsumedTagNames()
+		{
+			return FindTagNamesFromFilespec(Parameters.CredentialsFile);
+		}
+
+		/// <summary>
+		/// Find all the tags which are modified by this task
+		/// </summary>
+		/// <returns>The tag names which are modified by this task</returns>
+		public override IEnumerable<string> FindProducedTagNames()
+		{
+			yield break;
+		}
+	}
+}
