Remove hard-coded tool location.

diff --git a/Engine/Source/Programs/AutomationTool/BuildGraph/Tasks/DockerPushTask.cs b/Engine/Source/Programs/AutomationTool/BuildGraph/Tasks/DockerPushTask.cs
--- a/Engine/Source/Programs/AutomationTool/BuildGraph/Tasks/DockerPushTask.cs
+++ b/Engine/Source/Programs/AutomationTool/BuildGraph/Tasks/DockerPushTask.cs
@@ -90,7 +90,7 @@ namespace AutomationTool.Tasks
 					EnvVars.Add("AWS_SECRET_ACCESS_KEY", SecretKey);
 				}
 
-				FileReference AwsExe = new FileReference(@"C:\Program Files\Amazon\AWSCLIV2\aws.exe");// CommandUtils.FindToolInPath("aws");
+				FileReference AwsExe = CommandUtils.FindToolInPath("aws");
 				if (AwsExe == null)
 				{
 					throw new AutomationException("Unable to find path to AWSCLI. Check you have it installed, and it is on your PATH.");
