BuildGraph: Use a standard AWS config file to configure Aws/Docker/Helm tasks.

diff --git a/Engine/Source/Programs/AutomationTool/BuildGraph/Tasks/AwsTask.cs b/Engine/Source/Programs/AutomationTool/BuildGraph/Tasks/AwsTask.cs
--- a/Engine/Source/Programs/AutomationTool/BuildGraph/Tasks/AwsTask.cs
+++ b/Engine/Source/Programs/AutomationTool/BuildGraph/Tasks/AwsTask.cs
@@ -5,32 +5,104 @@ using System;
 using System.Collections.Generic;
 using System.IO;
 using System.Text;
+using System.Text.Json;
 using System.Xml;
 
 namespace AutomationTool.Tasks
 {
+	#region Settings
+
 	/// <summary>
-	/// Parameters for an AWS CLI task
+	/// AWS Settings file
 	/// </summary>
-	public class AwsTaskParameters
+	public class AwsSettings
 	{
 		/// <summary>
-		/// AWS command line arguments
+		/// The default region
 		/// </summary>
-		[TaskParameter(Optional = true)]
-		public string Arguments;
+		public string Region { get; set; }
 
 		/// <summary>
-		/// Base directory for running the command
+		/// The credentials
 		/// </summary>
-		[TaskParameter(Optional = true)]
-		public string BaseDir;
+		public AwsCredentials Credentials { get; set; }
 
 		/// <summary>
-		/// The minimum exit code, which is treated as an error.
+		/// Read settings from a file
+		/// </summary>
+		/// <param name="File"></param>
+		/// <returns></returns>
+		public static AwsSettings Read(FileReference File)
+		{
+			byte[] Data = FileReference.ReadAllBytes(File);
+
+			JsonSerializerOptions Options = new JsonSerializerOptions();
+			Options.PropertyNameCaseInsensitive = true;
+
+			return JsonSerializer.Deserialize<AwsSettings>(Data, Options);
+		}
+
+		/// <summary>
+		/// Gets environment variables in a dictionary
+		/// </summary>
+		/// <param name="EnvVars"></param>
+		public void GetEnvVars(Dictionary<string, string> EnvVars)
+		{
+			if (Region != null)
+			{
+				EnvVars.Add("AWS_DEFAULT_REGION", Region);
+			}
+			if (Credentials != null)
+			{
+				if (Credentials.AccessKeyId != null)
+				{
+					EnvVars.Add("AWS_ACCESS_KEY_ID", Credentials.AccessKeyId);
+				}
+				if (Credentials.SecretAccessKey != null)
+				{
+					EnvVars.Add("AWS_SECRET_ACCESS_KEY", Credentials.SecretAccessKey);
+				}
+				if (Credentials.SessionToken != null)
+				{
+					EnvVars.Add("AWS_SESSION_TOKEN", Credentials.SessionToken);
+				}
+			}
+		}
+	}
+
+	/// <summary>
+	/// Credentials for AWS
+	/// </summary>
+	public class AwsCredentials
+	{
+		/// <summary>
+		/// Access key for AWS
+		/// </summary>
+		public string AccessKeyId { get; set; }
+
+		/// <summary>
+		/// Secret key for AWS
+		/// </summary>
+		public string SecretAccessKey { get; set; }
+
+		/// <summary>
+		/// Session token for AWS
+		/// </summary>
+		public string SessionToken { get; set; }
+	}
+
+	#endregion
+
+	/// <summary>
+	/// Parameters for an AWS CLI task
+	/// </summary>
+	public class AwsTaskParameters
+	{
+		/// <summary>
+		/// AWS command line arguments
 		/// </summary>
 		[TaskParameter(Optional = true)]
-		public int ErrorLevel = 1;
+		public string Arguments;
 
 		/// <summary>
 		/// Whether or not to show the console output from the AWS CLI command.
@@ -42,10 +114,10 @@ namespace AutomationTool.Tasks
 		/// Path to a file containing the credentials needed to run the AWS command
 		/// </summary>
 		[TaskParameter(Optional = true)]
-		public string CredentialsFile;
+		public string SettingsFile;
 
 		/// <summary>
-		/// Path to a file to store any JSON output by the command.
+		/// Path to receive the ouptut
 		/// </summary>
 		[TaskParameter(Optional = true)]
 		public string OutputFile;
@@ -78,41 +150,50 @@ namespace AutomationTool.Tasks
 		/// <param name="BuildProducts">Set of build products produced by this node.</param>
 		/// <param name="TagNameToFileSet">Mapping from tag names to the set of files they include</param>
 		public override void Execute(JobContext Job, HashSet<FileReference> BuildProducts, Dictionary<string, HashSet<FileReference>> TagNameToFileSet)
+		{
+			FileReference SettingsFile = (Parameters.SettingsFile != null) ? ResolveFile(Parameters.SettingsFile) : null;
+			FileReference OutputFile = (Parameters.OutputFile != null) ? ResolveFile(Parameters.OutputFile) : null;
+			Run(Parameters.Arguments, SettingsFile: SettingsFile, OutputFile: OutputFile, LogOutput: Parameters.Verbose);
+		}
+
+		/// <summary>
+		/// Runs an AWS command
+		/// </summary>
+		/// <param name="Arguments"></param>
+		/// <param name="SettingsFile"></param>
+		/// <param name="LogOutput"></param>
+		/// <param name="OutputFile"></param>
+		/// <param name="BaseDir"></param>
+		public static string Run(string Arguments, FileReference SettingsFile = null, FileReference OutputFile = null, DirectoryReference BaseDir = null, bool LogOutput = true)
 		{
 			FileReference ToolFile = CommandUtils.FindToolInPath("aws");
-			if(ToolFile == null)
+			if (ToolFile == null)
 			{
 				throw new AutomationException("Unable to find path to AWS CLI. Check you have it installed, and it is on your PATH.");
 			}
 
-			string AWSArgs = Parameters.Arguments;
 			Dictionary<string, string> EnvVars = new Dictionary<string, string>();
-			if (!string.IsNullOrWhiteSpace(Parameters.CredentialsFile))
+			if (SettingsFile != null)
 			{
-				if (!File.Exists(Parameters.CredentialsFile))
-				{
-					throw new AutomationException("Credentials file {0} could not be found.", Parameters.CredentialsFile);
-				}
-
-				string[] CredentialTokens = File.ReadAllText(Parameters.CredentialsFile).Split(':');
-				EnvVars.Add("AWS_ACCESS_KEY_ID", CredentialTokens[0]);
-				EnvVars.Add("AWS_SECRET_ACCESS_KEY", CredentialTokens[1]);
+				AwsSettings Settings = AwsSettings.Read(SettingsFile);
+				Settings.GetEnvVars(EnvVars);
 			}
 
 			CommandUtils.ERunOptions Options = CommandUtils.ERunOptions.AppMustExist;
-			if (Parameters.Verbose)
+			if (LogOutput)
 			{
 				Options |= CommandUtils.ERunOptions.AllowSpew;
 			}
-			IProcessResult Result = CommandUtils.Run(ToolFile.FullName, AWSArgs, WorkingDir: Parameters.BaseDir, Env: EnvVars, Options: Options);
-			if (!string.IsNullOrWhiteSpace(Parameters.OutputFile))
+			IProcessResult Result = CommandUtils.Run(ToolFile.FullName, Arguments, WorkingDir: BaseDir?.FullName, Env: EnvVars, Options: Options);
+			if (OutputFile != null)
 			{
-				File.WriteAllText(Parameters.OutputFile, Result.Output);
+				FileReference.WriteAllText(OutputFile, Result.Output);
 			}
-			if (Result.ExitCode < 0 || Result.ExitCode >= Parameters.ErrorLevel)
+			if (Result.ExitCode != 0)
 			{
 				throw new AutomationException("AWS terminated with an exit code indicating an error ({0})", Result.ExitCode);
 			}
+			return Result.Output;
 		}
 
 		/// <summary>
diff --git a/Engine/Source/Programs/AutomationTool/BuildGraph/Tasks/DockerPushTask.cs b/Engine/Source/Programs/AutomationTool/BuildGraph/Tasks/DockerPushTask.cs
--- a/Engine/Source/Programs/AutomationTool/BuildGraph/Tasks/DockerPushTask.cs
+++ b/Engine/Source/Programs/AutomationTool/BuildGraph/Tasks/DockerPushTask.cs
@@ -35,10 +35,10 @@ namespace AutomationTool.Tasks
 		public string TargetImage;
 
 		/// <summary>
-		/// Path to a credentials file
+		/// Settings for logging in to AWS ECR
 		/// </summary>
 		[TaskParameter(Optional = true)]
-		public string CredentialsFile;
+		public string AwsSettings;
 	}
 
 	/// <summary>
@@ -47,72 +47,6 @@ namespace AutomationTool.Tasks
 	[TaskElement("Docker-Push", typeof(DockerPushTaskParameters))]
 	public class DockerPushTask : CustomTask
 	{
-		/// <summary>
-		/// Base class for credentials that can be read from a JSON file
-		/// </summary>
-		[JsonKnownTypes(typeof(BasicDockerCredentials), typeof(AwsDockerCredentials))]
-		abstract class DockerCredentials
-		{
-			public abstract (string UserName, string Password) Resolve();
-		}
-
-		/// <summary>
-		/// Basic credentials
-		/// </summary>
-		[JsonDiscriminator("Basic")]
-		class BasicDockerCredentials : DockerCredentials
-		{
-			public string UserName { get; set; }
-			public string Password { get; set; }
-
-			public override (string, string) Resolve() => (UserName, Password);
-		}
-
-		/// <summary>
-		/// AWS credentials
-		/// </summary>
-		[JsonDiscriminator("AWS")]
-		class AwsDockerCredentials : DockerCredentials
-		{
-			public string Region { get; set; }
-			public string AccessKey { get; set; }
-			public string SecretKey { get; set; }
-
-			public override (string, string) Resolve()
-			{
-				Dictionary<string, string> EnvVars = new Dictionary<string, string>();
-				if (AccessKey != null)
-				{
-					EnvVars.Add("AWS_ACCESS_KEY_ID", AccessKey);
-				}
-				if(SecretKey != null)
-				{
-					EnvVars.Add("AWS_SECRET_ACCESS_KEY", SecretKey);
-				}
-
-				FileReference AwsExe = CommandUtils.FindToolInPath("aws");
-				if (AwsExe == null)
-				{
-					throw new AutomationException("Unable to find path to AWSCLI. Check you have it installed, and it is on your PATH.");
-				}
-
-				StringBuilder Arguments = new StringBuilder("ecr get-login-password");
-				if(Region != null)
-				{
-					Arguments.Append($" --region {Region}");
-				}
-
-				IProcessResult Result = CommandUtils.Run(AwsExe.FullName, Arguments.ToString(), Env: EnvVars, Options: CommandUtils.ERunOptions.None);
-				if (Result.ExitCode != 0)
-				{
-					Log.TraceInformation(Result.Output);
-					throw new AutomationException("AWSCLI terminated with an exit code indicating an error ({0})", Result.ExitCode);
-				}
-
-				return ("AWS", Result.Output);
-			}
-		}
-
 		/// <summary>
 		/// Parameters for this task
 		/// </summary>
@@ -144,20 +78,10 @@ namespace AutomationTool.Tasks
 			Log.TraceInformation("Pushing Docker image");
 			using (LogIndentScope Scope = new LogIndentScope("  "))
 			{
-				if (Parameters.CredentialsFile != null)
+				if (Parameters.AwsSettings != null)
 				{
-					FileReference CredentialsFile = ResolveFile(Parameters.CredentialsFile);
-					byte[] Data = FileReference.ReadAllBytes(CredentialsFile);
-
-					JsonSerializerOptions Options = new JsonSerializerOptions();
-					Options.PropertyNamingPolicy = JsonNamingPolicy.CamelCase;
-					Options.PropertyNameCaseInsensitive = true;
-					Options.Converters.Add(new JsonKnownTypesConverterFactory());
-
-					DockerCredentials Credentials = JsonSerializer.Deserialize<DockerCredentials>(Data, Options);
-					(string UserName, string Password) = Credentials.Resolve();
-
-					RunDocker(DockerExe, $"login {Parameters.Repository} --username {UserName} --password-stdin", Password, CommandUtils.RootDirectory);
+					string Password = AwsTask.Run("ecr get-login-password", ResolveFile(Parameters.AwsSettings), LogOutput: false);
+					RunDocker(DockerExe, $"login {Parameters.Repository} --username AWS --password-stdin", Password, CommandUtils.RootDirectory);
 				}
 
 				string TargetImage = Parameters.TargetImage ?? Parameters.Image;
@@ -195,7 +119,7 @@ namespace AutomationTool.Tasks
 		/// <returns>The tag names which are read by this task</returns>
 		public override IEnumerable<string> FindConsumedTagNames()
 		{
-			return FindTagNamesFromFilespec(Parameters.CredentialsFile);
+			return FindTagNamesFromFilespec(Parameters.AwsSettings);
 		}
 
 		/// <summary>
diff --git a/Engine/Source/Programs/AutomationTool/BuildGraph/Tasks/HelmTask.cs b/Engine/Source/Programs/AutomationTool/BuildGraph/Tasks/HelmTask.cs
--- a/Engine/Source/Programs/AutomationTool/BuildGraph/Tasks/HelmTask.cs
+++ b/Engine/Source/Programs/AutomationTool/BuildGraph/Tasks/HelmTask.cs
@@ -24,7 +24,7 @@ namespace AutomationTool.Tasks
 		/// Name of the release
 		/// </summary>
 		[TaskParameter]
-		public string Release;
+		public string Deployment;
 
 		/// <summary>
 		/// The Kubernetes namespace
@@ -121,7 +121,7 @@ namespace AutomationTool.Tasks
 			// Build the argument list
 			List<string> Arguments = new List<string>();
 			Arguments.Add("upgrade");
-			Arguments.Add(Parameters.Release);
+			Arguments.Add(Parameters.Deployment);
 			Arguments.Add(Parameters.Chart);
 			Arguments.Add("--install");
 			Arguments.Add("--reset-values");
