BuildGraph: Remove --dry-run argument from Helm task.

diff --git a/Engine/Source/Programs/AutomationTool/BuildGraph/Tasks/HelmTask.cs b/Engine/Source/Programs/AutomationTool/BuildGraph/Tasks/HelmTask.cs
--- a/Engine/Source/Programs/AutomationTool/BuildGraph/Tasks/HelmTask.cs
+++ b/Engine/Source/Programs/AutomationTool/BuildGraph/Tasks/HelmTask.cs
@@ -135,7 +135,6 @@ namespace AutomationTool.Tasks
 				Arguments.Add("--set");
 				Arguments.Add(Value);
 			}
-			Arguments.Add("--dry-run");
 
 			IProcessResult Result = CommandUtils.Run(HelmExe.FullName, CommandLineArguments.Join(Arguments), WorkingDir: Parameters.BaseDir);
 			if (Result.ExitCode != 0)
