BuildGraph: Add support for passing text files containing environment variables into AWS, Docker, Helm and Spawn tasks. Also add an Aws-AssumeRole task which can assume a role and write a new list of environment variables.

diff --git a/Engine/Source/Programs/AutomationTool/AutomationUtils/ProcessUtils.cs b/Engine/Source/Programs/AutomationTool/AutomationUtils/ProcessUtils.cs
--- a/Engine/Source/Programs/AutomationTool/AutomationUtils/ProcessUtils.cs
+++ b/Engine/Source/Programs/AutomationTool/AutomationUtils/ProcessUtils.cs
@@ -10,6 +10,7 @@ using System.Management;
 using System.Runtime.InteropServices;
 using System.Runtime.CompilerServices;
 using EpicGames.Core;
+using UnrealBuildTool;
 
 namespace AutomationTool
 {
@@ -770,6 +771,11 @@ namespace AutomationTool
 				return App;
 			}
 
+			if (HostPlatform.Current.HostEditorPlatform == UnrealTargetPlatform.Win64 && !Path.HasExtension(App))
+			{
+				App += ".exe";
+			}
+
 			string ResolvedPath = null;
 
 			if (!App.Contains(Path.DirectorySeparatorChar) && !App.Contains(Path.AltDirectorySeparatorChar))
diff --git a/Engine/Source/Programs/AutomationTool/BuildGraph/Tasks/AwsAssumeRoleTask.cs b/Engine/Source/Programs/AutomationTool/BuildGraph/Tasks/AwsAssumeRoleTask.cs
new file mode 100644
--- /dev/null
+++ b/Engine/Source/Programs/AutomationTool/BuildGraph/Tasks/AwsAssumeRoleTask.cs
@@ -0,0 +1,162 @@
+// Copyright Epic Games, Inc. All Rights Reserved.
+
+using AutomationTool;
+using EpicGames.Core;
+using System;
+using System.Collections.Generic;
+using System.IO;
+using System.Linq;
+using System.Text;
+using System.Text.Json;
+using System.Xml;
+
+namespace BuildGraph.Tasks
+{
+	/// <summary>
+	/// Parameters for an AWS CLI task
+	/// </summary>
+	public class AwsAssumeRoleTaskParameters
+	{
+		/// <summary>
+		/// Role to assume
+		/// </summary>
+		[TaskParameter]
+		public string Arn;
+
+		/// <summary>
+		/// Name of this session
+		/// </summary>
+		[TaskParameter]
+		public string Session;
+
+		/// <summary>
+		/// Duration of the token in seconds
+		/// </summary>
+		[TaskParameter(Optional = true)]
+		public int Duration = 1000;
+
+		/// <summary>
+		/// Environment variables
+		/// </summary>
+		[TaskParameter(Optional = true)]
+		public string Environment;
+
+		/// <summary>
+		/// File to read environment variables from
+		/// </summary>
+		[TaskParameter(Optional = true)]
+		public string EnvironmentFile;
+
+		/// <summary>
+		/// Output file for the new environment
+		/// </summary>
+		[TaskParameter]
+		public string OutputFile;
+	}
+
+	/// <summary>
+	/// Spawns AWS CLI and waits for it to complete.
+	/// </summary>
+	[TaskElement("Aws-AssumeRole", typeof(AwsAssumeRoleTaskParameters))]
+	public class AwsAssumeRoleTask : SpawnTaskBase
+	{
+		class AwsSettings
+		{
+			public AwsCredentials Credentials { get; set; }
+		}
+
+		class AwsCredentials
+		{
+			public string AccessKeyId { get; set; }
+			public string SecretAccessKey { get; set; }
+			public string SessionToken { get; set; }
+		}
+
+		/// <summary>
+		/// Parameters for this task
+		/// </summary>
+		AwsAssumeRoleTaskParameters Parameters;
+
+		/// <summary>
+		/// Construct an AWS CLI task
+		/// </summary>
+		/// <param name="InParameters">Parameters for the task</param>
+		public AwsAssumeRoleTask(AwsAssumeRoleTaskParameters InParameters)
+		{
+			Parameters = InParameters;
+		}
+
+		/// <summary>
+		/// Execute the task.
+		/// </summary>
+		/// <param name="Job">Information about the current job</param>
+		/// <param name="BuildProducts">Set of build products produced by this node.</param>
+		/// <param name="TagNameToFileSet">Mapping from tag names to the set of files they include</param>
+		public override void Execute(JobContext Job, HashSet<FileReference> BuildProducts, Dictionary<string, HashSet<FileReference>> TagNameToFileSet)
+		{
+			StringBuilder Arguments = new StringBuilder("sts assume-role");
+			if(Parameters.Arn != null)
+			{
+				Arguments.Append($" --role-arn {Parameters.Arn}");
+			}
+			if (Parameters.Session != null)
+			{
+				Arguments.Append($" --role-session-name {Parameters.Session}");
+			}
+			Arguments.Append($" --duration-seconds {Parameters.Duration}");
+
+			Dictionary<string, string> Environment = SpawnTaskBase.ParseEnvVars(Parameters.Environment, Parameters.EnvironmentFile);
+			IProcessResult Result = SpawnTaskBase.Execute("aws", Arguments.ToString(), EnvVars: Environment, LogOutput: false);
+
+			JsonSerializerOptions Options = new JsonSerializerOptions();
+			Options.PropertyNameCaseInsensitive = true;
+
+			AwsSettings Settings = JsonSerializer.Deserialize<AwsSettings>(Result.Output, Options);
+			if (Settings.Credentials != null)
+			{
+				if (Settings.Credentials.AccessKeyId != null)
+				{
+					Environment["AWS_ACCESS_KEY_ID"] = Settings.Credentials.AccessKeyId;
+				}
+				if (Settings.Credentials.SecretAccessKey != null)
+				{
+					Environment["AWS_SECRET_ACCESS_KEY"] = Settings.Credentials.SecretAccessKey;
+				}
+				if (Settings.Credentials.SessionToken != null)
+				{
+					Environment["AWS_SESSION_TOKEN"] = Settings.Credentials.SessionToken;
+				}
+			}
+
+			FileReference OutputFile = ResolveFile(Parameters.OutputFile);
+			DirectoryReference.CreateDirectory(OutputFile.Directory);
+			FileReference.WriteAllLines(OutputFile, Environment.OrderBy(x => x.Key).Select(x => $"{x.Key}={x.Value}"));
+		}
+
+		/// <summary>
+		/// Output this task out to an XML writer.
+		/// </summary>
+		public override void Write(XmlWriter Writer)
+		{
+			Write(Writer, Parameters);
+		}
+
+		/// <summary>
+		/// Find all the tags which are used as inputs to this task
+		/// </summary>
+		/// <returns>The tag names which are read by this task</returns>
+		public override IEnumerable<string> FindConsumedTagNames()
+		{
+			yield break;
+		}
+
+		/// <summary>
+		/// Find all the tags which are modified by this task
+		/// </summary>
+		/// <returns>The tag names which are modified by this task</returns>
+		public override IEnumerable<string> FindProducedTagNames()
+		{
+			yield break;
+		}
+	}
+}
diff --git a/Engine/Source/Programs/AutomationTool/BuildGraph/Tasks/AwsTask.cs b/Engine/Source/Programs/AutomationTool/BuildGraph/Tasks/AwsTask.cs
--- a/Engine/Source/Programs/AutomationTool/BuildGraph/Tasks/AwsTask.cs
+++ b/Engine/Source/Programs/AutomationTool/BuildGraph/Tasks/AwsTask.cs
@@ -1,5 +1,6 @@
 // Copyright Epic Games, Inc. All Rights Reserved.
 
+using AutomationTool;
 using EpicGames.Core;
 using System;
 using System.Collections.Generic;
@@ -8,126 +9,43 @@ using System.Text;
 using System.Text.Json;
 using System.Xml;
 
-namespace AutomationTool.Tasks
+namespace BuildGraph.Tasks
 {
-	#region Settings
-
-	/// <summary>
-	/// AWS Settings file
-	/// </summary>
-	public class AwsSettings
-	{
-		/// <summary>
-		/// The default region
-		/// </summary>
-		public string Region { get; set; }
-
-		/// <summary>
-		/// The credentials
-		/// </summary>
-		public AwsCredentials Credentials { get; set; }
-
-		/// <summary>
-		/// Read settings from a file
-		/// </summary>
-		/// <param name="File"></param>
-		/// <returns></returns>
-		public static AwsSettings Read(FileReference File)
-		{
-			byte[] Data = FileReference.ReadAllBytes(File);
-
-			JsonSerializerOptions Options = new JsonSerializerOptions();
-			Options.PropertyNameCaseInsensitive = true;
-
-			return JsonSerializer.Deserialize<AwsSettings>(Data, Options);
-		}
-
-		/// <summary>
-		/// Gets environment variables in a dictionary
-		/// </summary>
-		/// <param name="EnvVars"></param>
-		public void GetEnvVars(Dictionary<string, string> EnvVars)
-		{
-			if (Region != null)
-			{
-				EnvVars.Add("AWS_DEFAULT_REGION", Region);
-			}
-			if (Credentials != null)
-			{
-				if (Credentials.AccessKeyId != null)
-				{
-					EnvVars.Add("AWS_ACCESS_KEY_ID", Credentials.AccessKeyId);
-				}
-				if (Credentials.SecretAccessKey != null)
-				{
-					EnvVars.Add("AWS_SECRET_ACCESS_KEY", Credentials.SecretAccessKey);
-				}
-				if (Credentials.SessionToken != null)
-				{
-					EnvVars.Add("AWS_SESSION_TOKEN", Credentials.SessionToken);
-				}
-			}
-		}
-	}
-
-	/// <summary>
-	/// Credentials for AWS
-	/// </summary>
-	public class AwsCredentials
-	{
-		/// <summary>
-		/// Access key for AWS
-		/// </summary>
-		public string AccessKeyId { get; set; }
-
-		/// <summary>
-		/// Secret key for AWS
-		/// </summary>
-		public string SecretAccessKey { get; set; }
-
-		/// <summary>
-		/// Session token for AWS
-		/// </summary>
-		public string SessionToken { get; set; }
-	}
-
-	#endregion
-
 	/// <summary>
-	/// Parameters for an AWS CLI task
+	/// Parameters for a spawn task
 	/// </summary>
 	public class AwsTaskParameters
 	{
 		/// <summary>
-		/// AWS command line arguments
+		/// Arguments for the newly created process.
 		/// </summary>
 		[TaskParameter(Optional = true)]
 		public string Arguments;
 
 		/// <summary>
-		/// Whether or not to show the console output from the AWS CLI command.
+		/// Environment variables
 		/// </summary>
 		[TaskParameter(Optional = true)]
-		public bool Verbose = false;
+		public string Environment;
 
 		/// <summary>
-		/// Path to a file containing the credentials needed to run the AWS command
+		/// File to read environment from
 		/// </summary>
 		[TaskParameter(Optional = true)]
-		public string SettingsFile;
+		public string EnvironmentFile;
 
 		/// <summary>
-		/// Path to receive the ouptut
+		/// Write output to the log
 		/// </summary>
 		[TaskParameter(Optional = true)]
-		public string OutputFile;
+		public bool LogOutput = false;
 	}
 
 	/// <summary>
 	/// Spawns AWS CLI and waits for it to complete.
 	/// </summary>
 	[TaskElement("Aws", typeof(AwsTaskParameters))]
-	public class AwsTask : CustomTask
+	public class AwsTask : SpawnTaskBase
 	{
 		/// <summary>
 		/// Parameters for this task
@@ -151,49 +69,7 @@ namespace AutomationTool.Tasks
 		/// <param name="TagNameToFileSet">Mapping from tag names to the set of files they include</param>
 		public override void Execute(JobContext Job, HashSet<FileReference> BuildProducts, Dictionary<string, HashSet<FileReference>> TagNameToFileSet)
 		{
-			FileReference SettingsFile = (Parameters.SettingsFile != null) ? ResolveFile(Parameters.SettingsFile) : null;
-			FileReference OutputFile = (Parameters.OutputFile != null) ? ResolveFile(Parameters.OutputFile) : null;
-			Run(Parameters.Arguments, SettingsFile: SettingsFile, OutputFile: OutputFile, LogOutput: Parameters.Verbose);
-		}
-
-		/// <summary>
-		/// Runs an AWS command
-		/// </summary>
-		/// <param name="Arguments"></param>
-		/// <param name="SettingsFile"></param>
-		/// <param name="LogOutput"></param>
-		/// <param name="OutputFile"></param>
-		/// <param name="BaseDir"></param>
-		public static string Run(string Arguments, FileReference SettingsFile = null, FileReference OutputFile = null, DirectoryReference BaseDir = null, bool LogOutput = true)
-		{
-			FileReference ToolFile = CommandUtils.FindToolInPath("aws");
-			if (ToolFile == null)
-			{
-				throw new AutomationException("Unable to find path to AWS CLI. Check you have it installed, and it is on your PATH.");
-			}
-
-			Dictionary<string, string> EnvVars = new Dictionary<string, string>();
-			if (SettingsFile != null)
-			{
-				AwsSettings Settings = AwsSettings.Read(SettingsFile);
-				Settings.GetEnvVars(EnvVars);
-			}
-
-			CommandUtils.ERunOptions Options = CommandUtils.ERunOptions.AppMustExist;
-			if (LogOutput)
-			{
-				Options |= CommandUtils.ERunOptions.AllowSpew;
-			}
-			IProcessResult Result = CommandUtils.Run(ToolFile.FullName, Arguments, WorkingDir: BaseDir?.FullName, Env: EnvVars, Options: Options);
-			if (OutputFile != null)
-			{
-				FileReference.WriteAllText(OutputFile, Result.Output);
-			}
-			if (Result.ExitCode != 0)
-			{
-				throw new AutomationException("AWS terminated with an exit code indicating an error ({0})", Result.ExitCode);
-			}
-			return Result.Output;
+			SpawnTaskBase.Execute("aws", Parameters.Arguments, EnvVars: ParseEnvVars(Parameters.Environment, Parameters.EnvironmentFile), LogOutput: Parameters.LogOutput);
 		}
 
 		/// <summary>
diff --git a/Engine/Source/Programs/AutomationTool/BuildGraph/Tasks/DockerBuildTask.cs b/Engine/Source/Programs/AutomationTool/BuildGraph/Tasks/DockerBuildTask.cs
--- a/Engine/Source/Programs/AutomationTool/BuildGraph/Tasks/DockerBuildTask.cs
+++ b/Engine/Source/Programs/AutomationTool/BuildGraph/Tasks/DockerBuildTask.cs
@@ -7,8 +7,9 @@ using System.IO;
 using System.Linq;
 using System.Text;
 using System.Xml;
+using AutomationTool;
 
-namespace AutomationTool.Tasks
+namespace BuildGraph.Tasks
 {
 	/// <summary>
 	/// Parameters for a Docker-Build task
@@ -50,13 +51,25 @@ namespace AutomationTool.Tasks
 		/// </summary>
 		[TaskParameter(Optional = true)]
 		public string OverlayDirs;
+
+		/// <summary>
+		/// Environment variables to set
+		/// </summary>
+		[TaskParameter(Optional = true)]
+		public string Environment;
+
+		/// <summary>
+		/// File to read environment variables from
+		/// </summary>
+		[TaskParameter(Optional = true)]
+		public string EnvironmentFile;
 	}
 
 	/// <summary>
 	/// Spawns Docker and waits for it to complete.
 	/// </summary>
 	[TaskElement("Docker-Build", typeof(DockerBuildTaskParameters))]
-	public class DockerBuildTask : CustomTask
+	public class DockerBuildTask : SpawnTaskBase
 	{
 		/// <summary>
 		/// Parameters for this task
@@ -80,12 +93,6 @@ namespace AutomationTool.Tasks
 		/// <param name="TagNameToFileSet">Mapping from tag names to the set of files they include</param>
 		public override void Execute(JobContext Job, HashSet<FileReference> BuildProducts, Dictionary<string, HashSet<FileReference>> TagNameToFileSet)
 		{
-			FileReference DockerExe = CommandUtils.FindToolInPath("docker");
-			if(DockerExe == null)
-			{
-				throw new AutomationException("Unable to find path to Docker. Check you have it installed, and it is on your PATH.");
-			}
-
 			Log.TraceInformation("Building Docker image");
 			using (LogIndentScope Scope = new LogIndentScope("  "))
 			{
@@ -124,11 +131,7 @@ namespace AutomationTool.Tasks
 					Arguments.Append($" {Parameters.Arguments}");
 				}
 
-				IProcessResult Result = CommandUtils.Run(DockerExe.FullName, Arguments.ToString(), WorkingDir: StagingDir.FullName, Options: CommandUtils.ERunOptions.AllowSpew);
-				if (Result.ExitCode != 0)
-				{
-					throw new AutomationException("Docker terminated with an exit code indicating an error ({0})", Result.ExitCode);
-				}
+				SpawnTaskBase.Execute("docker", Arguments.ToString(), EnvVars: ParseEnvVars(Parameters.Environment, Parameters.EnvironmentFile), WorkingDir: StagingDir.FullName);
 			}
 		}
 
diff --git a/Engine/Source/Programs/AutomationTool/BuildGraph/Tasks/DockerPushTask.cs b/Engine/Source/Programs/AutomationTool/BuildGraph/Tasks/DockerPushTask.cs
--- a/Engine/Source/Programs/AutomationTool/BuildGraph/Tasks/DockerPushTask.cs
+++ b/Engine/Source/Programs/AutomationTool/BuildGraph/Tasks/DockerPushTask.cs
@@ -1,5 +1,6 @@
 // Copyright Epic Games, Inc. All Rights Reserved.
 
+using AutomationTool;
 using EpicGames.Core;
 using System;
 using System.Collections.Generic;
@@ -9,7 +10,7 @@ using System.Text;
 using System.Text.Json;
 using System.Xml;
 
-namespace AutomationTool.Tasks
+namespace BuildGraph.Tasks
 {
 	/// <summary>
 	/// Parameters for a Docker-Build task
@@ -35,17 +36,29 @@ namespace AutomationTool.Tasks
 		public string TargetImage;
 
 		/// <summary>
-		/// Settings for logging in to AWS ECR
+		/// Additional environment variables
 		/// </summary>
 		[TaskParameter(Optional = true)]
-		public string AwsSettings;
+		public string Environment;
+
+		/// <summary>
+		/// File to read environment from
+		/// </summary>
+		[TaskParameter(Optional = true)]
+		public string EnvironmentFile;
+
+		/// <summary>
+		/// Whether to login to AWS ECR
+		/// </summary>
+		[TaskParameter(Optional = true)]
+		public bool AwsEcr;
 	}
 
 	/// <summary>
 	/// Spawns Docker and waits for it to complete.
 	/// </summary>
 	[TaskElement("Docker-Push", typeof(DockerPushTaskParameters))]
-	public class DockerPushTask : CustomTask
+	public class DockerPushTask : SpawnTaskBase
 	{
 		/// <summary>
 		/// Parameters for this task
@@ -69,24 +82,20 @@ namespace AutomationTool.Tasks
 		/// <param name="TagNameToFileSet">Mapping from tag names to the set of files they include</param>
 		public override void Execute(JobContext Job, HashSet<FileReference> BuildProducts, Dictionary<string, HashSet<FileReference>> TagNameToFileSet)
 		{
-			FileReference DockerExe = CommandUtils.FindToolInPath("docker");
-			if(DockerExe == null)
-			{
-				throw new AutomationException("Unable to find path to Docker. Check you have it installed, and it is on your PATH.");
-			}
-
 			Log.TraceInformation("Pushing Docker image");
 			using (LogIndentScope Scope = new LogIndentScope("  "))
 			{
-				if (Parameters.AwsSettings != null)
+				Dictionary<string, string> Environment = ParseEnvVars(Parameters.Environment, Parameters.EnvironmentFile);
+
+				if (Parameters.AwsEcr)
 				{
-					string Password = AwsTask.Run("ecr get-login-password", ResolveFile(Parameters.AwsSettings), LogOutput: false);
-					RunDocker(DockerExe, $"login {Parameters.Repository} --username AWS --password-stdin", Password, CommandUtils.RootDirectory);
+					IProcessResult Result = SpawnTaskBase.Execute("aws", "ecr get-login-password", EnvVars: Environment, LogOutput: false);
+					Execute("docker", $"login {Parameters.Repository} --username AWS --password-stdin", Result.Output);
 				}
 
 				string TargetImage = Parameters.TargetImage ?? Parameters.Image;
-				RunDocker(DockerExe, $"tag {Parameters.Image} {Parameters.Repository}/{TargetImage}", null, CommandUtils.RootDirectory);
-				RunDocker(DockerExe, $"push {Parameters.Repository}/{TargetImage}", null, CommandUtils.RootDirectory);
+				Execute("docker", $"tag {Parameters.Image} {Parameters.Repository}/{TargetImage}", EnvVars: Environment);
+				Execute("docker", $"push {Parameters.Repository}/{TargetImage}", EnvVars: Environment);
 			}
 		}
 
@@ -119,7 +128,7 @@ namespace AutomationTool.Tasks
 		/// <returns>The tag names which are read by this task</returns>
 		public override IEnumerable<string> FindConsumedTagNames()
 		{
-			return FindTagNamesFromFilespec(Parameters.AwsSettings);
+			yield break;
 		}
 
 		/// <summary>
diff --git a/Engine/Source/Programs/AutomationTool/BuildGraph/Tasks/DockerTask.cs b/Engine/Source/Programs/AutomationTool/BuildGraph/Tasks/DockerTask.cs
--- a/Engine/Source/Programs/AutomationTool/BuildGraph/Tasks/DockerTask.cs
+++ b/Engine/Source/Programs/AutomationTool/BuildGraph/Tasks/DockerTask.cs
@@ -1,5 +1,6 @@
 // Copyright Epic Games, Inc. All Rights Reserved.
 
+using AutomationTool;
 using EpicGames.Core;
 using System;
 using System.Collections.Generic;
@@ -7,7 +8,7 @@ using System.IO;
 using System.Text;
 using System.Xml;
 
-namespace AutomationTool.Tasks
+namespace BuildGraph.Tasks
 {
 	/// <summary>
 	/// Parameters for a Docker task
@@ -17,39 +18,33 @@ namespace AutomationTool.Tasks
 		/// <summary>
 		/// Docker command line arguments
 		/// </summary>
-		[TaskParameter(Optional = true)]
+		[TaskParameter]
 		public string Arguments;
 
 		/// <summary>
-		/// Base directory for running the command
-		/// </summary>
-		[TaskParameter(Optional = true)]
-		public string BaseDir;
-
-		/// <summary>
-		/// The minimum exit code, which is treated as an error.
+		/// Environment variables to set
 		/// </summary>
 		[TaskParameter(Optional = true)]
-		public int ErrorLevel = 1;
+		public string Environment;
 
 		/// <summary>
-		/// Whether or not to show the console output from the Docker command.
+		/// File to read environment variables from
 		/// </summary>
 		[TaskParameter(Optional = true)]
-		public bool Verbose = true;
+		public string EnvironmentFile;
 
 		/// <summary>
-		/// Path to a file containing the credentials needed to run the Docker command
+		/// Base directory for running the command
 		/// </summary>
 		[TaskParameter(Optional = true)]
-		public string CredentialsFile;
+		public string WorkingDir;
 	}
 
 	/// <summary>
 	/// Spawns Docker and waits for it to complete.
 	/// </summary>
 	[TaskElement("Docker", typeof(DockerTaskParameters))]
-	public class DockerTask : CustomTask
+	public class DockerTask : SpawnTaskBase
 	{
 		/// <summary>
 		/// Parameters for this task
@@ -73,50 +68,7 @@ namespace AutomationTool.Tasks
 		/// <param name="TagNameToFileSet">Mapping from tag names to the set of files they include</param>
 		public override void Execute(JobContext Job, HashSet<FileReference> BuildProducts, Dictionary<string, HashSet<FileReference>> TagNameToFileSet)
 		{
-			FileReference ToolFile = CommandUtils.FindToolInPath("docker");
-			if(ToolFile == null)
-			{
-				throw new AutomationException("Unable to find path to Docker. Check you have it installed, and it is on your PATH.");
-			}
-
-			string DockerArgs = Parameters.Arguments;
-			string CmdString = $"{ToolFile.FullName} {DockerArgs}";
-			string DockerInput = null;
-
-			if (!string.IsNullOrWhiteSpace(Parameters.CredentialsFile))
-			{
-				if (!File.Exists(Parameters.CredentialsFile))
-				{
-					throw new AutomationException("Credentials file {0} could not be found.", Parameters.CredentialsFile);
-				}
-
-				string[] CredentialTokens = File.ReadAllText(Parameters.CredentialsFile).Split(':');
-				if (CredentialTokens.Length == 2)
-				{
-					DockerArgs += $" --username {CredentialTokens[0]} --password-stdin";
-					DockerInput = CredentialTokens[1];
-				}
-				else if (CredentialTokens.Length == 1)
-				{
-					DockerArgs += $" --password-stdin";
-					DockerInput = CredentialTokens[0];
-				}
-				else
-				{
-					throw new AutomationException("Expected password or username:password in credentials file ({0})", Parameters.CredentialsFile);
-				}
-			}
-
-			CommandUtils.ERunOptions Options = CommandUtils.ERunOptions.AppMustExist;
-			if (Parameters.Verbose)
-			{
-				Options |= CommandUtils.ERunOptions.AllowSpew;
-			}
-			IProcessResult Result = CommandUtils.Run(ToolFile.FullName, DockerArgs, DockerInput, WorkingDir: Parameters.BaseDir, Options: Options);
-			if (Result.ExitCode < 0 || Result.ExitCode >= Parameters.ErrorLevel)
-			{
-				throw new AutomationException("Docker terminated with an exit code indicating an error ({0})", Result.ExitCode);
-			}
+			Execute("docker", Parameters.Arguments, EnvVars: ParseEnvVars(Parameters.Environment, Parameters.EnvironmentFile), WorkingDir: Parameters.WorkingDir);
 		}
 
 		/// <summary>
diff --git a/Engine/Source/Programs/AutomationTool/BuildGraph/Tasks/HelmTask.cs b/Engine/Source/Programs/AutomationTool/BuildGraph/Tasks/HelmTask.cs
--- a/Engine/Source/Programs/AutomationTool/BuildGraph/Tasks/HelmTask.cs
+++ b/Engine/Source/Programs/AutomationTool/BuildGraph/Tasks/HelmTask.cs
@@ -1,5 +1,6 @@
 // Copyright Epic Games, Inc. All Rights Reserved.
 
+using AutomationTool;
 using EpicGames.Core;
 using System;
 using System.Collections.Generic;
@@ -7,7 +8,7 @@ using System.IO;
 using System.Text;
 using System.Xml;
 
-namespace AutomationTool.Tasks
+namespace BuildGraph.Tasks
 {
 	/// <summary>
 	/// Parameters for a Helm task
@@ -44,6 +45,12 @@ namespace AutomationTool.Tasks
 		[TaskParameter(Optional = true)]
 		public List<string> Values;
 
+		/// <summary>
+		/// Environment variables to set
+		/// </summary>
+		[TaskParameter(Optional = true)]
+		public string Environment;
+
 		/// <summary>
 		/// File to parse environment variables from
 		/// </summary>
@@ -60,14 +67,14 @@ namespace AutomationTool.Tasks
 		/// Base directory for running the command
 		/// </summary>
 		[TaskParameter(Optional = true)]
-		public string BaseDir;
+		public string WorkingDir;
 	}
 
 	/// <summary>
 	/// Spawns Helm and waits for it to complete.
 	/// </summary>
 	[TaskElement("Helm", typeof(HelmTaskParameters))]
-	public class HelmTask : CustomTask
+	public class HelmTask : SpawnTaskBase
 	{
 		/// <summary>
 		/// Parameters for this task
@@ -91,31 +98,10 @@ namespace AutomationTool.Tasks
 		/// <param name="TagNameToFileSet">Mapping from tag names to the set of files they include</param>
 		public override void Execute(JobContext Job, HashSet<FileReference> BuildProducts, Dictionary<string, HashSet<FileReference>> TagNameToFileSet)
 		{
-			FileReference HelmExe = CommandUtils.FindToolInPath("helm");
-			if(HelmExe == null)
-			{
-				throw new AutomationException("Unable to find path to Helm. Check you have it installed, and it is on your PATH.");
-			}
-
-			Dictionary<string, string> Environment = new Dictionary<string, string>();
-			if (Parameters.EnvironmentFile != null)
-			{
-			}
-
 			// Switch Kubernetes config
 			if (Parameters.KubeContext != null)
 			{
-				FileReference KubectlExe = CommandUtils.FindToolInPath("kubectl");
-				if (KubectlExe == null)
-				{
-					throw new AutomationException("Unable to find path to Kubectl. Check you have it installed, and it is on your PATH.");
-				}
-
-				IProcessResult KubectlResult = CommandUtils.Run(KubectlExe.FullName, $"config use-context {Parameters.KubeContext}", null, WorkingDir: Parameters.BaseDir);
-				if (KubectlResult.ExitCode != 0)
-				{
-					throw new AutomationException("Kubectl terminated with an exit code indicating an error ({0})", KubectlResult.ExitCode);
-				}
+				Execute("kubectl", $"config use-context {Parameters.KubeContext}", WorkingDir: Parameters.WorkingDir);
 			}
 
 			// Build the argument list
@@ -136,11 +122,7 @@ namespace AutomationTool.Tasks
 				Arguments.Add(Value);
 			}
 
-			IProcessResult Result = CommandUtils.Run(HelmExe.FullName, CommandLineArguments.Join(Arguments), WorkingDir: Parameters.BaseDir);
-			if (Result.ExitCode != 0)
-			{
-				throw new AutomationException("Helm terminated with an exit code indicating an error ({0})", Result.ExitCode);
-			}
+			SpawnTaskBase.Execute("helm", CommandLineArguments.Join(Arguments), WorkingDir: Parameters.WorkingDir, EnvVars: ParseEnvVars(Parameters.Environment, Parameters.EnvironmentFile));
 		}
 
 		/// <summary>
diff --git a/Engine/Source/Programs/AutomationTool/BuildGraph/Tasks/SpawnTask.cs b/Engine/Source/Programs/AutomationTool/BuildGraph/Tasks/SpawnTask.cs
--- a/Engine/Source/Programs/AutomationTool/BuildGraph/Tasks/SpawnTask.cs
+++ b/Engine/Source/Programs/AutomationTool/BuildGraph/Tasks/SpawnTask.cs
@@ -30,6 +30,30 @@ namespace BuildGraph.Tasks
 		[TaskParameter(Optional = true)]
 		public string Arguments;
 
+		/// <summary>
+		/// Working directory for spawning the new task
+		/// </summary>
+		[TaskParameter(Optional = true)]
+		public string WorkingDir;
+
+		/// <summary>
+		/// Environment variables to set
+		/// </summary>
+		[TaskParameter(Optional = true)]
+		public string Environment;
+
+		/// <summary>
+		/// File to read environment from
+		/// </summary>
+		[TaskParameter(Optional = true)]
+		public string EnvironmentFile;
+
+		/// <summary>
+		/// Write output to the log
+		/// </summary>
+		[TaskParameter(Optional = true)]
+		public bool LogOutput;
+
 		/// <summary>
 		/// The minimum exit code, which is treated as an error.
 		/// </summary>
@@ -37,11 +61,92 @@ namespace BuildGraph.Tasks
 		public int ErrorLevel = 1;
 	}
 
+	/// <summary>
+	/// Base class for tasks that run an external tool
+	/// </summary>
+	public abstract class SpawnTaskBase : CustomTask
+	{
+		/// <summary>
+		/// Execute a command
+		/// </summary>
+		protected static IProcessResult Execute(string Exe, string Arguments, string WorkingDir = null, Dictionary<string, string> EnvVars = null, bool LogOutput = true, int ErrorLevel = 1)
+		{
+			if (WorkingDir != null)
+			{
+				WorkingDir = ResolveDirectory(WorkingDir).FullName;
+			}
+
+			CommandUtils.ERunOptions Options = CommandUtils.ERunOptions.Default;
+			if (!LogOutput)
+			{
+				Options &= ~CommandUtils.ERunOptions.AllowSpew;
+			}
+
+			IProcessResult Result = CommandUtils.Run(Exe, Arguments, Env: EnvVars, WorkingDir: WorkingDir, Options: Options);
+			if (Result.ExitCode < 0 || Result.ExitCode >= ErrorLevel)
+			{
+				throw new AutomationException("{0} terminated with an exit code indicating an error ({1})", Path.GetFileName(Exe), Result.ExitCode);
+			}
+
+			return Result;
+		}
+
+		/// <summary>
+		/// Parses environment from a property and file
+		/// </summary>
+		/// <param name="Environment"></param>
+		/// <param name="EnvironmentFile"></param>
+		/// <returns></returns>
+		protected static Dictionary<string, string> ParseEnvVars(string Environment, string EnvironmentFile)
+		{
+			Dictionary<string, string> EnvVars = new Dictionary<string, string>();
+			if (Environment != null)
+			{
+				ParseEnvironment(Environment, ';', EnvVars);
+			}
+			if (EnvironmentFile != null)
+			{
+				ParseEnvironment(FileUtils.ReadAllText(ResolveFile(EnvironmentFile)), '\n', EnvVars);
+			}
+			return EnvVars;
+		}
+
+		/// <summary>
+		/// Parse environment from a string
+		/// </summary>
+		/// <param name="Environment"></param>
+		/// <param name="Separator"></param>
+		/// <param name="EnvVars"></param>
+		static void ParseEnvironment(string Environment, char Separator, Dictionary<string, string> EnvVars)
+		{
+			for (int BaseIdx = 0; BaseIdx < Environment.Length;)
+			{
+				int EqualsIdx = Environment.IndexOf('=', BaseIdx);
+				if (EqualsIdx == -1)
+				{
+					throw new AutomationException("Missing value in environment variable string '{0}'", Environment);
+				}
+
+				int EndIdx = Environment.IndexOf(Separator, EqualsIdx + 1);
+				if (EndIdx == -1)
+				{
+					EndIdx = Environment.Length;
+				}
+
+				string Name = Environment.Substring(BaseIdx, EqualsIdx - BaseIdx).Trim();
+				string Value = Environment.Substring(EqualsIdx + 1, EndIdx - (EqualsIdx + 1)).Trim();
+				EnvVars[Name] = Value;
+
+				BaseIdx = EndIdx + 1;
+			}
+		}
+	}
+
 	/// <summary>
 	/// Spawns an external executable and waits for it to complete.
 	/// </summary>
 	[TaskElement("Spawn", typeof(SpawnTaskParameters))]
-	public class SpawnTask : CustomTask
+	public class SpawnTask : SpawnTaskBase
 	{
 		/// <summary>
 		/// Parameters for this task
@@ -65,11 +170,7 @@ namespace BuildGraph.Tasks
 		/// <param name="TagNameToFileSet">Mapping from tag names to the set of files they include</param>
 		public override void Execute(JobContext Job, HashSet<FileReference> BuildProducts, Dictionary<string, HashSet<FileReference>> TagNameToFileSet)
 		{
-			IProcessResult Result = CommandUtils.Run(Parameters.Exe, Parameters.Arguments);
-			if(Result.ExitCode < 0 || Result.ExitCode >= Parameters.ErrorLevel)
-			{
-				throw new AutomationException("{0} terminated with an exit code indicating an error ({1})", Path.GetFileName(Parameters.Exe), Result.ExitCode);
-			}
+			Execute(Parameters.Exe, Parameters.Arguments, Parameters.WorkingDir, EnvVars: ParseEnvVars(Parameters.Environment, Parameters.EnvironmentFile), LogOutput: Parameters.LogOutput, ErrorLevel: Parameters.ErrorLevel);
 		}
 
 		/// <summary>
diff --git a/Engine/Source/Programs/Shared/EpicGames.Core/FileUtils.cs b/Engine/Source/Programs/Shared/EpicGames.Core/FileUtils.cs
--- a/Engine/Source/Programs/Shared/EpicGames.Core/FileUtils.cs
+++ b/Engine/Source/Programs/Shared/EpicGames.Core/FileUtils.cs
@@ -53,6 +53,31 @@ namespace EpicGames.Core
 		/// </summary>
 		public static IEqualityComparer<ReadOnlyUtf8String> PlatformPathComparerUtf8 = Environment.OSVersion.Platform == PlatformID.Unix ? ReadOnlyUtf8StringComparer.Ordinal : ReadOnlyUtf8StringComparer.OrdinalIgnoreCase;
 
+		/// <summary>
+		/// Read all text for a file
+		/// </summary>
+		/// <param name="File"></param>
+		/// <returns></returns>
+		public static string ReadAllText(FileReference File)
+		{
+			try
+			{
+				return FileReference.ReadAllText(File);
+			}
+			catch (DirectoryNotFoundException Ex)
+			{
+				throw new WrappedFileOrDirectoryException(Ex, $"Unable to read file '{File}'. The directory does not exist.");
+			}
+			catch (FileNotFoundException Ex)
+			{
+				throw new WrappedFileOrDirectoryException(Ex, $"Unable to read file '{File}'. The file does not exist.");
+			}
+			catch (Exception Ex)
+			{
+				throw new WrappedFileOrDirectoryException(Ex, $"Unable to read file '{File}'");
+			}
+		}
+
 		/// <summary>
 		/// Finds the on-disk case of a a file
 		/// </summary>
