BuildGraph: Fix Docker AWS login.

diff --git a/Engine/Source/Programs/AutomationTool/BuildGraph/Tasks/DockerPushTask.cs b/Engine/Source/Programs/AutomationTool/BuildGraph/Tasks/DockerPushTask.cs
--- a/Engine/Source/Programs/AutomationTool/BuildGraph/Tasks/DockerPushTask.cs
+++ b/Engine/Source/Programs/AutomationTool/BuildGraph/Tasks/DockerPushTask.cs
@@ -90,7 +90,7 @@ namespace BuildGraph.Tasks
 				if (Parameters.AwsEcr)
 				{
 					IProcessResult Result = SpawnTaskBase.Execute("aws", "ecr get-login-password", EnvVars: Environment, LogOutput: false);
-					Execute("docker", $"login {Parameters.Repository} --username AWS --password-stdin", Result.Output);
+					Execute("docker", $"login {Parameters.Repository} --username AWS --password-stdin", Input: Result.Output);
 				}
 
 				string TargetImage = Parameters.TargetImage ?? Parameters.Image;
@@ -99,21 +99,6 @@ namespace BuildGraph.Tasks
 			}
 		}
 
-		/// <summary>
-		/// Runs Docker
-		/// </summary>
-		/// <param name="DockerExe"></param>
-		/// <param name="Arguments"></param>
-		/// <param name="StagingDir"></param>
-		public void RunDocker(FileReference DockerExe, string Arguments, string Input, DirectoryReference StagingDir)
-		{
-			IProcessResult Result = CommandUtils.Run(DockerExe.FullName, Arguments, Input: Input, WorkingDir: StagingDir.FullName, Options: CommandUtils.ERunOptions.AllowSpew);
-			if (Result.ExitCode != 0)
-			{
-				throw new AutomationException("Docker terminated with an exit code indicating an error ({0})", Result.ExitCode);
-			}
-		}
-
 		/// <summary>
 		/// Output this task out to an XML writer.
 		/// </summary>
diff --git a/Engine/Source/Programs/AutomationTool/BuildGraph/Tasks/SpawnTask.cs b/Engine/Source/Programs/AutomationTool/BuildGraph/Tasks/SpawnTask.cs
--- a/Engine/Source/Programs/AutomationTool/BuildGraph/Tasks/SpawnTask.cs
+++ b/Engine/Source/Programs/AutomationTool/BuildGraph/Tasks/SpawnTask.cs
@@ -69,7 +69,7 @@ namespace BuildGraph.Tasks
 		/// <summary>
 		/// Execute a command
 		/// </summary>
-		protected static IProcessResult Execute(string Exe, string Arguments, string WorkingDir = null, Dictionary<string, string> EnvVars = null, bool LogOutput = true, int ErrorLevel = 1)
+		protected static IProcessResult Execute(string Exe, string Arguments, string WorkingDir = null, Dictionary<string, string> EnvVars = null, bool LogOutput = true, int ErrorLevel = 1, string Input = null)
 		{
 			if (WorkingDir != null)
 			{
@@ -82,7 +82,7 @@ namespace BuildGraph.Tasks
 				Options &= ~CommandUtils.ERunOptions.AllowSpew;
 			}
 
-			IProcessResult Result = CommandUtils.Run(Exe, Arguments, Env: EnvVars, WorkingDir: WorkingDir, Options: Options);
+			IProcessResult Result = CommandUtils.Run(Exe, Arguments, Env: EnvVars, WorkingDir: WorkingDir, Options: Options, Input: Input);
 			if (Result.ExitCode < 0 || Result.ExitCode >= ErrorLevel)
 			{
 				throw new AutomationException("{0} terminated with an exit code indicating an error ({1})", Path.GetFileName(Exe), Result.ExitCode);
