BuildGraph: Use an absolute path when specifying a helm chart.

diff --git a/Engine/Source/Programs/AutomationTool/BuildGraph/Tasks/HelmTask.cs b/Engine/Source/Programs/AutomationTool/BuildGraph/Tasks/HelmTask.cs
--- a/Engine/Source/Programs/AutomationTool/BuildGraph/Tasks/HelmTask.cs
+++ b/Engine/Source/Programs/AutomationTool/BuildGraph/Tasks/HelmTask.cs
@@ -108,7 +108,7 @@ namespace BuildGraph.Tasks
 			List<string> Arguments = new List<string>();
 			Arguments.Add("upgrade");
 			Arguments.Add(Parameters.Deployment);
-			Arguments.Add(Parameters.Chart);
+			Arguments.Add(new FileReference(Parameters.Chart).FullName);
 			Arguments.Add("--install");
 			Arguments.Add("--reset-values");
 			if(Parameters.Namespace != null)
