BuildGraph: Fix additional arguments not being passed to Helm.

diff --git a/Engine/Source/Programs/AutomationTool/BuildGraph/Tasks/HelmTask.cs b/Engine/Source/Programs/AutomationTool/BuildGraph/Tasks/HelmTask.cs
--- a/Engine/Source/Programs/AutomationTool/BuildGraph/Tasks/HelmTask.cs
+++ b/Engine/Source/Programs/AutomationTool/BuildGraph/Tasks/HelmTask.cs
@@ -45,6 +45,12 @@ namespace BuildGraph.Tasks
 		[TaskParameter(Optional = true)]
 		public List<string> Values;
 
+		/// <summary>
+		/// Values to set for running the chart
+		/// </summary>
+		[TaskParameter(Optional = true)]
+		public string ValuesFile;
+
 		/// <summary>
 		/// Environment variables to set
 		/// </summary>
@@ -121,8 +127,14 @@ namespace BuildGraph.Tasks
 				Arguments.Add("--set");
 				Arguments.Add(Value);
 			}
+			foreach (FileReference ValuesFile in ResolveFilespec(CommandUtils.RootDirectory, Parameters.ValuesFile, TagNameToFileSet))
+			{
+				Arguments.Add("--values");
+				Arguments.Add(ValuesFile.FullName);
+			}
 
-			SpawnTaskBase.Execute("helm", CommandLineArguments.Join(Arguments), WorkingDir: Parameters.WorkingDir, EnvVars: ParseEnvVars(Parameters.Environment, Parameters.EnvironmentFile));
+			string AdditionalArguments = String.IsNullOrEmpty(Parameters.Arguments) ? String.Empty : $" {Parameters.Arguments}";
+			SpawnTaskBase.Execute("helm", CommandLineArguments.Join(Arguments) + AdditionalArguments, WorkingDir: Parameters.WorkingDir, EnvVars: ParseEnvVars(Parameters.Environment, Parameters.EnvironmentFile));
 		}
 
 		/// <summary>
