BuildGraph: Allow not specifying a ValuesFile attribute to Helm tasks.

diff --git a/Engine/Source/Programs/AutomationTool/BuildGraph/Tasks/HelmTask.cs b/Engine/Source/Programs/AutomationTool/BuildGraph/Tasks/HelmTask.cs
--- a/Engine/Source/Programs/AutomationTool/BuildGraph/Tasks/HelmTask.cs
+++ b/Engine/Source/Programs/AutomationTool/BuildGraph/Tasks/HelmTask.cs
@@ -127,10 +127,13 @@ namespace BuildGraph.Tasks
 				Arguments.Add("--set");
 				Arguments.Add(Value);
 			}
-			foreach (FileReference ValuesFile in ResolveFilespec(CommandUtils.RootDirectory, Parameters.ValuesFile, TagNameToFileSet))
+			if (!String.IsNullOrEmpty(Parameters.ValuesFile))
 			{
-				Arguments.Add("--values");
-				Arguments.Add(ValuesFile.FullName);
+				foreach (FileReference ValuesFile in ResolveFilespec(CommandUtils.RootDirectory, Parameters.ValuesFile, TagNameToFileSet))
+				{
+					Arguments.Add("--values");
+					Arguments.Add(ValuesFile.FullName);
+				}
 			}
 
 			string AdditionalArguments = String.IsNullOrEmpty(Parameters.Arguments) ? String.Empty : $" {Parameters.Arguments}";
