Add null check for property value lookup

diff --git a/Engine/Source/Programs/AutomationTool/BuildGraph/Script.cs b/Engine/Source/Programs/AutomationTool/BuildGraph/Script.cs
--- a/Engine/Source/Programs/AutomationTool/BuildGraph/Script.cs
+++ b/Engine/Source/Programs/AutomationTool/BuildGraph/Script.cs
@@ -505,7 +505,7 @@ namespace AutomationTool
 					Value = ScopeValue;
 
 					// Handle $(PropName:-6) where PropName might be "Foo"
-					if (Value.Length > Math.Abs(ValueLength))
+					if (Value != null && Value.Length > Math.Abs(ValueLength))
 					{
 						if (ValueLength > 0)
 						{
