BuildGraph: Fix <Spawn> task supressing output by default.

diff --git a/Engine/Source/Programs/AutomationTool/BuildGraph/Tasks/SpawnTask.cs b/Engine/Source/Programs/AutomationTool/BuildGraph/Tasks/SpawnTask.cs
--- a/Engine/Source/Programs/AutomationTool/BuildGraph/Tasks/SpawnTask.cs
+++ b/Engine/Source/Programs/AutomationTool/BuildGraph/Tasks/SpawnTask.cs
@@ -52,7 +52,7 @@ namespace BuildGraph.Tasks
 		/// Write output to the log
 		/// </summary>
 		[TaskParameter(Optional = true)]
-		public bool LogOutput;
+		public bool LogOutput = true;
 
 		/// <summary>
 		/// The minimum exit code, which is treated as an error.
