Basic Data Retention stuff for Replay Run. Added in BuildFarmTasks to not have this associated with any particular build and instead to just be run on a schedule.

diff --git a/Engine/Source/Programs/AutomationTool/Scripts/CleanAutomationReports.cs b/Engine/Source/Programs/AutomationTool/Scripts/CleanAutomationReports.cs
--- a/Engine/Source/Programs/AutomationTool/Scripts/CleanAutomationReports.cs
+++ b/Engine/Source/Programs/AutomationTool/Scripts/CleanAutomationReports.cs
@@ -76,6 +76,12 @@ namespace AutomationTool
 				int DirsFound = 0;
 				foreach(DirectoryInfo SubDirectory in Directory.EnumerateDirectories())
 				{
+					// If we find a DO_NOT_CLEANUP.txt file before we reach the target depth we want to ignore that entire portion of the directory structure.
+					if (File.Exists(Path.Combine(SubDirectory.FullName, "DO_NOT_CLEANUP.txt")))
+					{
+						CommandUtils.LogInformation("Found DO_NOT_CLEANUP.txt file in {0} before target depth, excluding it from cleanup process.", SubDirectory);
+						continue;
+					}
 					DirsFound += CleanDirectories(SubDirectory, RetainTime, DirectoriesToDelete, TargetDepth, CurrentDepth + 1);
 				}
 				return DirsFound;
@@ -88,6 +94,12 @@ namespace AutomationTool
 				{
 					try
 					{
+						// If we find a DO_NOT_CLEANUP.txt file recursively anywhere in a build folder, leave that build alone.
+						if (BuildDirectory.EnumerateFiles("*", SearchOption.AllDirectories).Any(x => x.Name.ToLower() == "do_not_cleanup.txt"))
+						{
+							CommandUtils.LogInformation("Found DO_NOT_CLEANUP.txt file in {0}, skipping it!", BuildDirectory);
+							continue;
+						}
 						if(!BuildDirectory.EnumerateFiles("*", SearchOption.AllDirectories).Any(x => x.LastWriteTimeUtc > RetainTime))
 						{
 							DirectoriesToDelete.Add(BuildDirectory);
