IterationProfile: Track Map_Load on Telemetry

Disable loading a map at startup to a) remove the map load from EditorInit, and b) make the automation plugin's map load the first, making it precisely tracked through Telemetry

Tested with local UAT runs.

diff --git a/Engine/Build/IterationProfile.xml b/Engine/Build/IterationProfile.xml
--- a/Engine/Build/IterationProfile.xml
+++ b/Engine/Build/IterationProfile.xml
@@ -69,6 +69,7 @@
 	<Property Name="ProfileArgs" Value="$(ProfileArgs) -EnablePlugins=$(ExtraPluginsCommaDelims)" If="'$(ExtraPluginsCommaDelims)' != ''"/>
 	<Property Name="ProfileArgs" Value="$(ProfileArgs) -TreatLogWarningsAsTestErrors=false -TreatLogErrorsAsTestErrors=false"/>
 	<Property Name="ProfileArgs" Value="$(ProfileArgs) -GameInstanceLostTimerSeconds=3600"/>
+	<Property Name="ProfileArgs" Value="$(ProfileArgs) -NoLoadLevelAtStartup"/>
 	<Property Name="ProfileArgs" Value="$(ProfileArgs) -NoShaderDistrib"/>
 
 	<!-- Gauntlet arguments for summarize commandlet runs -->
diff --git a/Engine/Source/Programs/AutomationTool/Gauntlet/Unreal/Editor/EditorTests.cs b/Engine/Source/Programs/AutomationTool/Gauntlet/Unreal/Editor/EditorTests.cs
--- a/Engine/Source/Programs/AutomationTool/Gauntlet/Unreal/Editor/EditorTests.cs
+++ b/Engine/Source/Programs/AutomationTool/Gauntlet/Unreal/Editor/EditorTests.cs
@@ -55,6 +55,12 @@ namespace UnrealEditor
 		[AutoParam]
 		public string GameInstanceLostTimerSeconds = string.Empty;
 
+		/// <summary>
+		/// Disable loading a level at startup (for profiling the map load)
+		/// </summary>
+		[AutoParam]
+		public bool NoLoadLevelAtStartup = false;
+
 		/// <summary>
 		/// Disable distribution of shader builds (but use worker processes still)
 		/// </summary>
@@ -97,6 +103,11 @@ namespace UnrealEditor
 				AppConfig.CommandLineParams.Add(string.Format("ini:Engine:[/Script/AutomationController.AutomationControllerSettings]:GameInstanceLostTimerSeconds={0}", GameInstanceLostTimerSeconds));
 			}
 
+			if (NoLoadLevelAtStartup)
+			{
+				AppConfig.CommandLineParams.Add("ini:EditorPerProjectUserSettings:[/Script/UnrealEd.EditorLoadingSavingSettings]:LoadLevelAtStartup=None");
+			}
+
 			if (NoShaderDistrib)
 			{
 				AppConfig.CommandLineParams.Add("-noxgeshadercompile");
