Report execution time of AutomationTool on shutdown, this is nice to have when doing local runs

This is a separate print statement to side step any sort of log analysis that may be looking for the subsequent print of the exit code

diff --git a/Engine/Source/Programs/AutomationTool/Program.cs b/Engine/Source/Programs/AutomationTool/Program.cs
--- a/Engine/Source/Programs/AutomationTool/Program.cs
+++ b/Engine/Source/Programs/AutomationTool/Program.cs
@@ -25,6 +25,8 @@ namespace AutomationTool
 		// Do not add [STAThread] here. It will cause deadlocks in platform automation code.
 		public static int Main(string[] Arguments)
 		{
+			Stopwatch Timer = Stopwatch.StartNew();
+
             // Ensure UTF8Output flag is respected, since we are initializing logging early in the program.
             if (CommandUtils.ParseParam(Arguments, "-Utf8output"))
             {
@@ -125,6 +127,7 @@ namespace AutomationTool
                 NoThrow(() => { if (ShouldKillProcesses && Utils.IsRunningOnWindows) ProcessManager.KillAll(); }, "Kill All Processes");
 
 				// Write the exit code
+                Log.TraceInformation("AutomationTool executed for {0}", Timer.Elapsed.ToString("h'h 'm'm 's's'"));
                 Log.TraceInformation("AutomationTool exiting with ExitCode={0} ({1})", (int)ReturnCode, ReturnCode);
 
                 // Can't use NoThrow here because the code logs exceptions. We're shutting down logging!
