Gauntlet - adding python as a log channel to listen to

#rnx

diff --git a/Engine/Source/Programs/AutomationTool/Gauntlet/Unreal/Utils/Gauntlet.UnrealLogParser.cs b/Engine/Source/Programs/AutomationTool/Gauntlet/Unreal/Utils/Gauntlet.UnrealLogParser.cs
--- a/Engine/Source/Programs/AutomationTool/Gauntlet/Unreal/Utils/Gauntlet.UnrealLogParser.cs
+++ b/Engine/Source/Programs/AutomationTool/Gauntlet/Unreal/Utils/Gauntlet.UnrealLogParser.cs
@@ -353,7 +353,7 @@ namespace Gauntlet
 		/// <returns></returns>
 		public IEnumerable<string> GetEditorBusyChannels ()
 		{
-			return GetLogChannels(new string[] { "Automation", "FunctionalTest", "Material", "DerivedDataCache", "ShaderCompilers", "Texture", "SkeletalMesh", "StaticMesh" }, false);
+			return GetLogChannels(new string[] { "Automation", "FunctionalTest", "Material", "DerivedDataCache", "ShaderCompilers", "Texture", "SkeletalMesh", "StaticMesh", "Python" }, false);
 		}
 
 		/// <summary>
