Gauntlet - Add support for Unit and Baseline to Telemetry data

#rnx

diff --git a/Engine/Source/Programs/AutomationTool/Gauntlet/Framework/Base/Gauntlet.TestReport.cs b/Engine/Source/Programs/AutomationTool/Gauntlet/Framework/Base/Gauntlet.TestReport.cs
--- a/Engine/Source/Programs/AutomationTool/Gauntlet/Framework/Base/Gauntlet.TestReport.cs
+++ b/Engine/Source/Programs/AutomationTool/Gauntlet/Framework/Base/Gauntlet.TestReport.cs
@@ -51,12 +51,16 @@ namespace Gauntlet
 		public string DataPoint { get; private set; }
 		public double Measurement { get; private set; }
 		public string Context { get; private set; }
-		public TelemetryData(string InTestName, string InDataPoint, double InMeasurement, string InContext = "")
+		public string Unit { get; private set; }
+		public double Baseline { get; private set; }
+		public TelemetryData(string InTestName, string InDataPoint, double InMeasurement, string InContext = "", string InUnit = "", double InBaseline = 0)
 		{
 			TestName = InTestName;
 			DataPoint = InDataPoint;
 			Measurement = InMeasurement;
 			Context = InContext;
+			Unit = InUnit;
+			Baseline = InBaseline;
 		}
 	}
 
@@ -72,8 +76,10 @@ namespace Gauntlet
 		/// <param name="DataPoint"></param>
 		/// <param name="Measurement"></param>
 		/// <param name="Context"></param>
+		/// <param name="Unit"></param>
+		/// <param name="Baseline"></param>
 		/// <returns></returns>
-		void AddTelemetry(string TestName, string DataPoint, double Measurement, string Context = "");
+		void AddTelemetry(string TestName, string DataPoint, double Measurement, string Context = "", string Unit = "", double Baseline = 0);
 
 		/// <summary>
 		/// Return the telemetry data accumulated
@@ -146,15 +152,17 @@ namespace Gauntlet
 		/// <param name="DataPoint"></param>
 		/// <param name="Measurement"></param>
 		/// <param name="Context"></param>
+		/// <param name="Unit"></param>
+		/// <param name="Baseline"></param>
 		/// <returns></returns>
-		public virtual void AddTelemetry(string TestName, string DataPoint, double Measurement, string Context = "")
+		public virtual void AddTelemetry(string TestName, string DataPoint, double Measurement, string Context = "", string Unit = "", double Baseline = 0)
 		{
 			if (TelemetryDataList is null)
 			{
 				TelemetryDataList = new List<TelemetryData>();
 			}
 
-			TelemetryDataList.Add(new TelemetryData(TestName, DataPoint, Measurement, Context));
+			TelemetryDataList.Add(new TelemetryData(TestName, DataPoint, Measurement, Context, Unit, Baseline));
 		}
 
 		/// <summary>
diff --git a/Engine/Source/Programs/AutomationTool/Gauntlet/Framework/Utils/Gauntlet.HordeReport.cs b/Engine/Source/Programs/AutomationTool/Gauntlet/Framework/Utils/Gauntlet.HordeReport.cs
--- a/Engine/Source/Programs/AutomationTool/Gauntlet/Framework/Utils/Gauntlet.HordeReport.cs
+++ b/Engine/Source/Programs/AutomationTool/Gauntlet/Framework/Utils/Gauntlet.HordeReport.cs
@@ -40,6 +40,22 @@ namespace Gauntlet
 				return !string.IsNullOrEmpty(Environment.GetEnvironmentVariable("UE_TESTDATA_DIR"));
 			}
 		}
+		/// <summary>
+		/// Default Job Horde link if running under Horde agent, otherwise return an empty string
+		/// </summary>
+		public static string DefaultHordeJobLink
+		{
+			get
+			{
+				string HordeJobId = Environment.GetEnvironmentVariable("UE_HORDE_JOBID");
+				if (!string.IsNullOrEmpty(HordeJobId))
+				{
+					return string.Format("https://horde.devtools.epicgames.com/job/{0}", HordeJobId);
+				}
+
+				return string.Empty;
+			}
+		}
 
 		public abstract class BaseHordeReport : BaseTestReport
 		{
diff --git a/Engine/Source/Programs/AutomationTool/Gauntlet/Unreal/Config/Gauntlet.EpicConfig.cs b/Engine/Source/Programs/AutomationTool/Gauntlet/Unreal/Config/Gauntlet.EpicConfig.cs
--- a/Engine/Source/Programs/AutomationTool/Gauntlet/Unreal/Config/Gauntlet.EpicConfig.cs
+++ b/Engine/Source/Programs/AutomationTool/Gauntlet/Unreal/Config/Gauntlet.EpicConfig.cs
@@ -48,7 +48,10 @@ namespace EpicConfig
 				"testName",
 				"context",
 				"dataPoint",
-				"recordedValue"
+				"recordedValue",
+				"unit",
+				"baseline",
+				"jobLink"
 			};
 		}
 		/// <summary>
@@ -68,7 +71,10 @@ namespace EpicConfig
 				Data.TestName,
 				Data.Context,
 				Data.DataPoint,
-				Data.Measurement.ToString()
+				Data.Measurement.ToString(),
+				Data.Unit,
+				Data.Baseline.ToString(),
+				Context.GetProperty("JobLink").ToString()
 			};
 		}
 	}
diff --git a/Engine/Source/Programs/AutomationTool/Gauntlet/Unreal/Utils/Gauntlet.UnrealTelemetry.cs b/Engine/Source/Programs/AutomationTool/Gauntlet/Unreal/Utils/Gauntlet.UnrealTelemetry.cs
--- a/Engine/Source/Programs/AutomationTool/Gauntlet/Unreal/Utils/Gauntlet.UnrealTelemetry.cs
+++ b/Engine/Source/Programs/AutomationTool/Gauntlet/Unreal/Utils/Gauntlet.UnrealTelemetry.cs
@@ -60,6 +60,8 @@ namespace Gauntlet
 			string DataPointKey = CSVColumns.DataPoint.ToString();
 			string MeasurementKey = CSVColumns.Measurement.ToString();
 			string ContextKey = CSVColumns.Context.ToString();
+			string UnitKey = CSVColumns.Unit.ToString();
+			string BaselineKey = CSVColumns.Baseline.ToString();
 			if (Mapping != null)
 			{
 				if (Mapping.ContainsKey(TestNameKey))
@@ -70,6 +72,10 @@ namespace Gauntlet
 					Mapping.TryGetValue(MeasurementKey, out MeasurementKey);
 				if (Mapping.ContainsKey(ContextKey))
 					Mapping.TryGetValue(ContextKey, out ContextKey);
+				if (Mapping.ContainsKey(UnitKey))
+					Mapping.TryGetValue(UnitKey, out UnitKey);
+				if (Mapping.ContainsKey(BaselineKey))
+					Mapping.TryGetValue(BaselineKey, out BaselineKey);
 			}
 			// Populate telemetry report
 			foreach (Dictionary<string, string> Row in CSVData)
@@ -82,6 +88,12 @@ namespace Gauntlet
 				Row.TryGetValue(MeasurementKey, out Measurement);
 				string Context;
 				Row.TryGetValue(ContextKey, out Context);
+				string Unit;
+				Row.TryGetValue(UnitKey, out Unit);
+				string Baseline;
+				Row.TryGetValue(BaselineKey, out Baseline);
+				double BaselineFloat = 0;
+				double.TryParse(Baseline, out BaselineFloat);
 				if (string.IsNullOrEmpty(TestName) || string.IsNullOrEmpty(DataPoint) || string.IsNullOrEmpty(Measurement))
 				{
 					Log.Warning("Telemetry - Missing data in CSV file '{0}':\n TestName='{1}', DataPoint='{2}', Measurement='{3}'", File, TestName, DataPoint, Measurement);
@@ -89,7 +101,7 @@ namespace Gauntlet
 				}
 				try
 				{
-					TelemetryReport.AddTelemetry(TestName, DataPoint, double.Parse(Measurement), Context);
+					TelemetryReport.AddTelemetry(TestName, DataPoint, double.Parse(Measurement), Context, Unit, BaselineFloat);
 				}
 				catch (FormatException)
 				{
@@ -103,7 +115,9 @@ namespace Gauntlet
 			TestName,
 			DataPoint,
 			Measurement,
-			Context
+			Context,
+			Unit,
+			Baseline
 		}
 
 	}
@@ -144,6 +158,7 @@ namespace Gauntlet
 		[Help("Branch", "Target Branch name. Default: Unknown.")]
 		[Help("Changelist", "Target Changelist number. Default: 0.")]
 		[Help("Configuration", "Target Configuration name. Default: Development.")]
+		[Help("JobLink", "Http Link to Build Job. Default: %UE_HORDE_JOBID%")]
 
 		public override ExitCode Execute()
 		{
@@ -158,6 +173,7 @@ namespace Gauntlet
 			string ChangelistString = ParseParamValue("Changelist=", "0");
 			string RoleString = ParseParamValue("Role=", "Editor");
 			string ConfigurationString = ParseParamValue("Configuration=", "Development");
+			string JobLink = ParseParamValue("JobLink=", "");
 
 			if (string.IsNullOrEmpty(CSVDirectory) && string.IsNullOrEmpty(CSVFile))
 			{
@@ -230,6 +246,12 @@ namespace Gauntlet
 			Context.SetProperty("Branch", BranchString);
 			Context.SetProperty("Changelist", ChangelistString);
 
+			if (string.IsNullOrEmpty(JobLink))
+			{
+				JobLink = HordeReport.DefaultHordeJobLink;
+			}
+			Context.SetProperty("JobLink", JobLink);
+
 			// Create a report
 			HordeReport.SimpleTestReport Report = new HordeReport.SimpleTestReport();
 
