Change gauntlet login special logic from looking for base XboxOne to XboxOneGDK

diff --git a/Engine/Source/Programs/AutomationTool/Gauntlet/Unreal/Game/EpicGame.TestConfig.cs b/Engine/Source/Programs/AutomationTool/Gauntlet/Unreal/Game/EpicGame.TestConfig.cs
--- a/Engine/Source/Programs/AutomationTool/Gauntlet/Unreal/Game/EpicGame.TestConfig.cs
+++ b/Engine/Source/Programs/AutomationTool/Gauntlet/Unreal/Game/EpicGame.TestConfig.cs
@@ -228,7 +228,7 @@ namespace EpicGame
 			if (ConfigRole.RoleType.IsClient() || RequiresLogin)
 			{
 				// select an account
-				if (NoMCP == false && ConfigRole.Platform != UnrealTargetPlatform.PS4 && PreAssignAccount == true)
+				if (NoMCP == false && ConfigRole.Platform != UnrealTargetPlatform.PS4 && ConfigRole.Platform != UnrealTargetPlatform.XboxOneGDK && PreAssignAccount == true)
 				{
 					Account UserAccount = AccountPool.Instance.ReserveAccount();
 					UserAccount.ApplyToConfig(AppConfig);
