Add XSX and PS5 to epic test node login logic.

diff --git a/Engine/Source/Programs/AutomationTool/Gauntlet/Unreal/Game/EpicGame.TestConfig.cs b/Engine/Source/Programs/AutomationTool/Gauntlet/Unreal/Game/EpicGame.TestConfig.cs
--- a/Engine/Source/Programs/AutomationTool/Gauntlet/Unreal/Game/EpicGame.TestConfig.cs
+++ b/Engine/Source/Programs/AutomationTool/Gauntlet/Unreal/Game/EpicGame.TestConfig.cs
@@ -228,7 +228,9 @@ namespace EpicGame
 			if (ConfigRole.RoleType.IsClient() || RequiresLogin)
 			{
 				// select an account
-				if (NoMCP == false && ConfigRole.Platform != UnrealTargetPlatform.PS4 && ConfigRole.Platform != UnrealTargetPlatform.XboxOneGDK && PreAssignAccount == true)
+				if (NoMCP == false && ConfigRole.Platform != UnrealTargetPlatform.PS4 && ConfigRole.Platform != UnrealTargetPlatform.XboxOneGDK
+					&& ConfigRole.Platform != UnrealTargetPlatform.PS5 && ConfigRole.Platform != UnrealTargetPlatform.XSX 
+					&& PreAssignAccount == true)
 				{
 					Account UserAccount = AccountPool.Instance.ReserveAccount();
 					UserAccount.ApplyToConfig(AppConfig);
