Remove hard references to platforms in EpicGame TestConfig login checks.

diff --git a/Engine/Source/Programs/AutomationTool/Gauntlet/Unreal/Game/EpicGame.TestConfig.cs b/Engine/Source/Programs/AutomationTool/Gauntlet/Unreal/Game/EpicGame.TestConfig.cs
--- a/Engine/Source/Programs/AutomationTool/Gauntlet/Unreal/Game/EpicGame.TestConfig.cs
+++ b/Engine/Source/Programs/AutomationTool/Gauntlet/Unreal/Game/EpicGame.TestConfig.cs
@@ -227,10 +227,9 @@ namespace EpicGame
 
 			if (ConfigRole.RoleType.IsClient() || RequiresLogin)
 			{
+				string[] NoAssignedAccountPlatforms = {"PS4","PS5","XboxOneGDK","XSX" };
 				// select an account
-				if (NoMCP == false && ConfigRole.Platform != UnrealTargetPlatform.PS4 && ConfigRole.Platform != UnrealTargetPlatform.XboxOneGDK
-					&& ConfigRole.Platform != UnrealTargetPlatform.PS5 && ConfigRole.Platform != UnrealTargetPlatform.XSX 
-					&& PreAssignAccount == true)
+				if (NoMCP == false && !NoAssignedAccountPlatforms.Contains(ConfigRole.Platform.ToString()) && PreAssignAccount == true)
 				{
 					Account UserAccount = AccountPool.Instance.ReserveAccount();
 					UserAccount.ApplyToConfig(AppConfig);
