BuildGraph: Clear the Docker staging directory before copying files into it.

diff --git a/Engine/Source/Programs/AutomationTool/BuildGraph/Tasks/DockerBuildTask.cs b/Engine/Source/Programs/AutomationTool/BuildGraph/Tasks/DockerBuildTask.cs
--- a/Engine/Source/Programs/AutomationTool/BuildGraph/Tasks/DockerBuildTask.cs
+++ b/Engine/Source/Programs/AutomationTool/BuildGraph/Tasks/DockerBuildTask.cs
@@ -100,6 +100,7 @@ namespace BuildGraph.Tasks
 				List<FileReference> SourceFiles = ResolveFilespec(BaseDir, Parameters.Files, TagNameToFileSet).ToList();
 
 				DirectoryReference StagingDir = DirectoryReference.Combine(CommandUtils.EngineDirectory, "Intermediate", "Docker");
+				FileUtils.ForceDeleteDirectoryContents(StagingDir);
 
 				List<FileReference> TargetFiles = SourceFiles.ConvertAll(x => FileReference.Combine(StagingDir, x.MakeRelativeTo(BaseDir)));
 				CommandUtils.ThreadedCopyFiles(SourceFiles, BaseDir, StagingDir);
