UAT build fix.

#rnx

diff --git a/Engine/Source/Programs/AutomationTool/Scripts/RunProjectCommand.Automation.cs b/Engine/Source/Programs/AutomationTool/Scripts/RunProjectCommand.Automation.cs
--- a/Engine/Source/Programs/AutomationTool/Scripts/RunProjectCommand.Automation.cs
+++ b/Engine/Source/Programs/AutomationTool/Scripts/RunProjectCommand.Automation.cs
@@ -1056,7 +1056,7 @@ public partial class Project : CommandUtils
 	private static bool StartUnrealTrace()
 	{
 		// [TEMPORARY] - UnrealTrace server is currently only available on Windows
-		if (!Utils.IsRunningOnWindows)
+		if (!RuntimePlatform.IsWindows)
 		{
 			return true;
 		}
