AutomationTool:

Replace dependency on UnrealBuildTool with BuildUtilities.

diff --git a/Engine/Source/Programs/AutomationTool/AutomationTool.csproj b/Engine/Source/Programs/AutomationTool/AutomationTool.csproj
--- a/Engine/Source/Programs/AutomationTool/AutomationTool.csproj
+++ b/Engine/Source/Programs/AutomationTool/AutomationTool.csproj
@@ -1,4 +1,4 @@
-﻿<Project>
+<Project>
   <PropertyGroup>
     <TargetFramework>netcoreapp3.1</TargetFramework>
     <Configuration Condition=" '$(Configuration)' == '' ">Development</Configuration>
@@ -78,7 +78,7 @@
   
   <ItemGroup>
     <ProjectReference Include="..\Shared\EpicGames.Core\EpicGames.Core.csproj" />
-    <ProjectReference Include="..\UnrealBuildTool\UnrealBuildTool.csproj" />
+    <ProjectReference Include="..\DotNETCommon\BuildUtilities\BuildUtilities.csproj" />
     <ProjectReference Include="AutomationUtils\AutomationUtils.Automation.csproj"><Private>true</Private></ProjectReference>
   </ItemGroup>
 
