AutomationTool:

Revert change to projectreference

diff --git a/Engine/Source/Programs/AutomationTool/AutomationTool.csproj b/Engine/Source/Programs/AutomationTool/AutomationTool.csproj
--- a/Engine/Source/Programs/AutomationTool/AutomationTool.csproj
+++ b/Engine/Source/Programs/AutomationTool/AutomationTool.csproj
@@ -71,6 +71,7 @@
     </Compile>
     <Compile Include="Program.cs" />
     <Compile Include="Properties\AssemblyInfo.cs" />
+    <Compile Include="Start.cs" />
     <None Include="app.manifest">
       <SubType>Designer</SubType>
     </None>
@@ -78,7 +79,7 @@
   
   <ItemGroup>
     <ProjectReference Include="..\Shared\EpicGames.Core\EpicGames.Core.csproj" />
-    <ProjectReference Include="..\DotNETCommon\BuildUtilities\BuildUtilities.csproj" />
+    <ProjectReference Include="..\UnrealBuildTool\UnrealBuildTool.csproj" />
     <ProjectReference Include="AutomationUtils\AutomationUtils.Automation.csproj"><Private>true</Private></ProjectReference>
   </ItemGroup>
 
