AutomationTool:

Remove spurious reference

diff --git a/Engine/Source/Programs/AutomationTool/AutomationTool.csproj b/Engine/Source/Programs/AutomationTool/AutomationTool.csproj
--- a/Engine/Source/Programs/AutomationTool/AutomationTool.csproj
+++ b/Engine/Source/Programs/AutomationTool/AutomationTool.csproj
@@ -71,7 +71,6 @@
     </Compile>
     <Compile Include="Program.cs" />
     <Compile Include="Properties\AssemblyInfo.cs" />
-    <Compile Include="Start.cs" />
     <None Include="app.manifest">
       <SubType>Designer</SubType>
     </None>
