Renamed "UnrealTrace.exe" to "UnrealTraceServer.exe"

#rnx

diff --git a/Engine/Source/Programs/AutomationTool/Scripts/RunProjectCommand.Automation.cs b/Engine/Source/Programs/AutomationTool/Scripts/RunProjectCommand.Automation.cs
--- a/Engine/Source/Programs/AutomationTool/Scripts/RunProjectCommand.Automation.cs
+++ b/Engine/Source/Programs/AutomationTool/Scripts/RunProjectCommand.Automation.cs
@@ -1065,7 +1065,7 @@ public partial class Project : CommandUtils
 		LogInformation("UnrealTrace: Starting server");
 
 		// Locate the UnrealTrace binary
-		var UnrealTracePath = HostPlatform.Current.GetUE4ExePath("UnrealTrace.exe");
+		var UnrealTracePath = HostPlatform.Current.GetUE4ExePath("UnrealTraceServer.exe");
 		if (!File.Exists(UnrealTracePath))
 		{
 			LogWarning("UnrealTrace: Unable to locate binary at " + UnrealTracePath);
diff --git a/Engine/Source/Runtime/Core/Private/ProfilingDebugging/TraceAuxiliary.cpp b/Engine/Source/Runtime/Core/Private/ProfilingDebugging/TraceAuxiliary.cpp
--- a/Engine/Source/Runtime/Core/Private/ProfilingDebugging/TraceAuxiliary.cpp
+++ b/Engine/Source/Runtime/Core/Private/ProfilingDebugging/TraceAuxiliary.cpp
@@ -492,7 +492,7 @@ static void StartUnrealTrace()
 {
 	FString BinPath = "\"";
 	BinPath += FPaths::EngineDir();
-	BinPath += "/Binaries/Win64/UnrealTrace.exe\"";
+	BinPath += "/Binaries/Win64/UnrealTraceServer.exe\"";
 
 	uint32 Flags = 0x0100'0000; // CREATE_BREAKAWAY_FROM_JOB
 	STARTUPINFOW StartupInfo = { sizeof(STARTUPINFOW) };
