Early out of CleanTempStorage if the temp storage path doesn't work.

#rnx

diff --git a/Engine/Source/Programs/AutomationTool/BuildGraph/TempStorage.cs b/Engine/Source/Programs/AutomationTool/BuildGraph/TempStorage.cs
--- a/Engine/Source/Programs/AutomationTool/BuildGraph/TempStorage.cs
+++ b/Engine/Source/Programs/AutomationTool/BuildGraph/TempStorage.cs
@@ -1320,6 +1320,12 @@ namespace AutomationTool
 				throw new AutomationException("Missing -TempStorageDir parameter");
 			}
 
+			if (!Directory.Exists(TempStorageDir))
+			{
+				CommandUtils.LogInformation("Temp Storage folder '{0}' does not exist, no work to do.", TempStorageDir);
+				return;
+			}
+
 			string Days = ParseParamValue("Days", null);
 			if (Days == null)
 			{
