UAT: Delete a zip file before trying to overwrite it.

diff --git a/Engine/Source/Programs/AutomationTool/AutomationUtils/CommandUtils.cs b/Engine/Source/Programs/AutomationTool/AutomationUtils/CommandUtils.cs
--- a/Engine/Source/Programs/AutomationTool/AutomationUtils/CommandUtils.cs
+++ b/Engine/Source/Programs/AutomationTool/AutomationUtils/CommandUtils.cs
@@ -2210,6 +2210,11 @@ namespace AutomationTool
 				DirectoryReference.CreateDirectory(OutputFile.Directory);
 			}
 
+			if (FileReference.Exists(OutputFile))
+			{
+				FileUtils.ForceDeleteFile(OutputFile);
+			}
+
 			using (ZipArchive ZipArchive = ZipFile.Open(OutputFile.FullName, ZipArchiveMode.Create))
 			{
 				foreach (FileReference File in Files)
