BuildGraph: Allow submitting files to a non-stream path.

diff --git a/Engine/Source/Programs/AutomationTool/BuildGraph/Tasks/SubmitTask.cs b/Engine/Source/Programs/AutomationTool/BuildGraph/Tasks/SubmitTask.cs
--- a/Engine/Source/Programs/AutomationTool/BuildGraph/Tasks/SubmitTask.cs
+++ b/Engine/Source/Programs/AutomationTool/BuildGraph/Tasks/SubmitTask.cs
@@ -4,6 +4,7 @@ using System;
 using System.Collections.Generic;
 using System.Linq;
 using System.Text;
+using System.Text.RegularExpressions;
 using System.Threading.Tasks;
 using System.Xml;
 using EpicGames.Core;
@@ -46,6 +47,12 @@ namespace AutomationTool.Tasks
 		[TaskParameter(Optional = true)]
 		public string Stream;
 
+		/// <summary>
+		/// Branch for the workspace (legacy P4 depot path). May not be used in conjunction with Stream.
+		/// </summary>
+		[TaskParameter(Optional = true)]
+		public string Branch;
+
 		/// <summary>
 		/// Root directory for the stream. If not specified, defaults to the current root directory.
 		/// </summary>
@@ -118,11 +125,18 @@ namespace AutomationTool.Tasks
 					P4ClientInfo Client = new P4ClientInfo();
 					Client.Owner = CommandUtils.P4Env.User;
 					Client.Host = Environment.MachineName;
-					Client.Stream = Parameters.Stream ?? CommandUtils.P4Env.Branch;
 					Client.RootPath = Parameters.RootDir.FullName ?? CommandUtils.RootDirectory.FullName;
-					Client.Name = Parameters.Workspace + "_" + ContentHash.MD5((CommandUtils.P4Env.ServerAndPort ?? "").ToUpperInvariant()).ToString();
+					Client.Name = $"{Parameters.Workspace}_{Regex.Replace(Client.Host, "[^a-zA-Z0-9]", "-")}_{ContentHash.MD5((CommandUtils.P4Env.ServerAndPort ?? "").ToUpperInvariant())}";
 					Client.Options = P4ClientOption.NoAllWrite | P4ClientOption.Clobber | P4ClientOption.NoCompress | P4ClientOption.Unlocked | P4ClientOption.NoModTime | P4ClientOption.RmDir;
 					Client.LineEnd = P4LineEnd.Local;
+					if (Parameters.Branch != null)
+					{
+						Client.View.Add(new KeyValuePair<string, string>($"{Parameters.Branch}/...", $"/..."));
+					}
+					else
+					{
+						Client.Stream = Parameters.Stream ?? CommandUtils.P4Env.Branch;
+					}
 					CommandUtils.P4.CreateClient(Client, AllowSpew: Parameters.P4Verbose);
 
 					// Create a new connection for it
