Gauntlet - Fix report issue with Automated tests

#rnx

diff --git a/Engine/Source/Programs/AutomationTool/Gauntlet/Unreal/Automation/UE.Automation.cs b/Engine/Source/Programs/AutomationTool/Gauntlet/Unreal/Automation/UE.Automation.cs
--- a/Engine/Source/Programs/AutomationTool/Gauntlet/Unreal/Automation/UE.Automation.cs
+++ b/Engine/Source/Programs/AutomationTool/Gauntlet/Unreal/Automation/UE.Automation.cs
@@ -434,9 +434,7 @@ namespace UE
 		/// Optional function that is called on test completion and gives an opportunity to create a report
 		/// </summary>
 		/// <param name="Result"></param>
-		/// <param name="Context"></param>
-		/// <param name="Build"></param>
-		public override ITestReport CreateReport(TestResult Result, UnrealTestContext Context, UnrealBuildSource Build, IEnumerable<UnrealRoleResult> InResults, string ArtifactPath)
+		public override ITestReport CreateReport(TestResult Result)
 		{
 			ITestReport Report = null;
 			if (GetConfiguration() is AutomationTestConfig)
@@ -448,7 +446,7 @@ namespace UE
 				{
 					if (Config.SimpleHordeReport)
 					{
-						Report = base.CreateReport(Result, Context, Build, InResults, ArtifactPath);
+						Report = base.CreateReport(Result);
 					}
 					else
 					{
diff --git a/Engine/Source/Programs/AutomationTool/Gauntlet/Unreal/Base/Gauntlet.UnrealTestNode.cs b/Engine/Source/Programs/AutomationTool/Gauntlet/Unreal/Base/Gauntlet.UnrealTestNode.cs
--- a/Engine/Source/Programs/AutomationTool/Gauntlet/Unreal/Base/Gauntlet.UnrealTestNode.cs
+++ b/Engine/Source/Programs/AutomationTool/Gauntlet/Unreal/Base/Gauntlet.UnrealTestNode.cs
@@ -926,15 +926,13 @@ namespace Gauntlet
 			string Message = string.Empty;
 			ITestReport Report = null;
 
-			var DeprecatedSignature = new[] { typeof(TestResult), typeof(UnrealTestContext), typeof(UnrealBuildSource), typeof(IEnumerable<UnrealRoleArtifacts>), typeof(string) };
-			var NewSignature = new[] { typeof(TestResult) };
-
 			try
 			{
-				// Check if the deprecated signature is overriden, call it anyway if that the case, and trigger a warning.
-				if (!Utils.InterfaceHelpers.HasOverriddenMethod(this.GetType(), "CreateReport", NewSignature) && Utils.InterfaceHelpers.HasOverriddenMethod(this.GetType(), "CreateReport", DeprecatedSignature))
+				// Check if the deprecated signature is overriden, call it anyway if that the case.
+				var OldSignature = new[] { typeof(TestResult), typeof(UnrealTestContext), typeof(UnrealBuildSource), typeof(IEnumerable<UnrealRoleResult>), typeof(string) };
+				var Simplifiedignature = new[] { typeof(TestResult) };
+				if (!Utils.InterfaceHelpers.HasOverriddenMethod(this.GetType(), "CreateReport", Simplifiedignature) && Utils.InterfaceHelpers.HasOverriddenMethod(this.GetType(), "CreateReport", OldSignature))
 				{
-					//Log.Warning("This CreateReport signature is deprecated, please use the overloaded version.");
 					Report = CreateReport(GetTestResult(), Context, Context.BuildInfo, RoleResults, ArtifactPath);
 				}
 				else
@@ -964,6 +962,7 @@ namespace Gauntlet
 			try
 			{
 				// Check if the deprecated signature is overriden, call it anyway if that the case, and trigger a warning.
+				var DeprecatedSignature = new[] { typeof(TestResult), typeof(UnrealTestContext), typeof(UnrealBuildSource), typeof(IEnumerable<UnrealRoleArtifacts>), typeof(string) };
 				if (Utils.InterfaceHelpers.HasOverriddenMethod(this.GetType(), "SubmitToDashboard", DeprecatedSignature))
 				{
 					Log.Warning("This SubmitToDashboard signature is deprecated, please use the overloaded version.");
