CrashReport: platform extensions can include crash report client.

#rnx

diff --git a/Engine/Source/Programs/AutomationTool/AutomationUtils/Platform.cs b/Engine/Source/Programs/AutomationTool/AutomationUtils/Platform.cs
--- a/Engine/Source/Programs/AutomationTool/AutomationUtils/Platform.cs
+++ b/Engine/Source/Programs/AutomationTool/AutomationUtils/Platform.cs
@@ -728,6 +728,14 @@ namespace AutomationTool
 
 		}
 
+		/// <summary>
+		/// Allows a platform to use the crash reporter from a different (built-in) platform
+		/// </summary>
+		public virtual UnrealTargetPlatform? CrashReportPlatform
+		{
+			get { return null; }
+		}
+
 		/// <summary>
 		/// General purpose command to run generic string commands inside the platform interfeace
 		/// </summary>
diff --git a/Engine/Source/Programs/AutomationTool/Scripts/BuildProjectCommand.Automation.cs b/Engine/Source/Programs/AutomationTool/Scripts/BuildProjectCommand.Automation.cs
--- a/Engine/Source/Programs/AutomationTool/Scripts/BuildProjectCommand.Automation.cs
+++ b/Engine/Source/Programs/AutomationTool/Scripts/BuildProjectCommand.Automation.cs
@@ -151,7 +151,8 @@ public partial class Project : CommandUtils
 			{
                 foreach (var ClientPlatformType in UniquePlatformTypes)
 				{
-                    CrashReportPlatforms.Add(ClientPlatformType);
+					UnrealTargetPlatform CrashReportPlatform = Platform.GetPlatform(ClientPlatformType).CrashReportPlatform ?? ClientPlatformType;
+                    CrashReportPlatforms.Add(CrashReportPlatform);
 					Agenda.AddTargets(Params.ClientCookedTargets.ToArray(), ClientPlatformType, BuildConfig, Params.CodeBasedUprojectPath, InAddArgs: " -remoteini=\"" + Params.RawProjectPath.Directory.FullName + "\"" + AdditionalArgs);
 				}
 			}
@@ -164,7 +165,8 @@ public partial class Project : CommandUtils
 			{
 				foreach (var ServerPlatformType in UniquePlatformTypes)
 				{
-                    CrashReportPlatforms.Add(ServerPlatformType);
+					UnrealTargetPlatform CrashReportPlatform = Platform.GetPlatform(ServerPlatformType).CrashReportPlatform ?? ServerPlatformType;
+                    CrashReportPlatforms.Add(CrashReportPlatform);
 					Agenda.AddTargets(Params.ServerCookedTargets.ToArray(), ServerPlatformType, BuildConfig, Params.CodeBasedUprojectPath, InAddArgs: " -remoteini=\"" + Params.RawProjectPath.Directory.FullName + "\"" + AdditionalArgs);
 				}
 			}
diff --git a/Engine/Source/Programs/AutomationTool/Scripts/CopyBuildToStagingDirectory.Automation.cs b/Engine/Source/Programs/AutomationTool/Scripts/CopyBuildToStagingDirectory.Automation.cs
--- a/Engine/Source/Programs/AutomationTool/Scripts/CopyBuildToStagingDirectory.Automation.cs
+++ b/Engine/Source/Programs/AutomationTool/Scripts/CopyBuildToStagingDirectory.Automation.cs
@@ -1044,7 +1044,8 @@ public partial class Project : CommandUtils
 
 				// CrashReportClient is a standalone slate app that does not look in the generated pak file, so it needs the Content/Slate and Shaders/StandaloneRenderer folders Non-UFS
 				// @todo Make CrashReportClient more portable so we don't have to do this
-				if (SC.bStageCrashReporter && PlatformSupportsCrashReporter(SC.StageTargetPlatform.PlatformType) && (Params.IterateSharedBuildUsePrecompiledExe == false))
+				UnrealTargetPlatform CrashReportPlatform = ThisPlatform.CrashReportPlatform ?? SC.StageTargetPlatform.PlatformType;
+				if (SC.bStageCrashReporter && PlatformSupportsCrashReporter(CrashReportPlatform) && (Params.IterateSharedBuildUsePrecompiledExe == false))
 				{
 					SC.StageCrashReporterFiles(StagedFileType.UFS, DirectoryReference.Combine(SC.EngineRoot, "Content", "Slate"), StageFilesSearch.AllDirectories);
 					SC.StageCrashReporterFiles(StagedFileType.UFS, DirectoryReference.Combine(SC.EngineRoot, "Shaders", "StandaloneRenderer"), StageFilesSearch.AllDirectories);
@@ -1055,14 +1056,14 @@ public partial class Project : CommandUtils
 					if (string.IsNullOrEmpty(Architecture))
 					{
 						Architecture = "";
-						if (PlatformExports.IsPlatformAvailable(SC.StageTargetPlatform.PlatformType))
+						if (PlatformExports.IsPlatformAvailable(CrashReportPlatform))
 						{
-							Architecture = PlatformExports.GetDefaultArchitecture(SC.StageTargetPlatform.PlatformType, Params.RawProjectPath);
+							Architecture = PlatformExports.GetDefaultArchitecture(CrashReportPlatform, Params.RawProjectPath);
 						}
 					}
 
 					// Get the target receipt path for CrashReportClient
-					FileReference ReceiptFileName = TargetReceipt.GetDefaultPath(SC.EngineRoot, "CrashReportClient", SC.StageTargetPlatform.PlatformType, UnrealTargetConfiguration.Shipping, Architecture);
+					FileReference ReceiptFileName = TargetReceipt.GetDefaultPath(SC.EngineRoot, "CrashReportClient", CrashReportPlatform, UnrealTargetConfiguration.Shipping, Architecture);
 					if (!FileReference.Exists(ReceiptFileName))
 					{
 						throw new AutomationException(ExitCode.Error_MissingExecutable, "Stage Failed. Missing receipt '{0}'. Check that this target has been built.", ReceiptFileName);
