UAT: Fix incorrect help text on FixPerforceCase command.

diff --git a/Engine/Source/Programs/AutomationTool/Scripts/FixPerforceCase.cs b/Engine/Source/Programs/AutomationTool/Scripts/FixPerforceCase.cs
--- a/Engine/Source/Programs/AutomationTool/Scripts/FixPerforceCase.cs
+++ b/Engine/Source/Programs/AutomationTool/Scripts/FixPerforceCase.cs
@@ -11,8 +11,7 @@ namespace AutomationTool
 	[RequireP4]
 	[DoesNotNeedP4CL]
 	[Help("Fixes the case of files on a case-insensitive Perforce server by removing and re-adding them.")]
-	[Help("Source", "Pattern for source files to match. Should be a full depot path. May end with a wildcard.")]
-	[Help("Target", "Pattern for target files. Should be identical to source, except for case.")]
+	[Help("Files", "Pattern for files to match. Should be a full depot path with the correct case. May end with a wildcard.")]
 	class FixPerforceCase : BuildCommand
 	{
 		const string BoilerplateText = "\n\n#rb none\n#rnx";
