ExitCode:

Move enum from AutomationUtils.Automation to BuildUtilities

#trivial

diff --git a/Engine/Source/Programs/AutomationTool/Android/AndroidPlatform.Automation.cs b/Engine/Source/Programs/AutomationTool/Android/AndroidPlatform.Automation.cs
--- a/Engine/Source/Programs/AutomationTool/Android/AndroidPlatform.Automation.cs
+++ b/Engine/Source/Programs/AutomationTool/Android/AndroidPlatform.Automation.cs
@@ -11,6 +11,7 @@ using AutomationTool;
 using UnrealBuildTool;
 using Ionic.Zip;
 using EpicGames.Core;
+using UnrealBuildBase;
 
 public class AndroidPlatform : Platform
 {
diff --git a/Engine/Source/Programs/AutomationTool/AutomationUtils/AutomationException.cs b/Engine/Source/Programs/AutomationTool/AutomationUtils/AutomationException.cs
--- a/Engine/Source/Programs/AutomationTool/AutomationUtils/AutomationException.cs
+++ b/Engine/Source/Programs/AutomationTool/AutomationUtils/AutomationException.cs
@@ -3,68 +3,10 @@ using System;
 using System.Collections.Generic;
 using System.Linq;
 using System.Text;
+using UnrealBuildBase;
 
 namespace AutomationTool
 {
-	// NOTE: this needs to be kept in sync with EditorAnalytics.h and iPhonePackager.cs
-	public enum ExitCode
-    {
-        Error_UATNotFound = -1,
-        Success = 0,
-        Error_Unknown = 1,
-        Error_Arguments = 2,
-        Error_UnknownCommand = 3,
-        Error_SDKNotFound = 10,
-        Error_ProvisionNotFound = 11,
-        Error_CertificateNotFound = 12,
-        Error_ProvisionAndCertificateNotFound = 13,
-        Error_InfoPListNotFound = 14,
-        Error_KeyNotFoundInPList = 15,
-        Error_ProvisionExpired = 16,
-        Error_CertificateExpired = 17,
-        Error_CertificateProvisionMismatch = 18,
-        Error_CodeUnsupported = 19,
-        Error_PluginsUnsupported = 20,
-        Error_UnknownCookFailure = 25,
-        Error_UnknownDeployFailure = 26,
-        Error_UnknownBuildFailure = 27,
-        Error_UnknownPackageFailure = 28,
-        Error_UnknownLaunchFailure = 29,
-        Error_StageMissingFile = 30,
-        Error_FailedToCreateIPA = 31,
-        Error_FailedToCodeSign = 32,
-        Error_DeviceBackupFailed = 33,
-        Error_AppUninstallFailed = 34,
-        Error_AppInstallFailed = 35,
-        Error_AppNotFound = 36,
-        Error_StubNotSignedCorrectly = 37,
-        Error_IPAMissingInfoPList = 38,
-        Error_DeleteFile = 39,
-        Error_DeleteDirectory = 40,
-        Error_CreateDirectory = 41,
-        Error_CopyFile = 42,
-        Error_OnlyOneObbFileSupported = 50,
-        Error_FailureGettingPackageInfo = 51,
-        Error_OnlyOneTargetConfigurationSupported = 52,
-        Error_ObbNotFound = 53,
-        Error_AndroidBuildToolsPathNotFound = 54,
-        Error_NoApkSuitableForArchitecture = 55,
-        Error_FilesInstallFailed = 56,
-        Error_RemoteCertificatesNotFound = 57,
-        Error_LauncherFailed = 100,
-        Error_UATLaunchFailure = 101,
-        Error_FailedToDeleteStagingDirectory = 102,
-        Error_MissingExecutable = 103,
-        Error_DeviceNotSetupForDevelopment = 150,
-        Error_DeviceOSNewerThanSDK = 151,
-		Error_TestFailure = 152,
-		Error_SymbolizedSONotFound = 153,
-		Error_LicenseNotAccepted = 154,
-		Error_AndroidOBBError = 155,
-		Error_SDKInstallFailed = 200,
-		Error_DeviceUpdateFailed = 201,
-	};
-
 	/// <summary>
 	/// How the exception should be output
 	/// </summary>
diff --git a/Engine/Source/Programs/AutomationTool/AutomationUtils/BuildCommand.cs b/Engine/Source/Programs/AutomationTool/AutomationUtils/BuildCommand.cs
--- a/Engine/Source/Programs/AutomationTool/AutomationUtils/BuildCommand.cs
+++ b/Engine/Source/Programs/AutomationTool/AutomationUtils/BuildCommand.cs
@@ -4,6 +4,7 @@ using System.Collections.Generic;
 using System.Linq;
 using System.Text;
 using EpicGames.Core;
+using UnrealBuildBase;
 
 namespace AutomationTool
 {
diff --git a/Engine/Source/Programs/AutomationTool/AutomationUtils/ProcessUtils.cs b/Engine/Source/Programs/AutomationTool/AutomationUtils/ProcessUtils.cs
--- a/Engine/Source/Programs/AutomationTool/AutomationUtils/ProcessUtils.cs
+++ b/Engine/Source/Programs/AutomationTool/AutomationUtils/ProcessUtils.cs
@@ -11,6 +11,7 @@ using System.Runtime.InteropServices;
 using System.Runtime.CompilerServices;
 using EpicGames.Core;
 using UnrealBuildTool;
+using UnrealBuildBase;
 
 namespace AutomationTool
 {
diff --git a/Engine/Source/Programs/AutomationTool/AutomationUtils/UBTUtils.cs b/Engine/Source/Programs/AutomationTool/AutomationUtils/UBTUtils.cs
--- a/Engine/Source/Programs/AutomationTool/AutomationUtils/UBTUtils.cs
+++ b/Engine/Source/Programs/AutomationTool/AutomationUtils/UBTUtils.cs
@@ -4,6 +4,7 @@ using System.Collections.Generic;
 using System.Linq;
 using System.Text;
 using EpicGames.Core;
+using UnrealBuildBase;
 using UnrealBuildTool;
 
 namespace AutomationTool
diff --git a/Engine/Source/Programs/AutomationTool/AutomationUtils/Utils.cs b/Engine/Source/Programs/AutomationTool/AutomationUtils/Utils.cs
--- a/Engine/Source/Programs/AutomationTool/AutomationUtils/Utils.cs
+++ b/Engine/Source/Programs/AutomationTool/AutomationUtils/Utils.cs
@@ -10,6 +10,7 @@ using System.Reflection;
 using UnrealBuildTool;
 using System.Text.RegularExpressions;
 using EpicGames.Core;
+using UnrealBuildBase;
 
 namespace AutomationTool
 {
diff --git a/Engine/Source/Programs/AutomationTool/Gauntlet/Editor/RunEditorTests.cs b/Engine/Source/Programs/AutomationTool/Gauntlet/Editor/RunEditorTests.cs
--- a/Engine/Source/Programs/AutomationTool/Gauntlet/Editor/RunEditorTests.cs
+++ b/Engine/Source/Programs/AutomationTool/Gauntlet/Editor/RunEditorTests.cs
@@ -6,6 +6,7 @@ using UnrealBuildTool;
 using Gauntlet;
 using AutomationTool;
 using System.Reflection;
+using UnrealBuildBase;
 
 namespace EditorTest
 {
diff --git a/Engine/Source/Programs/AutomationTool/Gauntlet/SelfTest/TestGauntlet.cs b/Engine/Source/Programs/AutomationTool/Gauntlet/SelfTest/TestGauntlet.cs
--- a/Engine/Source/Programs/AutomationTool/Gauntlet/SelfTest/TestGauntlet.cs
+++ b/Engine/Source/Programs/AutomationTool/Gauntlet/SelfTest/TestGauntlet.cs
@@ -7,6 +7,7 @@ using System.Linq;
 using System.Reflection;
 using System.Text;
 using System.Threading.Tasks;
+using UnrealBuildBase;
 
 namespace Gauntlet
 {
diff --git a/Engine/Source/Programs/AutomationTool/Gauntlet/Unreal/Engine/UnrealTests.cs b/Engine/Source/Programs/AutomationTool/Gauntlet/Unreal/Engine/UnrealTests.cs
--- a/Engine/Source/Programs/AutomationTool/Gauntlet/Unreal/Engine/UnrealTests.cs
+++ b/Engine/Source/Programs/AutomationTool/Gauntlet/Unreal/Engine/UnrealTests.cs
@@ -6,7 +6,7 @@ using System.Collections.Generic;
 using System.Linq;
 using AutomationTool;
 using UnrealBuildTool;
-
+using UnrealBuildBase;
 
 namespace Gauntlet
 {
diff --git a/Engine/Source/Programs/AutomationTool/Gauntlet/Unreal/RunUnreal.cs b/Engine/Source/Programs/AutomationTool/Gauntlet/Unreal/RunUnreal.cs
--- a/Engine/Source/Programs/AutomationTool/Gauntlet/Unreal/RunUnreal.cs
+++ b/Engine/Source/Programs/AutomationTool/Gauntlet/Unreal/RunUnreal.cs
@@ -14,6 +14,7 @@ using Newtonsoft.Json;
 using System.Reflection;
 using System.Text.RegularExpressions;
 using EpicGames.Core;
+using UnrealBuildBase;
 
 namespace Gauntlet
 {
diff --git a/Engine/Source/Programs/AutomationTool/Gauntlet/Unreal/Tasks/Gauntlet.Task.CleanDevices.cs b/Engine/Source/Programs/AutomationTool/Gauntlet/Unreal/Tasks/Gauntlet.Task.CleanDevices.cs
--- a/Engine/Source/Programs/AutomationTool/Gauntlet/Unreal/Tasks/Gauntlet.Task.CleanDevices.cs
+++ b/Engine/Source/Programs/AutomationTool/Gauntlet/Unreal/Tasks/Gauntlet.Task.CleanDevices.cs
@@ -6,6 +6,7 @@ using System.IO;
 using System.Linq;
 using System.Text;
 using System.Threading.Tasks;
+using UnrealBuildBase;
 using UnrealBuildTool;
 
 namespace Gauntlet
@@ -30,7 +31,7 @@ namespace Gauntlet
 		public int MaxDays;
 
 
-		public override AutomationTool.ExitCode Execute()
+		public override ExitCode Execute()
 		{
 			AutoParam.ApplyParamsAndDefaults(this, Environment.GetCommandLineArgs());
 
@@ -67,7 +68,7 @@ namespace Gauntlet
 			DevicePool.Instance.Dispose();
 
 
-			return AutomationTool.ExitCode.Success;
+			return ExitCode.Success;
 		}
 
 		protected void CleanDevice(ITargetDevice Device)
diff --git a/Engine/Source/Programs/AutomationTool/Gauntlet/Unreal/Utils/Gauntlet.UnrealTelemetry.cs b/Engine/Source/Programs/AutomationTool/Gauntlet/Unreal/Utils/Gauntlet.UnrealTelemetry.cs
--- a/Engine/Source/Programs/AutomationTool/Gauntlet/Unreal/Utils/Gauntlet.UnrealTelemetry.cs
+++ b/Engine/Source/Programs/AutomationTool/Gauntlet/Unreal/Utils/Gauntlet.UnrealTelemetry.cs
@@ -5,7 +5,7 @@ using System.IO;
 using System.Collections.Generic;
 using AutomationTool;
 using UnrealBuildTool;
-
+using UnrealBuildBase;
 
 namespace Gauntlet
 {
diff --git a/Engine/Source/Programs/AutomationTool/HoloLens/HoloLensPlatform.Automation.cs b/Engine/Source/Programs/AutomationTool/HoloLens/HoloLensPlatform.Automation.cs
--- a/Engine/Source/Programs/AutomationTool/HoloLens/HoloLensPlatform.Automation.cs
+++ b/Engine/Source/Programs/AutomationTool/HoloLens/HoloLensPlatform.Automation.cs
@@ -23,6 +23,7 @@ using Windows.ApplicationModel.Core;
 using System.Threading;
 using System.Security.Cryptography.X509Certificates;
 using Microsoft.Tools.WindowsDevicePortal;
+using UnrealBuildBase;
 
 namespace HoloLens.Automation
 {
diff --git a/Engine/Source/Programs/AutomationTool/IOS/ExportIPAFromArchive.cs b/Engine/Source/Programs/AutomationTool/IOS/ExportIPAFromArchive.cs
--- a/Engine/Source/Programs/AutomationTool/IOS/ExportIPAFromArchive.cs
+++ b/Engine/Source/Programs/AutomationTool/IOS/ExportIPAFromArchive.cs
@@ -8,6 +8,7 @@ using System.Text;
 using System.Threading.Tasks;
 using System.Xml;
 using AutomationTool;
+using UnrealBuildBase;
 
 [Help(@"Creates an IPA from an xarchive file")]
 [Help("method=<method>", @"Purpose of the IPA. (Development, Adhoc, Store)")]
diff --git a/Engine/Source/Programs/AutomationTool/IOS/MakeIPA.cs b/Engine/Source/Programs/AutomationTool/IOS/MakeIPA.cs
--- a/Engine/Source/Programs/AutomationTool/IOS/MakeIPA.cs
+++ b/Engine/Source/Programs/AutomationTool/IOS/MakeIPA.cs
@@ -8,6 +8,7 @@ using System.Text;
 using System.Threading.Tasks;
 using System.Xml;
 using AutomationTool;
+using UnrealBuildBase;
 
 [Help(@"Creates an IPA from an xarchive file")]
 [Help("method=<method>", @"Purpose of the IPA. (Development, Adhoc, Store)")]
diff --git a/Engine/Source/Programs/AutomationTool/IOS/WriteIniValueToPlist.cs b/Engine/Source/Programs/AutomationTool/IOS/WriteIniValueToPlist.cs
--- a/Engine/Source/Programs/AutomationTool/IOS/WriteIniValueToPlist.cs
+++ b/Engine/Source/Programs/AutomationTool/IOS/WriteIniValueToPlist.cs
@@ -9,6 +9,7 @@ using System.Text.RegularExpressions;
 using System.Threading.Tasks;
 using System.Xml;
 using AutomationTool;
+using UnrealBuildBase;
 
 [Help(@"Pulls a value from an ini file and inserts it into a plist.")]
 [Help(@"Note currently only looks at values irrespective of sections!")]
diff --git a/Engine/Source/Programs/AutomationTool/Scripts/BenchmarkBuild/BenchmarkBuildTests.cs b/Engine/Source/Programs/AutomationTool/Scripts/BenchmarkBuild/BenchmarkBuildTests.cs
--- a/Engine/Source/Programs/AutomationTool/Scripts/BenchmarkBuild/BenchmarkBuildTests.cs
+++ b/Engine/Source/Programs/AutomationTool/Scripts/BenchmarkBuild/BenchmarkBuildTests.cs
@@ -8,6 +8,7 @@ using System.Text;
 using System.Threading.Tasks;
 using EpicGames.Core;
 using UnrealBuildTool;
+using UnrealBuildBase;
 
 namespace AutomationTool.Benchmark
 {
diff --git a/Engine/Source/Programs/AutomationTool/Scripts/BenchmarkBuild/BenchmarkCookTask.cs b/Engine/Source/Programs/AutomationTool/Scripts/BenchmarkBuild/BenchmarkCookTask.cs
--- a/Engine/Source/Programs/AutomationTool/Scripts/BenchmarkBuild/BenchmarkCookTask.cs
+++ b/Engine/Source/Programs/AutomationTool/Scripts/BenchmarkBuild/BenchmarkCookTask.cs
@@ -9,6 +9,7 @@ using System.Text;
 using System.Threading.Tasks;
 using EpicGames.Core;
 using UnrealBuildTool;
+using UnrealBuildBase;
 
 namespace AutomationTool.Benchmark
 {
diff --git a/Engine/Source/Programs/AutomationTool/Scripts/BenchmarkBuild/BenchmarkRunEditorTask.cs b/Engine/Source/Programs/AutomationTool/Scripts/BenchmarkBuild/BenchmarkRunEditorTask.cs
--- a/Engine/Source/Programs/AutomationTool/Scripts/BenchmarkBuild/BenchmarkRunEditorTask.cs
+++ b/Engine/Source/Programs/AutomationTool/Scripts/BenchmarkBuild/BenchmarkRunEditorTask.cs
@@ -7,6 +7,7 @@ using System.Text;
 using System.Threading;
 using System.Threading.Tasks;
 using EpicGames.Core;
+using UnrealBuildBase;
 using UnrealBuildTool;
 
 namespace AutomationTool.Benchmark
diff --git a/Engine/Source/Programs/AutomationTool/Scripts/CheckRestrictedFolders.cs b/Engine/Source/Programs/AutomationTool/Scripts/CheckRestrictedFolders.cs
--- a/Engine/Source/Programs/AutomationTool/Scripts/CheckRestrictedFolders.cs
+++ b/Engine/Source/Programs/AutomationTool/Scripts/CheckRestrictedFolders.cs
@@ -7,6 +7,7 @@ using System.Linq;
 using System.Text;
 using System.Threading.Tasks;
 using EpicGames.Core;
+using UnrealBuildBase;
 using UnrealBuildTool;
 
 namespace AutomationTool
diff --git a/Engine/Source/Programs/AutomationTool/Scripts/OpenEditor.Automation.cs b/Engine/Source/Programs/AutomationTool/Scripts/OpenEditor.Automation.cs
--- a/Engine/Source/Programs/AutomationTool/Scripts/OpenEditor.Automation.cs
+++ b/Engine/Source/Programs/AutomationTool/Scripts/OpenEditor.Automation.cs
@@ -6,6 +6,7 @@ using System.IO;
 using System.Linq;
 using System.Text.RegularExpressions;
 using EpicGames.Core;
+using UnrealBuildBase;
 using UnrealBuildTool;
 
 namespace AutomationTool
diff --git a/Engine/Source/Programs/AutomationTool/Scripts/RebuildHLODCommand.Automation.cs b/Engine/Source/Programs/AutomationTool/Scripts/RebuildHLODCommand.Automation.cs
--- a/Engine/Source/Programs/AutomationTool/Scripts/RebuildHLODCommand.Automation.cs
+++ b/Engine/Source/Programs/AutomationTool/Scripts/RebuildHLODCommand.Automation.cs
@@ -9,6 +9,7 @@ using System.Net.Mail;
 using AutomationTool;
 using UnrealBuildTool;
 using EpicGames.Core;
+using UnrealBuildBase;
 
 /// <summary>
 /// Helper command used for rebuilding a projects Hierarchical LODs.
diff --git a/Engine/Source/Programs/AutomationTool/Scripts/RebuildLightMapsCommand.Automation.cs b/Engine/Source/Programs/AutomationTool/Scripts/RebuildLightMapsCommand.Automation.cs
--- a/Engine/Source/Programs/AutomationTool/Scripts/RebuildLightMapsCommand.Automation.cs
+++ b/Engine/Source/Programs/AutomationTool/Scripts/RebuildLightMapsCommand.Automation.cs
@@ -9,6 +9,7 @@ using System.Net.Mail;
 using AutomationTool;
 using UnrealBuildTool;
 using EpicGames.Core;
+using UnrealBuildBase;
 
 /// <summary>
 /// Helper command used for rebuilding a projects light maps.
diff --git a/Engine/Source/Programs/AutomationTool/Scripts/ResavePackagesCommand.Automation.cs b/Engine/Source/Programs/AutomationTool/Scripts/ResavePackagesCommand.Automation.cs
--- a/Engine/Source/Programs/AutomationTool/Scripts/ResavePackagesCommand.Automation.cs
+++ b/Engine/Source/Programs/AutomationTool/Scripts/ResavePackagesCommand.Automation.cs
@@ -9,6 +9,7 @@ using System.Net.Mail;
 using AutomationTool;
 using UnrealBuildTool;
 using EpicGames.Core;
+using UnrealBuildBase;
 
 /// <summary>
 /// Helper command used for resaving a projects packages.
diff --git a/Engine/Source/Programs/AutomationTool/Scripts/SyncProject.Automation.cs b/Engine/Source/Programs/AutomationTool/Scripts/SyncProject.Automation.cs
--- a/Engine/Source/Programs/AutomationTool/Scripts/SyncProject.Automation.cs
+++ b/Engine/Source/Programs/AutomationTool/Scripts/SyncProject.Automation.cs
@@ -9,6 +9,7 @@ using UnrealBuildTool;
 using EpicGames.Core;
 using System.Text.RegularExpressions;
 using System.Threading;
+using UnrealBuildBase;
 
 abstract class SyncProjectBase : BuildCommand
 {
diff --git a/Engine/Source/Programs/AutomationTool/Scripts/ZipUtils.Automation.cs b/Engine/Source/Programs/AutomationTool/Scripts/ZipUtils.Automation.cs
--- a/Engine/Source/Programs/AutomationTool/Scripts/ZipUtils.Automation.cs
+++ b/Engine/Source/Programs/AutomationTool/Scripts/ZipUtils.Automation.cs
@@ -8,6 +8,7 @@ using System.Text;
 using AutomationTool;
 using Ionic.Zip;
 using EpicGames.Core;
+using UnrealBuildBase;
 
 [Help("ZipUtils is used to zip/unzip (i.e:RunUAT.bat ZipUtils -archive=D:/Content.zip -add=D:/UE/Pojects/SampleGame/Content/) or (i.e:RunUAT.bat ZipUtils -archive=D:/Content.zip -extract=D:/UE/Pojects/SampleGame/Content/)")]
 [Help("archive=<PathToArchive>", "Path to folder that should be add to the archive.")]
diff --git a/Engine/Source/Programs/AutomationTool/Turnkey/Commands/VerifySdk.cs b/Engine/Source/Programs/AutomationTool/Turnkey/Commands/VerifySdk.cs
--- a/Engine/Source/Programs/AutomationTool/Turnkey/Commands/VerifySdk.cs
+++ b/Engine/Source/Programs/AutomationTool/Turnkey/Commands/VerifySdk.cs
@@ -8,6 +8,7 @@ using System.IO;
 using System.Linq;
 using EpicGames.Core;
 using UnrealBuildTool;
+using UnrealBuildBase;
 
 namespace Turnkey.Commands
 {
@@ -42,7 +43,7 @@ namespace Turnkey.Commands
 			}
 
 			TurnkeyUtils.Log("Installed Sdk validity:");
-			TurnkeyUtils.ExitCode = AutomationTool.ExitCode.Success;
+			TurnkeyUtils.ExitCode = ExitCode.Success;
 
 			TurnkeyContextImpl TurnkeyContext = new TurnkeyContextImpl();
 
@@ -140,7 +141,7 @@ namespace Turnkey.Commands
 					if (BestSdk == null)
 					{
 						TurnkeyUtils.Log("ERROR: {0}: Unable to find any Sdks that could be installed", Platform);
-						TurnkeyUtils.ExitCode = AutomationTool.ExitCode.Error_SDKNotFound;
+						TurnkeyUtils.ExitCode = ExitCode.Error_SDKNotFound;
 						continue;
 					}
 
@@ -212,7 +213,7 @@ namespace Turnkey.Commands
 									if (MatchingInstallableSdk == null)
 									{
 										TurnkeyUtils.Log("ERROR: {0}: Unable to find any Sdks that could be installed on {1}", Platform, Device.Name);
-										TurnkeyUtils.ExitCode = AutomationTool.ExitCode.Error_SDKNotFound;
+										TurnkeyUtils.ExitCode = ExitCode.Error_SDKNotFound;
 									}
 									else
 									{
diff --git a/Engine/Source/Programs/AutomationTool/Turnkey/FileSource.cs b/Engine/Source/Programs/AutomationTool/Turnkey/FileSource.cs
--- a/Engine/Source/Programs/AutomationTool/Turnkey/FileSource.cs
+++ b/Engine/Source/Programs/AutomationTool/Turnkey/FileSource.cs
@@ -553,7 +553,7 @@ namespace Turnkey
 		// 					// failed to install, nothing we can do
 		// 					if (string.IsNullOrEmpty(InstalledRoot))
 		// 					{
-		// 						TurnkeyUtils.ExitCode = AutomationTool.ExitCode.Error_SDKNotFound;
+		// 						TurnkeyUtils.ExitCode = ExitCode.Error_SDKNotFound;
 		// 						return false;
 		// 					}
 		// 
diff --git a/Engine/Source/Programs/AutomationTool/Turnkey/Turnkey.Automation.cs b/Engine/Source/Programs/AutomationTool/Turnkey/Turnkey.Automation.cs
--- a/Engine/Source/Programs/AutomationTool/Turnkey/Turnkey.Automation.cs
+++ b/Engine/Source/Programs/AutomationTool/Turnkey/Turnkey.Automation.cs
@@ -35,7 +35,7 @@ namespace Turnkey
 		}
 
 		static bool bHasBeenInitialized = false;
-		public static AutomationTool.ExitCode Execute(IOProvider IOProvider, BuildCommand CommandUtilHelper)
+		public static ExitCode Execute(IOProvider IOProvider, BuildCommand CommandUtilHelper)
 		{
 			if (!bHasBeenInitialized)
 			{
diff --git a/Engine/Source/Programs/AutomationTool/Turnkey/TurnkeyUtils.cs b/Engine/Source/Programs/AutomationTool/Turnkey/TurnkeyUtils.cs
--- a/Engine/Source/Programs/AutomationTool/Turnkey/TurnkeyUtils.cs
+++ b/Engine/Source/Programs/AutomationTool/Turnkey/TurnkeyUtils.cs
@@ -11,6 +11,7 @@ using System.Globalization;
 using EpicGames.Core;
 using AutomationTool;
 using UnrealBuildTool;
+using UnrealBuildBase;
 
 namespace Turnkey
 {
@@ -20,7 +21,7 @@ namespace Turnkey
 		static public BuildCommand CommandUtilHelper;
 
 		// replacement for Environment.ExitCode
-		static public AutomationTool.ExitCode ExitCode = ExitCode.Success;
+		static public ExitCode ExitCode = ExitCode.Success;
 
 		static public void Initialize(IOProvider InIOProvider, BuildCommand InCommandUtilHelper)
 		{
diff --git a/Engine/Source/Programs/DotNETCommon/BuildUtilities/Automation/ExitCode.cs b/Engine/Source/Programs/DotNETCommon/BuildUtilities/Automation/ExitCode.cs
new file mode 100644
--- /dev/null
+++ b/Engine/Source/Programs/DotNETCommon/BuildUtilities/Automation/ExitCode.cs
@@ -0,0 +1,65 @@
+using System;
+using System.Collections.Generic;
+using System.Text;
+
+namespace UnrealBuildBase
+{
+	// NOTE: this needs to be kept in sync with EditorAnalytics.h and iPhonePackager.cs
+	public enum ExitCode
+    {
+        Error_UATNotFound = -1,
+        Success = 0,
+        Error_Unknown = 1,
+        Error_Arguments = 2,
+        Error_UnknownCommand = 3,
+        Error_SDKNotFound = 10,
+        Error_ProvisionNotFound = 11,
+        Error_CertificateNotFound = 12,
+        Error_ProvisionAndCertificateNotFound = 13,
+        Error_InfoPListNotFound = 14,
+        Error_KeyNotFoundInPList = 15,
+        Error_ProvisionExpired = 16,
+        Error_CertificateExpired = 17,
+        Error_CertificateProvisionMismatch = 18,
+        Error_CodeUnsupported = 19,
+        Error_PluginsUnsupported = 20,
+        Error_UnknownCookFailure = 25,
+        Error_UnknownDeployFailure = 26,
+        Error_UnknownBuildFailure = 27,
+        Error_UnknownPackageFailure = 28,
+        Error_UnknownLaunchFailure = 29,
+        Error_StageMissingFile = 30,
+        Error_FailedToCreateIPA = 31,
+        Error_FailedToCodeSign = 32,
+        Error_DeviceBackupFailed = 33,
+        Error_AppUninstallFailed = 34,
+        Error_AppInstallFailed = 35,
+        Error_AppNotFound = 36,
+        Error_StubNotSignedCorrectly = 37,
+        Error_IPAMissingInfoPList = 38,
+        Error_DeleteFile = 39,
+        Error_DeleteDirectory = 40,
+        Error_CreateDirectory = 41,
+        Error_CopyFile = 42,
+        Error_OnlyOneObbFileSupported = 50,
+        Error_FailureGettingPackageInfo = 51,
+        Error_OnlyOneTargetConfigurationSupported = 52,
+        Error_ObbNotFound = 53,
+        Error_AndroidBuildToolsPathNotFound = 54,
+        Error_NoApkSuitableForArchitecture = 55,
+        Error_FilesInstallFailed = 56,
+        Error_RemoteCertificatesNotFound = 57,
+        Error_LauncherFailed = 100,
+        Error_UATLaunchFailure = 101,
+        Error_FailedToDeleteStagingDirectory = 102,
+        Error_MissingExecutable = 103,
+        Error_DeviceNotSetupForDevelopment = 150,
+        Error_DeviceOSNewerThanSDK = 151,
+		Error_TestFailure = 152,
+		Error_SymbolizedSONotFound = 153,
+		Error_LicenseNotAccepted = 154,
+		Error_AndroidOBBError = 155,
+		Error_SDKInstallFailed = 200,
+		Error_DeviceUpdateFailed = 201,
+	};
+}
