AutomationTool:

Remove dependencies from AutomationTool -> AutomationToolUtils

Removes compile-time dependencies, except for the call to AutomationTool.Automation.Process - which will be changed in a future CL.

This change is part of preparing AutomationTool to be able to build All *.Automation.csproj projects (including AutomationUtils.Automation.csproj)

diff --git a/Engine/Source/Programs/AutomationTool/AutomationTool.csproj b/Engine/Source/Programs/AutomationTool/AutomationTool.csproj
--- a/Engine/Source/Programs/AutomationTool/AutomationTool.csproj
+++ b/Engine/Source/Programs/AutomationTool/AutomationTool.csproj
@@ -79,6 +79,7 @@
   <ItemGroup>
     <ProjectReference Include="..\Shared\EpicGames.Core\EpicGames.Core.csproj" />
     <ProjectReference Include="..\UnrealBuildTool\UnrealBuildTool.csproj" />
+    <ProjectReference  Include="..\DotNETCommon\BuildUtilities\BuildUtilities.csproj" />
     <ProjectReference Include="AutomationUtils\AutomationUtils.Automation.csproj"><Private>true</Private></ProjectReference>
   </ItemGroup>
 
diff --git a/Engine/Source/Programs/AutomationTool/AutomationUtils/Automation.cs b/Engine/Source/Programs/AutomationTool/AutomationUtils/Automation.cs
--- a/Engine/Source/Programs/AutomationTool/AutomationUtils/Automation.cs
+++ b/Engine/Source/Programs/AutomationTool/AutomationUtils/Automation.cs
@@ -410,115 +410,204 @@ AutomationTool.exe [-verbose] [-compileonly] [-p4] Command0 [-Arg0 -Arg1 -Arg2 .
 			}
 		}
 
+		/// <summary>
+		/// Keep a persistent reference to the delegate for handling Ctrl-C events. Since it's passed to non-managed code, we have to prevent it from being garbage collected.
+		/// </summary>
+		static ProcessManager.CtrlHandlerDelegate CtrlHandlerDelegateInstance = CtrlHandler;
+
+        static bool CtrlHandler(CtrlTypes EventType)
+		{
+			Domain_ProcessExit(null, null);
+			if (EventType == CtrlTypes.CTRL_C_EVENT)
+			{
+				// Force exit
+				Environment.Exit(3);
+			}			
+			return true;
+		}
+
+		static void Domain_ProcessExit(object sender, EventArgs e)
+		{
+			// Kill all spawned processes (Console instead of Log because logging is closed at this time anyway)
+			if (ShouldKillProcesses && RuntimePlatform.IsWindows)
+			{			
+				ProcessManager.KillAll();
+			}
+			Trace.Close();
+		}
+
 		/// <summary>
 		/// Main method.
 		/// </summary>
 		/// <param name="Arguments">Command line</param>
 		public static ExitCode Process(string[] Arguments, StartupTraceListener StartupListener)
 		{
-			// Initial check for local or build machine runs BEFORE we parse the command line (We need this value set
-			// in case something throws the exception while parsing the command line)
-			IsBuildMachine = Arguments.Any(x => x.Equals("-BuildMachine", StringComparison.InvariantCultureIgnoreCase));
-			if (!IsBuildMachine)
+			// Hook up exit callbacks
+			AppDomain Domain = AppDomain.CurrentDomain;
+			Domain.ProcessExit += Domain_ProcessExit;
+			Domain.DomainUnload += Domain_ProcessExit;
+			HostPlatform.Current.SetConsoleCtrlHandler(CtrlHandlerDelegateInstance);
+
+			try
 			{
-				int Value;
-				if (int.TryParse(Environment.GetEnvironmentVariable("IsBuildMachine"), out Value) && Value != 0)
+				// Initial check for local or build machine runs BEFORE we parse the command line (We need this value set
+				// in case something throws the exception while parsing the command line)
+				IsBuildMachine = Arguments.Any(x => x.Equals("-BuildMachine", StringComparison.InvariantCultureIgnoreCase));
+				if (!IsBuildMachine)
 				{
-					IsBuildMachine = true;
+					int Value;
+					if (int.TryParse(Environment.GetEnvironmentVariable("IsBuildMachine"), out Value) && Value != 0)
+					{
+						IsBuildMachine = true;
+					}
 				}
-			}
-			Log.TraceVerbose("IsBuildMachine={0}", IsBuildMachine);
-			Environment.SetEnvironmentVariable("IsBuildMachine", IsBuildMachine ? "1" : "0");
+				Log.TraceVerbose("IsBuildMachine={0}", IsBuildMachine);
+				Environment.SetEnvironmentVariable("IsBuildMachine", IsBuildMachine ? "1" : "0");
 
-			// Scan the command line for commands to execute.
-			var CommandsToExecute = new List<CommandInfo>();
-			string OutScriptsForProjectFileName;
-			var AdditionalScriptsFolders = new List<string>();
-			ParseCommandLine(Arguments, CommandsToExecute, out OutScriptsForProjectFileName, AdditionalScriptsFolders);
+				// Scan the command line for commands to execute.
+				var CommandsToExecute = new List<CommandInfo>();
+				string OutScriptsForProjectFileName;
+				var AdditionalScriptsFolders = new List<string>();
+				ParseCommandLine(Arguments, CommandsToExecute, out OutScriptsForProjectFileName, AdditionalScriptsFolders);
 
-			// Get the path to the telemetry file, if present
-			string TelemetryFile = CommandUtils.ParseParamValue(Arguments, "-Telemetry");
+				// Get the path to the telemetry file, if present
+				string TelemetryFile = CommandUtils.ParseParamValue(Arguments, "-Telemetry");
 
-			// should we kill processes on exit
-			ShouldKillProcesses = !GlobalCommandLine.NoKill;
-			Log.TraceVerbose("ShouldKillProcesses={0}", ShouldKillProcesses);
+				// should we kill processes on exit
+				ShouldKillProcesses = !GlobalCommandLine.NoKill;
+				Log.TraceVerbose("ShouldKillProcesses={0}", ShouldKillProcesses);
 
-			if (CommandsToExecute.Count == 0 && GlobalCommandLine.Help)
-			{
-				DisplayHelp();
-				return ExitCode.Success;
-			}
+				if (CommandsToExecute.Count == 0 && GlobalCommandLine.Help)
+				{
+					DisplayHelp();
+					return ExitCode.Success;
+				}
 
-			// Disable AutoSDKs if specified on the command line
-			if (GlobalCommandLine.NoAutoSDK)
-			{
-				PlatformExports.PreventAutoSDKSwitching();
-			}
+				// Disable AutoSDKs if specified on the command line
+				if (GlobalCommandLine.NoAutoSDK)
+				{
+					PlatformExports.PreventAutoSDKSwitching();
+				}
 
-			// Setup environment
-			Log.TraceLog("Setting up command environment.");
-			CommandUtils.InitCommandEnvironment();
+				// Setup environment
+				Log.TraceLog("Setting up command environment.");
+				CommandUtils.InitCommandEnvironment();
 
-			// Create the log file, and flush the startup listener to it
-			TraceListener LogTraceListener = LogUtils.AddLogFileListener(CommandUtils.CmdEnv.LogFolder, CommandUtils.CmdEnv.FinalLogFolder);
-			StartupListener.CopyTo(LogTraceListener);
-			Trace.Listeners.Remove(StartupListener);
+				// Create the log file, and flush the startup listener to it
+				TraceListener LogTraceListener = LogUtils.AddLogFileListener(CommandUtils.CmdEnv.LogFolder, CommandUtils.CmdEnv.FinalLogFolder);
+				StartupListener.CopyTo(LogTraceListener);
+				Trace.Listeners.Remove(StartupListener);
 
-			// Initialize UBT
-			if(!UnrealBuildTool.PlatformExports.Initialize())
-			{
-				Log.TraceInformation("Failed to initialize UBT");
-				return ExitCode.Error_Unknown;
-			}
+				// Initialize UBT
+				if(!UnrealBuildTool.PlatformExports.Initialize())
+				{
+					Log.TraceInformation("Failed to initialize UBT");
+					return ExitCode.Error_Unknown;
+				}
 
-			// Clean rules folders up
-			ProjectUtils.CleanupFolders();
+				// Clean rules folders up
+				ProjectUtils.CleanupFolders();
 
-			// Compile scripts.
-			using(TelemetryStopwatch ScriptCompileStopwatch = new TelemetryStopwatch("ScriptCompile"))
-			{
-				ScriptCompiler.FindAndCompileAllScripts(OutScriptsForProjectFileName, AdditionalScriptsFolders);
-			}
+				// Compile scripts.
+				using(TelemetryStopwatch ScriptCompileStopwatch = new TelemetryStopwatch("ScriptCompile"))
+				{
+					ScriptCompiler.FindAndCompileAllScripts(OutScriptsForProjectFileName, AdditionalScriptsFolders);
+				}
 
-			if (GlobalCommandLine.List)
-			{
-				ListAvailableCommands(ScriptCompiler.Commands);
-				return ExitCode.Success;
-			}
+				if (GlobalCommandLine.List)
+				{
+					ListAvailableCommands(ScriptCompiler.Commands);
+					return ExitCode.Success;
+				}
 
-			if (GlobalCommandLine.Help)
-			{
-				DisplayHelp(CommandsToExecute, ScriptCompiler.Commands);
-				return ExitCode.Success;
-			}
+				if (GlobalCommandLine.Help)
+				{
+					DisplayHelp(CommandsToExecute, ScriptCompiler.Commands);
+					return ExitCode.Success;
+				}
 
-			// Enable or disable P4 support
-			CommandUtils.InitP4Support(CommandsToExecute, ScriptCompiler.Commands);
-			if (CommandUtils.P4Enabled)
-			{
-				Log.TraceLog("Setting up Perforce environment.");
-				CommandUtils.InitP4Environment();
-				CommandUtils.InitDefaultP4Connection();
-			}
+				// Enable or disable P4 support
+				CommandUtils.InitP4Support(CommandsToExecute, ScriptCompiler.Commands);
+				if (CommandUtils.P4Enabled)
+				{
+					Log.TraceLog("Setting up Perforce environment.");
+					CommandUtils.InitP4Environment();
+					CommandUtils.InitDefaultP4Connection();
+				}
 
-			try
+				try
+				{
+					// Find and execute commands.
+					ExitCode Result = Execute(CommandsToExecute, ScriptCompiler.Commands);
+					if (TelemetryFile != null)
+					{
+						Directory.CreateDirectory(Path.GetDirectoryName(TelemetryFile));
+						CommandUtils.Telemetry.Write(TelemetryFile);
+					}
+					return Result;
+				}
+				finally
+				{
+					// Flush any timing data
+					TraceSpan.Flush();
+				}
+			}
+			catch (AutomationException Ex)
 			{
-				// Find and execute commands.
-				ExitCode Result = Execute(CommandsToExecute, ScriptCompiler.Commands);
-				if (TelemetryFile != null)
+				// Output the message in the desired format
+				if(Ex.OutputFormat == AutomationExceptionOutputFormat.Silent)
+				{
+					Log.TraceLog("{0}", ExceptionUtils.FormatExceptionDetails(Ex));
+				}
+				else if(Ex.OutputFormat == AutomationExceptionOutputFormat.Minimal)
 				{
-					Directory.CreateDirectory(Path.GetDirectoryName(TelemetryFile));
-					CommandUtils.Telemetry.Write(TelemetryFile);
+					Log.TraceInformation("{0}", Ex.ToString().Replace("\n", "\n  "));
+					Log.TraceLog("{0}", ExceptionUtils.FormatExceptionDetails(Ex));
 				}
-				return Result;
+				else
+				{
+					Log.WriteException(Ex, LogUtils.FinalLogFileName);
+				}
+
+				// Take the exit code from the exception
+				return Ex.ErrorCode;
 			}
-			finally
+			catch (Exception Ex)
 			{
-				// Flush any timing data
-				TraceSpan.Flush();
+				// Use a default exit code
+				Log.WriteException(Ex, LogUtils.FinalLogFileName);
+				return ExitCode.Error_Unknown;
 			}
+            finally
+            {
+                // In all cases, do necessary shut down stuff, but don't let any additional exceptions leak out while trying to shut down.
+
+                // Make sure there's no directories on the stack.
+                NoThrow(() => CommandUtils.ClearDirStack(), "Clear Dir Stack");
+
+                // Try to kill process before app domain exits to leave the other KillAll call to extreme edge cases
+                NoThrow(() => { if (ShouldKillProcesses && RuntimePlatform.IsWindows) ProcessManager.KillAll(); }, "Kill All Processes");
+            }
 		}
 
+		/// <summary>
+		/// Wraps an action in an exception block.
+		/// Ensures individual actions can be performed and exceptions won't prevent further actions from being executed.
+		/// Useful for shutdown code where shutdown may be in several stages and it's important that all stages get a chance to run.
+		/// </summary>
+		/// <param name="Action"></param>
+		private static void NoThrow(System.Action Action, string ActionDesc)
+        {
+            try
+            {
+                Action();
+            }
+            catch (Exception Ex)
+            {
+                Log.TraceError("Exception performing nothrow action \"{0}\": {1}", ActionDesc, ExceptionUtils.FormatException(Ex));
+            }
+        }
+
 		/// <summary>
 		/// Execute commands specified in the command line.
 		/// </summary>
@@ -636,18 +725,14 @@ AutomationTool.exe [-verbose] [-compileonly] [-p4] Command0 [-Arg0 -Arg1 -Arg2 .
 		{
 			get
 			{
-				if (!bShouldKillProcesses.HasValue)
-				{
-					throw new AutomationException("Trying to access ShouldKillProcesses property before it was initialized.");
-				}
-				return (bool)bShouldKillProcesses;
+				return bShouldKillProcesses;
 			}
 			private set
 			{
 				bShouldKillProcesses = value;
 			}
 		}
-		private static bool? bShouldKillProcesses;
+		private static bool bShouldKillProcesses = true;
 	}
 
 }
diff --git a/Engine/Source/Programs/AutomationTool/AutomationUtils/CommandEnvironment.cs b/Engine/Source/Programs/AutomationTool/AutomationUtils/CommandEnvironment.cs
--- a/Engine/Source/Programs/AutomationTool/AutomationUtils/CommandEnvironment.cs
+++ b/Engine/Source/Programs/AutomationTool/AutomationUtils/CommandEnvironment.cs
@@ -113,7 +113,7 @@ namespace AutomationTool
 			}
 
 			// clear the logfolder if we're the only running instance
-			if (InternalUtils.IsSoleInstance)
+			if (ProcessSingleton.IsSoleInstance)
 			{
 				ClearLogFolder(LogFolder);
 			}
diff --git a/Engine/Source/Programs/AutomationTool/AutomationUtils/HostPlatform.cs b/Engine/Source/Programs/AutomationTool/AutomationUtils/HostPlatform.cs
--- a/Engine/Source/Programs/AutomationTool/AutomationUtils/HostPlatform.cs
+++ b/Engine/Source/Programs/AutomationTool/AutomationUtils/HostPlatform.cs
@@ -15,43 +15,23 @@ namespace AutomationTool
 	/// </summary>
 	public abstract class HostPlatform
 	{
-		private static HostPlatform RunningPlatform;
 		/// <summary>
 		/// Current running host platform.
 		/// </summary>
-		public static HostPlatform Current
-		{
-			get 
-			{ 
-				if (RunningPlatform == null)
-				{
-					throw new AutomationException("UnrealAutomationTool host platform not initialized.");
-				}
-				return RunningPlatform; 
-			}
-		}
+		public static readonly HostPlatform Current = Initialize();
 
 		/// <summary>
 		/// Initializes the current platform.
 		/// </summary>
-		public static void Initialize()
+		private static HostPlatform Initialize()
 		{
-			if (RuntimePlatform.IsWindows)
-			{
-				RunningPlatform = new WindowsHostPlatform();
-			}
-			else if (RuntimePlatform.IsMac)
-			{
-				RunningPlatform = new MacHostPlatform();
-			}
-			else if (RuntimePlatform.IsLinux)
-			{
-				RunningPlatform = new LinuxHostPlatform();
-			}
-			else
+			switch (RuntimePlatform.Current)
 			{
-				throw new Exception ("Unhandled runtime platform " + Environment.OSVersion.Platform);
+				case RuntimePlatform.Type.Windows: return new WindowsHostPlatform();
+				case RuntimePlatform.Type.Mac:     return new MacHostPlatform();
+				case RuntimePlatform.Type.Linux:   return new LinuxHostPlatform();
 			}
+			throw new Exception ("Unhandled runtime platform " + Environment.OSVersion.Platform);
 		}
 
 		/// <summary>
diff --git a/Engine/Source/Programs/AutomationTool/AutomationUtils/Utils.cs b/Engine/Source/Programs/AutomationTool/AutomationUtils/Utils.cs
--- a/Engine/Source/Programs/AutomationTool/AutomationUtils/Utils.cs
+++ b/Engine/Source/Programs/AutomationTool/AutomationUtils/Utils.cs
@@ -709,76 +709,6 @@ namespace AutomationTool
 			return Result;
 		}
 
-		/// <summary>
-		/// Returns true/false based on whether this is the only instance
-		/// running (checked at startup).
-		/// </summary>
-		public static bool IsSoleInstance { get; private set; }
-
-		/// <summary>
-		/// Runs the specified delegate checking if this is the only instance of the application.
-		/// </summary>
-		/// <param name="Main"></param>
-		/// <param name="Param"></param>
-		public static ExitCode RunSingleInstance(string[] Arguments, Func<ExitCode> Main)
-		{
-			bool AllowMultipleInsances = (Environment.GetEnvironmentVariable("uebp_UATMutexNoWait") == "1");
-	
-			var bCreatedMutex = false;
-            var EntryAssemblyLocation = Assembly.GetEntryAssembly().GetOriginalLocation();
-
-			string MutexName = GetUniqueMutexForPath(Path.GetFileNameWithoutExtension(EntryAssemblyLocation), EntryAssemblyLocation);
-			using (Mutex SingleInstanceMutex = new Mutex(true, MutexName, out bCreatedMutex))
-			{
-				IsSoleInstance = bCreatedMutex;
-
-				if (!IsSoleInstance && AllowMultipleInsances == false)
-				{
-					if (CommandUtils.ParseParam(Arguments, "-WaitForUATMutex"))
-					{
-						CommandUtils.LogWarning("Another instance of UAT at '{0}' is running, and the -WaitForUATMutex parameter has been used. Waiting for other UAT to finish...", EntryAssemblyLocation);
-						int Seconds = 0;
-						while (WaitMutexNoExceptions(SingleInstanceMutex, 15 * 1000) == false)
-						{
-							Seconds += 15;
-							CommandUtils.LogInformation("Still waiting for Mutex. {0} seconds has passed...", Seconds);
-						}
-					}
-					else
-					{
-						throw new AutomationException("A conflicting instance of AutomationTool is already running. Current location: {0}. A process manager may be used to determine the conflicting process and what tool may have launched it", EntryAssemblyLocation);
-					}
-				}
-
-				ExitCode Result = Main();
-
-				if (IsSoleInstance)
-				{
-					SingleInstanceMutex.ReleaseMutex();
-				}
-
-				return Result;
-			}
-		}
-
-		static bool WaitMutexNoExceptions(Mutex Mutex, int TimeoutMs)
-		{
-			try
-			{
-				return Mutex.WaitOne(15 * 1000);
-			}
-			catch (AbandonedMutexException)
-			{
-				return true;
-			}
-		}
-
-		public static string GetUniqueMutexForPath(string Name, string UniquePath)
-		{
-			// generate a md5 hash of the path, as GetHashCode is not guaranteed to generate a stable hash
-			return string.Format("Global\\{0}_{1}", Name, ContentHash.MD5(UniquePath.ToUpperInvariant()));
-		}
-
 	    public static void Robust_CopyFile(string InputFileName, string OutputFileName)
 	    {
 	        if (OutputFileName.StartsWith("/Volumes/"))
diff --git a/Engine/Source/Programs/AutomationTool/Program.cs b/Engine/Source/Programs/AutomationTool/Program.cs
--- a/Engine/Source/Programs/AutomationTool/Program.cs
+++ b/Engine/Source/Programs/AutomationTool/Program.cs
@@ -13,20 +13,15 @@ using System.Collections.Generic;
 using System.Text;
 using UnrealBuildBase;
 
-namespace AutomationTool
+namespace AutomationToolDriver
 {
 	public class Program
 	{
-		/// <summary>
-		/// Keep a persistent reference to the delegate for handling Ctrl-C events. Since it's passed to non-managed code, we have to prevent it from being garbage collected.
-		/// </summary>
-		static ProcessManager.CtrlHandlerDelegate CtrlHandlerDelegateInstance = CtrlHandler;
-
 		// Do not add [STAThread] here. It will cause deadlocks in platform automation code.
 		public static int Main(string[] Arguments)
 		{
 			// Wait for a debugger to be attached
-			if (CommandUtils.ParseParam(Arguments, "-WaitForDebugger"))
+			if (ParseParam(Arguments, "-WaitForDebugger"))
 			{
 				Console.WriteLine("Waiting for debugger to be attached...");
 				while (Debugger.IsAttached == false)
@@ -39,17 +34,17 @@ namespace AutomationTool
 			Stopwatch Timer = Stopwatch.StartNew();
 			
 			// Ensure UTF8Output flag is respected, since we are initializing logging early in the program.
-			if (CommandUtils.ParseParam(Arguments, "-Utf8output"))
+			if (ParseParam(Arguments, "-Utf8output"))
             {
                 Console.OutputEncoding = new System.Text.UTF8Encoding(false, false);
             }
 
 			// Parse the log level argument
-			if(CommandUtils.ParseParam(Arguments, "-Verbose"))
+			if(ParseParam(Arguments, "-Verbose"))
 			{
 				Log.OutputLevel = LogEventType.Verbose;
 			}
-			if(CommandUtils.ParseParam(Arguments, "-VeryVerbose"))
+			if(ParseParam(Arguments, "-VeryVerbose"))
 			{
 				Log.OutputLevel = LogEventType.VeryVerbose;
 			}
@@ -59,7 +54,7 @@ namespace AutomationTool
 			Trace.Listeners.Add(StartupListener);
 
 			// Configure log timestamps
-			Log.IncludeTimestamps = CommandUtils.ParseParam(Arguments, "-Timestamps");
+			Log.IncludeTimestamps = ParseParam(Arguments, "-Timestamps");
 
 			// Enter the main program section
             ExitCode ReturnCode = ExitCode.Success;
@@ -74,11 +69,8 @@ namespace AutomationTool
 				AssemblyUtils.InstallAssemblyResolver(PathToBinariesDotNET);
 				AssemblyUtils.InstallRecursiveAssemblyResolver(PathToBinariesDotNET);
 
-				// Initialize the host platform layer
-				HostPlatform.Initialize();
-
 				// Log the operating environment. Since we usually compile to AnyCPU, we may be executed using different system paths under WOW64.
-				Log.TraceVerbose("{2}: Running on {0} as a {1}-bit process.", HostPlatform.Current.GetType().Name, Environment.Is64BitProcess ? 64 : 32, DateTime.UtcNow.ToString("o"));
+				Log.TraceVerbose("{2}: Running on {0} as a {1}-bit process.", RuntimePlatform.Current.ToString(), Environment.Is64BitProcess ? 64 : 32, DateTime.UtcNow.ToString("o"));
 
 				// Log if we're running from the launcher
 				string ExecutingAssemblyLocation = Assembly.GetExecutingAssembly().Location;
@@ -88,55 +80,22 @@ namespace AutomationTool
 				}
 				Log.TraceVerbose("CWD={0}", Environment.CurrentDirectory);
 
-				// Hook up exit callbacks
-				AppDomain Domain = AppDomain.CurrentDomain;
-				Domain.ProcessExit += Domain_ProcessExit;
-				Domain.DomainUnload += Domain_ProcessExit;
-				HostPlatform.Current.SetConsoleCtrlHandler(CtrlHandlerDelegateInstance);
-
 				// Log the application version
 				FileVersionInfo Version = AssemblyUtils.ExecutableVersion;
 				Log.TraceVerbose("{0} ver. {1}", Version.ProductName, Version.ProductVersion);
 
+				bool bWaitForUATMutex = ParseParam(Arguments, "-WaitForUATMutex");
+
 				// Don't allow simultaneous execution of AT (in the same branch)
-				ReturnCode = InternalUtils.RunSingleInstance(Arguments, () => MainProc(Arguments, StartupListener));
-			}
-			catch (AutomationException Ex)
-			{
-				// Output the message in the desired format
-				if(Ex.OutputFormat == AutomationExceptionOutputFormat.Silent)
-				{
-					Log.TraceLog("{0}", ExceptionUtils.FormatExceptionDetails(Ex));
-				}
-				else if(Ex.OutputFormat == AutomationExceptionOutputFormat.Minimal)
-				{
-					Log.TraceInformation("{0}", Ex.ToString().Replace("\n", "\n  "));
-					Log.TraceLog("{0}", ExceptionUtils.FormatExceptionDetails(Ex));
-				}
-				else
-				{
-					Log.WriteException(Ex, LogUtils.FinalLogFileName);
-				}
 
-				// Take the exit code from the exception
-				ReturnCode = Ex.ErrorCode;
+				ReturnCode = ProcessSingleton.RunSingleInstance(() => MainProc(Arguments, StartupListener), bWaitForUATMutex);
 			}
 			catch (Exception Ex)
-			{
-				// Use a default exit code
-				Log.WriteException(Ex, LogUtils.FinalLogFileName);
-				ReturnCode = ExitCode.Error_Unknown;
-			}
+            {
+				Log.TraceError("Exception occurred between AutomationToolDriver.Main() and Automation.Process()" + ExceptionUtils.FormatException(Ex));
+            }
             finally
             {
-                // In all cases, do necessary shut down stuff, but don't let any additional exceptions leak out while trying to shut down.
-
-                // Make sure there's no directories on the stack.
-                NoThrow(() => CommandUtils.ClearDirStack(), "Clear Dir Stack");
-
-                // Try to kill process before app domain exits to leave the other KillAll call to extreme edge cases
-                NoThrow(() => { if (ShouldKillProcesses && RuntimePlatform.IsWindows) ProcessManager.KillAll(); }, "Kill All Processes");
-
 				// Write the exit code
                 Log.TraceInformation("AutomationTool executed for {0}", Timer.Elapsed.ToString("h'h 'm'm 's's'"));
                 Log.TraceInformation("AutomationTool exiting with ExitCode={0} ({1})", (int)ReturnCode, ReturnCode);
@@ -147,52 +106,35 @@ namespace AutomationTool
             return (int)ReturnCode;
         }
 
+		static ExitCode MainProc(string[] Arguments, StartupTraceListener StartupListener)
+		{
+			ExitCode Result = AutomationTool.Automation.Process(Arguments, StartupListener);
+			return Result;
+		}
+
+		// Code duplicated from CommandUtils.cs
 		/// <summary>
-		/// Wraps an action in an exception block.
-		/// Ensures individual actions can be performed and exceptions won't prevent further actions from being executed.
-		/// Useful for shutdown code where shutdown may be in several stages and it's important that all stages get a chance to run.
+		/// Parses the argument list for a parameter and returns whether it is defined or not.
 		/// </summary>
-		/// <param name="Action"></param>
-		private static void NoThrow(System.Action Action, string ActionDesc)
-        {
-            try
-            {
-                Action();
-            }
-            catch (Exception Ex)
+		/// <param name="ArgList">Argument list.</param>
+		/// <param name="Param">Param to check for.</param>
+		/// <returns>True if param was found, false otherwise.</returns>
+		public static bool ParseParam(string[] ArgList, string Param)
+		{
+            string ValueParam = Param;
+            if (!ValueParam.EndsWith("="))
             {
-                Log.TraceError("Exception performing nothrow action \"{0}\": {1}", ActionDesc, LogUtils.FormatException(Ex));
+                ValueParam += "=";
             }
-        }
 
-        static bool CtrlHandler(CtrlTypes EventType)
-		{
-			Domain_ProcessExit(null, null);
-			if (EventType == CtrlTypes.CTRL_C_EVENT)
+            foreach (string ArgStr in ArgList)
 			{
-				// Force exit
-				Environment.Exit(3);
-			}			
-			return true;
-		}
-
-		static void Domain_ProcessExit(object sender, EventArgs e)
-		{
-			// Kill all spawned processes (Console instead of Log because logging is closed at this time anyway)
-			if (ShouldKillProcesses && RuntimePlatform.IsWindows)
-			{			
-				ProcessManager.KillAll();
+                if (ArgStr.Equals(Param, StringComparison.InvariantCultureIgnoreCase) || ArgStr.StartsWith(ValueParam, StringComparison.InvariantCultureIgnoreCase))
+				{
+					return true;
+				}
 			}
-			Trace.Close();
-		}
-
-		static ExitCode MainProc(string[] Arguments, StartupTraceListener StartupListener)
-		{
-			ExitCode Result = Automation.Process(Arguments, StartupListener);
-			ShouldKillProcesses = Automation.ShouldKillProcesses;
-			return Result;
+			return false;
 		}
-
-		static bool ShouldKillProcesses = true;
 	}
 }
diff --git a/Engine/Source/Programs/DotNETCommon/BuildUtilities/Automation/ProcessSingleton.cs b/Engine/Source/Programs/DotNETCommon/BuildUtilities/Automation/ProcessSingleton.cs
new file mode 100644
--- /dev/null
+++ b/Engine/Source/Programs/DotNETCommon/BuildUtilities/Automation/ProcessSingleton.cs
@@ -0,0 +1,86 @@
+﻿// Copyright Epic Games, Inc. All Rights Reserved.
+
+using EpicGames.Core;
+using System;
+using System.Collections.Generic;
+using System.IO;
+using System.Reflection;
+using System.Text;
+using System.Threading;
+
+#nullable disable
+
+namespace UnrealBuildBase
+{
+	public class ProcessSingleton
+	{
+		public static string GetUniqueMutexForPath(string Name, string UniquePath)
+		{
+			// generate a md5 hash of the path, as GetHashCode is not guaranteed to generate a stable hash
+			return string.Format("Global\\{0}_{1}", Name, ContentHash.MD5(UniquePath.ToUpperInvariant()));
+		}
+
+		/// <summary>
+		/// Returns true/false based on whether this is the only instance
+		/// running (checked at startup).
+		/// </summary>
+		public static bool IsSoleInstance { get; private set; }
+
+		/// <summary>
+		/// Runs the specified delegate checking if this is the only instance of the application.
+		/// </summary>
+		/// <param name="Main"></param>
+		/// <param name="Param"></param>
+		public static ExitCode RunSingleInstance(Func<ExitCode> Main, bool bWaitForUATMutex)
+		{
+			bool AllowMultipleInsances = (Environment.GetEnvironmentVariable("uebp_UATMutexNoWait") == "1");
+	
+            string EntryAssemblyLocation = Assembly.GetEntryAssembly().GetOriginalLocation();
+
+			string MutexName = GetUniqueMutexForPath(Path.GetFileNameWithoutExtension(EntryAssemblyLocation), EntryAssemblyLocation);
+			using (Mutex SingleInstanceMutex = new Mutex(true, MutexName, out bool bCreatedMutex))
+			{
+				IsSoleInstance = bCreatedMutex;
+
+				if (!IsSoleInstance && AllowMultipleInsances == false)
+				{
+					if (bWaitForUATMutex)
+					{
+						Log.TraceWarning("Another instance of UAT at '{0}' is running, and the -WaitForUATMutex parameter has been used. Waiting for other UAT to finish...", EntryAssemblyLocation);
+						int Seconds = 0;
+						while (WaitMutexNoExceptions(SingleInstanceMutex, 15 * 1000) == false)
+						{
+							Seconds += 15;
+							Log.TraceInformation("Still waiting for Mutex. {0} seconds has passed...", Seconds);
+						}
+					}
+					else
+					{
+						throw new Exception($"A conflicting instance of AutomationTool is already running. Current location: {EntryAssemblyLocation}. A process manager may be used to determine the conflicting process and what tool may have launched it");
+					}
+				}
+
+				ExitCode Result = Main();
+
+				if (IsSoleInstance)
+				{
+					SingleInstanceMutex.ReleaseMutex();
+				}
+
+				return Result;
+			}
+		}
+
+		static bool WaitMutexNoExceptions(Mutex Mutex, int TimeoutMs)
+		{
+			try
+			{
+				return Mutex.WaitOne(15 * 1000);
+			}
+			catch (AbandonedMutexException)
+			{
+				return true;
+			}
+		}
+	}
+}
