Make sure to break after a successful copy.

#rnx

diff --git a/Engine/Source/Programs/AutomationTool/AutomationUtils/CommandUtils.cs b/Engine/Source/Programs/AutomationTool/AutomationUtils/CommandUtils.cs
--- a/Engine/Source/Programs/AutomationTool/AutomationUtils/CommandUtils.cs
+++ b/Engine/Source/Programs/AutomationTool/AutomationUtils/CommandUtils.cs
@@ -1329,6 +1329,7 @@ namespace AutomationTool
 					}
 					throw new AutomationException("Failed to copy {0} to {1}", Source, Dest);
 				}
+				break;
 			}
 		}
 
