Downgrade ImageMagick issues and handle them more gracefully, and automatically upload all client/server logs from build machine tests to their horde job.

diff --git a/Engine/Source/Programs/AutomationTool/Gauntlet/Unreal/Base/Gauntlet.UnrealSession.cs b/Engine/Source/Programs/AutomationTool/Gauntlet/Unreal/Base/Gauntlet.UnrealSession.cs
--- a/Engine/Source/Programs/AutomationTool/Gauntlet/Unreal/Base/Gauntlet.UnrealSession.cs
+++ b/Engine/Source/Programs/AutomationTool/Gauntlet/Unreal/Base/Gauntlet.UnrealSession.cs
@@ -1083,6 +1083,14 @@ namespace Gauntlet
 			LogOut.Append(InRunningRole.AppInstance.StdOut);
 
 			File.WriteAllText(ArtifactLogPath, LogOut.ToString());
+
+			// On build machines, copy all role logs to Horde.
+			if (CommandUtils.IsBuildMachine)
+			{
+				string HordeLogFileName = Path.Combine(CommandUtils.CmdEnv.LogFolder, RoleName + "Output.log");
+				File.WriteAllText(HordeLogFileName, LogOut.ToString());
+			}
+
 			Log.Info("Wrote Log to {0}", ArtifactLogPath);
 
 			if (IsServer == false)
@@ -1115,7 +1123,7 @@ namespace Gauntlet
 				}
 				catch (Exception Ex)
 				{
-					Log.Info("Failed to downsize and gif-ify images! {0}", Ex);
+					Log.Info("Failed to downsize and gif-ify images! {0}", Ex.Message);
 				}
 			}
 
