IterationProfile: Promote ShaderCompiler logging to Verbose so we don't time out building shaders for 30+ minutes

diff --git a/Engine/Build/IterationProfile.xml b/Engine/Build/IterationProfile.xml
--- a/Engine/Build/IterationProfile.xml
+++ b/Engine/Build/IterationProfile.xml
@@ -70,7 +70,7 @@
 	<Property Name="ProfileArgs" Value="$(ProfileArgs) -TreatLogWarningsAsTestErrors=false -TreatLogErrorsAsTestErrors=false"/>
 	<Property Name="ProfileArgs" Value="$(ProfileArgs) -GameInstanceLostTimerSeconds=3600"/>
 	<Property Name="ProfileArgs" Value="$(ProfileArgs) -NoLoadLevelAtStartup"/>
-	<Property Name="ProfileArgs" Value="$(ProfileArgs) -NoShaderDistrib"/>
+	<Property Name="ProfileArgs" Value="$(ProfileArgs) -NoShaderDistrib -VerboseShaderLogging"/>
 
 	<!-- Gauntlet arguments for summarize commandlet runs -->
 	<Property Name="SummarizeArgs" Value="-statsfile=&quot;$(RootDir)/Engine/Build/IterationProfileStats.csv&quot;"/>
diff --git a/Engine/Source/Programs/AutomationTool/Gauntlet/Unreal/Editor/EditorTests.cs b/Engine/Source/Programs/AutomationTool/Gauntlet/Unreal/Editor/EditorTests.cs
--- a/Engine/Source/Programs/AutomationTool/Gauntlet/Unreal/Editor/EditorTests.cs
+++ b/Engine/Source/Programs/AutomationTool/Gauntlet/Unreal/Editor/EditorTests.cs
@@ -67,6 +67,12 @@ namespace UnrealEditor
 		[AutoParam]
 		public bool NoShaderDistrib = false;
 
+		/// <summary>
+		/// Enable Verbose Shader logging so we don't time out compiling lots of shaders
+		/// </summary>
+		[AutoParam]
+		public bool VerboseShaderLogging = false;
+
 		/// <summary>
 		/// Applies these options to the provided app config
 		/// </summary>
@@ -112,6 +118,11 @@ namespace UnrealEditor
 			{
 				AppConfig.CommandLineParams.Add("-noxgeshadercompile");
 			}
+
+			if (VerboseShaderLogging)
+			{
+				AppConfig.CommandLineParams.Add("ini:Engine:[Core.Log]:LogShaderCompilers=Verbose");
+			}
 		}
 	}
 
