IterationProfile: Turn on feature toggles in the engine for stabilizing tick rate and deterministic updates

diff --git a/Engine/Build/IterationProfile.xml b/Engine/Build/IterationProfile.xml
--- a/Engine/Build/IterationProfile.xml
+++ b/Engine/Build/IterationProfile.xml
@@ -69,7 +69,7 @@
 	<Property Name="ProfileArgs" Value="$(ProfileArgs) -EnablePlugins=$(ExtraPluginsCommaDelims)" If="'$(ExtraPluginsCommaDelims)' != ''"/>
 	<Property Name="ProfileArgs" Value="$(ProfileArgs) -TreatLogWarningsAsTestErrors=false -TreatLogErrorsAsTestErrors=false"/>
 	<Property Name="ProfileArgs" Value="$(ProfileArgs) -GameInstanceLostTimerSeconds=3600"/>
-	<Property Name="ProfileArgs" Value="$(ProfileArgs) -NoLoadLevelAtStartup"/>
+	<Property Name="ProfileArgs" Value="$(ProfileArgs) -NoLoadLevelAtStartup -Benchmarking"/>
 	<Property Name="ProfileArgs" Value="$(ProfileArgs) -NoShaderDistrib -VerboseShaderLogging"/>
 
 	<!-- Gauntlet arguments for summarize commandlet runs -->
diff --git a/Engine/Source/Programs/AutomationTool/Gauntlet/Unreal/Editor/EditorTests.cs b/Engine/Source/Programs/AutomationTool/Gauntlet/Unreal/Editor/EditorTests.cs
--- a/Engine/Source/Programs/AutomationTool/Gauntlet/Unreal/Editor/EditorTests.cs
+++ b/Engine/Source/Programs/AutomationTool/Gauntlet/Unreal/Editor/EditorTests.cs
@@ -73,6 +73,12 @@ namespace UnrealEditor
 		[AutoParam]
 		public bool VerboseShaderLogging = false;
 
+		/// <summary>
+		/// Enable benchmarking features in the engine
+		/// </summary>
+		[AutoParam]
+		public bool Benchmarking = false;
+
 		/// <summary>
 		/// Applies these options to the provided app config
 		/// </summary>
@@ -116,13 +122,19 @@ namespace UnrealEditor
 
 			if (NoShaderDistrib)
 			{
-				AppConfig.CommandLineParams.Add("-noxgeshadercompile");
+				AppConfig.CommandLineParams.Add("noxgeshadercompile");
 			}
 
 			if (VerboseShaderLogging)
 			{
 				AppConfig.CommandLineParams.Add("ini:Engine:[Core.Log]:LogShaderCompilers=Verbose");
 			}
+
+			if (Benchmarking)
+			{
+				AppConfig.CommandLineParams.Add("BENCHMARK");
+				AppConfig.CommandLineParams.Add("Deterministic");
+			}
 		}
 	}
 
