Always log the client spec created in calls to CommandUtils.P4.CreateClient.

diff --git a/Engine/Source/Programs/AutomationTool/AutomationUtils/P4Utils.cs b/Engine/Source/Programs/AutomationTool/AutomationUtils/P4Utils.cs
--- a/Engine/Source/Programs/AutomationTool/AutomationUtils/P4Utils.cs
+++ b/Engine/Source/Programs/AutomationTool/AutomationUtils/P4Utils.cs
@@ -3551,7 +3551,7 @@ namespace AutomationTool
 					SpecInput += "\t" + Mapping.Key + " //" + ClientSpec.Name + Mapping.Value + Environment.NewLine;
 				}
 			}
-			if (AllowSpew) CommandUtils.LogLog(SpecInput);
+			CommandUtils.LogLog(SpecInput);
             LogP4("", "client -i", SpecInput, AllowSpew: AllowSpew, WithClient: false);
             return ClientSpec;
         }
