Add support for nested expansion of Build Graph Properties.

For example: $(Outer$(Inner))

diff --git a/Engine/Source/Programs/AutomationTool/BuildGraph/Script.cs b/Engine/Source/Programs/AutomationTool/BuildGraph/Script.cs
--- a/Engine/Source/Programs/AutomationTool/BuildGraph/Script.cs
+++ b/Engine/Source/Programs/AutomationTool/BuildGraph/Script.cs
@@ -2087,7 +2087,8 @@ namespace AutomationTool
 		string ExpandProperties(ScriptElement Element, string Text)
 		{
 			string Result = Text;
-			for (int Idx = Result.IndexOf("$("); Idx != -1; Idx = Result.IndexOf("$(", Idx))
+			// Iterate in reverse order to handle cases where there are nested expansions like $(Outer$(Inner))
+			for (int Idx = Result.LastIndexOf("$("); Idx != -1; Idx = Result.LastIndexOf("$(", Idx, Idx+1))
 			{
 				// Find the end of the variable name
 				int EndIdx = Result.IndexOf(')', Idx + 2);
@@ -2101,25 +2102,17 @@ namespace AutomationTool
 
 				// Find the value for it, either from the dictionary or the environment block
 				string Value;
-				if(!TryGetPropertyValue(Name, out Value))
+				if (!TryGetPropertyValue(Name, out Value))
 				{
 					LogWarning(Element, "Property '{0}' is not defined", Name);
 					Value = "";
 				}
 
 				// Check if we've got a value for this variable
-				if (Value == null)
-				{
-					// Do not expand it; must be preprocessing the script.
-					Idx = EndIdx;
-				}
-				else
+				if (Value != null)
 				{
 					// Replace the variable, or skip past it
 					Result = Result.Substring(0, Idx) + Value + Result.Substring(EndIdx + 1);
-
-					// Make sure we skip over the expanded variable; we don't want to recurse on it.
-					Idx += Value.Length;
 				}
 			}
 			return Result;
