Crowdin Localization - Fixing copyright.

diff --git a/Engine/Source/Programs/AutomationTool/CrowdinLocalization/Properties/AssemblyInfo.cs b/Engine/Source/Programs/AutomationTool/CrowdinLocalization/Properties/AssemblyInfo.cs
--- a/Engine/Source/Programs/AutomationTool/CrowdinLocalization/Properties/AssemblyInfo.cs
+++ b/Engine/Source/Programs/AutomationTool/CrowdinLocalization/Properties/AssemblyInfo.cs
@@ -1,4 +1,4 @@
-// Copyright 1998-2019 Epic Games, Inc. All Rights Reserved.
+// Copyright Epic Games, Inc. All Rights Reserved.
 
 using System.Reflection;
 using System.Runtime.CompilerServices;
