Updated automation mappings for XLOC

diff --git a/Engine/Source/Programs/AutomationTool/XLocLocalization/XLocLocalizationProvider.cs b/Engine/Source/Programs/AutomationTool/XLocLocalization/XLocLocalizationProvider.cs
--- a/Engine/Source/Programs/AutomationTool/XLocLocalization/XLocLocalizationProvider.cs
+++ b/Engine/Source/Programs/AutomationTool/XLocLocalization/XLocLocalizationProvider.cs
@@ -415,21 +415,31 @@ namespace EpicGames.XLocLocalization
 		{
 			return new Dictionary<string, string>
 			{
-				{ "en", "E" },		// English
-				{ "fr", "F" },		// French
-				{ "it", "I" },		// Italian
-				{ "de", "G" },		// German
-				{ "es", "S" },		// Spanish
-				{ "es-419", "Y" },	// LatAm Spanish
-				{ "ru", "R" },		// Russian
-				{ "pl", "P" },		// Polish
-				{ "ar", "A" },		// Arabic
-				{ "ko", "K" },		// Korean
-				{ "ja", "J" },		// Japanese
-				{ "zh-Hans", "2" },	// Simp Chinese
-				{ "zh-Hant", "1" },	// Trad Chinese
-				{ "tr", "3" },		// Turkish
-				{ "pt-BR", "$" },	// Brazilian Portuguese
+				{ "en", "E" },			// English
+				{ "en-US", "9" },		// English (US)
+				{ "en-GB", "6" },		// English (British)
+				{ "en-HK", "en-HK" },	// English (Hong Kong)
+				{ "fr", "F" },			// French
+				{ "fr-CA", "Q"},		// French (Canadian)
+				{ "it", "I" },			// Italian
+				{ "de", "G" },			// German
+				{ "es", "S" },			// Spanish
+				{ "es-419", "Y" },		// Spanish (Latin American)
+				{ "da", "D" },			// Danish
+				{ "nl", "U" },			// Dutch
+				{ "fi", "B" },			// Finnish
+				{ "sv", "W" },			// Swedish
+				{ "ru", "R" },			// Russian
+				{ "pl", "P" },			// Polish
+				{ "ar", "A" },			// Arabic
+				{ "ko", "K" },			// Korean
+				{ "ja", "J" },			// Japanese
+				{ "zh-Hans", "2" },		// Chinese (Simplified)
+				{ "zh-Hant", "1" },		// Chinese (Traditional)
+				{ "tr", "3" },			// Turkish
+				{ "th", "T" },			// Thai
+				{ "pt", "O" },			// Portuguese
+				{ "pt-BR", "$" },		// Portuguese (Brazilian)
 			};
 		}
 	}
