AutomationTool:

diff --git a/Engine/Source/Programs/AutomationTool/Scripts/CopyBuildToStagingDirectory.Automation.cs b/Engine/Source/Programs/AutomationTool/Scripts/CopyBuildToStagingDirectory.Automation.cs
--- a/Engine/Source/Programs/AutomationTool/Scripts/CopyBuildToStagingDirectory.Automation.cs
+++ b/Engine/Source/Programs/AutomationTool/Scripts/CopyBuildToStagingDirectory.Automation.cs
@@ -18,766 +18,754 @@ using System.Collections.Concurrent;
 using UnrealBuildBase;
 using System.Text.Json;
 
-/// <summary>
-/// Helper command used for cooking.
-/// </summary>
-/// <remarks>
-/// Command line parameters used by this command:
-/// -clean
-/// </remarks>
-public partial class Project : CommandUtils
+namespace AutomationScripts
 {
+	/// <summary>
+	/// Helper command used for cooking.
+	/// </summary>
+	/// <remarks>
+	/// Command line parameters used by this command:
+	/// -clean
+	/// </remarks>
+	public partial class Project : CommandUtils
+	{
 
-	private static readonly object SyncLock = new object();
+		private static readonly object SyncLock = new object();
 
-	public struct OrderFile : IComparable<OrderFile>
-	{
-		public enum OrderFileType
+		public struct OrderFile : IComparable<OrderFile>
 		{
-			Custom,
-			Game,
-			Cooker,
-			Editor
-		}
+			public enum OrderFileType
+			{
+				Custom,
+				Game,
+				Cooker,
+				Editor
+			}
 
-		public OrderFile(FileReference FileRef, OrderFileType InType, int SpecIndex, int InPriority)
-		{
-			File = FileRef;
-			OrderSpecIndex = SpecIndex;
+			public OrderFile(FileReference FileRef, OrderFileType InType, int SpecIndex, int InPriority)
+			{
+				File = FileRef;
+				OrderSpecIndex = SpecIndex;
 
-			string FileName = FileRef.GetFileNameWithoutExtension();
-			OrderType = InType;
-			Priority = InPriority;
+				string FileName = FileRef.GetFileNameWithoutExtension();
+				OrderType = InType;
+				Priority = InPriority;
 
-			//Check if we have an order number
-			AppendOrder = -1;
-			int index = FileName.LastIndexOf("_");
-			if (index != -1)
-			{
-				index++;
-				string Number = FileName.Substring(index, FileName.Length - index);
-				AppendOrder = Convert.ToInt32(Number);
-			}
+				//Check if we have an order number
+				AppendOrder = -1;
+				int index = FileName.LastIndexOf("_");
+				if (index != -1)
+				{
+					index++;
+					string Number = FileName.Substring(index, FileName.Length - index);
+					AppendOrder = Convert.ToInt32(Number);
+				}
 
-			if(OrderType == OrderFileType.Cooker && AppendOrder == -1)
-            {
-				// make sure to sort CookerOpenOrder.log at the end since its produce by the cooking stage
-				AppendOrder = Int32.MaxValue;
+				if (OrderType == OrderFileType.Cooker && AppendOrder == -1)
+				{
+					// make sure to sort CookerOpenOrder.log at the end since its produce by the cooking stage
+					AppendOrder = Int32.MaxValue;
+				}
 			}
-		}
 
-		public int CompareTo(OrderFile Other)
-		{
-			int compare = (OrderSpecIndex.CompareTo(Other.OrderSpecIndex));
-			if (compare == 0)
+			public int CompareTo(OrderFile Other)
 			{
-				return (AppendOrder.CompareTo(Other.AppendOrder));
+				int compare = (OrderSpecIndex.CompareTo(Other.OrderSpecIndex));
+				if (compare == 0)
+				{
+					return (AppendOrder.CompareTo(Other.AppendOrder));
+				}
+				return compare;
 			}
-			return compare;
-		}
 
-		// Order files should be sorted by OrderSpecIndex and then by AppendOrder
-		// OrderType can be used for filtering
-		public int AppendOrder;
-		public int OrderSpecIndex;
-		// Priority at which this order file can claim files even if it appears later in the list than another order file
-		public int Priority;
+			// Order files should be sorted by OrderSpecIndex and then by AppendOrder
+			// OrderType can be used for filtering
+			public int AppendOrder;
+			public int OrderSpecIndex;
+			// Priority at which this order file can claim files even if it appears later in the list than another order file
+			public int Priority;
 
-		public OrderFileType OrderType;
+			public OrderFileType OrderType;
 
-		public FileReference File;
-	};
+			public FileReference File;
+		};
 
-	struct OrderFileSpec
-	{
-		public string FileNamePattern;
-		public OrderFile.OrderFileType OrderType;
-		public bool Required;
-		public int Priority;
-	};
-
-	/// <returns>The path for the BuildPatchTool executable depending on host platform.</returns>
-	private static string GetBuildPatchToolExecutable()
-	{
-		if (UnrealBuildTool.BuildHostPlatform.Current.Platform == UnrealTargetPlatform.Win64)
-		{
-			return CommandUtils.CombinePaths(CommandUtils.CmdEnv.LocalRoot, "Engine/Binaries/Win64/BuildPatchTool.exe");
-		}
-		if (UnrealBuildTool.BuildHostPlatform.Current.Platform == UnrealTargetPlatform.Mac)
+		struct OrderFileSpec
 		{
-			return CommandUtils.CombinePaths(CommandUtils.CmdEnv.LocalRoot, "Engine/Binaries/Mac/BuildPatchTool");
-		}
-		if (UnrealBuildTool.BuildHostPlatform.Current.Platform == UnrealTargetPlatform.Linux)
+			public string FileNamePattern;
+			public OrderFile.OrderFileType OrderType;
+			public bool Required;
+			public int Priority;
+		};
+
+		/// <returns>The path for the BuildPatchTool executable depending on host platform.</returns>
+		private static string GetBuildPatchToolExecutable()
 		{
-			return CommandUtils.CombinePaths(CommandUtils.CmdEnv.LocalRoot, "Engine/Binaries/Linux/BuildPatchTool");
-		}
+			if (UnrealBuildTool.BuildHostPlatform.Current.Platform == UnrealTargetPlatform.Win64)
+			{
+				return CommandUtils.CombinePaths(CommandUtils.CmdEnv.LocalRoot, "Engine/Binaries/Win64/BuildPatchTool.exe");
+			}
+			if (UnrealBuildTool.BuildHostPlatform.Current.Platform == UnrealTargetPlatform.Mac)
+			{
+				return CommandUtils.CombinePaths(CommandUtils.CmdEnv.LocalRoot, "Engine/Binaries/Mac/BuildPatchTool");
+			}
+			if (UnrealBuildTool.BuildHostPlatform.Current.Platform == UnrealTargetPlatform.Linux)
+			{
+				return CommandUtils.CombinePaths(CommandUtils.CmdEnv.LocalRoot, "Engine/Binaries/Linux/BuildPatchTool");
+			}
 
-		throw new AutomationException(string.Format("Unknown host platform for BuildPatchTool - {0}", UnrealBuildTool.BuildHostPlatform.Current.Platform));
-	}
+			throw new AutomationException(string.Format("Unknown host platform for BuildPatchTool - {0}", UnrealBuildTool.BuildHostPlatform.Current.Platform));
+		}
 
-	/// <summary>
-	/// Checks the existence of the BuildPatchTool executable exists and builds it if it is missing
-	/// </summary>
-	private static void EnsureBuildPatchToolExists()
-	{
-		string BuildPatchToolExe = GetBuildPatchToolExecutable();
-		if (!CommandUtils.FileExists_NoExceptions(BuildPatchToolExe))
+		/// <summary>
+		/// Checks the existence of the BuildPatchTool executable exists and builds it if it is missing
+		/// </summary>
+		private static void EnsureBuildPatchToolExists()
 		{
-			lock (SyncLock)
+			string BuildPatchToolExe = GetBuildPatchToolExecutable();
+			if (!CommandUtils.FileExists_NoExceptions(BuildPatchToolExe))
 			{
-				if (!CommandUtils.FileExists_NoExceptions(BuildPatchToolExe))
+				lock (SyncLock)
 				{
-					UEBuildUtils.BuildBuildPatchTool(null, UnrealBuildTool.BuildHostPlatform.Current.Platform);
+					if (!CommandUtils.FileExists_NoExceptions(BuildPatchToolExe))
+					{
+						UEBuildUtils.BuildBuildPatchTool(null, UnrealBuildTool.BuildHostPlatform.Current.Platform);
+					}
 				}
 			}
 		}
-	}
 
-	/// <summary>
-	/// Writes a pak response file to disk
-	/// </summary>
-	/// <param name="Filename"></param>
-	/// <param name="ResponseFile"></param>
-	private static void WritePakResponseFile(string Filename, Dictionary<string, string> ResponseFile, bool Compressed, EncryptionAndSigning.CryptoSettings CryptoSettings, bool bForceFullEncryption)
-	{
-		using (var Writer = new StreamWriter(Filename, false, new System.Text.UTF8Encoding(true)))
+		/// <summary>
+		/// Writes a pak response file to disk
+		/// </summary>
+		/// <param name="Filename"></param>
+		/// <param name="ResponseFile"></param>
+		private static void WritePakResponseFile(string Filename, Dictionary<string, string> ResponseFile, bool Compressed, EncryptionAndSigning.CryptoSettings CryptoSettings, bool bForceFullEncryption)
 		{
-			foreach (var Entry in ResponseFile)
+			using (var Writer = new StreamWriter(Filename, false, new System.Text.UTF8Encoding(true)))
 			{
-				string Line = String.Format("\"{0}\" \"{1}\"", Entry.Key, Entry.Value);
-
-				// explicitly exclude some file types from compression
-				if (Compressed && !Path.GetExtension(Entry.Key).Contains(".mp4") && !Path.GetExtension(Entry.Key).Contains("ushaderbytecode") && !Path.GetExtension(Entry.Key).Contains("upipelinecache"))
+				foreach (var Entry in ResponseFile)
 				{
-					Line += " -compress";
-				}
+					string Line = String.Format("\"{0}\" \"{1}\"", Entry.Key, Entry.Value);
 
-				if (CryptoSettings != null)
-				{
-					bool bEncryptFile = bForceFullEncryption || CryptoSettings.bEnablePakFullAssetEncryption;
-					bEncryptFile = bEncryptFile || (CryptoSettings.bEnablePakUAssetEncryption && Path.GetExtension(Entry.Key).Contains(".uasset"));
-					bEncryptFile = bEncryptFile || (CryptoSettings.bEnablePakIniEncryption && Path.GetExtension(Entry.Key).Contains(".ini"));
+					// explicitly exclude some file types from compression
+					if (Compressed && !Path.GetExtension(Entry.Key).Contains(".mp4") && !Path.GetExtension(Entry.Key).Contains("ushaderbytecode") && !Path.GetExtension(Entry.Key).Contains("upipelinecache"))
+					{
+						Line += " -compress";
+					}
 
-					if (bEncryptFile)
+					if (CryptoSettings != null)
 					{
-						Line += " -encrypt";
+						bool bEncryptFile = bForceFullEncryption || CryptoSettings.bEnablePakFullAssetEncryption;
+						bEncryptFile = bEncryptFile || (CryptoSettings.bEnablePakUAssetEncryption && Path.GetExtension(Entry.Key).Contains(".uasset"));
+						bEncryptFile = bEncryptFile || (CryptoSettings.bEnablePakIniEncryption && Path.GetExtension(Entry.Key).Contains(".ini"));
+
+						if (bEncryptFile)
+						{
+							Line += " -encrypt";
+						}
 					}
-				}
 
-				Writer.WriteLine(Line);
+					Writer.WriteLine(Line);
+				}
 			}
 		}
-	}
-
-	/// <summary>
-	/// Loads streaming install chunk manifest file from disk
-	/// </summary>
-	/// <param name="Filename"></param>
-	/// <returns></returns>
-	private static HashSet<string> ReadPakChunkManifest(string Filename)
-	{
-		var ResponseFile = ReadAllLines(Filename);
-		var Result = new HashSet<string>(ResponseFile, StringComparer.InvariantCultureIgnoreCase);
-		return Result;
-	}
 
-	static public string GetUnrealPakArguments(FileReference ProjectPath, Dictionary<string, string> UnrealPakResponseFile, FileReference OutputLocation, List<OrderFile> PakOrderFileLocations, string PlatformOptions, bool Compressed, EncryptionAndSigning.CryptoSettings CryptoSettings, FileReference CryptoKeysCacheFilename, String PatchSourceContentPath, string EncryptionKeyGuid, List<OrderFile> SecondaryPakOrderFileLocations)
-	{
-		StringBuilder CmdLine = new StringBuilder(MakePathSafeToUseWithCommandLine(ProjectPath.FullName));
-		CmdLine.AppendFormat(" {0}", MakePathSafeToUseWithCommandLine(OutputLocation.FullName));
-
-		// Force encryption of ALL files if we're using specific encryption key. This should be made an option per encryption key in the settings, but for our initial
-		// implementation we will just assume that we require maximum security for this data.
-		bool bForceEncryption = !string.IsNullOrEmpty(EncryptionKeyGuid);
-		string PakName = Path.GetFileNameWithoutExtension(OutputLocation.FullName);
-		string UnrealPakResponseFileName = CombinePaths(CmdEnv.LogFolder, "PakList_" + PakName + ".txt");
-		WritePakResponseFile(UnrealPakResponseFileName, UnrealPakResponseFile, Compressed, CryptoSettings, bForceEncryption);
-		CmdLine.AppendFormat(" -create={0}", CommandUtils.MakePathSafeToUseWithCommandLine(UnrealPakResponseFileName));
-
-		if (CryptoKeysCacheFilename != null)
-		{
-			CmdLine.AppendFormat(" -cryptokeys={0}", CommandUtils.MakePathSafeToUseWithCommandLine(CryptoKeysCacheFilename.FullName));
-		}
-		if (PakOrderFileLocations != null && PakOrderFileLocations.Count() > 0)
-		{
-			CmdLine.AppendFormat(" -order={0}", CommandUtils.MakePathSafeToUseWithCommandLine(string.Join(",", PakOrderFileLocations.Select(u => u.File.FullName).ToArray())));
-		}
-		if (SecondaryPakOrderFileLocations != null && SecondaryPakOrderFileLocations.Count() > 0)
-		{
-			CmdLine.AppendFormat(" -secondaryOrder={0}", CommandUtils.MakePathSafeToUseWithCommandLine(string.Join(",", SecondaryPakOrderFileLocations.Select(u => u.File.FullName).ToArray())));
-		}
-		if (!String.IsNullOrEmpty(PatchSourceContentPath)) 
+		/// <summary>
+		/// Loads streaming install chunk manifest file from disk
+		/// </summary>
+		/// <param name="Filename"></param>
+		/// <returns></returns>
+		private static HashSet<string> ReadPakChunkManifest(string Filename)
 		{
-			CmdLine.AppendFormat(" -generatepatch={0} -tempfiles={1}", CommandUtils.MakePathSafeToUseWithCommandLine(PatchSourceContentPath), CommandUtils.MakePathSafeToUseWithCommandLine(CommandUtils.CombinePaths(CommandUtils.CmdEnv.LocalRoot, "TempFiles" + Path.GetFileNameWithoutExtension(OutputLocation.FullName))));
+			var ResponseFile = ReadAllLines(Filename);
+			var Result = new HashSet<string>(ResponseFile, StringComparer.InvariantCultureIgnoreCase);
+			return Result;
 		}
-		if (CryptoSettings != null && CryptoSettings.bDataCryptoRequired)
+
+		static public string GetUnrealPakArguments(FileReference ProjectPath, Dictionary<string, string> UnrealPakResponseFile, FileReference OutputLocation, List<OrderFile> PakOrderFileLocations, string PlatformOptions, bool Compressed, EncryptionAndSigning.CryptoSettings CryptoSettings, FileReference CryptoKeysCacheFilename, String PatchSourceContentPath, string EncryptionKeyGuid, List<OrderFile> SecondaryPakOrderFileLocations)
 		{
-			if (CryptoSettings.bEnablePakIndexEncryption)
+			StringBuilder CmdLine = new StringBuilder(MakePathSafeToUseWithCommandLine(ProjectPath.FullName));
+			CmdLine.AppendFormat(" {0}", MakePathSafeToUseWithCommandLine(OutputLocation.FullName));
+
+			// Force encryption of ALL files if we're using specific encryption key. This should be made an option per encryption key in the settings, but for our initial
+			// implementation we will just assume that we require maximum security for this data.
+			bool bForceEncryption = !string.IsNullOrEmpty(EncryptionKeyGuid);
+			string PakName = Path.GetFileNameWithoutExtension(OutputLocation.FullName);
+			string UnrealPakResponseFileName = CombinePaths(CmdEnv.LogFolder, "PakList_" + PakName + ".txt");
+			WritePakResponseFile(UnrealPakResponseFileName, UnrealPakResponseFile, Compressed, CryptoSettings, bForceEncryption);
+			CmdLine.AppendFormat(" -create={0}", CommandUtils.MakePathSafeToUseWithCommandLine(UnrealPakResponseFileName));
+
+			if (CryptoKeysCacheFilename != null)
 			{
-				CmdLine.AppendFormat(" -encryptindex");
+				CmdLine.AppendFormat(" -cryptokeys={0}", CommandUtils.MakePathSafeToUseWithCommandLine(CryptoKeysCacheFilename.FullName));
 			}
-			if (!string.IsNullOrEmpty(EncryptionKeyGuid))
+			if (PakOrderFileLocations != null && PakOrderFileLocations.Count() > 0)
 			{
-				CmdLine.AppendFormat(" -EncryptionKeyOverrideGuid={0}", EncryptionKeyGuid);
+				CmdLine.AppendFormat(" -order={0}", CommandUtils.MakePathSafeToUseWithCommandLine(string.Join(",", PakOrderFileLocations.Select(u => u.File.FullName).ToArray())));
 			}
-			if (CryptoSettings.bDataCryptoRequired && CryptoSettings.bEnablePakSigning && CryptoSettings.SigningKey.IsValid())
+			if (SecondaryPakOrderFileLocations != null && SecondaryPakOrderFileLocations.Count() > 0)
 			{
-				CmdLine.AppendFormat(" -sign");
+				CmdLine.AppendFormat(" -secondaryOrder={0}", CommandUtils.MakePathSafeToUseWithCommandLine(string.Join(",", SecondaryPakOrderFileLocations.Select(u => u.File.FullName).ToArray())));
+			}
+			if (!String.IsNullOrEmpty(PatchSourceContentPath))
+			{
+				CmdLine.AppendFormat(" -generatepatch={0} -tempfiles={1}", CommandUtils.MakePathSafeToUseWithCommandLine(PatchSourceContentPath), CommandUtils.MakePathSafeToUseWithCommandLine(CommandUtils.CombinePaths(CommandUtils.CmdEnv.LocalRoot, "TempFiles" + Path.GetFileNameWithoutExtension(OutputLocation.FullName))));
+			}
+			if (CryptoSettings != null && CryptoSettings.bDataCryptoRequired)
+			{
+				if (CryptoSettings.bEnablePakIndexEncryption)
+				{
+					CmdLine.AppendFormat(" -encryptindex");
+				}
+				if (!string.IsNullOrEmpty(EncryptionKeyGuid))
+				{
+					CmdLine.AppendFormat(" -EncryptionKeyOverrideGuid={0}", EncryptionKeyGuid);
+				}
+				if (CryptoSettings.bDataCryptoRequired && CryptoSettings.bEnablePakSigning && CryptoSettings.SigningKey.IsValid())
+				{
+					CmdLine.AppendFormat(" -sign");
+				}
 			}
-		}
-		
-		CmdLine.Append(PlatformOptions);
 
-		return CmdLine.ToString();
-	}
+			CmdLine.Append(PlatformOptions);
 
-	static private string GetIoStoreCommandArguments(
-		Dictionary<string, string> UnrealPakResponseFile,
-		string ContainerName,
-		FileReference PakOutputLocation,
-		bool bCompressed,
-		EncryptionAndSigning.CryptoSettings CryptoSettings,
-		string EncryptionKeyGuid,
-		string PatchSourceContentPath,
-		bool bGenerateDiffPatch,
-		bool bIsDLC)
-	{
-		StringBuilder CmdLine = new StringBuilder();
-		CmdLine.AppendFormat("-Output={0}", MakePathSafeToUseWithCommandLine(Path.ChangeExtension(PakOutputLocation.FullName, ".utoc")));
-		CmdLine.AppendFormat("-ContainerName={0}", ContainerName);
-		if (!bIsDLC && !String.IsNullOrEmpty(PatchSourceContentPath))
-		{
-			CmdLine.AppendFormat(" -PatchSource={0}", CommandUtils.MakePathSafeToUseWithCommandLine(PatchSourceContentPath));
-		}
-		if (bGenerateDiffPatch)
-		{
-			CmdLine.Append(" -GenerateDiffPatch");
+			return CmdLine.ToString();
 		}
 
-		// Force encryption of ALL files if we're using specific encryption key. This should be made an option per encryption key in the settings, but for our initial
-		// implementation we will just assume that we require maximum security for this data.
-		bool bForceEncryption = !string.IsNullOrEmpty(EncryptionKeyGuid);
-		string UnrealPakResponseFileName = CombinePaths(CmdEnv.LogFolder, "PakListIoStore_" + ContainerName + ".txt");
-		WritePakResponseFile(UnrealPakResponseFileName, UnrealPakResponseFile, bCompressed, CryptoSettings, bForceEncryption);
-		CmdLine.AppendFormat(" -ResponseFile={0}", CommandUtils.MakePathSafeToUseWithCommandLine(UnrealPakResponseFileName));
-
-		if (CryptoSettings != null && CryptoSettings.bDataCryptoRequired)
+		static private string GetIoStoreCommandArguments(
+			Dictionary<string, string> UnrealPakResponseFile,
+			string ContainerName,
+			FileReference PakOutputLocation,
+			bool bCompressed,
+			EncryptionAndSigning.CryptoSettings CryptoSettings,
+			string EncryptionKeyGuid,
+			string PatchSourceContentPath,
+			bool bGenerateDiffPatch,
+			bool bIsDLC)
 		{
-			if (CryptoSettings.bEnablePakIndexEncryption)
+			StringBuilder CmdLine = new StringBuilder();
+			CmdLine.AppendFormat("-Output={0}", MakePathSafeToUseWithCommandLine(Path.ChangeExtension(PakOutputLocation.FullName, ".utoc")));
+			CmdLine.AppendFormat("-ContainerName={0}", ContainerName);
+			if (!bIsDLC && !String.IsNullOrEmpty(PatchSourceContentPath))
 			{
-				CmdLine.AppendFormat(" -encryptindex");
+				CmdLine.AppendFormat(" -PatchSource={0}", CommandUtils.MakePathSafeToUseWithCommandLine(PatchSourceContentPath));
 			}
-			if (!string.IsNullOrEmpty(EncryptionKeyGuid))
+			if (bGenerateDiffPatch)
 			{
-				CmdLine.AppendFormat(" -EncryptionKeyOverrideGuid={0}", EncryptionKeyGuid);
+				CmdLine.Append(" -GenerateDiffPatch");
 			}
-			if (CryptoSettings.bDataCryptoRequired && CryptoSettings.bEnablePakSigning && CryptoSettings.SigningKey.IsValid())
+
+			// Force encryption of ALL files if we're using specific encryption key. This should be made an option per encryption key in the settings, but for our initial
+			// implementation we will just assume that we require maximum security for this data.
+			bool bForceEncryption = !string.IsNullOrEmpty(EncryptionKeyGuid);
+			string UnrealPakResponseFileName = CombinePaths(CmdEnv.LogFolder, "PakListIoStore_" + ContainerName + ".txt");
+			WritePakResponseFile(UnrealPakResponseFileName, UnrealPakResponseFile, bCompressed, CryptoSettings, bForceEncryption);
+			CmdLine.AppendFormat(" -ResponseFile={0}", CommandUtils.MakePathSafeToUseWithCommandLine(UnrealPakResponseFileName));
+
+			if (CryptoSettings != null && CryptoSettings.bDataCryptoRequired)
 			{
-				CmdLine.AppendFormat(" -sign");
+				if (CryptoSettings.bEnablePakIndexEncryption)
+				{
+					CmdLine.AppendFormat(" -encryptindex");
+				}
+				if (!string.IsNullOrEmpty(EncryptionKeyGuid))
+				{
+					CmdLine.AppendFormat(" -EncryptionKeyOverrideGuid={0}", EncryptionKeyGuid);
+				}
+				if (CryptoSettings.bDataCryptoRequired && CryptoSettings.bEnablePakSigning && CryptoSettings.SigningKey.IsValid())
+				{
+					CmdLine.AppendFormat(" -sign");
+				}
 			}
-		}
-
-		return CmdLine.ToString();
-	}
 
-	public static FileReference GetUnrealPakLocation()
-	{
-		if (HostPlatform.Current.HostEditorPlatform == UnrealTargetPlatform.Win64)
-		{
-			return FileReference.Combine(Unreal.EngineDirectory, "Binaries", "Win64", "UnrealPak.exe");
-		}
-		else
-		{
-			return FileReference.Combine(Unreal.EngineDirectory, "Binaries", HostPlatform.Current.HostEditorPlatform.ToString(), "UnrealPak");
+			return CmdLine.ToString();
 		}
-	}
 
-	static public void RunUnrealPak(ProjectParams Params, Dictionary<string, string> UnrealPakResponseFile, FileReference OutputLocation, List<OrderFile> PakOrderFileLocations, string PlatformOptions, bool Compressed, EncryptionAndSigning.CryptoSettings CryptoSettings, FileReference CryptoKeysCacheFilename, String PatchSourceContentPath, string EncryptionKeyGuid, List<OrderFile> SecondaryPakOrderFileLocations = null)
-	{
-		if (UnrealPakResponseFile.Count < 1)
+		public static FileReference GetUnrealPakLocation()
 		{
-			return;
+			if (HostPlatform.Current.HostEditorPlatform == UnrealTargetPlatform.Win64)
+			{
+				return FileReference.Combine(Unreal.EngineDirectory, "Binaries", "Win64", "UnrealPak.exe");
+			}
+			else
+			{
+				return FileReference.Combine(Unreal.EngineDirectory, "Binaries", HostPlatform.Current.HostEditorPlatform.ToString(), "UnrealPak");
+			}
 		}
 
-		string Arguments = GetUnrealPakArguments(Params.RawProjectPath, UnrealPakResponseFile, OutputLocation, PakOrderFileLocations, PlatformOptions, Compressed, CryptoSettings, CryptoKeysCacheFilename, PatchSourceContentPath, EncryptionKeyGuid, SecondaryPakOrderFileLocations);
-		RunAndLog(CmdEnv, GetUnrealPakLocation().FullName, Arguments, Options: ERunOptions.Default | ERunOptions.UTF8Output);
-	}
-
-	static public void LogDeploymentContext(DeploymentContext SC)
-	{
-		LogLog("Deployment Context **************");
-		LogLog("ProjectFile = {0}", SC.RawProjectPath);
-		LogLog("ArchiveDir = {0}", SC.ArchiveDirectory);
-		LogLog("IsCodeBasedUprojectFile = {0}", SC.IsCodeBasedProject);
-		LogLog("DedicatedServer = {0}", SC.DedicatedServer);
-		LogLog("Stage = {0}", SC.Stage);
-		LogLog("StageTargetPlatform = {0}", SC.StageTargetPlatform.PlatformType.ToString());
-		LogLog("InputRootDir = {0}", SC.LocalRoot);
-		LogLog("InputProjectDir = {0}", SC.ProjectRoot);
-		LogLog("PlatformDir = {0}", SC.PlatformDir);
-		LogLog("StagedOutputDir = {0}", SC.StageDirectory);
-		LogLog("ShortProjectName = {0}", SC.ShortProjectName);
-		LogLog("ProjectArgForCommandLines = {0}", SC.ProjectArgForCommandLines);
-		LogLog("RunRootDir = {0}", SC.RuntimeRootDir);
-		LogLog("RunProjectDir = {0}", SC.RuntimeProjectRootDir);
-		LogLog("PakFileInternalRoot = {0}", SC.PakFileInternalRoot);
-		LogLog("PlatformUsesChunkManifests = {0}", SC.PlatformUsesChunkManifests);
-		LogLog("End Deployment Context **************");
-	}
-
-	private static string GetInternationalizationPreset(ProjectParams Params, ConfigHierarchy PlatformGameConfig, bool bMustExist = true)
-	{
-		// Initialize internationalization preset.
-		string InternationalizationPreset = Params.InternationalizationPreset;
-
-		// Use configuration if otherwise lacking an internationalization preset.
-		if (string.IsNullOrEmpty(InternationalizationPreset))
+		static public void RunUnrealPak(ProjectParams Params, Dictionary<string, string> UnrealPakResponseFile, FileReference OutputLocation, List<OrderFile> PakOrderFileLocations, string PlatformOptions, bool Compressed, EncryptionAndSigning.CryptoSettings CryptoSettings, FileReference CryptoKeysCacheFilename, String PatchSourceContentPath, string EncryptionKeyGuid, List<OrderFile> SecondaryPakOrderFileLocations = null)
 		{
-			if (PlatformGameConfig != null)
+			if (UnrealPakResponseFile.Count < 1)
 			{
-				PlatformGameConfig.GetString("/Script/UnrealEd.ProjectPackagingSettings", "InternationalizationPreset", out InternationalizationPreset);
+				return;
 			}
-		}
 
-		// Error if no preset has been provided.
-		if (bMustExist && string.IsNullOrEmpty(InternationalizationPreset))
-		{
-			throw new AutomationException("No internationalization preset was specified for packaging. This will lead to fatal errors when launching. Specify preset via commandline (-I18NPreset=) or project packaging settings (InternationalizationPreset).");
+			string Arguments = GetUnrealPakArguments(Params.RawProjectPath, UnrealPakResponseFile, OutputLocation, PakOrderFileLocations, PlatformOptions, Compressed, CryptoSettings, CryptoKeysCacheFilename, PatchSourceContentPath, EncryptionKeyGuid, SecondaryPakOrderFileLocations);
+			RunAndLog(CmdEnv, GetUnrealPakLocation().FullName, Arguments, Options: ERunOptions.Default | ERunOptions.UTF8Output);
 		}
 
-		return InternationalizationPreset;
-	}
-
-	private static List<string> GetCulturesToStage(ProjectParams Params, ConfigHierarchy PlatformGameConfig, bool bMustExist = true)
-	{
-		// Initialize cultures to stage.
-		List<string> CulturesToStage = null;
-
-		// Use parameters if provided.
-		if (Params.CulturesToCook != null && Params.CulturesToCook.Count > 0)
+		static public void LogDeploymentContext(DeploymentContext SC)
 		{
-			CulturesToStage = Params.CulturesToCook;
+			LogLog("Deployment Context **************");
+			LogLog("ProjectFile = {0}", SC.RawProjectPath);
+			LogLog("ArchiveDir = {0}", SC.ArchiveDirectory);
+			LogLog("IsCodeBasedUprojectFile = {0}", SC.IsCodeBasedProject);
+			LogLog("DedicatedServer = {0}", SC.DedicatedServer);
+			LogLog("Stage = {0}", SC.Stage);
+			LogLog("StageTargetPlatform = {0}", SC.StageTargetPlatform.PlatformType.ToString());
+			LogLog("InputRootDir = {0}", SC.LocalRoot);
+			LogLog("InputProjectDir = {0}", SC.ProjectRoot);
+			LogLog("PlatformDir = {0}", SC.PlatformDir);
+			LogLog("StagedOutputDir = {0}", SC.StageDirectory);
+			LogLog("ShortProjectName = {0}", SC.ShortProjectName);
+			LogLog("ProjectArgForCommandLines = {0}", SC.ProjectArgForCommandLines);
+			LogLog("RunRootDir = {0}", SC.RuntimeRootDir);
+			LogLog("RunProjectDir = {0}", SC.RuntimeProjectRootDir);
+			LogLog("PakFileInternalRoot = {0}", SC.PakFileInternalRoot);
+			LogLog("PlatformUsesChunkManifests = {0}", SC.PlatformUsesChunkManifests);
+			LogLog("End Deployment Context **************");
 		}
 
-		// Use configuration if otherwise lacking cultures to stage.
-		if (CulturesToStage == null || CulturesToStage.Count == 0)
+		private static string GetInternationalizationPreset(ProjectParams Params, ConfigHierarchy PlatformGameConfig, bool bMustExist = true)
 		{
-			if (PlatformGameConfig != null)
+			// Initialize internationalization preset.
+			string InternationalizationPreset = Params.InternationalizationPreset;
+
+			// Use configuration if otherwise lacking an internationalization preset.
+			if (string.IsNullOrEmpty(InternationalizationPreset))
 			{
-				PlatformGameConfig.GetArray("/Script/UnrealEd.ProjectPackagingSettings", "CulturesToStage", out CulturesToStage);
+				if (PlatformGameConfig != null)
+				{
+					PlatformGameConfig.GetString("/Script/UnrealEd.ProjectPackagingSettings", "InternationalizationPreset", out InternationalizationPreset);
+				}
 			}
-		}
 
-		// Error if no cultures have been provided.
-		if (bMustExist && (CulturesToStage == null || CulturesToStage.Count == 0))
-		{
-			throw new AutomationException("No cultures were specified for cooking and packaging. This will lead to fatal errors when launching. Specify culture codes via commandline (-CookCultures=) or using project packaging settings (+CulturesToStage).");
-		}
-
-		return CulturesToStage;
-	}
+			// Error if no preset has been provided.
+			if (bMustExist && string.IsNullOrEmpty(InternationalizationPreset))
+			{
+				throw new AutomationException("No internationalization preset was specified for packaging. This will lead to fatal errors when launching. Specify preset via commandline (-I18NPreset=) or project packaging settings (InternationalizationPreset).");
+			}
 
-	private static bool ShouldStageLocalizationTarget(DeploymentContext SC, ConfigHierarchy PlatformGameConfig, string LocalizationTarget)
-	{
-		if (SC.BlacklistLocalizationTargets.Contains(LocalizationTarget))
-		{
-			return false;
+			return InternationalizationPreset;
 		}
 
-		// Chunked targets can be added to the manifest during cook, so skip them during staging otherwise we'll overwrite the chunk 0 data with the non-chunked version
-		if (PlatformGameConfig != null)
+		private static List<string> GetCulturesToStage(ProjectParams Params, ConfigHierarchy PlatformGameConfig, bool bMustExist = true)
 		{
-			List<string> LocalizationTargetsToChunk = null;
-			PlatformGameConfig.GetArray("/Script/UnrealEd.ProjectPackagingSettings", "LocalizationTargetsToChunk", out LocalizationTargetsToChunk);
+			// Initialize cultures to stage.
+			List<string> CulturesToStage = null;
 
-			if (LocalizationTargetsToChunk != null && LocalizationTargetsToChunk.Contains(LocalizationTarget))
+			// Use parameters if provided.
+			if (Params.CulturesToCook != null && Params.CulturesToCook.Count > 0)
 			{
-				var CookedLocalizationTargetDirectory = DirectoryReference.Combine(SC.PlatformCookDir, SC.ShortProjectName, "Content", "Localization", LocalizationTarget);
-				return !DirectoryReference.Exists(CookedLocalizationTargetDirectory);
+				CulturesToStage = Params.CulturesToCook;
 			}
-		}
-
-		return true;
-	}
 
-	private static void StageLocalizationDataForPlugin(DeploymentContext SC, List<string> CulturesToStage, FileReference Plugin)
-	{
-		PluginDescriptor Descriptor = PluginDescriptor.FromFile(Plugin);
-		if (Descriptor.LocalizationTargets != null)
-		{
-			foreach (LocalizationTargetDescriptor LocalizationTarget in Descriptor.LocalizationTargets)
+			// Use configuration if otherwise lacking cultures to stage.
+			if (CulturesToStage == null || CulturesToStage.Count == 0)
 			{
-				if (LocalizationTarget.LoadingPolicy == LocalizationTargetDescriptorLoadingPolicy.Always || LocalizationTarget.LoadingPolicy == LocalizationTargetDescriptorLoadingPolicy.Game)
+				if (PlatformGameConfig != null)
 				{
-					DirectoryReference PluginLocTargetDirectory = DirectoryReference.Combine(Plugin.Directory, "Content", "Localization", LocalizationTarget.Name);
-					if (ShouldStageLocalizationTarget(SC, null, LocalizationTarget.Name) && DirectoryReference.Exists(PluginLocTargetDirectory))
-					{
-						StageLocalizationDataForTarget(SC, CulturesToStage, PluginLocTargetDirectory);
-					}
+					PlatformGameConfig.GetArray("/Script/UnrealEd.ProjectPackagingSettings", "CulturesToStage", out CulturesToStage);
 				}
 			}
-		}
-	}
 
-	private static void StageLocalizationDataForTarget(DeploymentContext SC, List<string> CulturesToStage, DirectoryReference SourceDirectory)
-	{
-		var PlatformSourceDirectory = new DirectoryReference(CombinePaths(SourceDirectory.FullName, "Platforms", ConfigHierarchy.GetIniPlatformName(SC.StageTargetPlatform.IniPlatformType)));
-		if (!DirectoryReference.Exists(PlatformSourceDirectory))
-		{
-			PlatformSourceDirectory = null;
-		}
+			// Error if no cultures have been provided.
+			if (bMustExist && (CulturesToStage == null || CulturesToStage.Count == 0))
+			{
+				throw new AutomationException("No cultures were specified for cooking and packaging. This will lead to fatal errors when launching. Specify culture codes via commandline (-CookCultures=) or using project packaging settings (+CulturesToStage).");
+			}
 
-		foreach (FileReference SourceFile in DirectoryReference.EnumerateFiles(SourceDirectory, "*.locmeta"))
-		{
-			SC.StageFile(StagedFileType.UFS, SourceFile);
+			return CulturesToStage;
 		}
 
-		foreach (string Culture in CulturesToStage)
+		private static bool ShouldStageLocalizationTarget(DeploymentContext SC, ConfigHierarchy PlatformGameConfig, string LocalizationTarget)
 		{
-			StageLocalizationDataForCulture(SC, Culture, SourceDirectory);
+			if (SC.BlacklistLocalizationTargets.Contains(LocalizationTarget))
+			{
+				return false;
+			}
 
-			if (PlatformSourceDirectory != null)
+			// Chunked targets can be added to the manifest during cook, so skip them during staging otherwise we'll overwrite the chunk 0 data with the non-chunked version
+			if (PlatformGameConfig != null)
 			{
-				StageLocalizationDataForCulture(SC, Culture, PlatformSourceDirectory);
+				List<string> LocalizationTargetsToChunk = null;
+				PlatformGameConfig.GetArray("/Script/UnrealEd.ProjectPackagingSettings", "LocalizationTargetsToChunk", out LocalizationTargetsToChunk);
+
+				if (LocalizationTargetsToChunk != null && LocalizationTargetsToChunk.Contains(LocalizationTarget))
+				{
+					var CookedLocalizationTargetDirectory = DirectoryReference.Combine(SC.PlatformCookDir, SC.ShortProjectName, "Content", "Localization", LocalizationTarget);
+					return !DirectoryReference.Exists(CookedLocalizationTargetDirectory);
+				}
 			}
+
+			return true;
 		}
-	}
 
-	private static void StageLocalizationDataForTargetsInDirectory(DeploymentContext SC, ConfigHierarchy PlatformGameConfig, List<string> CulturesToStage, DirectoryReference RootDirectory)
-	{
-		if (DirectoryReference.Exists(RootDirectory))
+		private static void StageLocalizationDataForPlugin(DeploymentContext SC, List<string> CulturesToStage, FileReference Plugin)
 		{
-			foreach (DirectoryReference LocTargetDirectory in DirectoryReference.EnumerateDirectories(RootDirectory))
+			PluginDescriptor Descriptor = PluginDescriptor.FromFile(Plugin);
+			if (Descriptor.LocalizationTargets != null)
 			{
-				if (ShouldStageLocalizationTarget(SC, PlatformGameConfig, LocTargetDirectory.MakeRelativeTo(RootDirectory)))
+				foreach (LocalizationTargetDescriptor LocalizationTarget in Descriptor.LocalizationTargets)
 				{
-					StageLocalizationDataForTarget(SC, CulturesToStage, LocTargetDirectory);
+					if (LocalizationTarget.LoadingPolicy == LocalizationTargetDescriptorLoadingPolicy.Always || LocalizationTarget.LoadingPolicy == LocalizationTargetDescriptorLoadingPolicy.Game)
+					{
+						DirectoryReference PluginLocTargetDirectory = DirectoryReference.Combine(Plugin.Directory, "Content", "Localization", LocalizationTarget.Name);
+						if (ShouldStageLocalizationTarget(SC, null, LocalizationTarget.Name) && DirectoryReference.Exists(PluginLocTargetDirectory))
+						{
+							StageLocalizationDataForTarget(SC, CulturesToStage, PluginLocTargetDirectory);
+						}
+					}
 				}
 			}
 		}
-	}
-
-	private static void StageLocalizationDataForCulture(DeploymentContext SC, string CultureName, DirectoryReference SourceDirectory)
-	{
-		CultureName = CultureName.Replace('-', '_');
-
-		string[] LocaleTags = CultureName.Replace('-', '_').Split('_');
 
-		List<string> PotentialParentCultures = new List<string>();
-
-		if (LocaleTags.Length > 0)
+		private static void StageLocalizationDataForTarget(DeploymentContext SC, List<string> CulturesToStage, DirectoryReference SourceDirectory)
 		{
-			if (LocaleTags.Length > 1 && LocaleTags.Length > 2)
+			var PlatformSourceDirectory = new DirectoryReference(CombinePaths(SourceDirectory.FullName, "Platforms", ConfigHierarchy.GetIniPlatformName(SC.StageTargetPlatform.IniPlatformType)));
+			if (!DirectoryReference.Exists(PlatformSourceDirectory))
 			{
-				PotentialParentCultures.Add(string.Join("_", LocaleTags[0], LocaleTags[1], LocaleTags[2]));
+				PlatformSourceDirectory = null;
 			}
-			if (LocaleTags.Length > 2)
+
+			foreach (FileReference SourceFile in DirectoryReference.EnumerateFiles(SourceDirectory, "*.locmeta"))
 			{
-				PotentialParentCultures.Add(string.Join("_", LocaleTags[0], LocaleTags[2]));
+				SC.StageFile(StagedFileType.UFS, SourceFile);
 			}
-			if (LocaleTags.Length > 1)
+
+			foreach (string Culture in CulturesToStage)
 			{
-				PotentialParentCultures.Add(string.Join("_", LocaleTags[0], LocaleTags[1]));
+				StageLocalizationDataForCulture(SC, Culture, SourceDirectory);
+
+				if (PlatformSourceDirectory != null)
+				{
+					StageLocalizationDataForCulture(SC, Culture, PlatformSourceDirectory);
+				}
 			}
-			PotentialParentCultures.Add(LocaleTags[0]);
 		}
 
-		foreach (DirectoryReference FoundDirectory in DirectoryReference.EnumerateDirectories(SourceDirectory, "*", SearchOption.TopDirectoryOnly))
+		private static void StageLocalizationDataForTargetsInDirectory(DeploymentContext SC, ConfigHierarchy PlatformGameConfig, List<string> CulturesToStage, DirectoryReference RootDirectory)
 		{
-			string DirectoryName = CommandUtils.GetLastDirectoryName(FoundDirectory.FullName);
-			string CanonicalizedPotentialCulture = DirectoryName.Replace('-', '_');
-
-			if (PotentialParentCultures.Contains(CanonicalizedPotentialCulture))
+			if (DirectoryReference.Exists(RootDirectory))
 			{
-				foreach (FileReference SourceFile in DirectoryReference.EnumerateFiles(FoundDirectory, "*.locres", SearchOption.AllDirectories))
+				foreach (DirectoryReference LocTargetDirectory in DirectoryReference.EnumerateDirectories(RootDirectory))
 				{
-					SC.StageFile(StagedFileType.UFS, SourceFile);
+					if (ShouldStageLocalizationTarget(SC, PlatformGameConfig, LocTargetDirectory.MakeRelativeTo(RootDirectory)))
+					{
+						StageLocalizationDataForTarget(SC, CulturesToStage, LocTargetDirectory);
+					}
 				}
 			}
 		}
-	}
 
-	private static void StageAdditionalDirectoriesFromConfig(DeploymentContext SC, DirectoryReference ProjectContentRoot, StagedDirectoryReference StageContentRoot, ConfigHierarchy PlatformGameConfig, bool bUFS, string ConfigKeyName)
-	{
-		List<string> ExtraDirs;
-		if (PlatformGameConfig.GetArray("/Script/UnrealEd.ProjectPackagingSettings", ConfigKeyName, out ExtraDirs))
+		private static void StageLocalizationDataForCulture(DeploymentContext SC, string CultureName, DirectoryReference SourceDirectory)
 		{
-			// Each string has the format '(Path="TheDirToStage")'
-			foreach (var PathStr in ExtraDirs)
+			CultureName = CultureName.Replace('-', '_');
+
+			string[] LocaleTags = CultureName.Replace('-', '_').Split('_');
+
+			List<string> PotentialParentCultures = new List<string>();
+
+			if (LocaleTags.Length > 0)
 			{
-				string RelativePath = null;
-				var PathParts = PathStr.Split('"');
-				if (PathParts.Length == 3)
+				if (LocaleTags.Length > 1 && LocaleTags.Length > 2)
 				{
-					RelativePath = PathParts[1];
+					PotentialParentCultures.Add(string.Join("_", LocaleTags[0], LocaleTags[1], LocaleTags[2]));
 				}
-				else if (PathParts.Length == 1)
+				if (LocaleTags.Length > 2)
+				{
+					PotentialParentCultures.Add(string.Join("_", LocaleTags[0], LocaleTags[2]));
+				}
+				if (LocaleTags.Length > 1)
 				{
-					RelativePath = PathParts[0];
+					PotentialParentCultures.Add(string.Join("_", LocaleTags[0], LocaleTags[1]));
 				}
-				if (RelativePath != null)
+				PotentialParentCultures.Add(LocaleTags[0]);
+			}
+
+			foreach (DirectoryReference FoundDirectory in DirectoryReference.EnumerateDirectories(SourceDirectory, "*", SearchOption.TopDirectoryOnly))
+			{
+				string DirectoryName = CommandUtils.GetLastDirectoryName(FoundDirectory.FullName);
+				string CanonicalizedPotentialCulture = DirectoryName.Replace('-', '_');
+
+				if (PotentialParentCultures.Contains(CanonicalizedPotentialCulture))
 				{
-					DirectoryReference InputDir = DirectoryReference.Combine(ProjectContentRoot, RelativePath);
-					if (Directory.Exists(InputDir.FullName)==false)
+					foreach (FileReference SourceFile in DirectoryReference.EnumerateFiles(FoundDirectory, "*.locres", SearchOption.AllDirectories))
 					{
-						CommandUtils.LogWarning("Unable to find directory \"{0}\" for staging, retrieved from \"/Script/UnrealEd.ProjectPackagingSettings\" \"{1}\"", InputDir.FullName, ConfigKeyName);
-						continue;
+						SC.StageFile(StagedFileType.UFS, SourceFile);
 					}
-					StagedDirectoryReference OutputDir = StagedDirectoryReference.Combine(StageContentRoot, RelativePath);
-					if (bUFS)
+				}
+			}
+		}
+
+		private static void StageAdditionalDirectoriesFromConfig(DeploymentContext SC, DirectoryReference ProjectContentRoot, StagedDirectoryReference StageContentRoot, ConfigHierarchy PlatformGameConfig, bool bUFS, string ConfigKeyName)
+		{
+			List<string> ExtraDirs;
+			if (PlatformGameConfig.GetArray("/Script/UnrealEd.ProjectPackagingSettings", ConfigKeyName, out ExtraDirs))
+			{
+				// Each string has the format '(Path="TheDirToStage")'
+				foreach (var PathStr in ExtraDirs)
+				{
+					string RelativePath = null;
+					var PathParts = PathStr.Split('"');
+					if (PathParts.Length == 3)
 					{
-						List<FileReference> Files = SC.FindFilesToStage(InputDir, StageFilesSearch.AllDirectories);
-						Files.RemoveAll(x => x.HasExtension(".uasset") || x.HasExtension(".umap") || (SC.DedicatedServer && x.HasExtension(".mp4")));
-						SC.StageFiles(StagedFileType.UFS, InputDir, Files, OutputDir);
+						RelativePath = PathParts[1];
 					}
-					else
+					else if (PathParts.Length == 1)
 					{
-						SC.StageFiles(StagedFileType.NonUFS, InputDir, StageFilesSearch.AllDirectories, OutputDir);
+						RelativePath = PathParts[0];
+					}
+					if (RelativePath != null)
+					{
+						DirectoryReference InputDir = DirectoryReference.Combine(ProjectContentRoot, RelativePath);
+						if (Directory.Exists(InputDir.FullName) == false)
+						{
+							CommandUtils.LogWarning("Unable to find directory \"{0}\" for staging, retrieved from \"/Script/UnrealEd.ProjectPackagingSettings\" \"{1}\"", InputDir.FullName, ConfigKeyName);
+							continue;
+						}
+						StagedDirectoryReference OutputDir = StagedDirectoryReference.Combine(StageContentRoot, RelativePath);
+						if (bUFS)
+						{
+							List<FileReference> Files = SC.FindFilesToStage(InputDir, StageFilesSearch.AllDirectories);
+							Files.RemoveAll(x => x.HasExtension(".uasset") || x.HasExtension(".umap") || (SC.DedicatedServer && x.HasExtension(".mp4")));
+							SC.StageFiles(StagedFileType.UFS, InputDir, Files, OutputDir);
+						}
+						else
+						{
+							SC.StageFiles(StagedFileType.NonUFS, InputDir, StageFilesSearch.AllDirectories, OutputDir);
+						}
 					}
 				}
 			}
 		}
-	}
-
-	/// <summary>
-	/// Figures out which plugins are enabled for a content-only plugins. This may include content-only plugins and editor plugins, which will not be in the default 
-	/// target manifest, since it was compiled elsewhere.
-	/// </summary>
-	/// <param name="ProjectFile">The project being built</param>
-	/// <param name="Targets">List of targets being staged</param>
-	/// <returns>List of plugin files that should be staged</returns>
-	private static List<FileReference> GetPluginsForContentProject(FileReference ProjectFile, List<TargetReceipt> Targets)
-	{
-		ProjectDescriptor Project = ProjectDescriptor.FromFile(ProjectFile);
-		List<PluginInfo> AvailablePlugins = UnrealBuildTool.Plugins.ReadAvailablePlugins(Unreal.EngineDirectory, ProjectFile.Directory, null);
 
-		HashSet<FileReference> Plugins = new HashSet<FileReference>();
-		foreach (TargetReceipt Target in Targets)
+		/// <summary>
+		/// Figures out which plugins are enabled for a content-only plugins. This may include content-only plugins and editor plugins, which will not be in the default 
+		/// target manifest, since it was compiled elsewhere.
+		/// </summary>
+		/// <param name="ProjectFile">The project being built</param>
+		/// <param name="Targets">List of targets being staged</param>
+		/// <returns>List of plugin files that should be staged</returns>
+		private static List<FileReference> GetPluginsForContentProject(FileReference ProjectFile, List<TargetReceipt> Targets)
 		{
-			// Find all the specifically enabled plugins for this target
-			Dictionary<string, bool> EnabledPlugins = new Dictionary<string, bool>(StringComparer.OrdinalIgnoreCase);
-			if (Project.Plugins != null)
+			ProjectDescriptor Project = ProjectDescriptor.FromFile(ProjectFile);
+			List<PluginInfo> AvailablePlugins = UnrealBuildTool.Plugins.ReadAvailablePlugins(Unreal.EngineDirectory, ProjectFile.Directory, null);
+
+			HashSet<FileReference> Plugins = new HashSet<FileReference>();
+			foreach (TargetReceipt Target in Targets)
 			{
-				foreach (PluginReferenceDescriptor Reference in Project.Plugins)
+				// Find all the specifically enabled plugins for this target
+				Dictionary<string, bool> EnabledPlugins = new Dictionary<string, bool>(StringComparer.OrdinalIgnoreCase);
+				if (Project.Plugins != null)
 				{
-					bool bEnabled = false;
-					if (Reference.IsEnabledForPlatform(Target.Platform) && Reference.IsEnabledForTargetConfiguration(Target.Configuration) && Reference.IsEnabledForTarget(Target.TargetType))
+					foreach (PluginReferenceDescriptor Reference in Project.Plugins)
 					{
-						bEnabled = true;
+						bool bEnabled = false;
+						if (Reference.IsEnabledForPlatform(Target.Platform) && Reference.IsEnabledForTargetConfiguration(Target.Configuration) && Reference.IsEnabledForTarget(Target.TargetType))
+						{
+							bEnabled = true;
+						}
+						EnabledPlugins[Reference.Name] = bEnabled;
 					}
-					EnabledPlugins[Reference.Name] = bEnabled;
 				}
-			}
 
-			// Add all the enabled-by-default plugins
-			foreach (PluginInfo AvailablePlugin in AvailablePlugins)
-			{
-				if (!EnabledPlugins.ContainsKey(AvailablePlugin.Name))
+				// Add all the enabled-by-default plugins
+				foreach (PluginInfo AvailablePlugin in AvailablePlugins)
 				{
-					EnabledPlugins[AvailablePlugin.Name] = AvailablePlugin.Descriptor.SupportsTargetPlatform(Target.Platform) && AvailablePlugin.IsEnabledByDefault(!Project.DisableEnginePluginsByDefault);
+					if (!EnabledPlugins.ContainsKey(AvailablePlugin.Name))
+					{
+						EnabledPlugins[AvailablePlugin.Name] = AvailablePlugin.Descriptor.SupportsTargetPlatform(Target.Platform) && AvailablePlugin.IsEnabledByDefault(!Project.DisableEnginePluginsByDefault);
+					}
 				}
-			}
 
-			// Add the enabled plugins
-			foreach (PluginInfo AvailablePlugin in AvailablePlugins)
-			{
-				if (AvailablePlugin.Descriptor.SupportsTargetPlatform(Target.Platform) && EnabledPlugins[AvailablePlugin.Name])
+				// Add the enabled plugins
+				foreach (PluginInfo AvailablePlugin in AvailablePlugins)
 				{
-					Plugins.Add(AvailablePlugin.File);
+					if (AvailablePlugin.Descriptor.SupportsTargetPlatform(Target.Platform) && EnabledPlugins[AvailablePlugin.Name])
+					{
+						Plugins.Add(AvailablePlugin.File);
+					}
 				}
-			}
 
-			// Exclude all the plugins that would have been considered for the base target (editor only plugins that are enabled by default)
-			foreach (PluginInfo AvailablePlugin in AvailablePlugins)
-			{
-				if (AvailablePlugin.LoadedFrom == PluginLoadedFrom.Engine && !AvailablePlugin.Descriptor.bInstalled && AvailablePlugin.IsEnabledByDefault(!Project.DisableEnginePluginsByDefault))
+				// Exclude all the plugins that would have been considered for the base target (editor only plugins that are enabled by default)
+				foreach (PluginInfo AvailablePlugin in AvailablePlugins)
 				{
-					Plugins.Remove(AvailablePlugin.File);
+					if (AvailablePlugin.LoadedFrom == PluginLoadedFrom.Engine && !AvailablePlugin.Descriptor.bInstalled && AvailablePlugin.IsEnabledByDefault(!Project.DisableEnginePluginsByDefault))
+					{
+						Plugins.Remove(AvailablePlugin.File);
+					}
 				}
-			}
 
-			// Exclude all the existing plugin descriptors
-			Plugins.ExceptWith(Target.RuntimeDependencies.Select(x => x.Path));
+				// Exclude all the existing plugin descriptors
+				Plugins.ExceptWith(Target.RuntimeDependencies.Select(x => x.Path));
+			}
+			return Plugins.ToList();
 		}
-		return Plugins.ToList();
-	}
 
-	private class PackageStoreManifest
-	{
-		public static IList<FileReference> GetCookedFiles(FileReference PackageStoreManifestFile)
+		private class PackageStoreManifest
 		{
-			List<FileReference> CookedFiles = new List<FileReference>();
-			if (FileReference.Exists(PackageStoreManifestFile))
+			public static IList<FileReference> GetCookedFiles(FileReference PackageStoreManifestFile)
 			{
-				String JsonString = FileReference.ReadAllText(PackageStoreManifestFile);
-				PackageStoreManifest PackageStoreManifest = JsonSerializer.Deserialize<PackageStoreManifest>(JsonString);
-				foreach (PackageStoreManifest.FileInfo File in PackageStoreManifest.Files)
+				List<FileReference> CookedFiles = new List<FileReference>();
+				if (FileReference.Exists(PackageStoreManifestFile))
 				{
-					CookedFiles.Add(new FileReference(File.Path));
+					String JsonString = FileReference.ReadAllText(PackageStoreManifestFile);
+					PackageStoreManifest PackageStoreManifest = JsonSerializer.Deserialize<PackageStoreManifest>(JsonString);
+					foreach (PackageStoreManifest.FileInfo File in PackageStoreManifest.Files)
+					{
+						CookedFiles.Add(new FileReference(File.Path));
+					}
 				}
+				return CookedFiles;
 			}
-			return CookedFiles;
-		}
-		public IList<FileInfo> Files { get; set; }
+			public IList<FileInfo> Files { get; set; }
 
-		public class FileInfo
-		{
-			public String Path { get; set; }
+			public class FileInfo
+			{
+				public String Path { get; set; }
+			}
 		}
-	}
 
-	public static void CreateStagingManifest(ProjectParams Params, DeploymentContext SC)
-	{
-		if (!Params.Stage)
+		public static void CreateStagingManifest(ProjectParams Params, DeploymentContext SC)
 		{
-			return;
-		}
-		var ThisPlatform = SC.StageTargetPlatform;
-
-		LogInformation("Creating Staging Manifest...");
-
-		ConfigHierarchy PlatformGameConfig = ConfigCache.ReadHierarchy(ConfigHierarchyType.Game, DirectoryReference.FromFile(Params.RawProjectPath), SC.StageTargetPlatform.IniPlatformType);
+			if (!Params.Stage)
+			{
+				return;
+			}
+			var ThisPlatform = SC.StageTargetPlatform;
 
-		if (Params.HasIterateSharedCookedBuild)
-		{
-			// can't do shared cooked builds with DLC that's madness!!
-			//check( Params.HasDLCName == false );
+			LogInformation("Creating Staging Manifest...");
 
-			// stage all the previously staged files
-			SC.StageFiles(StagedFileType.NonUFS, DirectoryReference.Combine(SC.ProjectRoot, "Saved", "SharedIterativeBuild", SC.CookPlatform, "Staged"), StageFilesSearch.AllDirectories, StagedDirectoryReference.Root); // remap to the root directory
-		}
-		bool bCreatePluginManifest = false;
-		if (Params.HasDLCName)
-		{
-			// Making a plugin
-			DirectoryReference DLCRoot = Params.DLCFile.Directory;
-			string RelativeDLCRootPath = (Params.DLCOverrideCookedSubDir == null) ? DLCRoot.MakeRelativeTo(SC.LocalRoot) : Params.DLCOverrideCookedSubDir;
+			ConfigHierarchy PlatformGameConfig = ConfigCache.ReadHierarchy(ConfigHierarchyType.Game, DirectoryReference.FromFile(Params.RawProjectPath), SC.StageTargetPlatform.IniPlatformType);
 
-			// Put all of the cooked dir into the staged dir
-			if (String.IsNullOrEmpty(Params.CookOutputDir))
+			if (Params.HasIterateSharedCookedBuild)
 			{
-				SC.PlatformCookDir = DirectoryReference.Combine(DLCRoot, "Saved", "Cooked", SC.CookPlatform);
+				// can't do shared cooked builds with DLC that's madness!!
+				//check( Params.HasDLCName == false );
+
+				// stage all the previously staged files
+				SC.StageFiles(StagedFileType.NonUFS, DirectoryReference.Combine(SC.ProjectRoot, "Saved", "SharedIterativeBuild", SC.CookPlatform, "Staged"), StageFilesSearch.AllDirectories, StagedDirectoryReference.Root); // remap to the root directory
 			}
-			else
+			bool bCreatePluginManifest = false;
+			if (Params.HasDLCName)
 			{
-				SC.PlatformCookDir = new DirectoryReference(Params.CookOutputDir);
-				if (!SC.PlatformCookDir.GetDirectoryName().Equals(SC.CookPlatform, StringComparison.InvariantCultureIgnoreCase))
-				{
-					SC.PlatformCookDir = DirectoryReference.Combine(SC.PlatformCookDir, SC.CookPlatform);
-				}
-			}
-
-			DirectoryReference PlatformEngineDir = DirectoryReference.Combine(SC.PlatformCookDir, "Engine");
-			DirectoryReference ProjectMetadataDir = DirectoryReference.Combine(SC.PlatformCookDir, SC.ShortProjectName, "Metadata");
-			SC.MetadataDir = DirectoryReference.Combine(SC.PlatformCookDir, RelativeDLCRootPath, "Metadata");
+				// Making a plugin
+				DirectoryReference DLCRoot = Params.DLCFile.Directory;
+				string RelativeDLCRootPath = (Params.DLCOverrideCookedSubDir == null) ? DLCRoot.MakeRelativeTo(SC.LocalRoot) : Params.DLCOverrideCookedSubDir;
 
-			// The .uplugin file is staged differently for different DLC
-			// The .uplugin file doesn't actually exist for mobile DLC
-			if (FileReference.Exists(Params.DLCFile))
-			{
-				StagedFileType StageType = Params.DLCPakPluginFile ? StagedFileType.UFS : StagedFileType.NonUFS;
-				if (Params.DLCOverrideCookedSubDir != null)
+				// Put all of the cooked dir into the staged dir
+				if (String.IsNullOrEmpty(Params.CookOutputDir))
 				{
-					SC.StageFile(StageType, Params.DLCFile, new StagedFileReference(Path.Combine(Params.DLCOverrideCookedSubDir, Params.DLCFile.GetFileName())));
+					SC.PlatformCookDir = DirectoryReference.Combine(DLCRoot, "Saved", "Cooked", SC.CookPlatform);
 				}
 				else
 				{
-					SC.StageFile(StageType, Params.DLCFile);
-				}
-
-				// Stage DLC localization targets
-				List<string> CulturesToStage = GetCulturesToStage(Params, PlatformGameConfig, false);
-				if (CulturesToStage != null)
-				{
-					StageLocalizationDataForPlugin(SC, CulturesToStage, Params.DLCFile);
+					SC.PlatformCookDir = new DirectoryReference(Params.CookOutputDir);
+					if (!SC.PlatformCookDir.GetDirectoryName().Equals(SC.CookPlatform, StringComparison.InvariantCultureIgnoreCase))
+					{
+						SC.PlatformCookDir = DirectoryReference.Combine(SC.PlatformCookDir, SC.CookPlatform);
+					}
 				}
-			}
-
-			// Put the binaries into the staged dir
-			DirectoryReference BinariesDir = DirectoryReference.Combine(DLCRoot, "Binaries");
-			if (DirectoryReference.Exists(BinariesDir))
-			{
-				SC.StageFiles(StagedFileType.NonUFS, BinariesDir, "libUE4-*.so", StageFilesSearch.AllDirectories);
-				SC.StageFiles(StagedFileType.NonUFS, BinariesDir, "UE4-*.dll", StageFilesSearch.AllDirectories);
-				SC.StageFiles(StagedFileType.NonUFS, BinariesDir, "libUE4Server-*.so", StageFilesSearch.AllDirectories);
-				SC.StageFiles(StagedFileType.NonUFS, BinariesDir, "UnrealServer-*.dll", StageFilesSearch.AllDirectories);
-			}
-
-			// Put the config files into the staged dir
-			DirectoryReference ConfigDir = DirectoryReference.Combine(DLCRoot, "Config");
-			if (DirectoryReference.Exists(ConfigDir))
-			{
-				SC.StageFiles(StagedFileType.UFS, ConfigDir, "*.ini", StageFilesSearch.AllDirectories);
-			}
-
-			if (Params.DLCActLikePatch)
-			{
-				DirectoryReference DLCContent = DirectoryReference.Combine(DLCRoot, "Content");
-				string RelativeDLCContentPath = DLCContent.MakeRelativeTo(SC.LocalRoot);
-				string RelativeRootContentPath = DirectoryReference.Combine(SC.ProjectRoot, "Content").MakeRelativeTo(SC.LocalRoot);
 
-				SC.RemapDirectories.Add(Tuple.Create(new StagedDirectoryReference(RelativeDLCContentPath), new StagedDirectoryReference(RelativeRootContentPath)));
-			}
+				DirectoryReference PlatformEngineDir = DirectoryReference.Combine(SC.PlatformCookDir, "Engine");
+				DirectoryReference ProjectMetadataDir = DirectoryReference.Combine(SC.PlatformCookDir, SC.ShortProjectName, "Metadata");
+				SC.MetadataDir = DirectoryReference.Combine(SC.PlatformCookDir, RelativeDLCRootPath, "Metadata");
 
-			// Stage all the cooked data, this is the same rule as normal stage except we may skip Engine
-			List<FileReference> CookedFiles = DirectoryReference.EnumerateFiles(SC.PlatformCookDir, "*", SearchOption.AllDirectories).ToList();
+				// The .uplugin file is staged differently for different DLC
+				// The .uplugin file doesn't actually exist for mobile DLC
+				if (FileReference.Exists(Params.DLCFile))
+				{
+					StagedFileType StageType = Params.DLCPakPluginFile ? StagedFileType.UFS : StagedFileType.NonUFS;
+					if (Params.DLCOverrideCookedSubDir != null)
+					{
+						SC.StageFile(StageType, Params.DLCFile, new StagedFileReference(Path.Combine(Params.DLCOverrideCookedSubDir, Params.DLCFile.GetFileName())));
+					}
+					else
+					{
+						SC.StageFile(StageType, Params.DLCFile);
+					}
 
-			// When cooking directly to I/O store contaier files look for the package store manifest
-			if (Params.IoStore)
-			{
-				CookedFiles.AddRange(PackageStoreManifest.GetCookedFiles(FileReference.Combine(SC.MetadataDir, "packagestore.manifest")));
-			}
+					// Stage DLC localization targets
+					List<string> CulturesToStage = GetCulturesToStage(Params, PlatformGameConfig, false);
+					if (CulturesToStage != null)
+					{
+						StageLocalizationDataForPlugin(SC, CulturesToStage, Params.DLCFile);
+					}
+				}
 
-			foreach (FileReference CookedFile in CookedFiles)
-			{
-				// Skip metadata directory
-				if (CookedFile.Directory.IsUnderDirectory(SC.MetadataDir) || CookedFile.Directory.IsUnderDirectory(ProjectMetadataDir))
+				// Put the binaries into the staged dir
+				DirectoryReference BinariesDir = DirectoryReference.Combine(DLCRoot, "Binaries");
+				if (DirectoryReference.Exists(BinariesDir))
 				{
-					continue;
+					SC.StageFiles(StagedFileType.NonUFS, BinariesDir, "libUE4-*.so", StageFilesSearch.AllDirectories);
+					SC.StageFiles(StagedFileType.NonUFS, BinariesDir, "UE4-*.dll", StageFilesSearch.AllDirectories);
+					SC.StageFiles(StagedFileType.NonUFS, BinariesDir, "libUE4Server-*.so", StageFilesSearch.AllDirectories);
+					SC.StageFiles(StagedFileType.NonUFS, BinariesDir, "UnrealServer-*.dll", StageFilesSearch.AllDirectories);
 				}
 
-				// Dedicated server cook doesn't save shaders so no Engine dir is created
-				if ((!SC.DedicatedServer) && (!Params.DLCIncludeEngineContent) && CookedFile.Directory.IsUnderDirectory(PlatformEngineDir))
+				// Put the config files into the staged dir
+				DirectoryReference ConfigDir = DirectoryReference.Combine(DLCRoot, "Config");
+				if (DirectoryReference.Exists(ConfigDir))
 				{
-					continue;
+					SC.StageFiles(StagedFileType.UFS, ConfigDir, "*.ini", StageFilesSearch.AllDirectories);
 				}
 
-				// Skip region files. These are cook-time metadata files that never need to ship.
-				if (CookedFile.HasExtension(FileRegion.RegionsFileExtension))
+				if (Params.DLCActLikePatch)
 				{
-					continue;
+					DirectoryReference DLCContent = DirectoryReference.Combine(DLCRoot, "Content");
+					string RelativeDLCContentPath = DLCContent.MakeRelativeTo(SC.LocalRoot);
+					string RelativeRootContentPath = DirectoryReference.Combine(SC.ProjectRoot, "Content").MakeRelativeTo(SC.LocalRoot);
+
+					SC.RemapDirectories.Add(Tuple.Create(new StagedDirectoryReference(RelativeDLCContentPath), new StagedDirectoryReference(RelativeRootContentPath)));
 				}
 
-				// json files have never been staged
-				// metallib files cannot *currently* be staged as UFS as the Metal API needs to mmap them from files on disk in order to function efficiently
-				if (!CookedFile.HasExtension(".json") && !CookedFile.HasExtension(".metallib") && !CookedFile.HasExtension(".utoc") && !CookedFile.HasExtension(".ucas"))
+				// Stage all the cooked data, this is the same rule as normal stage except we may skip Engine
+				List<FileReference> CookedFiles = DirectoryReference.EnumerateFiles(SC.PlatformCookDir, "*", SearchOption.AllDirectories).ToList();
+
+				// When cooking directly to I/O store contaier files look for the package store manifest
+				if (Params.IoStore)
 				{
-					SC.StageFile(StagedFileType.UFS, CookedFile, new StagedFileReference(CookedFile.MakeRelativeTo(SC.PlatformCookDir)));
+					CookedFiles.AddRange(PackageStoreManifest.GetCookedFiles(FileReference.Combine(SC.MetadataDir, "packagestore.manifest")));
 				}
-			}
 
-			FileReference PluginSettingsFile = FileReference.Combine(DLCRoot, "Config", "PluginSettings.ini");
-			if (FileReference.Exists(PluginSettingsFile))
-			{
-				ConfigFile File = new ConfigFile(PluginSettingsFile);
-				ConfigFileSection StageSettings;
-				if (File.TryGetSection("StageSettings", out StageSettings))
+				foreach (FileReference CookedFile in CookedFiles)
 				{
-					foreach (ConfigLine Line in StageSettings.Lines)
+					// Skip metadata directory
+					if (CookedFile.Directory.IsUnderDirectory(SC.MetadataDir) || CookedFile.Directory.IsUnderDirectory(ProjectMetadataDir))
 					{
-						if (Line.Key == "AdditionalNonUSFDirectories")
-						{
-							SC.StageFiles(StagedFileType.NonUFS, DirectoryReference.Combine(DLCRoot, Line.Value), "*.*", StageFilesSearch.AllDirectories);
-						}
+						continue;
+					}
+
+					// Dedicated server cook doesn't save shaders so no Engine dir is created
+					if ((!SC.DedicatedServer) && (!Params.DLCIncludeEngineContent) && CookedFile.Directory.IsUnderDirectory(PlatformEngineDir))
+					{
+						continue;
+					}
+
+					// Skip region files. These are cook-time metadata files that never need to ship.
+					if (CookedFile.HasExtension(FileRegion.RegionsFileExtension))
+					{
+						continue;
+					}
+
+					// json files have never been staged
+					// metallib files cannot *currently* be staged as UFS as the Metal API needs to mmap them from files on disk in order to function efficiently
+					if (!CookedFile.HasExtension(".json") && !CookedFile.HasExtension(".metallib") && !CookedFile.HasExtension(".utoc") && !CookedFile.HasExtension(".ucas"))
+					{
+						SC.StageFile(StagedFileType.UFS, CookedFile, new StagedFileReference(CookedFile.MakeRelativeTo(SC.PlatformCookDir)));
 					}
 				}
 
-				if (SC.DedicatedServer)
+				FileReference PluginSettingsFile = FileReference.Combine(DLCRoot, "Config", "PluginSettings.ini");
+				if (FileReference.Exists(PluginSettingsFile))
 				{
-					ConfigFileSection StageSettingsServer;
-					if (File.TryGetSection("StageSettingsServer", out StageSettingsServer))
+					ConfigFile File = new ConfigFile(PluginSettingsFile);
+					ConfigFileSection StageSettings;
+					if (File.TryGetSection("StageSettings", out StageSettings))
 					{
-						foreach (ConfigLine Line in StageSettingsServer.Lines)
+						foreach (ConfigLine Line in StageSettings.Lines)
 						{
 							if (Line.Key == "AdditionalNonUSFDirectories")
 							{
@@ -785,3661 +773,3676 @@ public partial class Project : CommandUtils
 							}
 						}
 					}
-				}
-			}
 
-			bCreatePluginManifest = true;
-		}
-		else
-		{
-
-			if (!Params.IterateSharedBuildUsePrecompiledExe)
-			{
-				ThisPlatform.GetFilesToDeployOrStage(Params, SC);
-
-				// Stage any extra runtime dependencies from the receipts
-				foreach (StageTarget Target in SC.StageTargets)
-				{
-					SC.StageRuntimeDependenciesFromReceipt(Target.Receipt, Target.RequireFilesExist, Params.UsePak(SC.StageTargetPlatform));
+					if (SC.DedicatedServer)
+					{
+						ConfigFileSection StageSettingsServer;
+						if (File.TryGetSection("StageSettingsServer", out StageSettingsServer))
+						{
+							foreach (ConfigLine Line in StageSettingsServer.Lines)
+							{
+								if (Line.Key == "AdditionalNonUSFDirectories")
+								{
+									SC.StageFiles(StagedFileType.NonUFS, DirectoryReference.Combine(DLCRoot, Line.Value), "*.*", StageFilesSearch.AllDirectories);
+								}
+							}
+						}
+					}
 				}
 
-				// Stage any content-only plugins for content-only projects. We don't have a custom executable for these.
-				if (!Params.IsCodeBasedProject)
-				{
-					List<FileReference> PluginFiles = GetPluginsForContentProject(Params.RawProjectPath, SC.StageTargets.ConvertAll(x => x.Receipt));
-					SC.StageFiles(StagedFileType.UFS, PluginFiles);
-				}
+				bCreatePluginManifest = true;
 			}
-
-			// move the UECommandline.txt file to the root of the stage
-			// this file needs to be treated as a UFS file for casing, but NonUFS for being put into the .pak file
-			// @todo: Maybe there should be a new category - UFSNotForPak - and reevaluate all SystemNonUFS references
-			FileReference CommandLineFile = FileReference.Combine(GetIntermediateCommandlineDir(SC), "UECommandLine.txt");
-			if (FileReference.Exists(CommandLineFile))
+			else
 			{
-				StagedFileReference StagedCommandLineFile = new StagedFileReference("UECommandLine.txt");
-				if (SC.StageTargetPlatform.DeployLowerCaseFilenames(StagedFileType.SystemNonUFS))
+
+				if (!Params.IterateSharedBuildUsePrecompiledExe)
 				{
-					StagedCommandLineFile = StagedCommandLineFile.ToLowerInvariant();
-				}
-				SC.StageFile(StagedFileType.SystemNonUFS, CommandLineFile, StagedCommandLineFile);
-			}
+					ThisPlatform.GetFilesToDeployOrStage(Params, SC);
 
-			DirectoryReference ProjectContentRoot = DirectoryReference.Combine(SC.ProjectRoot, "Content");
-			StagedDirectoryReference StageContentRoot = StagedDirectoryReference.Combine(SC.RelativeProjectRootForStage, "Content");
+					// Stage any extra runtime dependencies from the receipts
+					foreach (StageTarget Target in SC.StageTargets)
+					{
+						SC.StageRuntimeDependenciesFromReceipt(Target.Receipt, Target.RequireFilesExist, Params.UsePak(SC.StageTargetPlatform));
+					}
 
-			if (!Params.CookOnTheFly && !Params.SkipCookOnTheFly) // only stage the UFS files if we are not using cook on the fly
-			{
-				// Get the final output directory for cooked data
-				DirectoryReference CookOutputDir;
-				if (!String.IsNullOrEmpty(Params.CookOutputDir))
-				{
-					CookOutputDir = new DirectoryReference(Params.CookOutputDir);
-					if (!CookOutputDir.GetDirectoryName().Equals(SC.CookPlatform, StringComparison.InvariantCultureIgnoreCase))
+					// Stage any content-only plugins for content-only projects. We don't have a custom executable for these.
+					if (!Params.IsCodeBasedProject)
 					{
-						CookOutputDir = DirectoryReference.Combine(CookOutputDir, SC.CookPlatform);
+						List<FileReference> PluginFiles = GetPluginsForContentProject(Params.RawProjectPath, SC.StageTargets.ConvertAll(x => x.Receipt));
+						SC.StageFiles(StagedFileType.UFS, PluginFiles);
 					}
 				}
-				else if (Params.CookInEditor)
-				{
-					CookOutputDir = DirectoryReference.Combine(SC.ProjectRoot, "Saved", "EditorCooked", SC.CookPlatform);
-				}
-				else
-				{
-					CookOutputDir = DirectoryReference.Combine(SC.ProjectRoot, "Saved", "Cooked", SC.CookPlatform);
-				}
 
-				SC.PlatformCookDir = CookOutputDir;
-				SC.MetadataDir = DirectoryReference.Combine(CookOutputDir, SC.ShortProjectName, "Metadata");
-
-				// Work out which ICU data version we use for this platform
-				var ICUDataVersion = "icudt64l";
-				if (SC.StageTargetPlatform.PlatformType == UnrealTargetPlatform.HoloLens || SC.StageTargetPlatform.PlatformType == UnrealTargetPlatform.TVOS)
+				// move the UECommandline.txt file to the root of the stage
+				// this file needs to be treated as a UFS file for casing, but NonUFS for being put into the .pak file
+				// @todo: Maybe there should be a new category - UFSNotForPak - and reevaluate all SystemNonUFS references
+				FileReference CommandLineFile = FileReference.Combine(GetIntermediateCommandlineDir(SC), "UECommandLine.txt");
+				if (FileReference.Exists(CommandLineFile))
 				{
-					ICUDataVersion = "icudt53l";
+					StagedFileReference StagedCommandLineFile = new StagedFileReference("UECommandLine.txt");
+					if (SC.StageTargetPlatform.DeployLowerCaseFilenames(StagedFileType.SystemNonUFS))
+					{
+						StagedCommandLineFile = StagedCommandLineFile.ToLowerInvariant();
+					}
+					SC.StageFile(StagedFileType.SystemNonUFS, CommandLineFile, StagedCommandLineFile);
 				}
 
-				// Initialize internationalization preset.
-				string InternationalizationPreset = GetInternationalizationPreset(Params, PlatformGameConfig);
-
-				// Initialize cultures to stage.
-				List<string> CulturesToStage = GetCulturesToStage(Params, PlatformGameConfig);
+				DirectoryReference ProjectContentRoot = DirectoryReference.Combine(SC.ProjectRoot, "Content");
+				StagedDirectoryReference StageContentRoot = StagedDirectoryReference.Combine(SC.RelativeProjectRootForStage, "Content");
 
-				// Initialize platform extensions to stage.
-				var PlatformExtensionsToStage = new List<string>();
+				if (!Params.CookOnTheFly && !Params.SkipCookOnTheFly) // only stage the UFS files if we are not using cook on the fly
 				{
-					string PlatformExtensionName = SC.StageTargetPlatform.PlatformType.ToString();
-					PlatformExtensionsToStage.Add(PlatformExtensionName);
-
-					DataDrivenPlatformInfo.ConfigDataDrivenPlatformInfo Info = DataDrivenPlatformInfo.GetDataDrivenInfoForPlatform(SC.StageTargetPlatform.PlatformType);
-					if (Info != null && Info.IniParentChain != null)
+					// Get the final output directory for cooked data
+					DirectoryReference CookOutputDir;
+					if (!String.IsNullOrEmpty(Params.CookOutputDir))
 					{
-						PlatformExtensionsToStage.AddRange(Info.IniParentChain);
+						CookOutputDir = new DirectoryReference(Params.CookOutputDir);
+						if (!CookOutputDir.GetDirectoryName().Equals(SC.CookPlatform, StringComparison.InvariantCultureIgnoreCase))
+						{
+							CookOutputDir = DirectoryReference.Combine(CookOutputDir, SC.CookPlatform);
+						}
+					}
+					else if (Params.CookInEditor)
+					{
+						CookOutputDir = DirectoryReference.Combine(SC.ProjectRoot, "Saved", "EditorCooked", SC.CookPlatform);
+					}
+					else
+					{
+						CookOutputDir = DirectoryReference.Combine(SC.ProjectRoot, "Saved", "Cooked", SC.CookPlatform);
 					}
-				}
-
-				// Stage ICU internationalization data from Engine.
-				SC.StageFiles(StagedFileType.UFS, DirectoryReference.Combine(SC.LocalRoot, "Engine", "Content", "Internationalization", InternationalizationPreset, ICUDataVersion), StageFilesSearch.AllDirectories, new StagedDirectoryReference(String.Format("Engine/Content/Internationalization/{0}", ICUDataVersion)));
-
-				// Engine ufs (content)
-				StageConfigFiles(SC, DirectoryReference.Combine(SC.LocalRoot, "Engine", "Config"), null);
 
-				// Stage the engine localization target
-				if (ShouldStageLocalizationTarget(SC, null, "Engine"))
-				{
-					StageLocalizationDataForTarget(SC, CulturesToStage, DirectoryReference.Combine(SC.LocalRoot, "Engine", "Content", "Localization", "Engine"));
-				}
+					SC.PlatformCookDir = CookOutputDir;
+					SC.MetadataDir = DirectoryReference.Combine(CookOutputDir, SC.ShortProjectName, "Metadata");
 
-				// Game ufs (content)
-				SC.StageFile(StagedFileType.UFS, SC.RawProjectPath);
+					// Work out which ICU data version we use for this platform
+					var ICUDataVersion = "icudt64l";
+					if (SC.StageTargetPlatform.PlatformType == UnrealTargetPlatform.HoloLens || SC.StageTargetPlatform.PlatformType == UnrealTargetPlatform.TVOS)
+					{
+						ICUDataVersion = "icudt53l";
+					}
 
-				StageConfigFiles(SC, DirectoryReference.Combine(SC.ProjectRoot, "Config"), null);
+					// Initialize internationalization preset.
+					string InternationalizationPreset = GetInternationalizationPreset(Params, PlatformGameConfig);
 
-				if (SC.DedicatedServer)
-				{
-					var ToRestoreRestrictedFolder = SC.RestrictedFolderNames;
+					// Initialize cultures to stage.
+					List<string> CulturesToStage = GetCulturesToStage(Params, PlatformGameConfig);
 
-					// For every platform we depend on also stage their platform config files
-					foreach (var PlatformDepends in Params.ServerDependentPlatformMap)
+					// Initialize platform extensions to stage.
+					var PlatformExtensionsToStage = new List<string>();
 					{
-						// StageConfigFiles will fail to find these config files due to filtering out restricted folder names
-						// Just for this loop we remove the filter for these platorms and stage the config files
-						// After the loop we restore the restricted folder to prevent staging anything unwanted
-						SC.RestrictedFolderNames.ExceptWith(PlatformExports.GetIncludedFolderNames(PlatformDepends.Key.Type));
+						string PlatformExtensionName = SC.StageTargetPlatform.PlatformType.ToString();
+						PlatformExtensionsToStage.Add(PlatformExtensionName);
 
-						StageConfigFiles(SC, DirectoryReference.Combine(SC.EngineRoot, "Config"), PlatformDepends.Key.ToString());
-						StageConfigFiles(SC, DirectoryReference.Combine(SC.ProjectRoot, "Config"), PlatformDepends.Key.ToString());
-						StagePlatformExtensionConfigFiles(SC, PlatformDepends.Key.ToString());
+						DataDrivenPlatformInfo.ConfigDataDrivenPlatformInfo Info = DataDrivenPlatformInfo.GetDataDrivenInfoForPlatform(SC.StageTargetPlatform.PlatformType);
+						if (Info != null && Info.IniParentChain != null)
+						{
+							PlatformExtensionsToStage.AddRange(Info.IniParentChain);
+						}
 					}
 
-					SC.RestrictedFolderNames = ToRestoreRestrictedFolder;
-				}
-				
-				// Stage platform extension config files
-				foreach (string PlatformExtensionToStage in PlatformExtensionsToStage)
-				{
-					StagePlatformExtensionConfigFiles(SC, PlatformExtensionToStage);
-				}
+					// Stage ICU internationalization data from Engine.
+					SC.StageFiles(StagedFileType.UFS, DirectoryReference.Combine(SC.LocalRoot, "Engine", "Content", "Internationalization", InternationalizationPreset, ICUDataVersion), StageFilesSearch.AllDirectories, new StagedDirectoryReference(String.Format("Engine/Content/Internationalization/{0}", ICUDataVersion)));
 
-				// Stage plugin config files
-				List<KeyValuePair<StagedFileReference, FileReference>> StagedPlugins = SC.FilesToStage.UFSFiles.Where(x => x.Value.HasExtension(".uplugin")).ToList();
-				foreach (KeyValuePair<StagedFileReference, FileReference> StagedPlugin in StagedPlugins)
-				{
-					//PluginDescriptor Descriptor = PluginDescriptor.FromFile(StagedPlugin.Value);
-					DirectoryReference PluginConfigDirectory = DirectoryReference.Combine(StagedPlugin.Value.Directory, "Config");
-					if (DirectoryReference.Exists(PluginConfigDirectory))
+					// Engine ufs (content)
+					StageConfigFiles(SC, DirectoryReference.Combine(SC.LocalRoot, "Engine", "Config"), null);
+
+					// Stage the engine localization target
+					if (ShouldStageLocalizationTarget(SC, null, "Engine"))
 					{
-						SC.StageFiles(StagedFileType.UFS, PluginConfigDirectory, "*.ini", StageFilesSearch.AllDirectories);
+						StageLocalizationDataForTarget(SC, CulturesToStage, DirectoryReference.Combine(SC.LocalRoot, "Engine", "Content", "Localization", "Engine"));
 					}
-				}
-
-				// Stage all project localization targets
-				StageLocalizationDataForTargetsInDirectory(SC, PlatformGameConfig, CulturesToStage, DirectoryReference.Combine(SC.ProjectRoot, "Content", "Localization"));
-
-				// Stage all plugin localization targets
-				foreach (KeyValuePair<StagedFileReference, FileReference> StagedPlugin in StagedPlugins)
-				{
-					StageLocalizationDataForPlugin(SC, CulturesToStage, StagedPlugin.Value);
-				}
 
-				// Stage any platform extension localization targets
-				foreach (string PlatformExtensionToStage in PlatformExtensionsToStage)
-				{
-					StageLocalizationDataForTargetsInDirectory(SC, null, CulturesToStage, DirectoryReference.Combine(SC.LocalRoot, "Engine", "Platforms", PlatformExtensionToStage, "Content", "Localization"));
-					StageLocalizationDataForTargetsInDirectory(SC, PlatformGameConfig, CulturesToStage, DirectoryReference.Combine(SC.ProjectRoot, "Platforms", PlatformExtensionToStage, "Content", "Localization"));
-				}
+					// Game ufs (content)
+					SC.StageFile(StagedFileType.UFS, SC.RawProjectPath);
 
-				// Stage any additional UFS and NonUFS paths specified in the project ini files; these dirs are relative to the game content directory
-				if (PlatformGameConfig != null)
-				{
-					StageAdditionalDirectoriesFromConfig(SC, ProjectContentRoot, StageContentRoot, PlatformGameConfig, true, "DirectoriesToAlwaysStageAsUFS");
-					// NonUFS files are never in pak files and should always be remapped
-					StageAdditionalDirectoriesFromConfig(SC, ProjectContentRoot, StageContentRoot, PlatformGameConfig, false, "DirectoriesToAlwaysStageAsNonUFS");
+					StageConfigFiles(SC, DirectoryReference.Combine(SC.ProjectRoot, "Config"), null);
 
 					if (SC.DedicatedServer)
 					{
-						StageAdditionalDirectoriesFromConfig(SC, ProjectContentRoot, StageContentRoot, PlatformGameConfig, true, "DirectoriesToAlwaysStageAsUFSServer");
-						// NonUFS files are never in pak files and should always be remapped
-						StageAdditionalDirectoriesFromConfig(SC, ProjectContentRoot, StageContentRoot, PlatformGameConfig, false, "DirectoriesToAlwaysStageAsNonUFSServer");
+						var ToRestoreRestrictedFolder = SC.RestrictedFolderNames;
+
+						// For every platform we depend on also stage their platform config files
+						foreach (var PlatformDepends in Params.ServerDependentPlatformMap)
+						{
+							// StageConfigFiles will fail to find these config files due to filtering out restricted folder names
+							// Just for this loop we remove the filter for these platorms and stage the config files
+							// After the loop we restore the restricted folder to prevent staging anything unwanted
+							SC.RestrictedFolderNames.ExceptWith(PlatformExports.GetIncludedFolderNames(PlatformDepends.Key.Type));
+
+							StageConfigFiles(SC, DirectoryReference.Combine(SC.EngineRoot, "Config"), PlatformDepends.Key.ToString());
+							StageConfigFiles(SC, DirectoryReference.Combine(SC.ProjectRoot, "Config"), PlatformDepends.Key.ToString());
+							StagePlatformExtensionConfigFiles(SC, PlatformDepends.Key.ToString());
+						}
+
+						SC.RestrictedFolderNames = ToRestoreRestrictedFolder;
 					}
-				}
 
-				// Per-project, per-platform setting to skip all movies. By default this is false.
-				bool bSkipMovies = false;
-				PlatformGameConfig.GetBool("/Script/UnrealEd.ProjectPackagingSettings", "bSkipMovies", out bSkipMovies);
+					// Stage platform extension config files
+					foreach (string PlatformExtensionToStage in PlatformExtensionsToStage)
+					{
+						StagePlatformExtensionConfigFiles(SC, PlatformExtensionToStage);
+					}
 
-				if (!bSkipMovies && !SC.DedicatedServer)
-				{
-					List<string> MovieBlackList;
-					PlatformGameConfig.GetArray("/Script/UnrealEd.ProjectPackagingSettings", "MovieBlacklist", out MovieBlackList);
-					if (MovieBlackList == null)
+					// Stage plugin config files
+					List<KeyValuePair<StagedFileReference, FileReference>> StagedPlugins = SC.FilesToStage.UFSFiles.Where(x => x.Value.HasExtension(".uplugin")).ToList();
+					foreach (KeyValuePair<StagedFileReference, FileReference> StagedPlugin in StagedPlugins)
 					{
-						// Make an empty list to avoid having to null check below
-						MovieBlackList = new List<string>();
+						//PluginDescriptor Descriptor = PluginDescriptor.FromFile(StagedPlugin.Value);
+						DirectoryReference PluginConfigDirectory = DirectoryReference.Combine(StagedPlugin.Value.Directory, "Config");
+						if (DirectoryReference.Exists(PluginConfigDirectory))
+						{
+							SC.StageFiles(StagedFileType.UFS, PluginConfigDirectory, "*.ini", StageFilesSearch.AllDirectories);
+						}
 					}
 
-					// UFS is required when using a file server
-					StagedFileType MovieFileType = Params.FileServer ? StagedFileType.UFS : StagedFileType.NonUFS;
+					// Stage all project localization targets
+					StageLocalizationDataForTargetsInDirectory(SC, PlatformGameConfig, CulturesToStage, DirectoryReference.Combine(SC.ProjectRoot, "Content", "Localization"));
 
-					DirectoryReference EngineMoviesDir = DirectoryReference.Combine(SC.EngineRoot, "Content", "Movies");
-					if (DirectoryReference.Exists(EngineMoviesDir))
+					// Stage all plugin localization targets
+					foreach (KeyValuePair<StagedFileReference, FileReference> StagedPlugin in StagedPlugins)
 					{
-						List<FileReference> MovieFiles = SC.FindFilesToStage(EngineMoviesDir, StageFilesSearch.AllDirectories);
-						SC.StageFiles(MovieFileType, MovieFiles.Where(x => !x.HasExtension(".uasset") && !x.HasExtension(".umap") && !MovieBlackList.Contains(x.GetFileNameWithoutAnyExtensions())));
+						StageLocalizationDataForPlugin(SC, CulturesToStage, StagedPlugin.Value);
 					}
 
-					DirectoryReference ProjectMoviesDir = DirectoryReference.Combine(SC.ProjectRoot, "Content", "Movies");
-					if (DirectoryReference.Exists(ProjectMoviesDir))
+					// Stage any platform extension localization targets
+					foreach (string PlatformExtensionToStage in PlatformExtensionsToStage)
 					{
-						List<FileReference> MovieFiles = SC.FindFilesToStage(ProjectMoviesDir, StageFilesSearch.AllDirectories);
-						SC.StageFiles(MovieFileType, MovieFiles.Where(x => !x.HasExtension(".uasset") && !x.HasExtension(".umap") && !MovieBlackList.Contains(x.GetFileNameWithoutAnyExtensions())));
+						StageLocalizationDataForTargetsInDirectory(SC, null, CulturesToStage, DirectoryReference.Combine(SC.LocalRoot, "Engine", "Platforms", PlatformExtensionToStage, "Content", "Localization"));
+						StageLocalizationDataForTargetsInDirectory(SC, PlatformGameConfig, CulturesToStage, DirectoryReference.Combine(SC.ProjectRoot, "Platforms", PlatformExtensionToStage, "Content", "Localization"));
 					}
-				}
-				else if (!SC.DedicatedServer)
-				{
-					// check to see if we have any specific movies we want to stage for non ufs and ufs files
-					// we still use the movie directories to find the paths to the movies and the list in the ini file is just the substring filename you wish to have staged
-					List<string> NonUFSMovieList;
-					List<string> UFSMovieList;
-					PlatformGameConfig.GetArray("/Script/UnrealEd.ProjectPackagingSettings", "NonUFSMovies", out NonUFSMovieList);
-					PlatformGameConfig.GetArray("/Script/UnrealEd.ProjectPackagingSettings", "UFSMovies", out UFSMovieList);
-					DirectoryReference EngineMoviesDir = DirectoryReference.Combine(SC.EngineRoot, "Content", "Movies");
-					if (DirectoryReference.Exists(EngineMoviesDir))
+
+					// Stage any additional UFS and NonUFS paths specified in the project ini files; these dirs are relative to the game content directory
+					if (PlatformGameConfig != null)
+					{
+						StageAdditionalDirectoriesFromConfig(SC, ProjectContentRoot, StageContentRoot, PlatformGameConfig, true, "DirectoriesToAlwaysStageAsUFS");
+						// NonUFS files are never in pak files and should always be remapped
+						StageAdditionalDirectoriesFromConfig(SC, ProjectContentRoot, StageContentRoot, PlatformGameConfig, false, "DirectoriesToAlwaysStageAsNonUFS");
+
+						if (SC.DedicatedServer)
+						{
+							StageAdditionalDirectoriesFromConfig(SC, ProjectContentRoot, StageContentRoot, PlatformGameConfig, true, "DirectoriesToAlwaysStageAsUFSServer");
+							// NonUFS files are never in pak files and should always be remapped
+							StageAdditionalDirectoriesFromConfig(SC, ProjectContentRoot, StageContentRoot, PlatformGameConfig, false, "DirectoriesToAlwaysStageAsNonUFSServer");
+						}
+					}
+
+					// Per-project, per-platform setting to skip all movies. By default this is false.
+					bool bSkipMovies = false;
+					PlatformGameConfig.GetBool("/Script/UnrealEd.ProjectPackagingSettings", "bSkipMovies", out bSkipMovies);
+
+					if (!bSkipMovies && !SC.DedicatedServer)
 					{
-						List<FileReference> MovieFiles = SC.FindFilesToStage(EngineMoviesDir, StageFilesSearch.AllDirectories);
-						if (NonUFSMovieList != null)
+						List<string> MovieBlackList;
+						PlatformGameConfig.GetArray("/Script/UnrealEd.ProjectPackagingSettings", "MovieBlacklist", out MovieBlackList);
+						if (MovieBlackList == null)
 						{
-							SC.StageFiles(StagedFileType.NonUFS, MovieFiles.Where(x => !x.HasExtension(".uasset") && !x.HasExtension(".umap") && NonUFSMovieList.Where(y => x.GetFileNameWithoutExtension().Contains(y)).Any()));
+							// Make an empty list to avoid having to null check below
+							MovieBlackList = new List<string>();
 						}
-						if (UFSMovieList != null)
+
+						// UFS is required when using a file server
+						StagedFileType MovieFileType = Params.FileServer ? StagedFileType.UFS : StagedFileType.NonUFS;
+
+						DirectoryReference EngineMoviesDir = DirectoryReference.Combine(SC.EngineRoot, "Content", "Movies");
+						if (DirectoryReference.Exists(EngineMoviesDir))
 						{
-							SC.StageFiles(StagedFileType.UFS, MovieFiles.Where(x => !x.HasExtension(".uasset") && !x.HasExtension(".umap") && UFSMovieList.Where(y => x.GetFileNameWithoutExtension().Contains(y)).Any()));
+							List<FileReference> MovieFiles = SC.FindFilesToStage(EngineMoviesDir, StageFilesSearch.AllDirectories);
+							SC.StageFiles(MovieFileType, MovieFiles.Where(x => !x.HasExtension(".uasset") && !x.HasExtension(".umap") && !MovieBlackList.Contains(x.GetFileNameWithoutAnyExtensions())));
 						}
-					}
 
-					DirectoryReference ProjectMoviesDir = DirectoryReference.Combine(SC.ProjectRoot, "Content", "Movies");
-					if (DirectoryReference.Exists(ProjectMoviesDir))
+						DirectoryReference ProjectMoviesDir = DirectoryReference.Combine(SC.ProjectRoot, "Content", "Movies");
+						if (DirectoryReference.Exists(ProjectMoviesDir))
+						{
+							List<FileReference> MovieFiles = SC.FindFilesToStage(ProjectMoviesDir, StageFilesSearch.AllDirectories);
+							SC.StageFiles(MovieFileType, MovieFiles.Where(x => !x.HasExtension(".uasset") && !x.HasExtension(".umap") && !MovieBlackList.Contains(x.GetFileNameWithoutAnyExtensions())));
+						}
+					}
+					else if (!SC.DedicatedServer)
 					{
-						List<FileReference> MovieFiles = SC.FindFilesToStage(ProjectMoviesDir, StageFilesSearch.AllDirectories);
-						if (NonUFSMovieList != null)
+						// check to see if we have any specific movies we want to stage for non ufs and ufs files
+						// we still use the movie directories to find the paths to the movies and the list in the ini file is just the substring filename you wish to have staged
+						List<string> NonUFSMovieList;
+						List<string> UFSMovieList;
+						PlatformGameConfig.GetArray("/Script/UnrealEd.ProjectPackagingSettings", "NonUFSMovies", out NonUFSMovieList);
+						PlatformGameConfig.GetArray("/Script/UnrealEd.ProjectPackagingSettings", "UFSMovies", out UFSMovieList);
+						DirectoryReference EngineMoviesDir = DirectoryReference.Combine(SC.EngineRoot, "Content", "Movies");
+						if (DirectoryReference.Exists(EngineMoviesDir))
 						{
-							SC.StageFiles(StagedFileType.NonUFS, MovieFiles.Where(x => !x.HasExtension(".uasset") && !x.HasExtension(".umap") && NonUFSMovieList.Where(y => x.GetFileNameWithoutExtension().Contains(y)).Any()));
+							List<FileReference> MovieFiles = SC.FindFilesToStage(EngineMoviesDir, StageFilesSearch.AllDirectories);
+							if (NonUFSMovieList != null)
+							{
+								SC.StageFiles(StagedFileType.NonUFS, MovieFiles.Where(x => !x.HasExtension(".uasset") && !x.HasExtension(".umap") && NonUFSMovieList.Where(y => x.GetFileNameWithoutExtension().Contains(y)).Any()));
+							}
+							if (UFSMovieList != null)
+							{
+								SC.StageFiles(StagedFileType.UFS, MovieFiles.Where(x => !x.HasExtension(".uasset") && !x.HasExtension(".umap") && UFSMovieList.Where(y => x.GetFileNameWithoutExtension().Contains(y)).Any()));
+							}
 						}
-						if (UFSMovieList != null)
+
+						DirectoryReference ProjectMoviesDir = DirectoryReference.Combine(SC.ProjectRoot, "Content", "Movies");
+						if (DirectoryReference.Exists(ProjectMoviesDir))
 						{
-							SC.StageFiles(StagedFileType.UFS, MovieFiles.Where(x => !x.HasExtension(".uasset") && !x.HasExtension(".umap") && UFSMovieList.Where(y => x.GetFileNameWithoutExtension().Contains(y)).Any()));
+							List<FileReference> MovieFiles = SC.FindFilesToStage(ProjectMoviesDir, StageFilesSearch.AllDirectories);
+							if (NonUFSMovieList != null)
+							{
+								SC.StageFiles(StagedFileType.NonUFS, MovieFiles.Where(x => !x.HasExtension(".uasset") && !x.HasExtension(".umap") && NonUFSMovieList.Where(y => x.GetFileNameWithoutExtension().Contains(y)).Any()));
+							}
+							if (UFSMovieList != null)
+							{
+								SC.StageFiles(StagedFileType.UFS, MovieFiles.Where(x => !x.HasExtension(".uasset") && !x.HasExtension(".umap") && UFSMovieList.Where(y => x.GetFileNameWithoutExtension().Contains(y)).Any()));
+							}
 						}
 					}
-				}
 
-				// shader cache
-				{
-					DirectoryReference ShaderCacheRoot = DirectoryReference.Combine(SC.ProjectRoot, "Content");
-					List<FileReference> ShaderCacheFiles = SC.FindFilesToStage(ShaderCacheRoot, "DrawCache-*.ushadercache", StageFilesSearch.TopDirectoryOnly);
-					SC.StageFiles(StagedFileType.UFS, ShaderCacheFiles);
-				}
-				// pipeline cache
-				{
-					DirectoryReference ShaderCacheRoot = DirectoryReference.Combine(SC.ProjectRoot, "Content", "PipelineCaches", SC.PlatformDir);
-					if (DirectoryReference.Exists(ShaderCacheRoot))
+					// shader cache
 					{
-						List<FileReference> ShaderCacheFiles = SC.FindFilesToStage(ShaderCacheRoot, "*.upipelinecache", StageFilesSearch.TopDirectoryOnly);
+						DirectoryReference ShaderCacheRoot = DirectoryReference.Combine(SC.ProjectRoot, "Content");
+						List<FileReference> ShaderCacheFiles = SC.FindFilesToStage(ShaderCacheRoot, "DrawCache-*.ushadercache", StageFilesSearch.TopDirectoryOnly);
 						SC.StageFiles(StagedFileType.UFS, ShaderCacheFiles);
 					}
-				}
-
-				// Stage all the cooked data. Currently not filtering this by restricted folders, since we shouldn't mask invalid references by filtering them out.
-				if (DirectoryReference.Exists(CookOutputDir))
-				{
-					List<FileReference> CookedFiles = DirectoryReference.EnumerateFiles(CookOutputDir, "*", SearchOption.AllDirectories).ToList();
-					
-					// When cooking directly to I/O store contaier files look for the package store manifest
-					if (Params.IoStore)
+					// pipeline cache
 					{
-						CookedFiles.AddRange(PackageStoreManifest.GetCookedFiles(FileReference.Combine(SC.MetadataDir, "packagestore.manifest")));
+						DirectoryReference ShaderCacheRoot = DirectoryReference.Combine(SC.ProjectRoot, "Content", "PipelineCaches", SC.PlatformDir);
+						if (DirectoryReference.Exists(ShaderCacheRoot))
+						{
+							List<FileReference> ShaderCacheFiles = SC.FindFilesToStage(ShaderCacheRoot, "*.upipelinecache", StageFilesSearch.TopDirectoryOnly);
+							SC.StageFiles(StagedFileType.UFS, ShaderCacheFiles);
+						}
 					}
 
-					foreach (FileReference CookedFile in CookedFiles)
+					// Stage all the cooked data. Currently not filtering this by restricted folders, since we shouldn't mask invalid references by filtering them out.
+					if (DirectoryReference.Exists(CookOutputDir))
 					{
-						// Skip metadata directory
-						if (CookedFile.Directory.IsUnderDirectory(SC.MetadataDir))
+						List<FileReference> CookedFiles = DirectoryReference.EnumerateFiles(CookOutputDir, "*", SearchOption.AllDirectories).ToList();
+
+						// When cooking directly to I/O store contaier files look for the package store manifest
+						if (Params.IoStore)
 						{
-							continue;
+							CookedFiles.AddRange(PackageStoreManifest.GetCookedFiles(FileReference.Combine(SC.MetadataDir, "packagestore.manifest")));
 						}
 
-						// Skip region files. These are cook-time metadata files that never need to ship.
-						if (CookedFile.HasExtension(FileRegion.RegionsFileExtension))
+						foreach (FileReference CookedFile in CookedFiles)
 						{
-							continue;
+							// Skip metadata directory
+							if (CookedFile.Directory.IsUnderDirectory(SC.MetadataDir))
+							{
+								continue;
+							}
+
+							// Skip region files. These are cook-time metadata files that never need to ship.
+							if (CookedFile.HasExtension(FileRegion.RegionsFileExtension))
+							{
+								continue;
+							}
+
+							// json files have never been staged
+							// metallib files cannot *currently* be staged as UFS as the Metal API needs to mmap them from files on disk in order to function efficiently
+							if (!CookedFile.HasExtension(".json") && !CookedFile.HasExtension(".metallib") && !CookedFile.HasExtension(".utoc") && !CookedFile.HasExtension(".ucas"))
+							{
+								SC.StageFile(StagedFileType.UFS, CookedFile, new StagedFileReference(CookedFile.MakeRelativeTo(CookOutputDir)));
+							}
 						}
+					}
+
+					// CrashReportClient is a standalone slate app that does not look in the generated pak file, so it needs the Content/Slate and Shaders/StandaloneRenderer folders Non-UFS
+					// @todo Make CrashReportClient more portable so we don't have to do this
+					UnrealTargetPlatform CrashReportPlatform = ThisPlatform.CrashReportPlatform ?? SC.StageTargetPlatform.PlatformType;
+					if (SC.bStageCrashReporter && PlatformSupportsCrashReporter(CrashReportPlatform) && (Params.IterateSharedBuildUsePrecompiledExe == false))
+					{
+						SC.StageCrashReporterFiles(StagedFileType.UFS, DirectoryReference.Combine(SC.EngineRoot, "Content", "Slate"), StageFilesSearch.AllDirectories);
+						SC.StageCrashReporterFiles(StagedFileType.UFS, DirectoryReference.Combine(SC.EngineRoot, "Shaders", "StandaloneRenderer"), StageFilesSearch.AllDirectories);
+						SC.StageCrashReporterFiles(StagedFileType.UFS, DirectoryReference.Combine(SC.EngineRoot, "Content", "Internationalization", InternationalizationPreset), StageFilesSearch.AllDirectories, new StagedDirectoryReference("Engine/Content/Internationalization"));
 
-						// json files have never been staged
-						// metallib files cannot *currently* be staged as UFS as the Metal API needs to mmap them from files on disk in order to function efficiently
-						if (!CookedFile.HasExtension(".json") && !CookedFile.HasExtension(".metallib") && !CookedFile.HasExtension(".utoc") && !CookedFile.HasExtension(".ucas"))
+						// Get the architecture in use
+						string Architecture = Params.SpecifiedArchitecture;
+						if (string.IsNullOrEmpty(Architecture))
 						{
-							SC.StageFile(StagedFileType.UFS, CookedFile, new StagedFileReference(CookedFile.MakeRelativeTo(CookOutputDir)));
+							Architecture = "";
+							if (PlatformExports.IsPlatformAvailable(CrashReportPlatform))
+							{
+								Architecture = PlatformExports.GetDefaultArchitecture(CrashReportPlatform, Params.RawProjectPath);
+							}
 						}
-					}
-				}
 
-				// CrashReportClient is a standalone slate app that does not look in the generated pak file, so it needs the Content/Slate and Shaders/StandaloneRenderer folders Non-UFS
-				// @todo Make CrashReportClient more portable so we don't have to do this
-				UnrealTargetPlatform CrashReportPlatform = ThisPlatform.CrashReportPlatform ?? SC.StageTargetPlatform.PlatformType;
-				if (SC.bStageCrashReporter && PlatformSupportsCrashReporter(CrashReportPlatform) && (Params.IterateSharedBuildUsePrecompiledExe == false))
-				{
-					SC.StageCrashReporterFiles(StagedFileType.UFS, DirectoryReference.Combine(SC.EngineRoot, "Content", "Slate"), StageFilesSearch.AllDirectories);
-					SC.StageCrashReporterFiles(StagedFileType.UFS, DirectoryReference.Combine(SC.EngineRoot, "Shaders", "StandaloneRenderer"), StageFilesSearch.AllDirectories);
-					SC.StageCrashReporterFiles(StagedFileType.UFS, DirectoryReference.Combine(SC.EngineRoot, "Content", "Internationalization", InternationalizationPreset), StageFilesSearch.AllDirectories, new StagedDirectoryReference("Engine/Content/Internationalization"));
+						// Get the target receipt path for CrashReportClient
+						FileReference ReceiptFileName = TargetReceipt.GetDefaultPath(SC.EngineRoot, "CrashReportClient", CrashReportPlatform, UnrealTargetConfiguration.Shipping, Architecture);
+						if (!FileReference.Exists(ReceiptFileName))
+						{
+							throw new AutomationException(ExitCode.Error_MissingExecutable, "Stage Failed. Missing receipt '{0}'. Check that this target has been built.", ReceiptFileName);
+						}
 
-					// Get the architecture in use
-					string Architecture = Params.SpecifiedArchitecture;
-					if (string.IsNullOrEmpty(Architecture))
-					{
-						Architecture = "";
-						if (PlatformExports.IsPlatformAvailable(CrashReportPlatform))
+						// Read the receipt for this target
+						TargetReceipt Receipt;
+						if (!TargetReceipt.TryRead(ReceiptFileName, out Receipt))
 						{
-							Architecture = PlatformExports.GetDefaultArchitecture(CrashReportPlatform, Params.RawProjectPath);
+							throw new AutomationException("Missing or invalid target receipt ({0})", ReceiptFileName);
 						}
-					}
 
-					// Get the target receipt path for CrashReportClient
-					FileReference ReceiptFileName = TargetReceipt.GetDefaultPath(SC.EngineRoot, "CrashReportClient", CrashReportPlatform, UnrealTargetConfiguration.Shipping, Architecture);
-					if (!FileReference.Exists(ReceiptFileName))
-					{
-						throw new AutomationException(ExitCode.Error_MissingExecutable, "Stage Failed. Missing receipt '{0}'. Check that this target has been built.", ReceiptFileName);
-					}
+						foreach (RuntimeDependency RuntimeDependency in Receipt.RuntimeDependencies)
+						{
+							StagedFileReference StagedFile = new StagedFileReference(RuntimeDependency.Path.MakeRelativeTo(Unreal.RootDirectory));
+							SC.StageCrashReporterFile(RuntimeDependency.Type, RuntimeDependency.Path, StagedFile);
+						}
 
-					// Read the receipt for this target
-					TargetReceipt Receipt;
-					if (!TargetReceipt.TryRead(ReceiptFileName, out Receipt))
-					{
-						throw new AutomationException("Missing or invalid target receipt ({0})", ReceiptFileName);
+						// Add config files.
+						SC.StageCrashReporterFiles(StagedFileType.UFS, DirectoryReference.Combine(SC.EngineRoot, "Programs", "CrashReportClient", "Config"), StageFilesSearch.AllDirectories);
 					}
-
-					foreach (RuntimeDependency RuntimeDependency in Receipt.RuntimeDependencies)
+					// check if the game will be verifying ssl connections - if not, we can skip staging files that won't be needed
+					bool bStageSSLCertificates = false;
+					ConfigHierarchy PlatformEngineConfig = ConfigCache.ReadHierarchy(ConfigHierarchyType.Engine, DirectoryReference.FromFile(Params.RawProjectPath), SC.StageTargetPlatform.IniPlatformType);
+					if (PlatformEngineConfig != null)
 					{
-						StagedFileReference StagedFile = new StagedFileReference(RuntimeDependency.Path.MakeRelativeTo(Unreal.RootDirectory));
-						SC.StageCrashReporterFile(RuntimeDependency.Type, RuntimeDependency.Path, StagedFile);
+						PlatformEngineConfig.GetBool("/Script/Engine.NetworkSettings", "n.VerifyPeer", out bStageSSLCertificates);
 					}
 
-					// Add config files.
-					SC.StageCrashReporterFiles(StagedFileType.UFS, DirectoryReference.Combine(SC.EngineRoot, "Programs", "CrashReportClient", "Config"), StageFilesSearch.AllDirectories);
-				}
-				// check if the game will be verifying ssl connections - if not, we can skip staging files that won't be needed
-				bool bStageSSLCertificates = false;
-				ConfigHierarchy PlatformEngineConfig = ConfigCache.ReadHierarchy(ConfigHierarchyType.Engine, DirectoryReference.FromFile(Params.RawProjectPath), SC.StageTargetPlatform.IniPlatformType);
-				if (PlatformEngineConfig != null)
-				{
-					PlatformEngineConfig.GetBool("/Script/Engine.NetworkSettings", "n.VerifyPeer", out bStageSSLCertificates);
-				}
-
-				if (bStageSSLCertificates)
-				{
-					// Game's SSL certs
-					FileReference ProjectCertFile = FileReference.Combine(SC.ProjectRoot, "Content", "Certificates", "cacert.pem");
-					if (FileReference.Exists(ProjectCertFile))
-					{
-						SC.StageFile(StagedFileType.UFS, ProjectCertFile);
-					}
-					else
+					if (bStageSSLCertificates)
 					{
-						// if the game had any files to be staged, then we don't need to stage the engine one - it will just added hundreds of kb of data that is never used
-						FileReference EngineCertFile = FileReference.Combine(SC.EngineRoot, "Content", "Certificates", "ThirdParty", "cacert.pem");
-						if (FileReference.Exists(EngineCertFile))
+						// Game's SSL certs
+						FileReference ProjectCertFile = FileReference.Combine(SC.ProjectRoot, "Content", "Certificates", "cacert.pem");
+						if (FileReference.Exists(ProjectCertFile))
 						{
-							SC.StageFile(StagedFileType.UFS, EngineCertFile);
+							SC.StageFile(StagedFileType.UFS, ProjectCertFile);
+						}
+						else
+						{
+							// if the game had any files to be staged, then we don't need to stage the engine one - it will just added hundreds of kb of data that is never used
+							FileReference EngineCertFile = FileReference.Combine(SC.EngineRoot, "Content", "Certificates", "ThirdParty", "cacert.pem");
+							if (FileReference.Exists(EngineCertFile))
+							{
+								SC.StageFile(StagedFileType.UFS, EngineCertFile);
+							}
 						}
-					}
 
-					// now stage any other game certs besides cacert
-					DirectoryReference CertificatesDir = DirectoryReference.Combine(SC.ProjectRoot, "Certificates");
-					if (DirectoryReference.Exists(CertificatesDir))
-					{
-						SC.StageFiles(StagedFileType.UFS, CertificatesDir, "*.pem", StageFilesSearch.AllDirectories);
+						// now stage any other game certs besides cacert
+						DirectoryReference CertificatesDir = DirectoryReference.Combine(SC.ProjectRoot, "Certificates");
+						if (DirectoryReference.Exists(CertificatesDir))
+						{
+							SC.StageFiles(StagedFileType.UFS, CertificatesDir, "*.pem", StageFilesSearch.AllDirectories);
+						}
 					}
-				}
 
-				// now make the fast-load ini/plugin file if requested
-				// if -skippak is on, this indicates that the data is already ready to go, we are staging non-data files to re-use
-				// data from a previous run, but still want to stage executables, etc. SkipPak indicates we _would_ have made a 
-				// pak file, but we aren't this time, so we can assume the binary config is already inside the .pak file (and
-				// the UnrealEditor.exe may not even exist if we are making a target-code-only build on top of existing build)
-				if (Params.GenerateOptimizationData && !Params.SkipPak)
-				{
-					// get the list of plugins that need to be processed
-					List<string> PluginFilesForTarget = new List<string>();
-					foreach (KeyValuePair<StagedFileReference, FileReference> StagedPlugin in StagedPlugins)
+					// now make the fast-load ini/plugin file if requested
+					// if -skippak is on, this indicates that the data is already ready to go, we are staging non-data files to re-use
+					// data from a previous run, but still want to stage executables, etc. SkipPak indicates we _would_ have made a 
+					// pak file, but we aren't this time, so we can assume the binary config is already inside the .pak file (and
+					// the UnrealEditor.exe may not even exist if we are making a target-code-only build on top of existing build)
+					if (Params.GenerateOptimizationData && !Params.SkipPak)
 					{
-						PluginFilesForTarget.Add(StagedPlugin.Value.FullName);
+						// get the list of plugins that need to be processed
+						List<string> PluginFilesForTarget = new List<string>();
+						foreach (KeyValuePair<StagedFileReference, FileReference> StagedPlugin in StagedPlugins)
+						{
+							PluginFilesForTarget.Add(StagedPlugin.Value.FullName);
+						}
+						FileReference PluginListFile = FileReference.Combine(SC.ProjectRoot, "Intermediate", "Config", "PluginList.txt");
+						DirectoryReference.CreateDirectory(PluginListFile.Directory);
+						File.WriteAllLines(PluginListFile.FullName, PluginFilesForTarget.ToArray());
+
+						// run the commandlet to generate a binary file
+						String TargetPlatformName = ThisPlatform.GetCookPlatform(Params.DedicatedServer, Params.Client);
+						FileReference OutputFile = FileReference.Combine(SC.ProjectRoot, "Intermediate", "Config", TargetPlatformName, "BinaryConfig.ini");
+						String CommandletParams = String.Format("-TargetPlatform={0} -OutputFile={1} -StagedPluginsFile={2}", TargetPlatformName, OutputFile.FullName, PluginListFile.FullName);
+						RunCommandlet(SC.RawProjectPath, Params.UE4Exe, "MakeBinaryConfig", CommandletParams);
+						SC.StageFile(StagedFileType.UFS, OutputFile, StagedFileReference.Combine(SC.RelativeProjectRootForStage, "Config", OutputFile.GetFileName()));
 					}
-					FileReference PluginListFile = FileReference.Combine(SC.ProjectRoot, "Intermediate", "Config", "PluginList.txt");
-					DirectoryReference.CreateDirectory(PluginListFile.Directory);
-					File.WriteAllLines(PluginListFile.FullName, PluginFilesForTarget.ToArray());
 
-					// run the commandlet to generate a binary file
-					String TargetPlatformName = ThisPlatform.GetCookPlatform(Params.DedicatedServer, Params.Client);
-					FileReference OutputFile = FileReference.Combine(SC.ProjectRoot, "Intermediate", "Config", TargetPlatformName, "BinaryConfig.ini");
-					String CommandletParams = String.Format("-TargetPlatform={0} -OutputFile={1} -StagedPluginsFile={2}", TargetPlatformName, OutputFile.FullName, PluginListFile.FullName);
-					RunCommandlet(SC.RawProjectPath, Params.UE4Exe, "MakeBinaryConfig", CommandletParams);
-					SC.StageFile(StagedFileType.UFS, OutputFile, StagedFileReference.Combine(SC.RelativeProjectRootForStage, "Config", OutputFile.GetFileName()));
+					bCreatePluginManifest = true;
 				}
-
-				bCreatePluginManifest = true;
-			}
-			else
-			{
-				if (PlatformGameConfig != null)
+				else
 				{
-					List<string> ExtraNonUFSDirs;
-					if (PlatformGameConfig.GetArray("/Script/UnrealEd.ProjectPackagingSettings", "DirectoriesToAlwaysStageAsNonUFS", out ExtraNonUFSDirs))
+					if (PlatformGameConfig != null)
 					{
-						// Each string has the format '(Path="TheDirToStage")'
-						foreach (var PathStr in ExtraNonUFSDirs)
+						List<string> ExtraNonUFSDirs;
+						if (PlatformGameConfig.GetArray("/Script/UnrealEd.ProjectPackagingSettings", "DirectoriesToAlwaysStageAsNonUFS", out ExtraNonUFSDirs))
 						{
-							var PathParts = PathStr.Split('"');
-							if (PathParts.Length == 3)
-							{
-								var RelativePath = PathParts[1];
-								SC.StageFiles(StagedFileType.NonUFS, DirectoryReference.Combine(ProjectContentRoot, RelativePath), StageFilesSearch.AllDirectories);
-							}
-							else if (PathParts.Length == 1)
+							// Each string has the format '(Path="TheDirToStage")'
+							foreach (var PathStr in ExtraNonUFSDirs)
 							{
-								var RelativePath = PathParts[0];
-								SC.StageFiles(StagedFileType.UFS, DirectoryReference.Combine(ProjectContentRoot, RelativePath), StageFilesSearch.AllDirectories, StagedDirectoryReference.Combine(StageContentRoot, RelativePath));
+								var PathParts = PathStr.Split('"');
+								if (PathParts.Length == 3)
+								{
+									var RelativePath = PathParts[1];
+									SC.StageFiles(StagedFileType.NonUFS, DirectoryReference.Combine(ProjectContentRoot, RelativePath), StageFilesSearch.AllDirectories);
+								}
+								else if (PathParts.Length == 1)
+								{
+									var RelativePath = PathParts[0];
+									SC.StageFiles(StagedFileType.UFS, DirectoryReference.Combine(ProjectContentRoot, RelativePath), StageFilesSearch.AllDirectories, StagedDirectoryReference.Combine(StageContentRoot, RelativePath));
+								}
 							}
 						}
 					}
-				}
 
-				// UE-58423
-				DirectoryReference CookOutputDir;
-				if (!String.IsNullOrEmpty(Params.CookOutputDir))
-				{
-					CookOutputDir = DirectoryReference.Combine(new DirectoryReference(Params.CookOutputDir), SC.CookPlatform);
+					// UE-58423
+					DirectoryReference CookOutputDir;
+					if (!String.IsNullOrEmpty(Params.CookOutputDir))
+					{
+						CookOutputDir = DirectoryReference.Combine(new DirectoryReference(Params.CookOutputDir), SC.CookPlatform);
+					}
+					else if (Params.CookInEditor)
+					{
+						CookOutputDir = DirectoryReference.Combine(SC.ProjectRoot, "Saved", "EditorCooked", SC.CookPlatform);
+					}
+					else
+					{
+						CookOutputDir = DirectoryReference.Combine(SC.ProjectRoot, "Saved", "Cooked", SC.CookPlatform);
+					}
+					SC.MetadataDir = DirectoryReference.Combine(CookOutputDir, SC.ShortProjectName, "Metadata");
 				}
-				else if (Params.CookInEditor)
+			}
+
+			// Allow the calling scripts to make modifications to the deployment context before we finalize it
+			if (Params.ModifyDeploymentContextCallback != null)
+			{
+				Params.ModifyDeploymentContextCallback(Params, SC);
+			}
+
+			// Apply all the directory mappings
+			SC.FilesToStage.UFSFiles = SC.FilesToStage.UFSFiles.ToDictionary(x => ApplyDirectoryRemap(SC, x.Key), x => x.Value);
+			SC.FilesToStage.NonUFSFiles = SC.FilesToStage.NonUFSFiles.ToDictionary(x => ApplyDirectoryRemap(SC, x.Key), x => x.Value);
+			SC.FilesToStage.NonUFSDebugFiles = SC.FilesToStage.NonUFSDebugFiles.ToDictionary(x => ApplyDirectoryRemap(SC, x.Key), x => x.Value);
+			SC.FilesToStage.NonUFSSystemFiles = SC.FilesToStage.NonUFSSystemFiles.ToDictionary(x => ApplyDirectoryRemap(SC, x.Key), x => x.Value);
+
+			// Create plugin manifests after the directory mappings
+			if (bCreatePluginManifest && Params.UsePak(SC.StageTargetPlatform))
+			{
+				// Generate a plugin manifest if we're using a pak file and not creating a mod. Mods can be enumerated independently by users copying them into the Mods directory.
+				if (Params.HasDLCName)
 				{
-					CookOutputDir = DirectoryReference.Combine(SC.ProjectRoot, "Saved", "EditorCooked", SC.CookPlatform);
+					CreatePluginManifest(SC, SC.FilesToStage.UFSFiles, StagedFileType.UFS, Params.DLCFile.GetFileNameWithoutAnyExtensions());
 				}
 				else
 				{
-					CookOutputDir = DirectoryReference.Combine(SC.ProjectRoot, "Saved", "Cooked", SC.CookPlatform);
+					CreatePluginManifest(SC, SC.FilesToStage.UFSFiles, StagedFileType.UFS, Params.ShortProjectName);
 				}
-				SC.MetadataDir = DirectoryReference.Combine(CookOutputDir, SC.ShortProjectName, "Metadata");
 			}
-		}
 
-		// Allow the calling scripts to make modifications to the deployment context before we finalize it
-		if (Params.ModifyDeploymentContextCallback != null)
-		{
-			Params.ModifyDeploymentContextCallback(Params, SC);
-		}
-
-		// Apply all the directory mappings
-		SC.FilesToStage.UFSFiles = SC.FilesToStage.UFSFiles.ToDictionary(x => ApplyDirectoryRemap(SC, x.Key), x => x.Value);
-		SC.FilesToStage.NonUFSFiles = SC.FilesToStage.NonUFSFiles.ToDictionary(x => ApplyDirectoryRemap(SC, x.Key), x => x.Value);
-		SC.FilesToStage.NonUFSDebugFiles = SC.FilesToStage.NonUFSDebugFiles.ToDictionary(x => ApplyDirectoryRemap(SC, x.Key), x => x.Value);
-		SC.FilesToStage.NonUFSSystemFiles = SC.FilesToStage.NonUFSSystemFiles.ToDictionary(x => ApplyDirectoryRemap(SC, x.Key), x => x.Value);
+			// Make all the filenames lowercase
+			if (SC.StageTargetPlatform.DeployLowerCaseFilenames(StagedFileType.NonUFS))
+			{
+				SC.FilesToStage.NonUFSFiles = SC.FilesToStage.NonUFSFiles.ToDictionary(x => x.Key.ToLowerInvariant(), x => x.Value);
+			}
+			// DebugNonUFS is weird, so ask per-file
+			SC.FilesToStage.NonUFSDebugFiles = SC.FilesToStage.NonUFSDebugFiles.ToDictionary(x => SC.StageTargetPlatform.DeployLowerCaseFile(x.Value, StagedFileType.DebugNonUFS) ? x.Key.ToLowerInvariant() : x.Key, x => x.Value);
+			if (!Params.UsePak(SC.StageTargetPlatform) && SC.StageTargetPlatform.DeployLowerCaseFilenames(StagedFileType.UFS))
+			{
+				SC.FilesToStage.UFSFiles = SC.FilesToStage.UFSFiles.ToDictionary(x => x.Key.ToLowerInvariant(), x => x.Value);
+			}
 
-		// Create plugin manifests after the directory mappings
-		if (bCreatePluginManifest && Params.UsePak(SC.StageTargetPlatform))
-		{
-			// Generate a plugin manifest if we're using a pak file and not creating a mod. Mods can be enumerated independently by users copying them into the Mods directory.
-			if (Params.HasDLCName)
+			// Remap all the non-ufs files if not using a PAK file
+			if (Params.HasIterateSharedCookedBuild)
 			{
-				CreatePluginManifest(SC, SC.FilesToStage.UFSFiles, StagedFileType.UFS, Params.DLCFile.GetFileNameWithoutAnyExtensions());
+				// Shared NonUFS files are staged in their remapped location, and may be duplicated in the to-stage list.
+				Dictionary<StagedFileReference, FileReference> NonUFSToStage = new Dictionary<StagedFileReference, FileReference>();
+				foreach (KeyValuePair<StagedFileReference, FileReference> StagedFilePair in SC.FilesToStage.NonUFSFiles)
+				{
+					NonUFSToStage[SC.StageTargetPlatform.Remap(StagedFilePair.Key)] = StagedFilePair.Value;
+				}
+				SC.FilesToStage.NonUFSFiles = NonUFSToStage;
 			}
 			else
 			{
-				CreatePluginManifest(SC, SC.FilesToStage.UFSFiles, StagedFileType.UFS, Params.ShortProjectName);
+				SC.FilesToStage.NonUFSFiles = SC.FilesToStage.NonUFSFiles.ToDictionary(x => SC.StageTargetPlatform.Remap(x.Key), x => x.Value);
 			}
-		}
 
-		// Make all the filenames lowercase
-		if (SC.StageTargetPlatform.DeployLowerCaseFilenames(StagedFileType.NonUFS))
-		{
-			SC.FilesToStage.NonUFSFiles = SC.FilesToStage.NonUFSFiles.ToDictionary(x => x.Key.ToLowerInvariant(), x => x.Value);
-		}
-		// DebugNonUFS is weird, so ask per-file
-		SC.FilesToStage.NonUFSDebugFiles = SC.FilesToStage.NonUFSDebugFiles.ToDictionary(x => SC.StageTargetPlatform.DeployLowerCaseFile(x.Value, StagedFileType.DebugNonUFS) ? x.Key.ToLowerInvariant() : x.Key, x => x.Value);
-		if (!Params.UsePak(SC.StageTargetPlatform) && SC.StageTargetPlatform.DeployLowerCaseFilenames(StagedFileType.UFS))
-		{
-			SC.FilesToStage.UFSFiles = SC.FilesToStage.UFSFiles.ToDictionary(x => x.Key.ToLowerInvariant(), x => x.Value);
-		}
+			if (!Params.UsePak(SC.StageTargetPlatform))
+			{
+				SC.FilesToStage.UFSFiles = SC.FilesToStage.UFSFiles.ToDictionary(x => SC.StageTargetPlatform.Remap(x.Key), x => x.Value);
+			}
 
-		// Remap all the non-ufs files if not using a PAK file
-		if (Params.HasIterateSharedCookedBuild)
-		{
-			// Shared NonUFS files are staged in their remapped location, and may be duplicated in the to-stage list.
-			Dictionary<StagedFileReference, FileReference> NonUFSToStage = new Dictionary<StagedFileReference, FileReference>();
-			foreach (KeyValuePair<StagedFileReference, FileReference> StagedFilePair in SC.FilesToStage.NonUFSFiles)
+			// Merge all the NonUFS system files back into the NonUFS list. Deployment is currently only set up to read from that.
+			foreach (KeyValuePair<StagedFileReference, FileReference> Pair in SC.FilesToStage.NonUFSSystemFiles)
 			{
-				NonUFSToStage[SC.StageTargetPlatform.Remap(StagedFilePair.Key)] = StagedFilePair.Value;
+				SC.FilesToStage.NonUFSFiles[Pair.Key] = Pair.Value;
 			}
-			SC.FilesToStage.NonUFSFiles = NonUFSToStage;
-		}
-		else
-		{
-			SC.FilesToStage.NonUFSFiles = SC.FilesToStage.NonUFSFiles.ToDictionary(x => SC.StageTargetPlatform.Remap(x.Key), x => x.Value);
-		}
 
-		if (!Params.UsePak(SC.StageTargetPlatform))
-		{
-			SC.FilesToStage.UFSFiles = SC.FilesToStage.UFSFiles.ToDictionary(x => SC.StageTargetPlatform.Remap(x.Key), x => x.Value);
-		}
+			// Make sure there are no restricted folders in the output
+			HashSet<StagedFileReference> RestrictedFiles = new HashSet<StagedFileReference>();
+			foreach (string RestrictedName in SC.RestrictedFolderNames)
+			{
+				RestrictedFiles.UnionWith(SC.FilesToStage.UFSFiles.Keys.Where(x => x.ContainsName(RestrictedName)));
+				RestrictedFiles.UnionWith(SC.FilesToStage.NonUFSFiles.Keys.Where(x => x.ContainsName(RestrictedName)));
+				RestrictedFiles.UnionWith(SC.FilesToStage.NonUFSDebugFiles.Keys.Where(x => x.ContainsName(RestrictedName)));
+				RestrictedFiles.UnionWith(SC.CrashReporterUFSFiles.Keys.Where(x => x.ContainsName(RestrictedName)));
+			}
+			RestrictedFiles.RemoveWhere(RestrictedFile => SC.WhitelistDirectories.Any(WhitelistDirectory => RestrictedFile.Directory.IsUnderDirectory(WhitelistDirectory)));
 
-		// Merge all the NonUFS system files back into the NonUFS list. Deployment is currently only set up to read from that.
-		foreach (KeyValuePair<StagedFileReference, FileReference> Pair in SC.FilesToStage.NonUFSSystemFiles)
-		{
-			SC.FilesToStage.NonUFSFiles[Pair.Key] = Pair.Value;
-		}
+			if (RestrictedFiles.Count > 0)
+			{
+				List<string> RestrictedNames = new List<string>();
+				foreach (string RestrictedFolderName in SC.RestrictedFolderNames)
+				{
+					if (RestrictedFiles.Any(x => x.ContainsName(RestrictedFolderName)))
+					{
+						RestrictedNames.Add(String.Format("\"{0}\"", RestrictedFolderName));
+					}
+				}
 
-		// Make sure there are no restricted folders in the output
-		HashSet<StagedFileReference> RestrictedFiles = new HashSet<StagedFileReference>();
-		foreach (string RestrictedName in SC.RestrictedFolderNames)
-		{
-			RestrictedFiles.UnionWith(SC.FilesToStage.UFSFiles.Keys.Where(x => x.ContainsName(RestrictedName)));
-			RestrictedFiles.UnionWith(SC.FilesToStage.NonUFSFiles.Keys.Where(x => x.ContainsName(RestrictedName)));
-			RestrictedFiles.UnionWith(SC.FilesToStage.NonUFSDebugFiles.Keys.Where(x => x.ContainsName(RestrictedName)));
-			RestrictedFiles.UnionWith(SC.CrashReporterUFSFiles.Keys.Where(x => x.ContainsName(RestrictedName)));
+				StringBuilder Message = new StringBuilder();
+				Message.AppendFormat("The following files are set to be staged, but contain restricted folder names ({0}):", String.Join(", ", RestrictedNames));
+				foreach (StagedFileReference RestrictedFile in RestrictedFiles.OrderBy(x => x.Name))
+				{
+					Message.AppendFormat("\n{0}", RestrictedFile);
+				}
+				Message.Append("\n[Restrictions]");
+				foreach (string RestrictedName in SC.RestrictedFolderNames)
+				{
+					Message.AppendFormat("\n{0}", RestrictedName);
+				}
+				Message.Append("\nIf these files are intended to be distributed in packaged builds, move the source files out of a restricted folder, or remap them during staging using the following syntax in DefaultGame.ini:");
+				Message.Append("\n[Staging]");
+				Message.Append("\n+RemapDirectories=(From=\"Foo/NoRedist\", To=\"Foo\")");
+				if (RestrictedNames.Any(x => x != "NotForLicensees" && x != "NoRedist")) // We don't ever want internal stuff white-listing folders like this
+				{
+					Message.Append("\nAlternatively, whitelist them using this syntax in DefaultGame.ini:");
+					Message.Append("\n[Staging]");
+					Message.Append("\n+WhitelistDirectories=MyGame/Content/Foo");
+				}
+				throw new AutomationException(Message.ToString());
+			}
 		}
-		RestrictedFiles.RemoveWhere(RestrictedFile => SC.WhitelistDirectories.Any(WhitelistDirectory => RestrictedFile.Directory.IsUnderDirectory(WhitelistDirectory)));
 
-		if (RestrictedFiles.Count > 0)
+		/// <summary>
+		/// Stage the appropriate config files from the given input directory
+		/// </summary>
+		/// <param name="SC">The staging context</param>
+		/// <param name="ConfigDir">Directory containing the config files</param>
+		/// <param name="PlatformExtensionName">The name of the platform when staging config files from a platform extension</param>
+		static void StageConfigFiles(DeploymentContext SC, DirectoryReference ConfigDir, string PlatformExtensionName)
 		{
-			List<string> RestrictedNames = new List<string>();
-			foreach (string RestrictedFolderName in SC.RestrictedFolderNames)
+			List<FileReference> ConfigFiles = SC.FindFilesToStage(ConfigDir, "*.ini", StageFilesSearch.AllDirectories);
+			foreach (FileReference ConfigFile in ConfigFiles)
 			{
-				if (RestrictedFiles.Any(x => x.ContainsName(RestrictedFolderName)))
+				Nullable<bool> ShouldStage = ShouldStageConfigFile(SC, ConfigDir, ConfigFile, PlatformExtensionName);
+				if (ShouldStage == null)
 				{
-					RestrictedNames.Add(String.Format("\"{0}\"", RestrictedFolderName));
+					CommandUtils.LogWarning("The config file '{0}' will be staged, but is not whitelisted or blacklisted. Add +WhitelistConfigFiles={0} or +BlacklistConfigFiles={0} to the [Staging] section of DefaultGame.ini", SC.GetStagedFileLocation(ConfigFile));
 				}
-			}
 
-			StringBuilder Message = new StringBuilder();
-			Message.AppendFormat("The following files are set to be staged, but contain restricted folder names ({0}):", String.Join(", ", RestrictedNames));
-			foreach (StagedFileReference RestrictedFile in RestrictedFiles.OrderBy(x => x.Name))
-			{
-				Message.AppendFormat("\n{0}", RestrictedFile);
+				if (ShouldStage ?? true)
+				{
+					SC.StageFile(StagedFileType.UFS, ConfigFile);
+				}
+				else
+				{
+					CommandUtils.LogInformation("Excluding config file {0}", ConfigFile);
+				}
 			}
-			Message.Append("\n[Restrictions]");
-			foreach (string RestrictedName in SC.RestrictedFolderNames)
+		}
+
+		/// <summary>
+		/// Looks for a potential PlatformExtension config directory
+		/// </summary>
+		/// <param name="SC">The staging context</param>
+		/// <param name="PlatformName">The name of the platform/platform parent</param>
+		static void StagePlatformExtensionConfigFiles(DeploymentContext SC, string PlatformName)
+		{
+			DirectoryReference PlatformEngineConfigDir = DirectoryReference.Combine(SC.LocalRoot, "Engine", "Platforms", PlatformName, "Config");
+			DirectoryReference PlatformProjectConfigDir = DirectoryReference.Combine(SC.ProjectRoot, "Platforms", PlatformName, "Config");
+
+			if (DirectoryReference.Exists(PlatformEngineConfigDir))
 			{
-				Message.AppendFormat("\n{0}", RestrictedName);
+				StageConfigFiles(SC, PlatformEngineConfigDir, PlatformName);
+
 			}
-			Message.Append("\nIf these files are intended to be distributed in packaged builds, move the source files out of a restricted folder, or remap them during staging using the following syntax in DefaultGame.ini:");
-			Message.Append("\n[Staging]");
-			Message.Append("\n+RemapDirectories=(From=\"Foo/NoRedist\", To=\"Foo\")");
-			if (RestrictedNames.Any(x => x != "NotForLicensees" && x != "NoRedist")) // We don't ever want internal stuff white-listing folders like this
+			if (DirectoryReference.Exists(PlatformProjectConfigDir))
 			{
-				Message.Append("\nAlternatively, whitelist them using this syntax in DefaultGame.ini:");
-				Message.Append("\n[Staging]");
-				Message.Append("\n+WhitelistDirectories=MyGame/Content/Foo");
+				StageConfigFiles(SC, PlatformProjectConfigDir, PlatformName);
 			}
-			throw new AutomationException(Message.ToString());
 		}
-	}
 
-	/// <summary>
-	/// Stage the appropriate config files from the given input directory
-	/// </summary>
-	/// <param name="SC">The staging context</param>
-	/// <param name="ConfigDir">Directory containing the config files</param>
-	/// <param name="PlatformExtensionName">The name of the platform when staging config files from a platform extension</param>
-	static void StageConfigFiles(DeploymentContext SC, DirectoryReference ConfigDir, string PlatformExtensionName)
-	{
-		List<FileReference> ConfigFiles = SC.FindFilesToStage(ConfigDir, "*.ini", StageFilesSearch.AllDirectories);
-		foreach (FileReference ConfigFile in ConfigFiles)
+		/// <summary>
+		/// Determines if an individual config file should be staged
+		/// </summary>
+		/// <param name="SC">The staging context</param>
+		/// <param name="ConfigDir">Directory containing the config files</param>
+		/// <param name="ConfigFile">The config file to check</param>
+		/// <returns>True if the file should be staged, false otherwise</returns>
+		static Nullable<bool> ShouldStageConfigFile(DeploymentContext SC, DirectoryReference ConfigDir, FileReference ConfigFile, string PlatformExtensionName)
 		{
-			Nullable<bool> ShouldStage = ShouldStageConfigFile(SC, ConfigDir, ConfigFile, PlatformExtensionName);
-			if (ShouldStage == null)
+			StagedFileReference StagedConfigFile = SC.GetStagedFileLocation(ConfigFile);
+			if (SC.WhitelistConfigFiles.Contains(StagedConfigFile))
 			{
-				CommandUtils.LogWarning("The config file '{0}' will be staged, but is not whitelisted or blacklisted. Add +WhitelistConfigFiles={0} or +BlacklistConfigFiles={0} to the [Staging] section of DefaultGame.ini", SC.GetStagedFileLocation(ConfigFile));
-			}
-
-			if (ShouldStage ?? true)
-			{
-				SC.StageFile(StagedFileType.UFS, ConfigFile);
+				return true;
 			}
-			else
+			if (SC.BlacklistConfigFiles.Contains(StagedConfigFile))
 			{
-				CommandUtils.LogInformation("Excluding config file {0}", ConfigFile);
+				return false;
 			}
-		}
-	}
 
-	/// <summary>
-	/// Looks for a potential PlatformExtension config directory
-	/// </summary>
-	/// <param name="SC">The staging context</param>
-	/// <param name="PlatformName">The name of the platform/platform parent</param>
-	static void StagePlatformExtensionConfigFiles(DeploymentContext SC, string PlatformName)
-	{
-		DirectoryReference PlatformEngineConfigDir = DirectoryReference.Combine(SC.LocalRoot, "Engine", "Platforms", PlatformName, "Config");
-		DirectoryReference PlatformProjectConfigDir = DirectoryReference.Combine(SC.ProjectRoot, "Platforms", PlatformName, "Config");
+			string NormalizedPath = ConfigFile.MakeRelativeTo(ConfigDir).ToLowerInvariant().Replace('\\', '/');
 
-		if (DirectoryReference.Exists(PlatformEngineConfigDir))
-		{
-			StageConfigFiles(SC, PlatformEngineConfigDir, PlatformName);
-
-		}
-		if (DirectoryReference.Exists(PlatformProjectConfigDir))
-		{
-			StageConfigFiles(SC, PlatformProjectConfigDir, PlatformName);
-		}
-	}
+			int DirectoryIdx = NormalizedPath.IndexOf('/');
+			if (DirectoryIdx == -1)
+			{
+				const string BasePrefix = "base";
+				if (NormalizedPath.StartsWith(BasePrefix))
+				{
+					string ShortName = NormalizedPath.Substring(BasePrefix.Length);
+					if (PlatformExtensionName != null)
+					{
+						if (!ShortName.StartsWith(PlatformExtensionName, StringComparison.InvariantCultureIgnoreCase))
+						{
+							// Ignore config files in the platform directory that don't start with the platform name.
+							return false;
+						}
 
-	/// <summary>
-	/// Determines if an individual config file should be staged
-	/// </summary>
-	/// <param name="SC">The staging context</param>
-	/// <param name="ConfigDir">Directory containing the config files</param>
-	/// <param name="ConfigFile">The config file to check</param>
-	/// <returns>True if the file should be staged, false otherwise</returns>
-	static Nullable<bool> ShouldStageConfigFile(DeploymentContext SC, DirectoryReference ConfigDir, FileReference ConfigFile, string PlatformExtensionName)
-	{
-		StagedFileReference StagedConfigFile = SC.GetStagedFileLocation(ConfigFile);
-		if (SC.WhitelistConfigFiles.Contains(StagedConfigFile))
-		{
-			return true;
-		}
-		if (SC.BlacklistConfigFiles.Contains(StagedConfigFile))
-		{
-			return false;
-		}
+						ShortName = ShortName.Substring(PlatformExtensionName.Length);
+					}
 
-		string NormalizedPath = ConfigFile.MakeRelativeTo(ConfigDir).ToLowerInvariant().Replace('\\', '/');
+					return ShouldStageConfigSuffix(SC, ConfigFile, ShortName);
+				}
 
-		int DirectoryIdx = NormalizedPath.IndexOf('/');
-		if (DirectoryIdx == -1)
-		{
-			const string BasePrefix = "base";
-			if (NormalizedPath.StartsWith(BasePrefix))
-			{
-				string ShortName = NormalizedPath.Substring(BasePrefix.Length);
-				if (PlatformExtensionName != null)
+				const string DefaultPrefix = "default";
+				if (NormalizedPath.StartsWith(DefaultPrefix))
 				{
-					if (!ShortName.StartsWith(PlatformExtensionName, StringComparison.InvariantCultureIgnoreCase))
+					string ShortName = NormalizedPath.Substring(DefaultPrefix.Length);
+					if (PlatformExtensionName != null)
 					{
-						// Ignore config files in the platform directory that don't start with the platform name.
-						return false;
+						if (!ShortName.StartsWith(PlatformExtensionName, StringComparison.InvariantCultureIgnoreCase))
+						{
+							// Ignore config files in the platform directory that don't start with the platform name.
+							return false;
+						}
+
+						ShortName = ShortName.Substring(PlatformExtensionName.Length);
 					}
 
-					ShortName = ShortName.Substring(PlatformExtensionName.Length);
+					return ShouldStageConfigSuffix(SC, ConfigFile, ShortName);
 				}
 
-				return ShouldStageConfigSuffix(SC, ConfigFile, ShortName);
-			}
+				const string DedicatedServerPrefix = "dedicatedserver";
+				if (NormalizedPath.StartsWith(DedicatedServerPrefix))
+				{
+					return SC.DedicatedServer ? ShouldStageConfigSuffix(SC, ConfigFile, NormalizedPath.Substring(DedicatedServerPrefix.Length)) : false;
+				}
+
+				if (NormalizedPath == "consolevariables.ini")
+				{
+					return SC.StageTargetConfigurations.Any(x => x != UnrealTargetConfiguration.Test && x != UnrealTargetConfiguration.Shipping);
+				}
+
+				if (NormalizedPath == "locgatherconfig.ini")
+				{
+					return false;
+				}
+
+				if (NormalizedPath == "designertoolsconfig.ini")
+				{
+					return false;
+				}
 
-			const string DefaultPrefix = "default";
-			if (NormalizedPath.StartsWith(DefaultPrefix))
-			{
-				string ShortName = NormalizedPath.Substring(DefaultPrefix.Length);
 				if (PlatformExtensionName != null)
 				{
-					if (!ShortName.StartsWith(PlatformExtensionName, StringComparison.InvariantCultureIgnoreCase))
+					if (NormalizedPath.StartsWith(PlatformExtensionName, StringComparison.InvariantCultureIgnoreCase))
 					{
-						// Ignore config files in the platform directory that don't start with the platform name.
-						return false;
+						string ShortName = NormalizedPath.Substring(PlatformExtensionName.Length);
+						return ShouldStageConfigSuffix(SC, ConfigFile, ShortName);
 					}
 
-					ShortName = ShortName.Substring(PlatformExtensionName.Length);
+					if (NormalizedPath == "datadrivenplatforminfo.ini")
+					{
+						return true;
+					}
 				}
-
-				return ShouldStageConfigSuffix(SC, ConfigFile, ShortName);
-			}
-
-			const string DedicatedServerPrefix = "dedicatedserver";
-			if (NormalizedPath.StartsWith(DedicatedServerPrefix))
-			{
-				return SC.DedicatedServer ? ShouldStageConfigSuffix(SC, ConfigFile, NormalizedPath.Substring(DedicatedServerPrefix.Length)) : false;
 			}
-
-			if (NormalizedPath == "consolevariables.ini")
+			else
 			{
-				return SC.StageTargetConfigurations.Any(x => x != UnrealTargetConfiguration.Test && x != UnrealTargetConfiguration.Shipping);
-			}
+				if (NormalizedPath.StartsWith("layouts/"))
+				{
+					return true;
+				}
 
-			if (NormalizedPath == "locgatherconfig.ini")
-			{
-				return false;
-			}
+				if (NormalizedPath.StartsWith("localization/"))
+				{
+					return false;
+				}
 
-			if (NormalizedPath == "designertoolsconfig.ini")
-			{
-				return false;
-			}
+				string PlatformPrefix = String.Format("{0}/{0}", NormalizedPath.Substring(0, DirectoryIdx));
+				if (NormalizedPath.StartsWith(PlatformPrefix))
+				{
+					return ShouldStageConfigSuffix(SC, ConfigFile, NormalizedPath.Substring(PlatformPrefix.Length));
+				}
 
-			if (PlatformExtensionName != null)
-			{
-				if (NormalizedPath.StartsWith(PlatformExtensionName, StringComparison.InvariantCultureIgnoreCase))
+				string PlatformBasePrefix = String.Format("{0}/base{0}", NormalizedPath.Substring(0, DirectoryIdx));
+				if (NormalizedPath.StartsWith(PlatformBasePrefix))
 				{
-					string ShortName = NormalizedPath.Substring(PlatformExtensionName.Length);
-					return ShouldStageConfigSuffix(SC, ConfigFile, ShortName);
+					return ShouldStageConfigSuffix(SC, ConfigFile, NormalizedPath.Substring(PlatformBasePrefix.Length));
 				}
 
-				if (NormalizedPath == "datadrivenplatforminfo.ini")
+				if (NormalizedPath.EndsWith("/datadrivenplatforminfo.ini"))
 				{
 					return true;
 				}
+
 			}
+			return null;
 		}
-		else
+
+		/// <summary>
+		/// Determines if the given config file suffix ("engine", "game", etc...) should be staged for the given context.
+		/// </summary>
+		/// <param name="SC">The staging context</param>
+		/// <param name="ConfigFile">Full path to the config file</param>
+		/// <param name="InvariantSuffix">Suffix for the config file, as a lowercase invariant string</param>
+		/// <returns>True if the suffix should be staged, false if not, null if unknown</returns>
+		static Nullable<bool> ShouldStageConfigSuffix(DeploymentContext SC, FileReference ConfigFile, string InvariantSuffix)
 		{
-			if (NormalizedPath.StartsWith("layouts/"))
+			if (SC.IniSuffixWhitelist != null && SC.IniSuffixWhitelist.Contains(InvariantSuffix))
 			{
 				return true;
 			}
 
-			if (NormalizedPath.StartsWith("localization/"))
+			if (SC.IniSuffixBlacklist != null && SC.IniSuffixBlacklist.Contains(InvariantSuffix))
 			{
 				return false;
 			}
 
-			string PlatformPrefix = String.Format("{0}/{0}", NormalizedPath.Substring(0, DirectoryIdx));
-			if (NormalizedPath.StartsWith(PlatformPrefix))
-			{
-				return ShouldStageConfigSuffix(SC, ConfigFile, NormalizedPath.Substring(PlatformPrefix.Length));
-			}
-
-			string PlatformBasePrefix = String.Format("{0}/base{0}", NormalizedPath.Substring(0, DirectoryIdx));
-			if (NormalizedPath.StartsWith(PlatformBasePrefix))
-			{
-				return ShouldStageConfigSuffix(SC, ConfigFile, NormalizedPath.Substring(PlatformBasePrefix.Length));
-			}
-
-			if (NormalizedPath.EndsWith("/datadrivenplatforminfo.ini"))
-			{
-				return true;
-			}
-
-		}
-		return null;
-	}
-
-	/// <summary>
-	/// Determines if the given config file suffix ("engine", "game", etc...) should be staged for the given context.
-	/// </summary>
-	/// <param name="SC">The staging context</param>
-	/// <param name="ConfigFile">Full path to the config file</param>
-	/// <param name="InvariantSuffix">Suffix for the config file, as a lowercase invariant string</param>
-	/// <returns>True if the suffix should be staged, false if not, null if unknown</returns>
-	static Nullable<bool> ShouldStageConfigSuffix(DeploymentContext SC, FileReference ConfigFile, string InvariantSuffix)
-	{
-		if (SC.IniSuffixWhitelist != null && SC.IniSuffixWhitelist.Contains(InvariantSuffix))
-		{
-			return true;
-		}
-
-		if (SC.IniSuffixBlacklist != null && SC.IniSuffixBlacklist.Contains(InvariantSuffix))
-		{
-			return false;
+			return null;
 		}
 
-		return null;
-	}
-
-	public static StagedFileReference ApplyDirectoryRemap(DeploymentContext SC, StagedFileReference InputFile)
-	{
-		StagedFileReference CurrentFile = InputFile;
-		foreach (Tuple<StagedDirectoryReference, StagedDirectoryReference> RemapDirectory in SC.RemapDirectories)
+		public static StagedFileReference ApplyDirectoryRemap(DeploymentContext SC, StagedFileReference InputFile)
 		{
-			StagedFileReference NewFile;
-			if (StagedFileReference.TryRemap(CurrentFile, RemapDirectory.Item1, RemapDirectory.Item2, out NewFile))
+			StagedFileReference CurrentFile = InputFile;
+			foreach (Tuple<StagedDirectoryReference, StagedDirectoryReference> RemapDirectory in SC.RemapDirectories)
 			{
-				CurrentFile = NewFile;
+				StagedFileReference NewFile;
+				if (StagedFileReference.TryRemap(CurrentFile, RemapDirectory.Item1, RemapDirectory.Item2, out NewFile))
+				{
+					CurrentFile = NewFile;
+				}
 			}
+			return CurrentFile;
 		}
-		return CurrentFile;
-	}
 
-	static List<string> ParseInputPaths(List<string> ConfigLines)
-	{
-		// Each string has the format '(Path="TheDirToStage")'
-		List<string> InputPaths = new List<string>();
-		foreach (string ConfigLine in ConfigLines)
+		static List<string> ParseInputPaths(List<string> ConfigLines)
 		{
-			string[] PathParts = ConfigLine.Split('"');
-			if (PathParts.Length == 3)
-			{
-				InputPaths.Add(PathParts[1]);
-			}
-			else if (PathParts.Length == 1)
+			// Each string has the format '(Path="TheDirToStage")'
+			List<string> InputPaths = new List<string>();
+			foreach (string ConfigLine in ConfigLines)
 			{
-				InputPaths.Add(PathParts[0]);
+				string[] PathParts = ConfigLine.Split('"');
+				if (PathParts.Length == 3)
+				{
+					InputPaths.Add(PathParts[1]);
+				}
+				else if (PathParts.Length == 1)
+				{
+					InputPaths.Add(PathParts[0]);
+				}
 			}
+			return InputPaths;
 		}
-		return InputPaths;
-	}
 
-	static void CreatePluginManifest(DeploymentContext SC, Dictionary<StagedFileReference, FileReference> FileMapping, StagedFileType FileType, string ManifestName)
-	{
-		// Get the path to the project's mods directory. We wont include anything under here in the manifest
-		DirectoryReference ModsDir = DirectoryReference.Combine(SC.ProjectRoot, "Mods");
-
-		// Find all the plugins that are being staged
-		Dictionary<StagedFileReference, PluginDescriptor> StagedPlugins = new Dictionary<StagedFileReference, PluginDescriptor>();
-		foreach (KeyValuePair<StagedFileReference, FileReference> File in FileMapping)
+		static void CreatePluginManifest(DeploymentContext SC, Dictionary<StagedFileReference, FileReference> FileMapping, StagedFileType FileType, string ManifestName)
 		{
-			if (!File.Value.IsUnderDirectory(ModsDir) && File.Value.HasExtension(".uplugin"))
+			// Get the path to the project's mods directory. We wont include anything under here in the manifest
+			DirectoryReference ModsDir = DirectoryReference.Combine(SC.ProjectRoot, "Mods");
+
+			// Find all the plugins that are being staged
+			Dictionary<StagedFileReference, PluginDescriptor> StagedPlugins = new Dictionary<StagedFileReference, PluginDescriptor>();
+			foreach (KeyValuePair<StagedFileReference, FileReference> File in FileMapping)
 			{
-				PluginDescriptor Descriptor = PluginDescriptor.FromFile(File.Value);
-				StagedPlugins[File.Key] = Descriptor;
+				if (!File.Value.IsUnderDirectory(ModsDir) && File.Value.HasExtension(".uplugin"))
+				{
+					PluginDescriptor Descriptor = PluginDescriptor.FromFile(File.Value);
+					StagedPlugins[File.Key] = Descriptor;
+				}
 			}
-		}
 
-		// If we have plugins, write out the manifest
-		if (StagedPlugins.Count > 0)
-		{
-			string PluginManifestName = String.Format("{0}.upluginmanifest", ManifestName);
+			// If we have plugins, write out the manifest
+			if (StagedPlugins.Count > 0)
+			{
+				string PluginManifestName = String.Format("{0}.upluginmanifest", ManifestName);
 
-			FileReference PluginManifestFile = FileReference.Combine(SC.ProjectRoot, "Intermediate", "Staging", PluginManifestName);
-			DirectoryReference.CreateDirectory(PluginManifestFile.Directory);
+				FileReference PluginManifestFile = FileReference.Combine(SC.ProjectRoot, "Intermediate", "Staging", PluginManifestName);
+				DirectoryReference.CreateDirectory(PluginManifestFile.Directory);
 
-			using (JsonWriter Writer = new JsonWriter(PluginManifestFile.FullName))
-			{
-				Writer.WriteObjectStart();
-				Writer.WriteArrayStart("Contents");
-				foreach (KeyValuePair<StagedFileReference, PluginDescriptor> StagedPlugin in StagedPlugins)
+				using (JsonWriter Writer = new JsonWriter(PluginManifestFile.FullName))
 				{
 					Writer.WriteObjectStart();
-					Writer.WriteValue("File", String.Format("../../../{0}", StagedPlugin.Key.Name));
-					Writer.WriteObjectStart("Descriptor");
-					StagedPlugin.Value.Write(Writer);
-					Writer.WriteObjectEnd();
+					Writer.WriteArrayStart("Contents");
+					foreach (KeyValuePair<StagedFileReference, PluginDescriptor> StagedPlugin in StagedPlugins)
+					{
+						Writer.WriteObjectStart();
+						Writer.WriteValue("File", String.Format("../../../{0}", StagedPlugin.Key.Name));
+						Writer.WriteObjectStart("Descriptor");
+						StagedPlugin.Value.Write(Writer);
+						Writer.WriteObjectEnd();
+						Writer.WriteObjectEnd();
+					}
+					Writer.WriteArrayEnd();
 					Writer.WriteObjectEnd();
 				}
-				Writer.WriteArrayEnd();
-				Writer.WriteObjectEnd();
-			}
 
-			SC.StageFile(FileType, PluginManifestFile, StagedFileReference.Combine(SC.RelativeProjectRootForStage, "Plugins", PluginManifestName));
+				SC.StageFile(FileType, PluginManifestFile, StagedFileReference.Combine(SC.RelativeProjectRootForStage, "Plugins", PluginManifestName));
+			}
 		}
-	}
 
-	public static void DumpTargetManifest(Dictionary<StagedFileReference, FileReference> Mapping, FileReference Filename, DirectoryReference StageDir, HashSet<StagedFileReference> CRCFiles)
-	{
-		// const string Iso8601DateTimeFormat = "yyyy-MM-ddTHH:mm:ssZ"; // probably should work
-		// const string Iso8601DateTimeFormat = "o"; // predefined universal Iso standard format (has too many millisecond spaces for our read code in FDateTime.ParseISO8601
-		const string Iso8601DateTimeFormat = "yyyy'-'MM'-'dd'T'HH':'mm':'ss'.'fffZ";
+		public static void DumpTargetManifest(Dictionary<StagedFileReference, FileReference> Mapping, FileReference Filename, DirectoryReference StageDir, HashSet<StagedFileReference> CRCFiles)
+		{
+			// const string Iso8601DateTimeFormat = "yyyy-MM-ddTHH:mm:ssZ"; // probably should work
+			// const string Iso8601DateTimeFormat = "o"; // predefined universal Iso standard format (has too many millisecond spaces for our read code in FDateTime.ParseISO8601
+			const string Iso8601DateTimeFormat = "yyyy'-'MM'-'dd'T'HH':'mm':'ss'.'fffZ";
 
 
-		if (Mapping.Count > 0)
-		{
-			var Lines = new List<string>();
-			foreach (var Pair in Mapping)
+			if (Mapping.Count > 0)
 			{
-				string TimeStamp = FileReference.GetLastWriteTimeUtc(Pair.Value).ToString(Iso8601DateTimeFormat);
-				if (CRCFiles.Contains(Pair.Key))
+				var Lines = new List<string>();
+				foreach (var Pair in Mapping)
 				{
-					byte[] FileData = File.ReadAllBytes(StageDir + "/" + Pair.Key);
-					TimeStamp = BitConverter.ToString(System.Security.Cryptography.MD5.Create().ComputeHash(FileData)).Replace("-", string.Empty);
+					string TimeStamp = FileReference.GetLastWriteTimeUtc(Pair.Value).ToString(Iso8601DateTimeFormat);
+					if (CRCFiles.Contains(Pair.Key))
+					{
+						byte[] FileData = File.ReadAllBytes(StageDir + "/" + Pair.Key);
+						TimeStamp = BitConverter.ToString(System.Security.Cryptography.MD5.Create().ComputeHash(FileData)).Replace("-", string.Empty);
+					}
+					string Dest = Pair.Key.Name + "\t" + TimeStamp;
+
+					Lines.Add(Dest);
 				}
-				string Dest = Pair.Key.Name + "\t" + TimeStamp;
 
-				Lines.Add(Dest);
+				DirectoryReference.CreateDirectory(Filename.Directory);
+				WriteAllLines(Filename.FullName, Lines.ToArray());
 			}
-
-			DirectoryReference.CreateDirectory(Filename.Directory);
-			WriteAllLines(Filename.FullName, Lines.ToArray());
 		}
-	}
 
-	private static void CopyManifestFilesToStageDir(DeploymentContext SC, Dictionary<StagedFileReference, FileReference> Mapping, string ManifestName)
-	{
-		CopyOrWriteManifestFilesToStageDir(true, SC, Mapping, ManifestName);
-	}
+		private static void CopyManifestFilesToStageDir(DeploymentContext SC, Dictionary<StagedFileReference, FileReference> Mapping, string ManifestName)
+		{
+			CopyOrWriteManifestFilesToStageDir(true, SC, Mapping, ManifestName);
+		}
 
-	private static void WriteManifestFilesToStageDir(DeploymentContext SC, Dictionary<StagedFileReference, FileReference> Mapping, string ManifestName)
-	{
-		CopyOrWriteManifestFilesToStageDir(false, SC, Mapping, ManifestName);
-	}
+		private static void WriteManifestFilesToStageDir(DeploymentContext SC, Dictionary<StagedFileReference, FileReference> Mapping, string ManifestName)
+		{
+			CopyOrWriteManifestFilesToStageDir(false, SC, Mapping, ManifestName);
+		}
 
-	private static void CopyOrWriteManifestFilesToStageDir(bool bPerformCopy, DeploymentContext SC, Dictionary<StagedFileReference, FileReference> Mapping, string ManifestName)
-	{
-		if (bPerformCopy)
+		private static void CopyOrWriteManifestFilesToStageDir(bool bPerformCopy, DeploymentContext SC, Dictionary<StagedFileReference, FileReference> Mapping, string ManifestName)
 		{
-			LogInformation("Copying {0} to staging directory: {1}", ManifestName, SC.StageDirectory);
-			foreach (KeyValuePair<StagedFileReference, FileReference> Pair in Mapping)
+			if (bPerformCopy)
 			{
-				FileReference Src = Pair.Value;
-				FileReference Dest = FileReference.Combine(SC.StageDirectory, Pair.Key.Name);
-				if (Src != Dest)  // special case for things created in the staging directory, like the pak file
+				LogInformation("Copying {0} to staging directory: {1}", ManifestName, SC.StageDirectory);
+				foreach (KeyValuePair<StagedFileReference, FileReference> Pair in Mapping)
 				{
-					CopyFileIncremental(Src, Dest, IniKeyBlacklist: SC.IniKeyBlacklist, IniSectionBlacklist: SC.IniSectionBlacklist);
+					FileReference Src = Pair.Value;
+					FileReference Dest = FileReference.Combine(SC.StageDirectory, Pair.Key.Name);
+					if (Src != Dest)  // special case for things created in the staging directory, like the pak file
+					{
+						CopyFileIncremental(Src, Dest, IniKeyBlacklist: SC.IniKeyBlacklist, IniSectionBlacklist: SC.IniSectionBlacklist);
+					}
 				}
 			}
-		}
 
-		if (Mapping.Count > 0)
-		{
-			var ManifestPathStaged = FileReference.Combine(SC.DebugStageDirectory, $"Manifest_{ManifestName}_{SC.StageTargetPlatform.PlatformType}.txt");
-			var ManifestPathLog = FileReference.Combine(new DirectoryReference(CmdEnv.LogFolder), ManifestPathStaged.GetFileName());
+			if (Mapping.Count > 0)
+			{
+				var ManifestPathStaged = FileReference.Combine(SC.DebugStageDirectory, $"Manifest_{ManifestName}_{SC.StageTargetPlatform.PlatformType}.txt");
+				var ManifestPathLog = FileReference.Combine(new DirectoryReference(CmdEnv.LogFolder), ManifestPathStaged.GetFileName());
 
-			DeleteFile(ManifestPathStaged.FullName);
-			DeleteFile(ManifestPathLog.FullName);
+				DeleteFile(ManifestPathStaged.FullName);
+				DeleteFile(ManifestPathLog.FullName);
 
-			DumpTargetManifest(Mapping, ManifestPathLog, SC.StageDirectory, SC.StageTargetPlatform.GetFilesForCRCCheck());
+				DumpTargetManifest(Mapping, ManifestPathLog, SC.StageDirectory, SC.StageTargetPlatform.GetFilesForCRCCheck());
 
-			if (!FileReference.Exists(ManifestPathLog))
-			{
-				throw new AutomationException("Failed to write manifest {0}", ManifestPathLog);
-			}
+				if (!FileReference.Exists(ManifestPathLog))
+				{
+					throw new AutomationException("Failed to write manifest {0}", ManifestPathLog);
+				}
 
-			if (SC.StageTargetPlatform.RequiresManifestFiles)
-			{
-				// Also copy the manifest file to the staged directory
-				CopyFile(ManifestPathLog.FullName, ManifestPathStaged.FullName);
+				if (SC.StageTargetPlatform.RequiresManifestFiles)
+				{
+					// Also copy the manifest file to the staged directory
+					CopyFile(ManifestPathLog.FullName, ManifestPathStaged.FullName);
+				}
 			}
 		}
-	}
 
-	public static void DumpManifest(Dictionary<StagedFileReference, FileReference> Mapping, string Filename)
-	{
-		if (Mapping.Count > 0)
+		public static void DumpManifest(Dictionary<StagedFileReference, FileReference> Mapping, string Filename)
 		{
-			List<string> Lines = new List<string>();
-			foreach (KeyValuePair<StagedFileReference, FileReference> Pair in Mapping)
+			if (Mapping.Count > 0)
 			{
-				Lines.Add("\"" + Pair.Value + "\" \"" + Pair.Key + "\"");
+				List<string> Lines = new List<string>();
+				foreach (KeyValuePair<StagedFileReference, FileReference> Pair in Mapping)
+				{
+					Lines.Add("\"" + Pair.Value + "\" \"" + Pair.Key + "\"");
+				}
+				WriteAllLines(Filename, Lines.ToArray());
 			}
-			WriteAllLines(Filename, Lines.ToArray());
 		}
-	}
 
-	public static void DumpManifest(DeploymentContext SC, string BaseFilename, bool DumpUFSFiles = true)
-	{
-		DumpManifest(SC.FilesToStage.NonUFSFiles, BaseFilename + "_NonUFSFiles.txt");
-		if (DumpUFSFiles)
+		public static void DumpManifest(DeploymentContext SC, string BaseFilename, bool DumpUFSFiles = true)
 		{
-			DumpManifest(SC.FilesToStage.NonUFSDebugFiles, BaseFilename + "_NonUFSFilesDebug.txt");
+			DumpManifest(SC.FilesToStage.NonUFSFiles, BaseFilename + "_NonUFSFiles.txt");
+			if (DumpUFSFiles)
+			{
+				DumpManifest(SC.FilesToStage.NonUFSDebugFiles, BaseFilename + "_NonUFSFilesDebug.txt");
+			}
+			DumpManifest(SC.FilesToStage.UFSFiles, BaseFilename + "_UFSFiles.txt");
 		}
-		DumpManifest(SC.FilesToStage.UFSFiles, BaseFilename + "_UFSFiles.txt");
-	}
-
-	public static void CopyUsingStagingManifest(ProjectParams Params, DeploymentContext SC)
-	{
-		CopyManifestFilesToStageDir(SC, SC.FilesToStage.NonUFSFiles, "NonUFSFiles");
 
-		Dictionary<StagedFileReference, FileReference> UFSFiles = new Dictionary<StagedFileReference, FileReference>(SC.FilesToStage.UFSFiles);
-		bool bStageUnrealFileSystemFiles = !Params.CookOnTheFly && !Params.UsePak(SC.StageTargetPlatform) && !Params.FileServer;
-		if (bStageUnrealFileSystemFiles)
+		public static void CopyUsingStagingManifest(ProjectParams Params, DeploymentContext SC)
 		{
-			foreach (KeyValuePair<StagedFileReference, FileReference> Pair in SC.CrashReporterUFSFiles)
+			CopyManifestFilesToStageDir(SC, SC.FilesToStage.NonUFSFiles, "NonUFSFiles");
+
+			Dictionary<StagedFileReference, FileReference> UFSFiles = new Dictionary<StagedFileReference, FileReference>(SC.FilesToStage.UFSFiles);
+			bool bStageUnrealFileSystemFiles = !Params.CookOnTheFly && !Params.UsePak(SC.StageTargetPlatform) && !Params.FileServer;
+			if (bStageUnrealFileSystemFiles)
 			{
-				FileReference ExistingLocation;
-				if (!UFSFiles.TryGetValue(Pair.Key, out ExistingLocation))
-				{
-					UFSFiles.Add(Pair.Key, Pair.Value);
-				}
-				else if (ExistingLocation != Pair.Value)
+				foreach (KeyValuePair<StagedFileReference, FileReference> Pair in SC.CrashReporterUFSFiles)
 				{
-					throw new AutomationException("File '{0}' is set to be staged from '{1}' for project and '{2}' for crash reporter", Pair.Key, ExistingLocation, Pair.Value);
+					FileReference ExistingLocation;
+					if (!UFSFiles.TryGetValue(Pair.Key, out ExistingLocation))
+					{
+						UFSFiles.Add(Pair.Key, Pair.Value);
+					}
+					else if (ExistingLocation != Pair.Value)
+					{
+						throw new AutomationException("File '{0}' is set to be staged from '{1}' for project and '{2}' for crash reporter", Pair.Key, ExistingLocation, Pair.Value);
+					}
 				}
+				CopyManifestFilesToStageDir(SC, UFSFiles, "UFSFiles");
+			}
+			else
+			{
+				// write out the manifest file even if we don't copy the files (so we can easily tell what went into the .pak file)
+				WriteManifestFilesToStageDir(SC, UFSFiles, "UFSFiles");
 			}
-			CopyManifestFilesToStageDir(SC, UFSFiles, "UFSFiles");
-		}
-		else
-		{
-			// write out the manifest file even if we don't copy the files (so we can easily tell what went into the .pak file)
-			WriteManifestFilesToStageDir(SC, UFSFiles, "UFSFiles");
-		}
 
-		// Copy debug files last
-		// they do not respect the DeployLowerCaseFilenames() setting, but if copied to a case-insensitive staging directory first they determine the casing for outer directories (like Engine/Content) 
-		if (!Params.NoDebugInfo)
-		{
-			CopyManifestFilesToStageDir(SC, SC.FilesToStage.NonUFSDebugFiles, "DebugFiles");
+			// Copy debug files last
+			// they do not respect the DeployLowerCaseFilenames() setting, but if copied to a case-insensitive staging directory first they determine the casing for outer directories (like Engine/Content) 
+			if (!Params.NoDebugInfo)
+			{
+				CopyManifestFilesToStageDir(SC, SC.FilesToStage.NonUFSDebugFiles, "DebugFiles");
+			}
 		}
-	}
-
-	// Pak file rules to apply to a list of files
-	private struct PakFileRules
-	{
-		// Name of config section
-		public string Name;
-
-		// Combined filter from all +Files= entries
-		public FileFilter Filter;
 
-		// An alternative and efficient way to match files
-		public HashSet<string> ExactFileMatches;
-
-		// To use exact file match or file filter
-		public bool bUseExactFilePathMatch;
-
-		// Rather to exclude entirely from paks
-		public bool bExcludeFromPaks;
-
-		// Rather to allow overriding the chunk manifest, if false will only modify loose files
-		public bool bOverrideChunkManifest;
-
-		// Whether pak file rule is disabled
-		public bool bDisabled;
-
-		// List of pak files to use instead of manifest
-		public List<string> OverridePaks;
-	};
-
-	/// <summary>
-	/// Reads Default/BasePakFileRules.ini and returns all PakFileRules objects that apply for this deployment
-	/// </summary>
-	/// <param name="Params"></param>
-	/// <param name="SC"></param>
-	private static List<PakFileRules> GetPakFileRules(ProjectParams Params, DeploymentContext SC)
-	{
-		// if we want to ignore the rules while staging, just don't read any rules in!
-		if (SC.UsePakFileRulesIni == false)
+		// Pak file rules to apply to a list of files
+		private struct PakFileRules
 		{
-			return new List<PakFileRules>();
-		}
+			// Name of config section
+			public string Name;
 
-		bool bWarnedAboutMultipleTargets = false;
+			// Combined filter from all +Files= entries
+			public FileFilter Filter;
 
-		/*List<ConfigFile> ConfigFiles = new List<ConfigFile>();
-		FileReference BaseConfigFileReference = FileReference.Combine(SC.EngineRoot, "Config", "BasePakFileRules.ini");
-		if (FileReference.Exists(BaseConfigFileReference))
-		{
-			ConfigFiles.Add(new ConfigFile(BaseConfigFileReference));
-			bFoundConfig = true;
-		}
+			// An alternative and efficient way to match files
+			public HashSet<string> ExactFileMatches;
 
-		FileReference ProjectConfigFileReference = FileReference.Combine(SC.ProjectRoot, "Config", "DefaultPakFileRules.ini");
-		if (FileReference.Exists(ProjectConfigFileReference))
-		{
-			ConfigFiles.Add(new ConfigFile(ProjectConfigFileReference));
-			bFoundConfig = true;
-		}
-	
-		if (!bFoundConfig)
-		{
-			return null;
-		}
-		ConfigHierarchy PakRulesConfig = new ConfigHierarchy(ConfigFiles);*/
+			// To use exact file match or file filter
+			public bool bUseExactFilePathMatch;
 
+			// Rather to exclude entirely from paks
+			public bool bExcludeFromPaks;
 
-		
-		ConfigHierarchy PakRulesConfig = ConfigCache.ReadHierarchy(ConfigHierarchyType.PakFileRules, DirectoryReference.FromFile(Params.RawProjectPath), SC.StageTargetPlatform.IniPlatformType);
+			// Rather to allow overriding the chunk manifest, if false will only modify loose files
+			public bool bOverrideChunkManifest;
 
-		bool bChunkedBuild = SC.PlatformUsesChunkManifests && DoesChunkPakManifestExist(Params, SC);
+			// Whether pak file rule is disabled
+			public bool bDisabled;
 
-		List<PakFileRules> RulesList = new List<PakFileRules>();
-		foreach (string SectionName in PakRulesConfig.SectionNames)
-		{
-			//LogInformation("Building PakFileRules for Section {0}", SectionName);
+			// List of pak files to use instead of manifest
+			public List<string> OverridePaks;
+		};
 
-			bool bOnlyChunkedBuilds = false;
-			if (!PakRulesConfig.TryGetValue(SectionName, "bOnlyChunkedBuilds", out bOnlyChunkedBuilds))
+		/// <summary>
+		/// Reads Default/BasePakFileRules.ini and returns all PakFileRules objects that apply for this deployment
+		/// </summary>
+		/// <param name="Params"></param>
+		/// <param name="SC"></param>
+		private static List<PakFileRules> GetPakFileRules(ProjectParams Params, DeploymentContext SC)
+		{
+			// if we want to ignore the rules while staging, just don't read any rules in!
+			if (SC.UsePakFileRulesIni == false)
 			{
-				bOnlyChunkedBuilds = false;
+				return new List<PakFileRules>();
 			}
 
-			bool bOnlyNonChunkedBuilds = false;
-			if (!PakRulesConfig.TryGetValue(SectionName, "bOnlyNonChunkedBuilds", out bOnlyNonChunkedBuilds))
+			bool bWarnedAboutMultipleTargets = false;
+
+			/*List<ConfigFile> ConfigFiles = new List<ConfigFile>();
+			FileReference BaseConfigFileReference = FileReference.Combine(SC.EngineRoot, "Config", "BasePakFileRules.ini");
+			if (FileReference.Exists(BaseConfigFileReference))
 			{
-				bOnlyNonChunkedBuilds = false;
+				ConfigFiles.Add(new ConfigFile(BaseConfigFileReference));
+				bFoundConfig = true;
 			}
 
-			if (bChunkedBuild && bOnlyNonChunkedBuilds)
+			FileReference ProjectConfigFileReference = FileReference.Combine(SC.ProjectRoot, "Config", "DefaultPakFileRules.ini");
+			if (FileReference.Exists(ProjectConfigFileReference))
 			{
-				continue;
+				ConfigFiles.Add(new ConfigFile(ProjectConfigFileReference));
+				bFoundConfig = true;
 			}
 
-			if(!bChunkedBuild && bOnlyChunkedBuilds)
+			if (!bFoundConfig)
 			{
-				continue;
+				return null;
 			}
+			ConfigHierarchy PakRulesConfig = new ConfigHierarchy(ConfigFiles);*/
+
+
+
+			ConfigHierarchy PakRulesConfig = ConfigCache.ReadHierarchy(ConfigHierarchyType.PakFileRules, DirectoryReference.FromFile(Params.RawProjectPath), SC.StageTargetPlatform.IniPlatformType);
 
-			string PlatformString;
-			if (PakRulesConfig.TryGetValue(SectionName, "Platforms", out PlatformString))
+			bool bChunkedBuild = SC.PlatformUsesChunkManifests && DoesChunkPakManifestExist(Params, SC);
+
+			List<PakFileRules> RulesList = new List<PakFileRules>();
+			foreach (string SectionName in PakRulesConfig.SectionNames)
 			{
-				string[] PlatformStrings = PlatformString.Split(new char[] { ',', ' ' }, StringSplitOptions.RemoveEmptyEntries);
-				bool bMatches = false;
+				//LogInformation("Building PakFileRules for Section {0}", SectionName);
 
-				// Check platform string
-				foreach (string Platform in PlatformStrings)
+				bool bOnlyChunkedBuilds = false;
+				if (!PakRulesConfig.TryGetValue(SectionName, "bOnlyChunkedBuilds", out bOnlyChunkedBuilds))
 				{
-					if (SC.StageTargetPlatform.PlatformType.ToString().Equals(Platform, StringComparison.OrdinalIgnoreCase))
-					{
-						bMatches = true;
-						break;
-					}
-					else if (SC.StageTargetPlatform.IniPlatformType.ToString().Equals(Platform, StringComparison.OrdinalIgnoreCase))
-					{
-						bMatches = true;
-						break;
-					}
+					bOnlyChunkedBuilds = false;
+				}
+
+				bool bOnlyNonChunkedBuilds = false;
+				if (!PakRulesConfig.TryGetValue(SectionName, "bOnlyNonChunkedBuilds", out bOnlyNonChunkedBuilds))
+				{
+					bOnlyNonChunkedBuilds = false;
 				}
 
-				if (!bMatches)
+				if (bChunkedBuild && bOnlyNonChunkedBuilds)
 				{
-					//LogInformation("No matching platform for PakFileRules for Section {0} : {1}, {2}", SectionName, SC.StageTargetPlatform.PlatformType.ToString(), SC.StageTargetPlatform.IniPlatformType.ToString());
 					continue;
 				}
-			}
 
-			string TargetString;
-			if (PakRulesConfig.TryGetValue(SectionName, "Targets", out TargetString))
-			{
-				string[] TargetStrings = TargetString.Split(new char[] { ',', ' ' }, StringSplitOptions.RemoveEmptyEntries);
-				bool bMatches = false;
+				if (!bChunkedBuild && bOnlyChunkedBuilds)
+				{
+					continue;
+				}
 
-				// Check target string
-				foreach (string Target in TargetStrings)
+				string PlatformString;
+				if (PakRulesConfig.TryGetValue(SectionName, "Platforms", out PlatformString))
 				{
-					foreach (UnrealTargetConfiguration C in SC.StageTargetConfigurations)
+					string[] PlatformStrings = PlatformString.Split(new char[] { ',', ' ' }, StringSplitOptions.RemoveEmptyEntries);
+					bool bMatches = false;
+
+					// Check platform string
+					foreach (string Platform in PlatformStrings)
 					{
-						if (C.ToString() == Target)
+						if (SC.StageTargetPlatform.PlatformType.ToString().Equals(Platform, StringComparison.OrdinalIgnoreCase))
+						{
+							bMatches = true;
+							break;
+						}
+						else if (SC.StageTargetPlatform.IniPlatformType.ToString().Equals(Platform, StringComparison.OrdinalIgnoreCase))
 						{
 							bMatches = true;
 							break;
 						}
 					}
-				}
 
-				if (!bMatches)
-				{
-					continue;
-				}
-				else if (SC.StageTargetConfigurations.Count > 0 && !bWarnedAboutMultipleTargets)
-				{
-					bWarnedAboutMultipleTargets = true;
-					LogInformation("Staging with more than one target, PakFileRules may apply too many rules!");
+					if (!bMatches)
+					{
+						//LogInformation("No matching platform for PakFileRules for Section {0} : {1}, {2}", SectionName, SC.StageTargetPlatform.PlatformType.ToString(), SC.StageTargetPlatform.IniPlatformType.ToString());
+						continue;
+					}
 				}
-			}
 
-			PakFileRules PakRules = new PakFileRules();
-			PakRules.Name = SectionName;
-			PakRulesConfig.TryGetValue(SectionName, "bUseExactFilePathMatch", out PakRules.bUseExactFilePathMatch);
-			PakRulesConfig.TryGetValue(SectionName, "bExcludeFromPaks", out PakRules.bExcludeFromPaks);
-			PakRulesConfig.TryGetValue(SectionName, "bOverrideChunkManifest", out PakRules.bOverrideChunkManifest);
-			PakRulesConfig.TryGetValue(SectionName, "bDisabled", out PakRules.bDisabled);
-			string PakString;
-			PakRulesConfig.TryGetValue(SectionName, "OverridePaks", out PakString);
+				string TargetString;
+				if (PakRulesConfig.TryGetValue(SectionName, "Targets", out TargetString))
+				{
+					string[] TargetStrings = TargetString.Split(new char[] { ',', ' ' }, StringSplitOptions.RemoveEmptyEntries);
+					bool bMatches = false;
 
-			if (PakString != null && PakString.Length > 0)
-			{
-				PakRules.OverridePaks = PakString.Split(new char[] { ',', ' ' }, StringSplitOptions.RemoveEmptyEntries).ToList();
-			}
-			else
-			{
-				PakRules.OverridePaks = null;
-			}
+					// Check target string
+					foreach (string Target in TargetStrings)
+					{
+						foreach (UnrealTargetConfiguration C in SC.StageTargetConfigurations)
+						{
+							if (C.ToString() == Target)
+							{
+								bMatches = true;
+								break;
+							}
+						}
+					}
 
-			if (PakRules.bExcludeFromPaks && PakRules.OverridePaks != null)
-			{
-				LogWarning("Error in PakFileRules {0}, set to exclude but also sets override!", PakRules.Name);
-				continue;
-			}
-			else if (!PakRules.bExcludeFromPaks && PakRules.OverridePaks == null)
-			{
-				LogWarning("Error in PakFileRules {0}, set to include but did not specify paks!", PakRules.Name);
-				continue;
-			}
+					if (!bMatches)
+					{
+						continue;
+					}
+					else if (SC.StageTargetConfigurations.Count > 0 && !bWarnedAboutMultipleTargets)
+					{
+						bWarnedAboutMultipleTargets = true;
+						LogInformation("Staging with more than one target, PakFileRules may apply too many rules!");
+					}
+				}
 
-			IReadOnlyList<string> FilesEnumberable;
-			if (PakRulesConfig.TryGetValues(SectionName, "Files", out FilesEnumberable))
-			{
-				// Only add if we have actual files, we can end up with none due to config overriding
-				PakRules.Filter = new FileFilter();
-				PakRules.ExactFileMatches = new HashSet<string>();
+				PakFileRules PakRules = new PakFileRules();
+				PakRules.Name = SectionName;
+				PakRulesConfig.TryGetValue(SectionName, "bUseExactFilePathMatch", out PakRules.bUseExactFilePathMatch);
+				PakRulesConfig.TryGetValue(SectionName, "bExcludeFromPaks", out PakRules.bExcludeFromPaks);
+				PakRulesConfig.TryGetValue(SectionName, "bOverrideChunkManifest", out PakRules.bOverrideChunkManifest);
+				PakRulesConfig.TryGetValue(SectionName, "bDisabled", out PakRules.bDisabled);
+				string PakString;
+				PakRulesConfig.TryGetValue(SectionName, "OverridePaks", out PakString);
 
-				if (PakRules.bUseExactFilePathMatch)
+				if (PakString != null && PakString.Length > 0)
 				{
-					PakRules.ExactFileMatches.UnionWith(FilesEnumberable);
+					PakRules.OverridePaks = PakString.Split(new char[] { ',', ' ' }, StringSplitOptions.RemoveEmptyEntries).ToList();
 				}
 				else
 				{
-					PakRules.Filter.AddRules(FilesEnumberable);
+					PakRules.OverridePaks = null;
 				}
 
-				RulesList.Add(PakRules);
-			}
-		}
-		if (RulesList.Count == 0)
-		{
-			return null;
-		}
-		return RulesList;
-	}
-
-	/// <summary>
-	/// Attempts to apply the pak file rules to a specific staging file, returns false if it should be excluded
-	/// </summary>
-	/// <param name="Params"></param>
-	/// <param name="SC"></param>
-	private static void ApplyPakFileRules(List<PakFileRules> RulesList, KeyValuePair<string, string> StagingFile, HashSet<ChunkDefinition> ModifyPakList, ConcurrentDictionary<string, ChunkDefinition> ChunkNameToDefinition, out bool bExcludeFromPaks)
-	{
-		bExcludeFromPaks = false;
-
-		if (RulesList == null)
-		{
-			return;
-		}
-
-		// Search in order, return on first match
-		foreach (var PakRules in RulesList)
-		{
-			bool bMatched = !PakRules.bDisabled && 
-				((!PakRules.bUseExactFilePathMatch && PakRules.Filter != null && PakRules.Filter.Matches(StagingFile.Key)) ||
-				(PakRules.bUseExactFilePathMatch && PakRules.ExactFileMatches != null && PakRules.ExactFileMatches.Contains(StagingFile.Value)));
+				if (PakRules.bExcludeFromPaks && PakRules.OverridePaks != null)
+				{
+					LogWarning("Error in PakFileRules {0}, set to exclude but also sets override!", PakRules.Name);
+					continue;
+				}
+				else if (!PakRules.bExcludeFromPaks && PakRules.OverridePaks == null)
+				{
+					LogWarning("Error in PakFileRules {0}, set to include but did not specify paks!", PakRules.Name);
+					continue;
+				}
 
-			if (bMatched)
-			{
-				bool bOverrideChunkAssignment = false;
-				if (ModifyPakList != null && ModifyPakList.Count > 0)
+				IReadOnlyList<string> FilesEnumberable;
+				if (PakRulesConfig.TryGetValues(SectionName, "Files", out FilesEnumberable))
 				{
-					// Only override the existing list if bOverrideChunkManifest is set
-					if (!PakRules.bOverrideChunkManifest)
-					{
-						return;
-					}
+					// Only add if we have actual files, we can end up with none due to config overriding
+					PakRules.Filter = new FileFilter();
+					PakRules.ExactFileMatches = new HashSet<string>();
 
-					bOverrideChunkAssignment = true;
-					if (PakRules.OverridePaks != null)
+					if (PakRules.bUseExactFilePathMatch)
 					{
-						LogInformation("Overridding chunk assignment {0} to {1}, this can cause broken references", StagingFile.Key, string.Join(", ", PakRules.OverridePaks));
+						PakRules.ExactFileMatches.UnionWith(FilesEnumberable);
 					}
-					else if (bExcludeFromPaks)
+					else
 					{
-						LogInformation("Removing {0} from pak despite chunk assignment, this can cause broken references", StagingFile.Key);
+						PakRules.Filter.AddRules(FilesEnumberable);
 					}
+
+					RulesList.Add(PakRules);
 				}
+			}
+			if (RulesList.Count == 0)
+			{
+				return null;
+			}
+			return RulesList;
+		}
+
+		/// <summary>
+		/// Attempts to apply the pak file rules to a specific staging file, returns false if it should be excluded
+		/// </summary>
+		/// <param name="Params"></param>
+		/// <param name="SC"></param>
+		private static void ApplyPakFileRules(List<PakFileRules> RulesList, KeyValuePair<string, string> StagingFile, HashSet<ChunkDefinition> ModifyPakList, ConcurrentDictionary<string, ChunkDefinition> ChunkNameToDefinition, out bool bExcludeFromPaks)
+		{
+			bExcludeFromPaks = false;
+
+			if (RulesList == null)
+			{
+				return;
+			}
+
+			// Search in order, return on first match
+			foreach (var PakRules in RulesList)
+			{
+				bool bMatched = !PakRules.bDisabled &&
+					((!PakRules.bUseExactFilePathMatch && PakRules.Filter != null && PakRules.Filter.Matches(StagingFile.Key)) ||
+					(PakRules.bUseExactFilePathMatch && PakRules.ExactFileMatches != null && PakRules.ExactFileMatches.Contains(StagingFile.Value)));
 
-				bExcludeFromPaks = PakRules.bExcludeFromPaks;
-				if (PakRules.OverridePaks != null && ModifyPakList != null && ChunkNameToDefinition != null)
+				if (bMatched)
 				{
-					if (!bOverrideChunkAssignment)
+					bool bOverrideChunkAssignment = false;
+					if (ModifyPakList != null && ModifyPakList.Count > 0)
 					{
-						LogInformation("Setting pak assignment for file {0} to {1}", StagingFile.Key, string.Join(", ", PakRules.OverridePaks));
+						// Only override the existing list if bOverrideChunkManifest is set
+						if (!PakRules.bOverrideChunkManifest)
+						{
+							return;
+						}
+
+						bOverrideChunkAssignment = true;
+						if (PakRules.OverridePaks != null)
+						{
+							LogInformation("Overridding chunk assignment {0} to {1}, this can cause broken references", StagingFile.Key, string.Join(", ", PakRules.OverridePaks));
+						}
+						else if (bExcludeFromPaks)
+						{
+							LogInformation("Removing {0} from pak despite chunk assignment, this can cause broken references", StagingFile.Key);
+						}
 					}
 
-					ModifyPakList.Clear();
-					ModifyPakList.UnionWith(PakRules.OverridePaks.Select(OverrideChunkName =>
+					bExcludeFromPaks = PakRules.bExcludeFromPaks;
+					if (PakRules.OverridePaks != null && ModifyPakList != null && ChunkNameToDefinition != null)
 					{
-						if (!ChunkNameToDefinition.ContainsKey(OverrideChunkName))
+						if (!bOverrideChunkAssignment)
 						{
-							ChunkNameToDefinition.TryAdd(OverrideChunkName, new ChunkDefinition(OverrideChunkName));
+							LogInformation("Setting pak assignment for file {0} to {1}", StagingFile.Key, string.Join(", ", PakRules.OverridePaks));
 						}
 
-						return ChunkNameToDefinition[OverrideChunkName];
-					}));
-				}
-				else if (bExcludeFromPaks)
-				{
-					if (!bOverrideChunkAssignment)
+						ModifyPakList.Clear();
+						ModifyPakList.UnionWith(PakRules.OverridePaks.Select(OverrideChunkName =>
+						{
+							if (!ChunkNameToDefinition.ContainsKey(OverrideChunkName))
+							{
+								ChunkNameToDefinition.TryAdd(OverrideChunkName, new ChunkDefinition(OverrideChunkName));
+							}
+
+							return ChunkNameToDefinition[OverrideChunkName];
+						}));
+					}
+					else if (bExcludeFromPaks)
 					{
-						LogInformation("Excluding {0} from pak files", StagingFile.Key);
+						if (!bOverrideChunkAssignment)
+						{
+							LogInformation("Excluding {0} from pak files", StagingFile.Key);
+						}
 					}
-				}
 
-				return;
+					return;
+				}
 			}
 		}
-	}
 
 
-	/// <summary>
-	/// Creates a pak file using staging context (single manifest)
-	/// </summary>
-	/// <param name="Params"></param>
-	/// <param name="SC"></param>
-	private static void CreatePakUsingStagingManifest(ProjectParams Params, DeploymentContext SC)
-	{
-		LogInformation("Creating pak using staging manifest.");
-
-		DumpManifest(SC, CombinePaths(CmdEnv.LogFolder, "PrePak" + (SC.DedicatedServer ? "_Server" : "") + "_" + SC.CookPlatform));
+		/// <summary>
+		/// Creates a pak file using staging context (single manifest)
+		/// </summary>
+		/// <param name="Params"></param>
+		/// <param name="SC"></param>
+		private static void CreatePakUsingStagingManifest(ProjectParams Params, DeploymentContext SC)
+		{
+			LogInformation("Creating pak using staging manifest.");
 
-		var UnrealPakResponseFile = CreatePakResponseFileFromStagingManifest(SC, SC.FilesToStage.UFSFiles);
+			DumpManifest(SC, CombinePaths(CmdEnv.LogFolder, "PrePak" + (SC.DedicatedServer ? "_Server" : "") + "_" + SC.CookPlatform));
 
-		List<PakFileRules> PakRulesList = GetPakFileRules(Params, SC);
+			var UnrealPakResponseFile = CreatePakResponseFileFromStagingManifest(SC, SC.FilesToStage.UFSFiles);
 
-		List<string> FilesToRemove = new List<string>();
+			List<PakFileRules> PakRulesList = GetPakFileRules(Params, SC);
 
-		// Apply the pak file rules, this can remove things but will not override the pak file name
-		foreach (var StagingFile in UnrealPakResponseFile)
-		{
-			bool bExcludeFromPaks = false;
-			ApplyPakFileRules(PakRulesList, StagingFile, null, null, out bExcludeFromPaks);
+			List<string> FilesToRemove = new List<string>();
 
-			if (bExcludeFromPaks)
+			// Apply the pak file rules, this can remove things but will not override the pak file name
+			foreach (var StagingFile in UnrealPakResponseFile)
 			{
-				FilesToRemove.Add(StagingFile.Key);
+				bool bExcludeFromPaks = false;
+				ApplyPakFileRules(PakRulesList, StagingFile, null, null, out bExcludeFromPaks);
+
+				if (bExcludeFromPaks)
+				{
+					FilesToRemove.Add(StagingFile.Key);
+				}
 			}
-		}
 
-		foreach (var FileToRemove in FilesToRemove)
-		{
-			UnrealPakResponseFile.Remove(FileToRemove);
-		}
+			foreach (var FileToRemove in FilesToRemove)
+			{
+				UnrealPakResponseFile.Remove(FileToRemove);
+			}
 
-		EncryptionAndSigning.CryptoSettings PakCryptoSettings = EncryptionAndSigning.ParseCryptoSettings(DirectoryReference.FromFile(Params.RawProjectPath), SC.StageTargetPlatform.IniPlatformType);
-		FileReference CryptoKeysCacheFilename = FileReference.Combine(SC.MetadataDir, "Crypto.json");
-		PakCryptoSettings.Save(CryptoKeysCacheFilename);
+			EncryptionAndSigning.CryptoSettings PakCryptoSettings = EncryptionAndSigning.ParseCryptoSettings(DirectoryReference.FromFile(Params.RawProjectPath), SC.StageTargetPlatform.IniPlatformType);
+			FileReference CryptoKeysCacheFilename = FileReference.Combine(SC.MetadataDir, "Crypto.json");
+			PakCryptoSettings.Save(CryptoKeysCacheFilename);
 
-		List<CreatePakParams> PakInputs = new List<CreatePakParams>();
-		PakInputs.Add(new CreatePakParams(SC.ShortProjectName, UnrealPakResponseFile, Params.Compressed, null));
-		CreatePaks(Params, SC, PakInputs, PakCryptoSettings, CryptoKeysCacheFilename);
-	}
+			List<CreatePakParams> PakInputs = new List<CreatePakParams>();
+			PakInputs.Add(new CreatePakParams(SC.ShortProjectName, UnrealPakResponseFile, Params.Compressed, null));
+			CreatePaks(Params, SC, PakInputs, PakCryptoSettings, CryptoKeysCacheFilename);
+		}
 
-	/// <summary>
-	/// Creates a standalone pak file for crash reporter
-	/// </summary>
-	/// <param name="Params">The packaging parameters</param>
-	/// <param name="SC">Staging context</param>
-	private static void CreatePakForCrashReporter(ProjectParams Params, DeploymentContext SC)
-	{
-		LogInformation("Creating pak for crash reporter.");
+		/// <summary>
+		/// Creates a standalone pak file for crash reporter
+		/// </summary>
+		/// <param name="Params">The packaging parameters</param>
+		/// <param name="SC">Staging context</param>
+		private static void CreatePakForCrashReporter(ProjectParams Params, DeploymentContext SC)
+		{
+			LogInformation("Creating pak for crash reporter.");
 
-		Dictionary<string, string> PakResponseFile = CreatePakResponseFileFromStagingManifest(SC, SC.CrashReporterUFSFiles);
-		FileReference OutputLocation = FileReference.Combine(SC.RuntimeRootDir, "Engine", "Programs", "CrashReportClient", "Content", "Paks", "CrashReportClient.pak");
+			Dictionary<string, string> PakResponseFile = CreatePakResponseFileFromStagingManifest(SC, SC.CrashReporterUFSFiles);
+			FileReference OutputLocation = FileReference.Combine(SC.RuntimeRootDir, "Engine", "Programs", "CrashReportClient", "Content", "Paks", "CrashReportClient.pak");
 
-		RunUnrealPak(Params, PakResponseFile, OutputLocation, null, null, Params.Compressed, null, null, null, null, null);
-	}
+			RunUnrealPak(Params, PakResponseFile, OutputLocation, null, null, Params.Compressed, null, null, null, null, null);
+		}
 
-	/// <summary>
-	/// Creates a pak response file using stage context
-	/// </summary>
-	/// <param name="SC"></param>
-	/// <returns></returns>
-	private static Dictionary<string, string> CreatePakResponseFileFromStagingManifest(DeploymentContext SC, Dictionary<StagedFileReference, FileReference> FilesToStage)
-	{
-		// look for optional packaging blacklist if only one config active
-		List<string> Blacklist = null;
-		if (SC.StageTargetConfigurations.Count == 1)
+		/// <summary>
+		/// Creates a pak response file using stage context
+		/// </summary>
+		/// <param name="SC"></param>
+		/// <returns></returns>
+		private static Dictionary<string, string> CreatePakResponseFileFromStagingManifest(DeploymentContext SC, Dictionary<StagedFileReference, FileReference> FilesToStage)
 		{
-			FileReference PakBlacklistFilename = FileReference.Combine(SC.ProjectRoot, "Platforms", SC.PlatformDir, "Build", string.Format("PakBlacklist-{0}.txt", SC.StageTargetConfigurations[0].ToString()));
-			if (!FileReference.Exists(PakBlacklistFilename))
-			{
-				PakBlacklistFilename = FileReference.Combine(SC.ProjectRoot, "Build", SC.PlatformDir, string.Format("PakBlacklist-{0}.txt", SC.StageTargetConfigurations[0].ToString()));
-			}
-			if (FileReference.Exists(PakBlacklistFilename))
+			// look for optional packaging blacklist if only one config active
+			List<string> Blacklist = null;
+			if (SC.StageTargetConfigurations.Count == 1)
 			{
-				LogInformation("Applying PAK blacklist file {0}. This is deprecated in favor of DefaultPakFileRules.ini", PakBlacklistFilename);
-				string[] BlacklistContents = FileReference.ReadAllLines(PakBlacklistFilename);
-				foreach (string Candidate in BlacklistContents)
+				FileReference PakBlacklistFilename = FileReference.Combine(SC.ProjectRoot, "Platforms", SC.PlatformDir, "Build", string.Format("PakBlacklist-{0}.txt", SC.StageTargetConfigurations[0].ToString()));
+				if (!FileReference.Exists(PakBlacklistFilename))
+				{
+					PakBlacklistFilename = FileReference.Combine(SC.ProjectRoot, "Build", SC.PlatformDir, string.Format("PakBlacklist-{0}.txt", SC.StageTargetConfigurations[0].ToString()));
+				}
+				if (FileReference.Exists(PakBlacklistFilename))
 				{
-					if (Candidate.Trim().Length > 0)
+					LogInformation("Applying PAK blacklist file {0}. This is deprecated in favor of DefaultPakFileRules.ini", PakBlacklistFilename);
+					string[] BlacklistContents = FileReference.ReadAllLines(PakBlacklistFilename);
+					foreach (string Candidate in BlacklistContents)
 					{
-						if (Blacklist == null)
+						if (Candidate.Trim().Length > 0)
 						{
-							Blacklist = new List<string>();
+							if (Blacklist == null)
+							{
+								Blacklist = new List<string>();
+							}
+							Blacklist.Add(Candidate);
 						}
-						Blacklist.Add(Candidate);
 					}
 				}
 			}
-		}
 
-		var UnrealPakResponseFile = new Dictionary<string, string>(StringComparer.InvariantCultureIgnoreCase);
-		foreach (KeyValuePair<StagedFileReference, FileReference> Pair in FilesToStage)
-		{
-			FileReference Src = Pair.Value;
-			string Dest = Pair.Key.Name;
+			var UnrealPakResponseFile = new Dictionary<string, string>(StringComparer.InvariantCultureIgnoreCase);
+			foreach (KeyValuePair<StagedFileReference, FileReference> Pair in FilesToStage)
+			{
+				FileReference Src = Pair.Value;
+				string Dest = Pair.Key.Name;
 
-			Dest = CombinePaths(PathSeparator.Slash, SC.PakFileInternalRoot, Dest);
+				Dest = CombinePaths(PathSeparator.Slash, SC.PakFileInternalRoot, Dest);
 
-			if (Blacklist != null)
-			{
-				bool bExcludeFile = false;
-				foreach (string ExcludePath in Blacklist)
+				if (Blacklist != null)
 				{
-					if (Dest.StartsWith(ExcludePath))
+					bool bExcludeFile = false;
+					foreach (string ExcludePath in Blacklist)
 					{
-						bExcludeFile = true;
-						break;
+						if (Dest.StartsWith(ExcludePath))
+						{
+							bExcludeFile = true;
+							break;
+						}
+					}
+
+					if (bExcludeFile)
+					{
+						LogInformation("Excluding {0}", Src);
+						continue;
 					}
 				}
 
-				if (bExcludeFile)
+				// Filter I/O store container files
+				if (Src.HasExtension(".ucas") || Src.HasExtension(".utoc"))
 				{
 					LogInformation("Excluding {0}", Src);
 					continue;
 				}
-			}
-
-			// Filter I/O store container files
-			if (Src.HasExtension(".ucas") || Src.HasExtension(".utoc"))
-			{
-				LogInformation("Excluding {0}", Src);
-				continue;
-			}
 
-			// Do a filtered copy of all ini files to allow stripping of values that we don't want to distribute
-			if (Src.HasExtension(".ini"))
-			{
-				string SubFolder = Pair.Key.Name.Replace('/', Path.DirectorySeparatorChar);
-				FileReference NewIniFilename = FileReference.Combine(SC.ProjectRoot, "Saved", "Temp", SC.PlatformDir, SubFolder);
-				InternalUtils.SafeCreateDirectory(NewIniFilename.Directory.FullName, true);
-				InternalUtils.SafeCopyFile(Src.FullName, NewIniFilename.FullName, IniKeyBlacklist:SC.IniKeyBlacklist, IniSectionBlacklist:SC.IniSectionBlacklist);
-				Src = NewIniFilename;
-			}
+				// Do a filtered copy of all ini files to allow stripping of values that we don't want to distribute
+				if (Src.HasExtension(".ini"))
+				{
+					string SubFolder = Pair.Key.Name.Replace('/', Path.DirectorySeparatorChar);
+					FileReference NewIniFilename = FileReference.Combine(SC.ProjectRoot, "Saved", "Temp", SC.PlatformDir, SubFolder);
+					InternalUtils.SafeCreateDirectory(NewIniFilename.Directory.FullName, true);
+					InternalUtils.SafeCopyFile(Src.FullName, NewIniFilename.FullName, IniKeyBlacklist: SC.IniKeyBlacklist, IniSectionBlacklist: SC.IniSectionBlacklist);
+					Src = NewIniFilename;
+				}
 
-			// there can be files that only differ in case only, we don't support that in paks as paks are case-insensitive
-			if (UnrealPakResponseFile.ContainsKey(Src.FullName))
-			{
-				if (UnrealPakResponseFile[Src.FullName] != Dest)
+				// there can be files that only differ in case only, we don't support that in paks as paks are case-insensitive
+				if (UnrealPakResponseFile.ContainsKey(Src.FullName))
 				{
-					throw new AutomationException("Staging manifest already contains {0} (or a file that differs in case only)", Src);
+					if (UnrealPakResponseFile[Src.FullName] != Dest)
+					{
+						throw new AutomationException("Staging manifest already contains {0} (or a file that differs in case only)", Src);
+					}
+					LogWarning("Tried to add duplicate file to stage " + Src + " ignoring second attempt pls fix");
+					continue;
 				}
-				LogWarning("Tried to add duplicate file to stage " + Src + " ignoring second attempt pls fix");
-				continue;
+
+				UnrealPakResponseFile.Add(Src.FullName, Dest);
 			}
 
-			UnrealPakResponseFile.Add(Src.FullName, Dest);
+			return UnrealPakResponseFile;
 		}
 
-		return UnrealPakResponseFile;
-	}
-
-	/// <summary>
-	/// Parameter class for CreatePaks(). Each instance stores information about a pak file to be created.
-	/// </summary>
-	[DebuggerDisplay("PakName")]
-	class CreatePakParams
-	{
 		/// <summary>
-		/// Path to the base output file for this pak file 
+		/// Parameter class for CreatePaks(). Each instance stores information about a pak file to be created.
 		/// </summary>
-		public string PakName;
+		[DebuggerDisplay("PakName")]
+		class CreatePakParams
+		{
+			/// <summary>
+			/// Path to the base output file for this pak file 
+			/// </summary>
+			public string PakName;
 
-		/// <summary>
-		/// Map of files within the pak file to their source file on disk
-		/// </summary>
-		public Dictionary<string, string> UnrealPakResponseFile;
+			/// <summary>
+			/// Map of files within the pak file to their source file on disk
+			/// </summary>
+			public Dictionary<string, string> UnrealPakResponseFile;
 
-		/// <summary>
-		/// Whether to enable compression
-		/// </summary>
-		public bool bCompressed;
+			/// <summary>
+			/// Whether to enable compression
+			/// </summary>
+			public bool bCompressed;
 
-		/// <summary>
-		/// GUID of the encryption key for this pak file
-		/// </summary>
-		public string EncryptionKeyGuid;
+			/// <summary>
+			/// GUID of the encryption key for this pak file
+			/// </summary>
+			public string EncryptionKeyGuid;
 
-		/// <summary>
-		/// Constructor
-		/// </summary>
-		/// <param name="PakName">Path to the base output file for this pak file</param>
-		/// <param name="UnrealPakResponseFile">Map of files within the pak file to their source file on disk</param>
-		/// <param name="bCompressed">Whether to enable compression</param>
-		public CreatePakParams(string PakName, Dictionary<string, string> UnrealPakResponseFile, bool bCompressed, string EncryptionKeyGuid)
-		{
-			this.PakName = PakName;
-			this.UnrealPakResponseFile = UnrealPakResponseFile;
-			this.bCompressed = bCompressed;
-			this.EncryptionKeyGuid = EncryptionKeyGuid;
+			/// <summary>
+			/// Constructor
+			/// </summary>
+			/// <param name="PakName">Path to the base output file for this pak file</param>
+			/// <param name="UnrealPakResponseFile">Map of files within the pak file to their source file on disk</param>
+			/// <param name="bCompressed">Whether to enable compression</param>
+			public CreatePakParams(string PakName, Dictionary<string, string> UnrealPakResponseFile, bool bCompressed, string EncryptionKeyGuid)
+			{
+				this.PakName = PakName;
+				this.UnrealPakResponseFile = UnrealPakResponseFile;
+				this.bCompressed = bCompressed;
+				this.EncryptionKeyGuid = EncryptionKeyGuid;
+			}
 		}
-	}
 
-	private static bool ShouldSkipGeneratingPatch(ConfigHierarchy PlatformGameConfig, string Filename)
-	{
-		bool bIsEarlyDownloaderPakFile = false;
-		string EarlyDownloaderPakFilePrefix = "";
-		if (PlatformGameConfig.GetString("/Script/UnrealEd.ProjectPackagingSettings", "EarlyDownloaderPakFilePrefix", out EarlyDownloaderPakFilePrefix))
+		private static bool ShouldSkipGeneratingPatch(ConfigHierarchy PlatformGameConfig, string Filename)
 		{
-			bIsEarlyDownloaderPakFile = String.Equals(EarlyDownloaderPakFilePrefix, Filename, StringComparison.OrdinalIgnoreCase);
-		}
+			bool bIsEarlyDownloaderPakFile = false;
+			string EarlyDownloaderPakFilePrefix = "";
+			if (PlatformGameConfig.GetString("/Script/UnrealEd.ProjectPackagingSettings", "EarlyDownloaderPakFilePrefix", out EarlyDownloaderPakFilePrefix))
+			{
+				bIsEarlyDownloaderPakFile = String.Equals(EarlyDownloaderPakFilePrefix, Filename, StringComparison.OrdinalIgnoreCase);
+			}
 
-		return bIsEarlyDownloaderPakFile;
-	}
+			return bIsEarlyDownloaderPakFile;
+		}
 
-	private static string GetPakFilePostFix(bool bShouldGeneratePatch, bool bHasIterateSharedCookedBuild, ConfigHierarchy PlatformGameConfig, string Filename)
-	{
-		if (ShouldSkipGeneratingPatch(PlatformGameConfig, Filename))
+		private static string GetPakFilePostFix(bool bShouldGeneratePatch, bool bHasIterateSharedCookedBuild, ConfigHierarchy PlatformGameConfig, string Filename)
 		{
+			if (ShouldSkipGeneratingPatch(PlatformGameConfig, Filename))
+			{
+				return String.Empty;
+			}
+			else if (bShouldGeneratePatch)
+			{
+				return "_P";
+			}
+			else if (bHasIterateSharedCookedBuild)
+			{
+				// shared cooked builds will produce a patch
+				// then be combined with the shared cooked build
+				return "_S_P";
+			}
+
 			return String.Empty;
 		}
-		else if(bShouldGeneratePatch)
-		{
-			return "_P";
-		}
-		else if(bHasIterateSharedCookedBuild)
-		{
-			// shared cooked builds will produce a patch
-			// then be combined with the shared cooked build
-			return "_S_P";
-		}
-
-		return String.Empty;
-	}
-
-	/// <summary>
-	/// Creates a pak file using response file.
-	/// </summary>
-	/// <param name="Params"></param>
-	/// <param name="SC"></param>
-	/// <param name="UnrealPakResponseFile"></param>
-	/// <param name="PakName"></param>
-	private static void CreatePaks(ProjectParams Params, DeploymentContext SC, List<CreatePakParams> PakParamsList, EncryptionAndSigning.CryptoSettings CryptoSettings, FileReference CryptoKeysCacheFilename)
-	{
-		bool bShouldGeneratePatch = Params.IsGeneratingPatch && SC.StageTargetPlatform.GetPlatformPatchesWithDiffPak(Params, SC);
 
-		if (bShouldGeneratePatch && !Params.HasBasedOnReleaseVersion)
+		/// <summary>
+		/// Creates a pak file using response file.
+		/// </summary>
+		/// <param name="Params"></param>
+		/// <param name="SC"></param>
+		/// <param name="UnrealPakResponseFile"></param>
+		/// <param name="PakName"></param>
+		private static void CreatePaks(ProjectParams Params, DeploymentContext SC, List<CreatePakParams> PakParamsList, EncryptionAndSigning.CryptoSettings CryptoSettings, FileReference CryptoKeysCacheFilename)
 		{
-			LogInformation("Generating patch required a based on release version flag");
-		}
+			bool bShouldGeneratePatch = Params.IsGeneratingPatch && SC.StageTargetPlatform.GetPlatformPatchesWithDiffPak(Params, SC);
 
-		const string OutputFilenameExtension = ".pak";
+			if (bShouldGeneratePatch && !Params.HasBasedOnReleaseVersion)
+			{
+				LogInformation("Generating patch required a based on release version flag");
+			}
 
-		// read some compression settings from the project (once, shared across all pak commands)
+			const string OutputFilenameExtension = ".pak";
 
-		ConfigHierarchy PlatformGameConfig = ConfigCache.ReadHierarchy(ConfigHierarchyType.Game, DirectoryReference.FromFile(Params.RawProjectPath), SC.StageTargetPlatform.IniPlatformType);
+			// read some compression settings from the project (once, shared across all pak commands)
 
-		// in standard runs, ProjectPackagingSettings/bCompressed is read by the program invoking this script
-		//	and used to pass "-compressed" on the command line
-		// so typically "-compressed" == ProjectPackagingSettings/bCompressed
-		// if ProjectPackagingSettings/bCompressed is true, but "-compressed" was not passed, we now turn it on
-		// if ProjectPackagingSettings/bCompressed is true but you really don't want compression, use "-forceuncompressed"
-		bool bCompressed_ProjectPackagingSettings = false;
-		PlatformGameConfig.GetBool("/Script/UnrealEd.ProjectPackagingSettings", "bCompressed", out bCompressed_ProjectPackagingSettings);
+			ConfigHierarchy PlatformGameConfig = ConfigCache.ReadHierarchy(ConfigHierarchyType.Game, DirectoryReference.FromFile(Params.RawProjectPath), SC.StageTargetPlatform.IniPlatformType);
 
-		bool bForceUseProjectCompressionFormatIgnoreHardwareOverride = false;// do we want to override HW compression with project? if so, read it below
-		PlatformGameConfig.GetBool("/Script/UnrealEd.ProjectPackagingSettings", "bForceUseProjectCompressionFormatIgnoreHardwareOverride", out bForceUseProjectCompressionFormatIgnoreHardwareOverride);
+			// in standard runs, ProjectPackagingSettings/bCompressed is read by the program invoking this script
+			//	and used to pass "-compressed" on the command line
+			// so typically "-compressed" == ProjectPackagingSettings/bCompressed
+			// if ProjectPackagingSettings/bCompressed is true, but "-compressed" was not passed, we now turn it on
+			// if ProjectPackagingSettings/bCompressed is true but you really don't want compression, use "-forceuncompressed"
+			bool bCompressed_ProjectPackagingSettings = false;
+			PlatformGameConfig.GetBool("/Script/UnrealEd.ProjectPackagingSettings", "bCompressed", out bCompressed_ProjectPackagingSettings);
 
-		string HardwareCompressionFormat = null;
-		if ( ! bForceUseProjectCompressionFormatIgnoreHardwareOverride )
-		{
-			DataDrivenPlatformInfo.ConfigDataDrivenPlatformInfo DDPI = DataDrivenPlatformInfo.GetDataDrivenInfoForPlatform(SC.StageTargetPlatform.IniPlatformType);
-			if (DDPI != null)
+			bool bForceUseProjectCompressionFormatIgnoreHardwareOverride = false;// do we want to override HW compression with project? if so, read it below
+			PlatformGameConfig.GetBool("/Script/UnrealEd.ProjectPackagingSettings", "bForceUseProjectCompressionFormatIgnoreHardwareOverride", out bForceUseProjectCompressionFormatIgnoreHardwareOverride);
+
+			string HardwareCompressionFormat = null;
+			if (!bForceUseProjectCompressionFormatIgnoreHardwareOverride)
 			{
-				HardwareCompressionFormat = DDPI.HardwareCompressionFormat;
+				DataDrivenPlatformInfo.ConfigDataDrivenPlatformInfo DDPI = DataDrivenPlatformInfo.GetDataDrivenInfoForPlatform(SC.StageTargetPlatform.IniPlatformType);
+				if (DDPI != null)
+				{
+					HardwareCompressionFormat = DDPI.HardwareCompressionFormat;
+				}
 			}
-		}
-
-		// get the compression format from platform or the project
-		string CompressionFormats;
-		if (!string.IsNullOrEmpty(HardwareCompressionFormat))
-		{
-			LogInformation("Overriding to HardwareCompressionFormat = {0}",HardwareCompressionFormat);
-			CompressionFormats = HardwareCompressionFormat;
-		}
-		else
-		{
-			PlatformGameConfig.GetString("/Script/UnrealEd.ProjectPackagingSettings", "PackageCompressionFormat", out CompressionFormats);
-			// empty CompressionFormats string means "use default" eg. zlib
-		}
 
-		if (Params.ForceUncompressed)
-		{
-			//-ForceUncompressed on the command line means really don't do it even if the ini or DDPI say to do it
-			Params.Compressed = false;
-			bCompressed_ProjectPackagingSettings = false;
-		}
+			// get the compression format from platform or the project
+			string CompressionFormats;
+			if (!string.IsNullOrEmpty(HardwareCompressionFormat))
+			{
+				LogInformation("Overriding to HardwareCompressionFormat = {0}", HardwareCompressionFormat);
+				CompressionFormats = HardwareCompressionFormat;
+			}
+			else
+			{
+				PlatformGameConfig.GetString("/Script/UnrealEd.ProjectPackagingSettings", "PackageCompressionFormat", out CompressionFormats);
+				// empty CompressionFormats string means "use default" eg. zlib
+			}
 
-		if (CompressionFormats == "None" || CompressionFormats == "none")
-		{
-			// rather than pass "-compressed" + format=None through to UnrealPak/iostore
-			//	just turn off compression
-			if (Params.Compressed)
+			if (Params.ForceUncompressed)
 			{
-				//LogInformation(" -compressed was on command line but override to None, turning OFF!");
+				//-ForceUncompressed on the command line means really don't do it even if the ini or DDPI say to do it
 				Params.Compressed = false;
+				bCompressed_ProjectPackagingSettings = false;
 			}
-		}
-		else if (!Params.Compressed && bCompressed_ProjectPackagingSettings)
-		{
-			//LogInformation(" -compressed was not on command line but was in ini, turning on!");
-			Params.Compressed = true;
-		}
-
-		// Params.Compressed is done being changed; push it to the PakParams list :
-		// note this is a change of behavior; PakParams.bCompressed could have been true from a Chunk Manifest
-		//	even though compression was otherwise turned off
-		// we now overwrite that with the overall compression setting
-		//	this should always be an improvement in behavior
-		foreach (CreatePakParams PakParams in PakParamsList)
-		{
-			PakParams.bCompressed = Params.Compressed;
-		}
-
-		if (Params.Compressed && !string.IsNullOrWhiteSpace(CompressionFormats))
-		{
-			CompressionFormats = " -compressionformats=" + CompressionFormats;
-		}
-		else
-		{
-			CompressionFormats = "";
-		}
 
-		string AdditionalCompressionOptionsOnCommandLine = "";
-		if (Params.Compressed)
-		{
-			// the game may want to control compression settings, but since it may be in a plugin that checks the commandline for the settings, we need to pass
-			// the settings directly on the UnrealPak commandline, and not put it into the batch file lines (plugins can't get the unrealpak command list, and
-			// there's not a great way to communicate random strings down into the plugins during plugin init time)
-			PlatformGameConfig.GetString("/Script/UnrealEd.ProjectPackagingSettings", "PackageAdditionalCompressionOptions", out AdditionalCompressionOptionsOnCommandLine);
-
-			//LogInformation("PackageAdditionalCompressionOptions = {0}", AdditionalCompressionOptionsOnCommandLine);
-
-			string CompressionMethod;
-			PlatformGameConfig.GetString("/Script/UnrealEd.ProjectPackagingSettings", "PackageCompressionMethod", out CompressionMethod);
-			if (!string.IsNullOrWhiteSpace(CompressionMethod))
+			if (CompressionFormats == "None" || CompressionFormats == "none")
+			{
+				// rather than pass "-compressed" + format=None through to UnrealPak/iostore
+				//	just turn off compression
+				if (Params.Compressed)
+				{
+					//LogInformation(" -compressed was on command line but override to None, turning OFF!");
+					Params.Compressed = false;
+				}
+			}
+			else if (!Params.Compressed && bCompressed_ProjectPackagingSettings)
 			{
-				//LogInformation("CompressionMethod = {0}", CompressionMethod);
-				CompressionFormats += " -compressmethod=" + CompressionMethod;
+				//LogInformation(" -compressed was not on command line but was in ini, turning on!");
+				Params.Compressed = true;
 			}
 
-			// get the encoder compression effort level based on the cook type
-			int CompressionLevel = 0;
-			// GetInt32 fills out with = 0 if not found
-			if (Params.Distribution)
+			// Params.Compressed is done being changed; push it to the PakParams list :
+			// note this is a change of behavior; PakParams.bCompressed could have been true from a Chunk Manifest
+			//	even though compression was otherwise turned off
+			// we now overwrite that with the overall compression setting
+			//	this should always be an improvement in behavior
+			foreach (CreatePakParams PakParams in PakParamsList)
 			{
-				PlatformGameConfig.GetInt32("/Script/UnrealEd.ProjectPackagingSettings", "PackageCompressionLevel_Distribution", out CompressionLevel);
+				PakParams.bCompressed = Params.Compressed;
 			}
-			else if (SC.StageTargetConfigurations.Any(Config => Config == UnrealTargetConfiguration.Test || Config == UnrealTargetConfiguration.Shipping))
+
+			if (Params.Compressed && !string.IsNullOrWhiteSpace(CompressionFormats))
 			{
-				PlatformGameConfig.GetInt32("/Script/UnrealEd.ProjectPackagingSettings", "PackageCompressionLevel_TestShipping", out CompressionLevel);
+				CompressionFormats = " -compressionformats=" + CompressionFormats;
 			}
 			else
 			{
-				PlatformGameConfig.GetInt32("/Script/UnrealEd.ProjectPackagingSettings", "PackageCompressionLevel_DebugDevelopment", out CompressionLevel);
+				CompressionFormats = "";
 			}
 
-			if (CompressionLevel != 0)
+			string AdditionalCompressionOptionsOnCommandLine = "";
+			if (Params.Compressed)
 			{
-				//LogInformation(" CompressionLevel = {0}", CompressionLevel);
-				CompressionFormats += " -compresslevel=" + CompressionLevel;
-			}
-		}
+				// the game may want to control compression settings, but since it may be in a plugin that checks the commandline for the settings, we need to pass
+				// the settings directly on the UnrealPak commandline, and not put it into the batch file lines (plugins can't get the unrealpak command list, and
+				// there's not a great way to communicate random strings down into the plugins during plugin init time)
+				PlatformGameConfig.GetString("/Script/UnrealEd.ProjectPackagingSettings", "PackageAdditionalCompressionOptions", out AdditionalCompressionOptionsOnCommandLine);
 
-		Func<string, string> GetPostFix = PakFilename => GetPakFilePostFix(bShouldGeneratePatch, Params.HasIterateSharedCookedBuild, PlatformGameConfig, PakFilename);
+				//LogInformation("PackageAdditionalCompressionOptions = {0}", AdditionalCompressionOptionsOnCommandLine);
 
-		List<string> Commands = new List<string>();
-		List<string> LogNames = new List<string>();
-		List<string> IoStoreCommands = new List<string>();
+				string CompressionMethod;
+				PlatformGameConfig.GetString("/Script/UnrealEd.ProjectPackagingSettings", "PackageCompressionMethod", out CompressionMethod);
+				if (!string.IsNullOrWhiteSpace(CompressionMethod))
+				{
+					//LogInformation("CompressionMethod = {0}", CompressionMethod);
+					CompressionFormats += " -compressmethod=" + CompressionMethod;
+				}
 
-		// Calculate target patch index by iterating all pak files in source build
-		int TargetPatchIndex = 0;
-		int NumPakFiles = 0;
-		if (bShouldGeneratePatch)
-		{
-			foreach (CreatePakParams PakParams in PakParamsList)
-			{
-				if(ShouldSkipGeneratingPatch(PlatformGameConfig, PakParams.PakName))
+				// get the encoder compression effort level based on the cook type
+				int CompressionLevel = 0;
+				// GetInt32 fills out with = 0 if not found
+				if (Params.Distribution)
 				{
-					continue;
+					PlatformGameConfig.GetInt32("/Script/UnrealEd.ProjectPackagingSettings", "PackageCompressionLevel_Distribution", out CompressionLevel);
+				}
+				else if (SC.StageTargetConfigurations.Any(Config => Config == UnrealTargetConfiguration.Test || Config == UnrealTargetConfiguration.Shipping))
+				{
+					PlatformGameConfig.GetInt32("/Script/UnrealEd.ProjectPackagingSettings", "PackageCompressionLevel_TestShipping", out CompressionLevel);
+				}
+				else
+				{
+					PlatformGameConfig.GetInt32("/Script/UnrealEd.ProjectPackagingSettings", "PackageCompressionLevel_DebugDevelopment", out CompressionLevel);
 				}
 
-				string OutputFilename = PakParams.PakName + "-" + SC.FinalCookPlatform;
-				string ExistingPatchSearchPath = SC.StageTargetPlatform.GetReleasePakFilePath(SC, Params, null);
-				if (Directory.Exists(ExistingPatchSearchPath))
+				if (CompressionLevel != 0)
+				{
+					//LogInformation(" CompressionLevel = {0}", CompressionLevel);
+					CompressionFormats += " -compresslevel=" + CompressionLevel;
+				}
+			}
+
+			Func<string, string> GetPostFix = PakFilename => GetPakFilePostFix(bShouldGeneratePatch, Params.HasIterateSharedCookedBuild, PlatformGameConfig, PakFilename);
+
+			List<string> Commands = new List<string>();
+			List<string> LogNames = new List<string>();
+			List<string> IoStoreCommands = new List<string>();
+
+			// Calculate target patch index by iterating all pak files in source build
+			int TargetPatchIndex = 0;
+			int NumPakFiles = 0;
+			if (bShouldGeneratePatch)
+			{
+				foreach (CreatePakParams PakParams in PakParamsList)
 				{
-					string PostFix = GetPostFix(PakParams.PakName);
-					IEnumerable<string> PakFileSet = Directory.EnumerateFiles(ExistingPatchSearchPath, OutputFilename + "*" + PostFix + OutputFilenameExtension);
-					NumPakFiles += PakFileSet.Count();
+					if (ShouldSkipGeneratingPatch(PlatformGameConfig, PakParams.PakName))
+					{
+						continue;
+					}
 
-					foreach (string PakFilePath in PakFileSet)
+					string OutputFilename = PakParams.PakName + "-" + SC.FinalCookPlatform;
+					string ExistingPatchSearchPath = SC.StageTargetPlatform.GetReleasePakFilePath(SC, Params, null);
+					if (Directory.Exists(ExistingPatchSearchPath))
 					{
-						string PakFileName = Path.GetFileName(PakFilePath);
-						int StartIndex = OutputFilename.Length + 1;
-						int LengthVar = PakFileName.Length - (OutputFilename.Length + 1 + PostFix.Length + OutputFilenameExtension.Length);
-						if (LengthVar > 0)
+						string PostFix = GetPostFix(PakParams.PakName);
+						IEnumerable<string> PakFileSet = Directory.EnumerateFiles(ExistingPatchSearchPath, OutputFilename + "*" + PostFix + OutputFilenameExtension);
+						NumPakFiles += PakFileSet.Count();
+
+						foreach (string PakFilePath in PakFileSet)
 						{
-							string PakFileIndex = PakFileName.Substring(StartIndex, LengthVar);
-							int ChunkIndex;
-							if (int.TryParse(PakFileIndex, out ChunkIndex))
+							string PakFileName = Path.GetFileName(PakFilePath);
+							int StartIndex = OutputFilename.Length + 1;
+							int LengthVar = PakFileName.Length - (OutputFilename.Length + 1 + PostFix.Length + OutputFilenameExtension.Length);
+							if (LengthVar > 0)
 							{
-								if (ChunkIndex > TargetPatchIndex)
+								string PakFileIndex = PakFileName.Substring(StartIndex, LengthVar);
+								int ChunkIndex;
+								if (int.TryParse(PakFileIndex, out ChunkIndex))
 								{
-									TargetPatchIndex = ChunkIndex;
+									if (ChunkIndex > TargetPatchIndex)
+									{
+										TargetPatchIndex = ChunkIndex;
+									}
 								}
 							}
 						}
 					}
 				}
-			}
 
-			if (Params.ShouldAddPatchLevel && NumPakFiles > 0)
-			{
-				TargetPatchIndex++;
+				if (Params.ShouldAddPatchLevel && NumPakFiles > 0)
+				{
+					TargetPatchIndex++;
+				}
 			}
-		}
 
-		string BulkOption = "";
-		{
-			ConfigHierarchy PlatformEngineConfig;
-			if (Params.EngineConfigs.TryGetValue(SC.StageTargetPlatform.PlatformType, out PlatformEngineConfig))
+			string BulkOption = "";
 			{
-				bool bMasterEnable = false;
-				PlatformEngineConfig.GetBool("MemoryMappedFiles", "MasterEnable", out bMasterEnable);
-				if (bMasterEnable)
+				ConfigHierarchy PlatformEngineConfig;
+				if (Params.EngineConfigs.TryGetValue(SC.StageTargetPlatform.PlatformType, out PlatformEngineConfig))
 				{
-					int Value = 0;
-					PlatformEngineConfig.GetInt32("MemoryMappedFiles", "Alignment", out Value);
-					if (Value > 0)
+					bool bMasterEnable = false;
+					PlatformEngineConfig.GetBool("MemoryMappedFiles", "MasterEnable", out bMasterEnable);
+					if (bMasterEnable)
 					{
-						BulkOption = String.Format(" -AlignForMemoryMapping={0}", Value);
+						int Value = 0;
+						PlatformEngineConfig.GetInt32("MemoryMappedFiles", "Alignment", out Value);
+						if (Value > 0)
+						{
+							BulkOption = String.Format(" -AlignForMemoryMapping={0}", Value);
+						}
 					}
 				}
 			}
-		}
 
-		string PatchOptions = String.Empty;
-		if (bShouldGeneratePatch)
-		{
-			string PatchSeekOptMode = String.Empty;
-			string PatchSeekOptMaxGapSize = String.Empty;
-			string PatchSeekOptMaxInflationPercent = String.Empty;
-			string PatchSeekOptMaxAdjacentOrderDiff = String.Empty;
-			if (PlatformGameConfig.GetString("/Script/UnrealEd.ProjectPackagingSettings", "PatchSeekOptMode", out PatchSeekOptMode))
-			{
-				PatchOptions += String.Format(" -PatchSeekOptMode={0}", PatchSeekOptMode);
-			}
-			if (PlatformGameConfig.GetString("/Script/UnrealEd.ProjectPackagingSettings", "PatchSeekOptMaxGapSize", out PatchSeekOptMaxGapSize))
+			string PatchOptions = String.Empty;
+			if (bShouldGeneratePatch)
 			{
-				PatchOptions += String.Format(" -PatchSeekOptMaxGapSize={0}", PatchSeekOptMaxGapSize);
-			}
-			if (PlatformGameConfig.GetString("/Script/UnrealEd.ProjectPackagingSettings", "PatchSeekOptMaxInflationPercent", out PatchSeekOptMaxInflationPercent))
-			{
-				PatchOptions += String.Format(" -PatchSeekOptMaxInflationPercent={0}", PatchSeekOptMaxInflationPercent);
-			}
-			if (PlatformGameConfig.GetString("/Script/UnrealEd.ProjectPackagingSettings", "PatchSeekOptMaxAdjacentOrderDiff", out PatchSeekOptMaxAdjacentOrderDiff))
-			{
-				PatchOptions += String.Format(" -PatchSeekOptMaxAdjacentOrderDiff={0}", PatchSeekOptMaxAdjacentOrderDiff);
+				string PatchSeekOptMode = String.Empty;
+				string PatchSeekOptMaxGapSize = String.Empty;
+				string PatchSeekOptMaxInflationPercent = String.Empty;
+				string PatchSeekOptMaxAdjacentOrderDiff = String.Empty;
+				if (PlatformGameConfig.GetString("/Script/UnrealEd.ProjectPackagingSettings", "PatchSeekOptMode", out PatchSeekOptMode))
+				{
+					PatchOptions += String.Format(" -PatchSeekOptMode={0}", PatchSeekOptMode);
+				}
+				if (PlatformGameConfig.GetString("/Script/UnrealEd.ProjectPackagingSettings", "PatchSeekOptMaxGapSize", out PatchSeekOptMaxGapSize))
+				{
+					PatchOptions += String.Format(" -PatchSeekOptMaxGapSize={0}", PatchSeekOptMaxGapSize);
+				}
+				if (PlatformGameConfig.GetString("/Script/UnrealEd.ProjectPackagingSettings", "PatchSeekOptMaxInflationPercent", out PatchSeekOptMaxInflationPercent))
+				{
+					PatchOptions += String.Format(" -PatchSeekOptMaxInflationPercent={0}", PatchSeekOptMaxInflationPercent);
+				}
+				if (PlatformGameConfig.GetString("/Script/UnrealEd.ProjectPackagingSettings", "PatchSeekOptMaxAdjacentOrderDiff", out PatchSeekOptMaxAdjacentOrderDiff))
+				{
+					PatchOptions += String.Format(" -PatchSeekOptMaxAdjacentOrderDiff={0}", PatchSeekOptMaxAdjacentOrderDiff);
+				}
 			}
-		}
 
-		// Find primary and secondary order files if they exist
-		List<FileReference> PakOrderFileLocations = new List<FileReference>();
+			// Find primary and secondary order files if they exist
+			List<FileReference> PakOrderFileLocations = new List<FileReference>();
 
-		List<OrderFileSpec> OrderFileSpecs = new List<OrderFileSpec>();
-		if (PlatformGameConfig.TryGetValues("/Script/UnrealEd.ProjectPackagingSettings", "PakOrderFileSpecs", out IReadOnlyList<string> OrderFileConfigValues))
-		{
-			foreach (string Config in OrderFileConfigValues)
+			List<OrderFileSpec> OrderFileSpecs = new List<OrderFileSpec>();
+			if (PlatformGameConfig.TryGetValues("/Script/UnrealEd.ProjectPackagingSettings", "PakOrderFileSpecs", out IReadOnlyList<string> OrderFileConfigValues))
 			{
-				Dictionary<string, string> Fields;
-				if (!ConfigHierarchy.TryParse(Config, out Fields))
+				foreach (string Config in OrderFileConfigValues)
 				{
-					throw new AutomationException(String.Format("Failed to parse pak order file config \"{0}\"", Config));
-				}
+					Dictionary<string, string> Fields;
+					if (!ConfigHierarchy.TryParse(Config, out Fields))
+					{
+						throw new AutomationException(String.Format("Failed to parse pak order file config \"{0}\"", Config));
+					}
 
-				string FilePattern;
-				if (!Fields.TryGetValue("Pattern", out FilePattern))
-				{
-					throw new AutomationException(String.Format("Missing field 'Pattern' in pak order file config \"{0}\"", Config));
-				}
+					string FilePattern;
+					if (!Fields.TryGetValue("Pattern", out FilePattern))
+					{
+						throw new AutomationException(String.Format("Missing field 'Pattern' in pak order file config \"{0}\"", Config));
+					}
 
-				// Priority missing is fine, we default to 0
-				int Priority = 0;
-				if( Fields.TryGetValue("Priority", out string PriorityString)) 
-				{
-					if (!int.TryParse(PriorityString, out Priority))
+					// Priority missing is fine, we default to 0
+					int Priority = 0;
+					if (Fields.TryGetValue("Priority", out string PriorityString))
 					{
-						throw new AutomationException(String.Format("Failed to parse int from field 'Priority' in pak order file config \"{0}\"", Config));
+						if (!int.TryParse(PriorityString, out Priority))
+						{
+							throw new AutomationException(String.Format("Failed to parse int from field 'Priority' in pak order file config \"{0}\"", Config));
+						}
 					}
-				}
 
-				OrderFileSpecs.Add(new OrderFileSpec { FileNamePattern = FilePattern, OrderType = OrderFile.OrderFileType.Custom, Required = FilePattern.Contains("*") == false, Priority = Priority });
+					OrderFileSpecs.Add(new OrderFileSpec { FileNamePattern = FilePattern, OrderType = OrderFile.OrderFileType.Custom, Required = FilePattern.Contains("*") == false, Priority = Priority });
+				}
 			}
-		}
 
-		if(OrderFileSpecs.Count == 0)
-		{
-			// Default filespecs
-			OrderFileSpecs.AddRange(new OrderFileSpec[] {
+			if (OrderFileSpecs.Count == 0)
+			{
+				// Default filespecs
+				OrderFileSpecs.AddRange(new OrderFileSpec[] {
 				new OrderFileSpec{ FileNamePattern = "GameOpenOrder*.log", OrderType = OrderFile.OrderFileType.Game, Required=false, Priority = 0 },
 				new OrderFileSpec{ FileNamePattern = "CookerOpenOrder*.log", OrderType = OrderFile.OrderFileType.Cooker, Required=false, Priority = 0 },
 				new OrderFileSpec{ FileNamePattern = "EditorOpenOrder.log", OrderType = OrderFile.OrderFileType.Editor, Required=false, Priority = 0 }
 			});
-		}
+			}
 
-		// search CookPlaform (e.g. IOSClient and then regular platform (e.g. IOS).
-		string[] OrderPlatformNames = new string[] { SC.FinalCookPlatform, SC.StageTargetPlatform.GetTargetPlatformDescriptor().Type.ToString() };
+			// search CookPlaform (e.g. IOSClient and then regular platform (e.g. IOS).
+			string[] OrderPlatformNames = new string[] { SC.FinalCookPlatform, SC.StageTargetPlatform.GetTargetPlatformDescriptor().Type.ToString() };
 
-		List<DirectoryReference> OrderPathsToSearch = new List<DirectoryReference>();
-		OrderPathsToSearch.Add(DirectoryReference.Combine(SC.ProjectRoot, "Build", "FileOpenOrder")); // Search common directory first
-		foreach (string OrderLocation in OrderPlatformNames)
-		{
-			DirectoryReference PakOrderFileLocationBase = DirectoryReference.Combine(SC.ProjectRoot, "Platforms", OrderLocation, "Build", "FileOpenOrder");
-			if (!DirectoryReference.Exists(PakOrderFileLocationBase))
+			List<DirectoryReference> OrderPathsToSearch = new List<DirectoryReference>();
+			OrderPathsToSearch.Add(DirectoryReference.Combine(SC.ProjectRoot, "Build", "FileOpenOrder")); // Search common directory first
+			foreach (string OrderLocation in OrderPlatformNames)
 			{
-				PakOrderFileLocationBase = DirectoryReference.Combine(SC.ProjectRoot, "Build", OrderLocation, "FileOpenOrder");
-			}
+				DirectoryReference PakOrderFileLocationBase = DirectoryReference.Combine(SC.ProjectRoot, "Platforms", OrderLocation, "Build", "FileOpenOrder");
+				if (!DirectoryReference.Exists(PakOrderFileLocationBase))
+				{
+					PakOrderFileLocationBase = DirectoryReference.Combine(SC.ProjectRoot, "Build", OrderLocation, "FileOpenOrder");
+				}
 
-			OrderPathsToSearch.Add(PakOrderFileLocationBase);
-		}
+				OrderPathsToSearch.Add(PakOrderFileLocationBase);
+			}
 
-		List<OrderFile> OrderFiles = new List<OrderFile>();
+			List<OrderFile> OrderFiles = new List<OrderFile>();
 
-		for (int OrderFileIndex = 0; OrderFileIndex < OrderFileSpecs.Count; OrderFileIndex++)
-		{
-			OrderFileSpec Spec = OrderFileSpecs[OrderFileIndex];
-			bool bAnyFound = false;
-			foreach (DirectoryReference BaseDir in OrderPathsToSearch)
+			for (int OrderFileIndex = 0; OrderFileIndex < OrderFileSpecs.Count; OrderFileIndex++)
 			{
-				// Add input file to control order of file within the pak
-				FileReference[] FileLocations = CommandUtils.FindFiles(Spec.FileNamePattern, false, BaseDir);
-
-				if (FileLocations != null)
+				OrderFileSpec Spec = OrderFileSpecs[OrderFileIndex];
+				bool bAnyFound = false;
+				foreach (DirectoryReference BaseDir in OrderPathsToSearch)
 				{
-					foreach (var File in FileLocations)
-					{
-						bAnyFound = true;
-						OrderFile FileOrder = new OrderFile(File, Spec.OrderType, OrderFileIndex, Spec.Priority);
+					// Add input file to control order of file within the pak
+					FileReference[] FileLocations = CommandUtils.FindFiles(Spec.FileNamePattern, false, BaseDir);
 
-						//check if the file is alreay in the list
-						if (!OrderFiles.Any(o => o.File.FullName.Equals(File.FullName)))
+					if (FileLocations != null)
+					{
+						foreach (var File in FileLocations)
 						{
-							OrderFiles.Add(FileOrder);
+							bAnyFound = true;
+							OrderFile FileOrder = new OrderFile(File, Spec.OrderType, OrderFileIndex, Spec.Priority);
+
+							//check if the file is alreay in the list
+							if (!OrderFiles.Any(o => o.File.FullName.Equals(File.FullName)))
+							{
+								OrderFiles.Add(FileOrder);
+							}
 						}
 					}
 				}
+				if (!bAnyFound && Spec.Required)
+				{
+					throw new AutomationException(String.Format("Did not find any order files from rule {0}", Spec.FileNamePattern));
+				}
 			}
-			if (!bAnyFound && Spec.Required)
-			{
-				throw new AutomationException(String.Format("Did not find any order files from rule {0}", Spec.FileNamePattern));
-			}
-		}
-
-		OrderFiles.Sort();
 
-		//Only allow Editor open order as a primary input
-		if (OrderFiles.Count() > 1 && OrderFiles[0].OrderType != OrderFile.OrderFileType.Editor)
-		{
-			OrderFiles.RemoveAll(x => x.OrderType == OrderFile.OrderFileType.Editor);
-		}
+			OrderFiles.Sort();
 
-        // Check second order 
-        bool bUseSecondaryOrder = false;
-        PlatformGameConfig.GetBool("/Script/UnrealEd.ProjectPackagingSettings", "bPakUsesSecondaryOrder", out bUseSecondaryOrder);
+			//Only allow Editor open order as a primary input
+			if (OrderFiles.Count() > 1 && OrderFiles[0].OrderType != OrderFile.OrderFileType.Editor)
+			{
+				OrderFiles.RemoveAll(x => x.OrderType == OrderFile.OrderFileType.Editor);
+			}
 
-		if(!bUseSecondaryOrder)
-        {
-			OrderFiles.RemoveAll(x => (x.OrderType == OrderFile.OrderFileType.Cooker && x.AppendOrder == Int32.MaxValue));
+			// Check second order 
+			bool bUseSecondaryOrder = false;
+			PlatformGameConfig.GetBool("/Script/UnrealEd.ProjectPackagingSettings", "bPakUsesSecondaryOrder", out bUseSecondaryOrder);
 
-			// chekc if theres a merging cookopenorder file
-			bUseSecondaryOrder = OrderFiles.Any(x => x.OrderType == OrderFile.OrderFileType.Cooker);
-		}
+			if (!bUseSecondaryOrder)
+			{
+				OrderFiles.RemoveAll(x => (x.OrderType == OrderFile.OrderFileType.Cooker && x.AppendOrder == Int32.MaxValue));
 
-		LogInformation("Using {0} pak order files:", OrderFiles.Count);
-		foreach( OrderFile File in OrderFiles )
-		{
-			LogInformation("    {0} Priority {1}", File.File.ToString(), File.Priority);
-		}
+				// chekc if theres a merging cookopenorder file
+				bUseSecondaryOrder = OrderFiles.Any(x => x.OrderType == OrderFile.OrderFileType.Cooker);
+			}
 
-		List<Tuple<FileReference, StagedFileReference, string>> Outputs = new List<Tuple<FileReference, StagedFileReference, string>>();
-		foreach (CreatePakParams PakParams in PakParamsList)
-		{
-			string OutputFilename = PakParams.PakName + "-" + SC.FinalCookPlatform;
-			string PostFix = GetPostFix(PakParams.PakName);
-			if (bShouldGeneratePatch && !ShouldSkipGeneratingPatch(PlatformGameConfig, PakParams.PakName))
+			LogInformation("Using {0} pak order files:", OrderFiles.Count);
+			foreach (OrderFile File in OrderFiles)
 			{
-				OutputFilename = OutputFilename + "_" + TargetPatchIndex;
+				LogInformation("    {0} Priority {1}", File.File.ToString(), File.Priority);
 			}
-			OutputFilename = OutputFilename + PostFix;
 
-			StagedFileReference OutputRelativeLocation;
-			if (Params.HasDLCName)
+			List<Tuple<FileReference, StagedFileReference, string>> Outputs = new List<Tuple<FileReference, StagedFileReference, string>>();
+			foreach (CreatePakParams PakParams in PakParamsList)
 			{
-				if (Params.DLCOverrideStagedSubDir != null)
+				string OutputFilename = PakParams.PakName + "-" + SC.FinalCookPlatform;
+				string PostFix = GetPostFix(PakParams.PakName);
+				if (bShouldGeneratePatch && !ShouldSkipGeneratingPatch(PlatformGameConfig, PakParams.PakName))
 				{
-					OutputRelativeLocation = StagedFileReference.Combine(SC.RelativeProjectRootForStage, Params.DLCOverrideStagedSubDir, "Content", "Paks", Params.DLCFile.GetFileNameWithoutExtension() + OutputFilename + ".pak");
+					OutputFilename = OutputFilename + "_" + TargetPatchIndex;
 				}
-				else
+				OutputFilename = OutputFilename + PostFix;
+
+				StagedFileReference OutputRelativeLocation;
+				if (Params.HasDLCName)
 				{
-					string PluginSubdirectory;
-					if (Params.DLCFile.Directory.IsUnderDirectory(SC.ProjectRoot))
+					if (Params.DLCOverrideStagedSubDir != null)
 					{
-						PluginSubdirectory = Params.DLCFile.Directory.MakeRelativeTo(SC.ProjectRoot);
+						OutputRelativeLocation = StagedFileReference.Combine(SC.RelativeProjectRootForStage, Params.DLCOverrideStagedSubDir, "Content", "Paks", Params.DLCFile.GetFileNameWithoutExtension() + OutputFilename + ".pak");
 					}
 					else
 					{
-						// if the plugin is external, the MakeRelativeTo above will fail, so just use the plugin name as the directory
-						PluginSubdirectory = Params.DLCFile.GetFileNameWithoutAnyExtensions();
-					}
-					OutputRelativeLocation = StagedFileReference.Combine(SC.RelativeProjectRootForStage, PluginSubdirectory, "Content", "Paks", SC.FinalCookPlatform, Params.DLCFile.GetFileNameWithoutExtension() + OutputFilename + ".pak");
-				}
-			}
-			else
-			{
-				OutputRelativeLocation = StagedFileReference.Combine(SC.RelativeProjectRootForStage, "Content", "Paks", OutputFilename + OutputFilenameExtension);
-			}
-			if (SC.StageTargetPlatform.DeployLowerCaseFilenames(StagedFileType.UFS))
-			{
-				OutputRelativeLocation = OutputRelativeLocation.ToLowerInvariant();
-			}
-			OutputRelativeLocation = SC.StageTargetPlatform.Remap(OutputRelativeLocation);
-
-			FileReference OutputLocation = FileReference.Combine(SC.RuntimeRootDir, OutputRelativeLocation.Name);
+						string PluginSubdirectory;
+						if (Params.DLCFile.Directory.IsUnderDirectory(SC.ProjectRoot))
+						{
+							PluginSubdirectory = Params.DLCFile.Directory.MakeRelativeTo(SC.ProjectRoot);
+						}
+						else
+						{
+							// if the plugin is external, the MakeRelativeTo above will fail, so just use the plugin name as the directory
+							PluginSubdirectory = Params.DLCFile.GetFileNameWithoutAnyExtensions();
+						}
+						OutputRelativeLocation = StagedFileReference.Combine(SC.RelativeProjectRootForStage, PluginSubdirectory, "Content", "Paks", SC.FinalCookPlatform, Params.DLCFile.GetFileNameWithoutExtension() + OutputFilename + ".pak");
+					}
+				}
+				else
+				{
+					OutputRelativeLocation = StagedFileReference.Combine(SC.RelativeProjectRootForStage, "Content", "Paks", OutputFilename + OutputFilenameExtension);
+				}
+				if (SC.StageTargetPlatform.DeployLowerCaseFilenames(StagedFileType.UFS))
+				{
+					OutputRelativeLocation = OutputRelativeLocation.ToLowerInvariant();
+				}
+				OutputRelativeLocation = SC.StageTargetPlatform.Remap(OutputRelativeLocation);
 
-			bool bCopiedExistingPak = false;
+				FileReference OutputLocation = FileReference.Combine(SC.RuntimeRootDir, OutputRelativeLocation.Name);
 
-			if (SC.StageTargetPlatform != SC.CookSourcePlatform && !Params.IgnorePaksFromDifferentCookSource)
-			{
-				// Check to see if we have an existing pak file we can use
+				bool bCopiedExistingPak = false;
 
-				StagedFileReference SourceOutputRelativeLocation = StagedFileReference.Combine(SC.RelativeProjectRootForStage, "Content/Paks/", PakParams.PakName + "-" + SC.CookPlatform + PostFix + ".pak");
-				if (SC.CookSourcePlatform.DeployLowerCaseFilenames(StagedFileType.UFS))
+				if (SC.StageTargetPlatform != SC.CookSourcePlatform && !Params.IgnorePaksFromDifferentCookSource)
 				{
-					SourceOutputRelativeLocation = SourceOutputRelativeLocation.ToLowerInvariant();
-				}
-				SourceOutputRelativeLocation = SC.CookSourcePlatform.Remap(SourceOutputRelativeLocation);
+					// Check to see if we have an existing pak file we can use
 
-				FileReference SourceOutputLocation = FileReference.Combine(SC.CookSourceRuntimeRootDir, SourceOutputRelativeLocation.Name);
-				if (FileExists_NoExceptions(SourceOutputLocation.FullName))
-				{
-					InternalUtils.SafeCreateDirectory(Path.GetDirectoryName(OutputLocation.FullName), true);
+					StagedFileReference SourceOutputRelativeLocation = StagedFileReference.Combine(SC.RelativeProjectRootForStage, "Content/Paks/", PakParams.PakName + "-" + SC.CookPlatform + PostFix + ".pak");
+					if (SC.CookSourcePlatform.DeployLowerCaseFilenames(StagedFileType.UFS))
+					{
+						SourceOutputRelativeLocation = SourceOutputRelativeLocation.ToLowerInvariant();
+					}
+					SourceOutputRelativeLocation = SC.CookSourcePlatform.Remap(SourceOutputRelativeLocation);
 
-					if (InternalUtils.SafeCopyFile(SourceOutputLocation.FullName, OutputLocation.FullName))
+					FileReference SourceOutputLocation = FileReference.Combine(SC.CookSourceRuntimeRootDir, SourceOutputRelativeLocation.Name);
+					if (FileExists_NoExceptions(SourceOutputLocation.FullName))
 					{
-						LogInformation("Copying source pak from {0} to {1} instead of creating new pak", SourceOutputLocation, OutputLocation);
-						bCopiedExistingPak = true;
+						InternalUtils.SafeCreateDirectory(Path.GetDirectoryName(OutputLocation.FullName), true);
 
-						FileReference InSigFile = SourceOutputLocation.ChangeExtension(".sig");
-						if (FileReference.Exists(InSigFile))
+						if (InternalUtils.SafeCopyFile(SourceOutputLocation.FullName, OutputLocation.FullName))
 						{
-							FileReference OutSigFile = OutputLocation.ChangeExtension(".sig");
+							LogInformation("Copying source pak from {0} to {1} instead of creating new pak", SourceOutputLocation, OutputLocation);
+							bCopiedExistingPak = true;
 
-							LogInformation("Copying pak sig from {0} to {1}", InSigFile, OutSigFile);
-
-							if (!InternalUtils.SafeCopyFile(InSigFile.FullName, OutSigFile.FullName))
+							FileReference InSigFile = SourceOutputLocation.ChangeExtension(".sig");
+							if (FileReference.Exists(InSigFile))
 							{
-								LogInformation("Failed to copy pak sig {0} to {1}, creating new pak", InSigFile, InSigFile);
-								bCopiedExistingPak = false;
+								FileReference OutSigFile = OutputLocation.ChangeExtension(".sig");
+
+								LogInformation("Copying pak sig from {0} to {1}", InSigFile, OutSigFile);
+
+								if (!InternalUtils.SafeCopyFile(InSigFile.FullName, OutSigFile.FullName))
+								{
+									LogInformation("Failed to copy pak sig {0} to {1}, creating new pak", InSigFile, InSigFile);
+									bCopiedExistingPak = false;
+								}
 							}
-						}
 
+						}
+					}
+					if (!bCopiedExistingPak)
+					{
+						LogInformation("Failed to copy source pak from {0} to {1}, creating new pak", SourceOutputLocation, OutputLocation);
 					}
 				}
-				if (!bCopiedExistingPak)
-				{
-					LogInformation("Failed to copy source pak from {0} to {1}, creating new pak", SourceOutputLocation, OutputLocation);
-				}
-			}
-
-			string PatchSourceContentPath = null;
-			if (bShouldGeneratePatch && !ShouldSkipGeneratingPatch(PlatformGameConfig, PakParams.PakName))
-			{
-				// don't include the post fix in this filename because we are looking for the source pak path
-				string PakFilename = PakParams.PakName + "-" + SC.FinalCookPlatform + "*.pak";
-				PatchSourceContentPath = SC.StageTargetPlatform.GetReleasePakFilePath(SC, Params, PakFilename);
-			}
 
-			if (!bCopiedExistingPak)
-			{
-				if (FileReference.Exists(OutputLocation) && !Params.IgnorePaksFromDifferentCookSource)
+				string PatchSourceContentPath = null;
+				if (bShouldGeneratePatch && !ShouldSkipGeneratingPatch(PlatformGameConfig, PakParams.PakName))
 				{
-					string UnrealPakResponseFileName = CombinePaths(CmdEnv.LogFolder, "PakList_" + OutputLocation.GetFileNameWithoutExtension() + ".txt");
-					if (File.Exists(UnrealPakResponseFileName) && FileReference.GetLastWriteTimeUtc(OutputLocation) > File.GetLastWriteTimeUtc(UnrealPakResponseFileName))
-					{
-						bCopiedExistingPak = true;
-					}
+					// don't include the post fix in this filename because we are looking for the source pak path
+					string PakFilename = PakParams.PakName + "-" + SC.FinalCookPlatform + "*.pak";
+					PatchSourceContentPath = SC.StageTargetPlatform.GetReleasePakFilePath(SC, Params, PakFilename);
 				}
+
 				if (!bCopiedExistingPak)
 				{
-					List<OrderFile> PrimaryOrderFiles = OrderFiles.FindAll(x => (x.OrderType != OrderFile.OrderFileType.Cooker));
-					List<OrderFile> SecondaryOrderFiles = null;
-
-					// Add a secondary order if there is one specified
-					if (bUseSecondaryOrder && OrderFiles.Count >= 1)
-					{
-						SecondaryOrderFiles = OrderFiles.FindAll(x => x.OrderType == OrderFile.OrderFileType.Cooker);
-					}
-
-					string AdditionalArgs = PatchOptions;
-					if (bShouldGeneratePatch && ShouldSkipGeneratingPatch(PlatformGameConfig, PakParams.PakName))
+					if (FileReference.Exists(OutputLocation) && !Params.IgnorePaksFromDifferentCookSource)
 					{
-						AdditionalArgs = " -SkipPatch";
+						string UnrealPakResponseFileName = CombinePaths(CmdEnv.LogFolder, "PakList_" + OutputLocation.GetFileNameWithoutExtension() + ".txt");
+						if (File.Exists(UnrealPakResponseFileName) && FileReference.GetLastWriteTimeUtc(OutputLocation) > File.GetLastWriteTimeUtc(UnrealPakResponseFileName))
+						{
+							bCopiedExistingPak = true;
+						}
 					}
-
-					bool bPakFallbackOrderForNonUassetFiles = false;
-					PlatformGameConfig.GetBool("/Script/UnrealEd.ProjectPackagingSettings", "bPakFallbackOrderForNonUassetFiles", out bPakFallbackOrderForNonUassetFiles);
-					if (bPakFallbackOrderForNonUassetFiles)
+					if (!bCopiedExistingPak)
 					{
-						AdditionalArgs += " -fallbackOrderForNonUassetFiles";
-					}
+						List<OrderFile> PrimaryOrderFiles = OrderFiles.FindAll(x => (x.OrderType != OrderFile.OrderFileType.Cooker));
+						List<OrderFile> SecondaryOrderFiles = null;
 
-					ConfigHierarchy PlatformEngineConfig;
-					if (Params.EngineConfigs.TryGetValue(SC.StageTargetPlatform.PlatformType, out PlatformEngineConfig))					
-					{
-						// if the runtime will want to reduce memory usage, we have to disable the pak index freezing
-						bool bUnloadPakEntries = false;
-						bool bShrinkPakEntries = false;
-						PlatformEngineConfig.GetBool("Pak", "UnloadPakEntryFilenamesIfPossible", out bUnloadPakEntries);
-						PlatformEngineConfig.GetBool("Pak", "ShrinkPakEntriesMemoryUsage", out bShrinkPakEntries);
-						if (bUnloadPakEntries || bShrinkPakEntries)
+						// Add a secondary order if there is one specified
+						if (bUseSecondaryOrder && OrderFiles.Count >= 1)
 						{
-							AdditionalArgs += " -allowForIndexUnload";
+							SecondaryOrderFiles = OrderFiles.FindAll(x => x.OrderType == OrderFile.OrderFileType.Cooker);
 						}
-					}
 
-					// pass the targetplatform so the index may be able to be frozen
-					AdditionalArgs += " -platform=" + ConfigHierarchy.GetIniPlatformName(SC.StageTargetPlatform.IniPlatformType);
+						string AdditionalArgs = PatchOptions;
+						if (bShouldGeneratePatch && ShouldSkipGeneratingPatch(PlatformGameConfig, PakParams.PakName))
+						{
+							AdditionalArgs = " -SkipPatch";
+						}
 
-					Dictionary<string, string> UnrealPakResponseFile = PakParams.UnrealPakResponseFile;
-					if (ShouldCreateIoStoreContainerFiles(Params, SC.StageTargetPlatform))
-					{
-						bool bAllowBulkDataInIoStore = true;
-						if(!PlatformEngineConfig.GetBool("Core.System", "AllowBulkDataInIoStore", out bAllowBulkDataInIoStore))
+						bool bPakFallbackOrderForNonUassetFiles = false;
+						PlatformGameConfig.GetBool("/Script/UnrealEd.ProjectPackagingSettings", "bPakFallbackOrderForNonUassetFiles", out bPakFallbackOrderForNonUassetFiles);
+						if (bPakFallbackOrderForNonUassetFiles)
 						{
-							bAllowBulkDataInIoStore = true; // Default is to allow it in the IoStore
+							AdditionalArgs += " -fallbackOrderForNonUassetFiles";
 						}
 
-						bool bAllowShadersInIoStore;
-						if (!PlatformEngineConfig.GetBool("Core.System", "AllowShadersInIoStore", out bAllowShadersInIoStore))
+						ConfigHierarchy PlatformEngineConfig;
+						if (Params.EngineConfigs.TryGetValue(SC.StageTargetPlatform.PlatformType, out PlatformEngineConfig))
 						{
-							bAllowShadersInIoStore = false;
+							// if the runtime will want to reduce memory usage, we have to disable the pak index freezing
+							bool bUnloadPakEntries = false;
+							bool bShrinkPakEntries = false;
+							PlatformEngineConfig.GetBool("Pak", "UnloadPakEntryFilenamesIfPossible", out bUnloadPakEntries);
+							PlatformEngineConfig.GetBool("Pak", "ShrinkPakEntriesMemoryUsage", out bShrinkPakEntries);
+							if (bUnloadPakEntries || bShrinkPakEntries)
+							{
+								AdditionalArgs += " -allowForIndexUnload";
+							}
 						}
 
+						// pass the targetplatform so the index may be able to be frozen
+						AdditionalArgs += " -platform=" + ConfigHierarchy.GetIniPlatformName(SC.StageTargetPlatform.IniPlatformType);
 
-						UnrealPakResponseFile = new Dictionary<string, string>();
-						Dictionary<string, string> IoStoreResponseFile = new Dictionary<string, string>();
-						foreach (var Entry in PakParams.UnrealPakResponseFile)
+						Dictionary<string, string> UnrealPakResponseFile = PakParams.UnrealPakResponseFile;
+						if (ShouldCreateIoStoreContainerFiles(Params, SC.StageTargetPlatform))
 						{
-							// Temporary solution to filter non cooked packages from I/O store container file(s)
-							if (SC.OnlyAllowPackagesFromStdCookPathInIoStore && !Entry.Key.ToLower().Contains("\\saved\\cooked\\"))
+							bool bAllowBulkDataInIoStore = true;
+							if (!PlatformEngineConfig.GetBool("Core.System", "AllowBulkDataInIoStore", out bAllowBulkDataInIoStore))
 							{
-								UnrealPakResponseFile.Add(Entry.Key, Entry.Value);
-								continue;
+								bAllowBulkDataInIoStore = true; // Default is to allow it in the IoStore
 							}
 
-							if (Path.GetExtension(Entry.Key).Contains(".uasset") ||
-								Path.GetExtension(Entry.Key).Contains(".umap"))
+							bool bAllowShadersInIoStore;
+							if (!PlatformEngineConfig.GetBool("Core.System", "AllowShadersInIoStore", out bAllowShadersInIoStore))
 							{
-								IoStoreResponseFile.Add(Entry.Key, Entry.Value);
+								bAllowShadersInIoStore = false;
 							}
-							else if(Path.GetExtension(Entry.Key).Contains(".ubulk") ||
-									Path.GetExtension(Entry.Key).Contains(".uptnl"))
+
+
+							UnrealPakResponseFile = new Dictionary<string, string>();
+							Dictionary<string, string> IoStoreResponseFile = new Dictionary<string, string>();
+							foreach (var Entry in PakParams.UnrealPakResponseFile)
 							{
-								if(bAllowBulkDataInIoStore)
+								// Temporary solution to filter non cooked packages from I/O store container file(s)
+								if (SC.OnlyAllowPackagesFromStdCookPathInIoStore && !Entry.Key.ToLower().Contains("\\saved\\cooked\\"))
+								{
+									UnrealPakResponseFile.Add(Entry.Key, Entry.Value);
+									continue;
+								}
+
+								if (Path.GetExtension(Entry.Key).Contains(".uasset") ||
+									Path.GetExtension(Entry.Key).Contains(".umap"))
 								{
 									IoStoreResponseFile.Add(Entry.Key, Entry.Value);
 								}
-								else
+								else if (Path.GetExtension(Entry.Key).Contains(".ubulk") ||
+										Path.GetExtension(Entry.Key).Contains(".uptnl"))
 								{
-									UnrealPakResponseFile.Add(Entry.Key, Entry.Value);
+									if (bAllowBulkDataInIoStore)
+									{
+										IoStoreResponseFile.Add(Entry.Key, Entry.Value);
+									}
+									else
+									{
+										UnrealPakResponseFile.Add(Entry.Key, Entry.Value);
+									}
 								}
-							}
-							else if (Path.GetExtension(Entry.Key).Contains(".ushaderbytecode"))
-							{
-								if (bAllowShadersInIoStore)
+								else if (Path.GetExtension(Entry.Key).Contains(".ushaderbytecode"))
 								{
-									IoStoreResponseFile.Add(Entry.Key, Entry.Value);
+									if (bAllowShadersInIoStore)
+									{
+										IoStoreResponseFile.Add(Entry.Key, Entry.Value);
+									}
+									else
+									{
+										UnrealPakResponseFile.Add(Entry.Key, Entry.Value);
+									}
 								}
-								else
+								else if (!Path.GetExtension(Entry.Key).Contains(".uexp"))
 								{
+									// all other extensions go in Pak not iostore
+									//	for example the compression-excluded extensions (ini,res) go in here
 									UnrealPakResponseFile.Add(Entry.Key, Entry.Value);
 								}
+
 							}
-							else if (!Path.GetExtension(Entry.Key).Contains(".uexp"))
+
+							string ContainerPatchSourcePath = null;
+							if (Params.HasBasedOnReleaseVersion)
 							{
-								// all other extensions go in Pak not iostore
-								//	for example the compression-excluded extensions (ini,res) go in here
-								UnrealPakResponseFile.Add(Entry.Key, Entry.Value);
+								string ContainerWildcard = PakParams.PakName + "-" + SC.FinalCookPlatform + "*.utoc";
+								ContainerPatchSourcePath = CombinePaths(Params.GetBasedOnReleaseVersionPath(SC, Params.Client), ContainerWildcard);
 							}
-
-						}
-
-						string ContainerPatchSourcePath = null;
-						if (Params.HasBasedOnReleaseVersion)
-						{
-							string ContainerWildcard = PakParams.PakName + "-" + SC.FinalCookPlatform + "*.utoc";
-							ContainerPatchSourcePath = CombinePaths(Params.GetBasedOnReleaseVersionPath(SC, Params.Client), ContainerWildcard);
+							bool bGenerateDiffPatch = bShouldGeneratePatch && !ShouldSkipGeneratingPatch(PlatformGameConfig, PakParams.PakName);
+							bool bCompressContainers = PakParams.bCompressed;
+
+							IoStoreCommands.Add(GetIoStoreCommandArguments(
+								IoStoreResponseFile,
+								PakParams.PakName,
+								OutputLocation,
+								bCompressContainers,
+								CryptoSettings,
+								PakParams.EncryptionKeyGuid,
+								ContainerPatchSourcePath,
+								bGenerateDiffPatch,
+								Params.HasDLCName));
 						}
-						bool bGenerateDiffPatch = bShouldGeneratePatch && !ShouldSkipGeneratingPatch(PlatformGameConfig, PakParams.PakName);
-						bool bCompressContainers = PakParams.bCompressed;
 
-						IoStoreCommands.Add(GetIoStoreCommandArguments(
-							IoStoreResponseFile,
-							PakParams.PakName,
+						string PakEncryptionKeyGuid = Params.SkipEncryption ? "" : PakParams.EncryptionKeyGuid;
+						Commands.Add(GetUnrealPakArguments(
+							Params.RawProjectPath,
+							UnrealPakResponseFile,
 							OutputLocation,
-							bCompressContainers,
-							CryptoSettings,
-							PakParams.EncryptionKeyGuid,
-							ContainerPatchSourcePath,
-							bGenerateDiffPatch,
-							Params.HasDLCName));
+							PrimaryOrderFiles,
+							SC.StageTargetPlatform.GetPlatformPakCommandLine(Params, SC) + AdditionalArgs + BulkOption + CompressionFormats + " " + Params.AdditionalPakOptions,
+							PakParams.bCompressed,
+							Params.SkipEncryption ? null : CryptoSettings,
+							Params.SkipEncryption ? null : CryptoKeysCacheFilename,
+							PatchSourceContentPath,
+							Params.SkipEncryption ? "" : PakParams.EncryptionKeyGuid,
+							SecondaryOrderFiles));
+
+						LogNames.Add(OutputLocation.GetFileNameWithoutExtension());
 					}
-
-					string PakEncryptionKeyGuid = Params.SkipEncryption ? "" : PakParams.EncryptionKeyGuid;
-					Commands.Add(GetUnrealPakArguments(
-						Params.RawProjectPath,
-						UnrealPakResponseFile,
-						OutputLocation,
-						PrimaryOrderFiles,
-						SC.StageTargetPlatform.GetPlatformPakCommandLine(Params, SC) + AdditionalArgs + BulkOption + CompressionFormats + " " + Params.AdditionalPakOptions,
-						PakParams.bCompressed,
-						Params.SkipEncryption ? null : CryptoSettings,
-						Params.SkipEncryption ? null : CryptoKeysCacheFilename,
-						PatchSourceContentPath,
-						Params.SkipEncryption ? "" : PakParams.EncryptionKeyGuid,
-						SecondaryOrderFiles));
-
-					LogNames.Add(OutputLocation.GetFileNameWithoutExtension());
-				}
-			}
-
-			Outputs.Add(Tuple.Create(OutputLocation, OutputRelativeLocation, PatchSourceContentPath));
-		}
-
-		// Actually execute UnrealPak
-		if (Commands.Count > 0)
-		{
-			RunUnrealPakInParallel( Commands, LogNames, AdditionalCompressionOptionsOnCommandLine);
-		}
-
-		if (IoStoreCommands.Count > 0)
-		{
-			string IoStoreCommandsFileName = CombinePaths(CmdEnv.LogFolder, "IoStoreCommands.txt");
-			using (var Writer = new StreamWriter(IoStoreCommandsFileName, false, new System.Text.UTF8Encoding(true)))
-			{
-				foreach (string Command in IoStoreCommands)
-				{
-					Writer.WriteLine(Command);
 				}
-			}
 
-			List<OrderFile> FinalOrderFiles = OrderFiles.FindAll(x => (x.OrderType != OrderFile.OrderFileType.Cooker));
-			if (bUseSecondaryOrder && OrderFiles.Count >= 1)
-			{
-				FinalOrderFiles.AddRange(OrderFiles.FindAll(x => x.OrderType == OrderFile.OrderFileType.Cooker));
+				Outputs.Add(Tuple.Create(OutputLocation, OutputRelativeLocation, PatchSourceContentPath));
 			}
 
-			string AdditionalArgs =
-				SC.StageTargetPlatform.GetPlatformPakCommandLine(Params, SC)
-				+ SC.StageTargetPlatform.GetPlatformIoStoreCommandLine(Params, SC)
-				+ BulkOption
-				+ CompressionFormats
-				+ " " + AdditionalCompressionOptionsOnCommandLine
-				+ " " + Params.AdditionalPakOptions;
-
-			if (CryptoKeysCacheFilename != null)
+			// Actually execute UnrealPak
+			if (Commands.Count > 0)
 			{
-				AdditionalArgs += String.Format(" -cryptokeys={0}", CommandUtils.MakePathSafeToUseWithCommandLine(CryptoKeysCacheFilename.FullName));
+				RunUnrealPakInParallel(Commands, LogNames, AdditionalCompressionOptionsOnCommandLine);
 			}
 
-			if (CryptoSettings != null)
+			if (IoStoreCommands.Count > 0)
 			{
-				if (CryptoSettings.bDataCryptoRequired && CryptoSettings.bEnablePakSigning && CryptoSettings.SigningKey.IsValid())
+				string IoStoreCommandsFileName = CombinePaths(CmdEnv.LogFolder, "IoStoreCommands.txt");
+				using (var Writer = new StreamWriter(IoStoreCommandsFileName, false, new System.Text.UTF8Encoding(true)))
 				{
-					AdditionalArgs += String.Format(" -sign");
+					foreach (string Command in IoStoreCommands)
+					{
+						Writer.WriteLine(Command);
+					}
 				}
-			}
 
-			int MaxIoStorePartitionSizeMB = 0;
-			if (PlatformGameConfig.GetInt32("/Script/UnrealEd.ProjectPackagingSettings", "MaxIoStorePartitionSizeMB", out MaxIoStorePartitionSizeMB))
-			{
-				if (MaxIoStorePartitionSizeMB > 0)
+				List<OrderFile> FinalOrderFiles = OrderFiles.FindAll(x => (x.OrderType != OrderFile.OrderFileType.Cooker));
+				if (bUseSecondaryOrder && OrderFiles.Count >= 1)
 				{
-					AdditionalArgs += String.Format(" -maxPartitionSize={0}MB", MaxIoStorePartitionSizeMB);
+					FinalOrderFiles.AddRange(OrderFiles.FindAll(x => x.OrderType == OrderFile.OrderFileType.Cooker));
 				}
-			}
-
-			AdditionalArgs += " " + Params.AdditionalIoStoreOptions;
 
-			RunIoStore(Params, SC, IoStoreCommandsFileName, FinalOrderFiles, AdditionalArgs);
-		}
+				string AdditionalArgs =
+					SC.StageTargetPlatform.GetPlatformPakCommandLine(Params, SC)
+					+ SC.StageTargetPlatform.GetPlatformIoStoreCommandLine(Params, SC)
+					+ BulkOption
+					+ CompressionFormats
+					+ " " + AdditionalCompressionOptionsOnCommandLine
+					+ " " + Params.AdditionalPakOptions;
 
-		// Do any additional processing on the command output
-		for (int Idx = 0; Idx < PakParamsList.Count; Idx++)
-		{
-			string PakName = PakParamsList[Idx].PakName;
-			FileReference OutputLocation = Outputs[Idx].Item1;
-			StagedFileReference OutputRelativeLocation = Outputs[Idx].Item2;
-			string PatchSourceContentPath = Outputs[Idx].Item3;
+				if (CryptoKeysCacheFilename != null)
+				{
+					AdditionalArgs += String.Format(" -cryptokeys={0}", CommandUtils.MakePathSafeToUseWithCommandLine(CryptoKeysCacheFilename.FullName));
+				}
 
-			if (Params.HasCreateReleaseVersion)
-			{
-				// copy the created pak to the release version directory we might need this later if we want to generate patches
-				//string ReleaseVersionPath = CombinePaths( SC.ProjectRoot, "Releases", Params.CreateReleaseVersion, SC.StageTargetPlatform.GetCookPlatform(Params.DedicatedServer, false), Path.GetFileName(OutputLocation) );
-				string ReleaseVersionPath = SC.StageTargetPlatform.GetReleasePakFilePath(SC, Params, OutputLocation.GetFileName());
+				if (CryptoSettings != null)
+				{
+					if (CryptoSettings.bDataCryptoRequired && CryptoSettings.bEnablePakSigning && CryptoSettings.SigningKey.IsValid())
+					{
+						AdditionalArgs += String.Format(" -sign");
+					}
+				}
 
-				InternalUtils.SafeCreateDirectory(Path.GetDirectoryName(ReleaseVersionPath));
-				InternalUtils.SafeCopyFile(OutputLocation.FullName, ReleaseVersionPath);
-				if (ShouldCreateIoStoreContainerFiles(Params, SC.StageTargetPlatform))
+				int MaxIoStorePartitionSizeMB = 0;
+				if (PlatformGameConfig.GetInt32("/Script/UnrealEd.ProjectPackagingSettings", "MaxIoStorePartitionSizeMB", out MaxIoStorePartitionSizeMB))
 				{
-					InternalUtils.SafeCopyFile(Path.ChangeExtension(OutputLocation.FullName, ".utoc"), Path.ChangeExtension(ReleaseVersionPath, ".utoc"));
-					InternalUtils.SafeCopyFile(Path.ChangeExtension(OutputLocation.FullName, ".ucas"), Path.ChangeExtension(ReleaseVersionPath, ".ucas"));
+					if (MaxIoStorePartitionSizeMB > 0)
+					{
+						AdditionalArgs += String.Format(" -maxPartitionSize={0}MB", MaxIoStorePartitionSizeMB);
+					}
 				}
+
+				AdditionalArgs += " " + Params.AdditionalIoStoreOptions;
+
+				RunIoStore(Params, SC, IoStoreCommandsFileName, FinalOrderFiles, AdditionalArgs);
 			}
 
-			if (Params.CreateChunkInstall)
+			// Do any additional processing on the command output
+			for (int Idx = 0; Idx < PakParamsList.Count; Idx++)
 			{
-				var RegEx = new Regex("pakchunk(\\d+)", RegexOptions.IgnoreCase | RegexOptions.Compiled);
-				var Matches = RegEx.Matches(PakName);
+				string PakName = PakParamsList[Idx].PakName;
+				FileReference OutputLocation = Outputs[Idx].Item1;
+				StagedFileReference OutputRelativeLocation = Outputs[Idx].Item2;
+				string PatchSourceContentPath = Outputs[Idx].Item3;
 
-				int ChunkID = 0;
-				if (Matches.Count != 0 && Matches[0].Groups.Count > 1)
-				{
-					ChunkID = Convert.ToInt32(Matches[0].Groups[1].ToString());
-				}
-				else if (Params.HasDLCName)
-				{
-					// Assuming DLC is a single pack file
-					ChunkID = 1;
-				}
-				else
+				if (Params.HasCreateReleaseVersion)
 				{
-					throw new AutomationException(String.Format("Failed Creating Chunk Install Data, Unable to parse chunk id from {0}", PakName));
+					// copy the created pak to the release version directory we might need this later if we want to generate patches
+					//string ReleaseVersionPath = CombinePaths( SC.ProjectRoot, "Releases", Params.CreateReleaseVersion, SC.StageTargetPlatform.GetCookPlatform(Params.DedicatedServer, false), Path.GetFileName(OutputLocation) );
+					string ReleaseVersionPath = SC.StageTargetPlatform.GetReleasePakFilePath(SC, Params, OutputLocation.GetFileName());
+
+					InternalUtils.SafeCreateDirectory(Path.GetDirectoryName(ReleaseVersionPath));
+					InternalUtils.SafeCopyFile(OutputLocation.FullName, ReleaseVersionPath);
+					if (ShouldCreateIoStoreContainerFiles(Params, SC.StageTargetPlatform))
+					{
+						InternalUtils.SafeCopyFile(Path.ChangeExtension(OutputLocation.FullName, ".utoc"), Path.ChangeExtension(ReleaseVersionPath, ".utoc"));
+						InternalUtils.SafeCopyFile(Path.ChangeExtension(OutputLocation.FullName, ".ucas"), Path.ChangeExtension(ReleaseVersionPath, ".ucas"));
+					}
 				}
 
-				if (ChunkID != 0)
+				if (Params.CreateChunkInstall)
 				{
-					var BPTExe = GetBuildPatchToolExecutable();
-					EnsureBuildPatchToolExists();
+					var RegEx = new Regex("pakchunk(\\d+)", RegexOptions.IgnoreCase | RegexOptions.Compiled);
+					var Matches = RegEx.Matches(PakName);
 
-					string VersionString = Params.ChunkInstallVersionString;
-					string ChunkInstallBasePath = CombinePaths(Params.ChunkInstallDirectory, SC.FinalCookPlatform);
-					string RawDataPath = CombinePaths(ChunkInstallBasePath, VersionString, PakName);
-					string RawDataPakPath = CombinePaths(RawDataPath, PakName + "-" + SC.FinalCookPlatform + GetPostFix(PakName) + ".pak");
-					bool bPakFilesAreSigned = InternalUtils.SafeFileExists(Path.ChangeExtension(OutputLocation.FullName, ".sig"));
-
-					//copy the pak chunk to the raw data folder
-					InternalUtils.SafeDeleteFile(RawDataPakPath, true);
-					InternalUtils.SafeDeleteFile(Path.ChangeExtension(RawDataPakPath, ".sig"), true);
-					InternalUtils.SafeCreateDirectory(RawDataPath, true);
-					InternalUtils.SafeCopyFile(OutputLocation.FullName, RawDataPakPath);
-					if (bPakFilesAreSigned)
+					int ChunkID = 0;
+					if (Matches.Count != 0 && Matches[0].Groups.Count > 1)
+					{
+						ChunkID = Convert.ToInt32(Matches[0].Groups[1].ToString());
+					}
+					else if (Params.HasDLCName)
+					{
+						// Assuming DLC is a single pack file
+						ChunkID = 1;
+					}
+					else
 					{
-						InternalUtils.SafeCopyFile(Path.ChangeExtension(OutputLocation.FullName, ".sig"), Path.ChangeExtension(RawDataPakPath, ".sig"), true);
+						throw new AutomationException(String.Format("Failed Creating Chunk Install Data, Unable to parse chunk id from {0}", PakName));
 					}
-					InternalUtils.SafeDeleteFile(OutputLocation.FullName, true);
 
-					if (Params.IsGeneratingPatch)
+					if (ChunkID != 0)
 					{
-						if (String.IsNullOrEmpty(PatchSourceContentPath))
-						{
-							throw new AutomationException(String.Format("Failed Creating Chunk Install Data. No source pak for patch pak {0} given", OutputLocation));
-						}
-						// If we are generating a patch, then we need to copy the original pak across
-						// for distribution.
-						string SourceRawDataPakPath = CombinePaths(RawDataPath, PakName + "-" + SC.FinalCookPlatform + ".pak");
-						InternalUtils.SafeCopyFile(PatchSourceContentPath, SourceRawDataPakPath);
+						var BPTExe = GetBuildPatchToolExecutable();
+						EnsureBuildPatchToolExists();
+
+						string VersionString = Params.ChunkInstallVersionString;
+						string ChunkInstallBasePath = CombinePaths(Params.ChunkInstallDirectory, SC.FinalCookPlatform);
+						string RawDataPath = CombinePaths(ChunkInstallBasePath, VersionString, PakName);
+						string RawDataPakPath = CombinePaths(RawDataPath, PakName + "-" + SC.FinalCookPlatform + GetPostFix(PakName) + ".pak");
+						bool bPakFilesAreSigned = InternalUtils.SafeFileExists(Path.ChangeExtension(OutputLocation.FullName, ".sig"));
+
+						//copy the pak chunk to the raw data folder
+						InternalUtils.SafeDeleteFile(RawDataPakPath, true);
+						InternalUtils.SafeDeleteFile(Path.ChangeExtension(RawDataPakPath, ".sig"), true);
+						InternalUtils.SafeCreateDirectory(RawDataPath, true);
+						InternalUtils.SafeCopyFile(OutputLocation.FullName, RawDataPakPath);
 						if (bPakFilesAreSigned)
 						{
-							InternalUtils.SafeCopyFile(PatchSourceContentPath, Path.ChangeExtension(SourceRawDataPakPath, ".sig"), true);
+							InternalUtils.SafeCopyFile(Path.ChangeExtension(OutputLocation.FullName, ".sig"), Path.ChangeExtension(RawDataPakPath, ".sig"), true);
 						}
-					}
-
-					string BuildRoot = MakePathSafeToUseWithCommandLine(RawDataPath);
-					string CloudDir = MakePathSafeToUseWithCommandLine(CombinePaths(ChunkInstallBasePath, "CloudDir"));
-					string ManifestDir = CombinePaths(ChunkInstallBasePath, "ManifestDir");
-					var AppID = 1; // For a chunk install this value doesn't seem to matter
-					string AppName = String.Format("{0}_{1}", SC.ShortProjectName, PakName);
-					string AppLaunch = ""; // For a chunk install this value doesn't seem to matter
-					string ManifestFilename = AppName + VersionString + ".manifest";
-					string SourceManifestPath = CombinePaths(CloudDir, ManifestFilename);
-					string DestManifestPath = CombinePaths(ManifestDir, ManifestFilename);
-					InternalUtils.SafeCreateDirectory(ManifestDir, true);
-
-					string CmdLine = String.Format("-BuildRoot={0} -CloudDir={1} -AppID={2} -AppName=\"{3}\" -BuildVersion=\"{4}\" -AppLaunch=\"{5}\"", BuildRoot, CloudDir, AppID, AppName, VersionString, AppLaunch);
-					CmdLine += " -AppArgs=\"\"";
-					CmdLine += " -custom=\"bIsPatch=false\"";
-					CmdLine += String.Format(" -customint=\"ChunkID={0}\"", ChunkID);
-					CmdLine += " -customint=\"PakReadOrdering=0\"";
-					CmdLine += " -stdout";
-
-					string BuildPatchToolLogFileName = "BuildPatchTool_" + PakName;
-					RunAndLog(CmdEnv, BPTExe, CmdLine, BuildPatchToolLogFileName, Options: ERunOptions.Default | ERunOptions.UTF8Output);
+						InternalUtils.SafeDeleteFile(OutputLocation.FullName, true);
 
-					InternalUtils.SafeCopyFile(SourceManifestPath, DestManifestPath);
+						if (Params.IsGeneratingPatch)
+						{
+							if (String.IsNullOrEmpty(PatchSourceContentPath))
+							{
+								throw new AutomationException(String.Format("Failed Creating Chunk Install Data. No source pak for patch pak {0} given", OutputLocation));
+							}
+							// If we are generating a patch, then we need to copy the original pak across
+							// for distribution.
+							string SourceRawDataPakPath = CombinePaths(RawDataPath, PakName + "-" + SC.FinalCookPlatform + ".pak");
+							InternalUtils.SafeCopyFile(PatchSourceContentPath, SourceRawDataPakPath);
+							if (bPakFilesAreSigned)
+							{
+								InternalUtils.SafeCopyFile(PatchSourceContentPath, Path.ChangeExtension(SourceRawDataPakPath, ".sig"), true);
+							}
+						}
 
-					// generate the master manifest
-					GenerateMasterChunkManifest(CloudDir, Params.ChunkInstallVersionString, SC.FinalCookPlatform);
+						string BuildRoot = MakePathSafeToUseWithCommandLine(RawDataPath);
+						string CloudDir = MakePathSafeToUseWithCommandLine(CombinePaths(ChunkInstallBasePath, "CloudDir"));
+						string ManifestDir = CombinePaths(ChunkInstallBasePath, "ManifestDir");
+						var AppID = 1; // For a chunk install this value doesn't seem to matter
+						string AppName = String.Format("{0}_{1}", SC.ShortProjectName, PakName);
+						string AppLaunch = ""; // For a chunk install this value doesn't seem to matter
+						string ManifestFilename = AppName + VersionString + ".manifest";
+						string SourceManifestPath = CombinePaths(CloudDir, ManifestFilename);
+						string DestManifestPath = CombinePaths(ManifestDir, ManifestFilename);
+						InternalUtils.SafeCreateDirectory(ManifestDir, true);
+
+						string CmdLine = String.Format("-BuildRoot={0} -CloudDir={1} -AppID={2} -AppName=\"{3}\" -BuildVersion=\"{4}\" -AppLaunch=\"{5}\"", BuildRoot, CloudDir, AppID, AppName, VersionString, AppLaunch);
+						CmdLine += " -AppArgs=\"\"";
+						CmdLine += " -custom=\"bIsPatch=false\"";
+						CmdLine += String.Format(" -customint=\"ChunkID={0}\"", ChunkID);
+						CmdLine += " -customint=\"PakReadOrdering=0\"";
+						CmdLine += " -stdout";
+
+						string BuildPatchToolLogFileName = "BuildPatchTool_" + PakName;
+						RunAndLog(CmdEnv, BPTExe, CmdLine, BuildPatchToolLogFileName, Options: ERunOptions.Default | ERunOptions.UTF8Output);
+
+						InternalUtils.SafeCopyFile(SourceManifestPath, DestManifestPath);
+
+						// generate the master manifest
+						GenerateMasterChunkManifest(CloudDir, Params.ChunkInstallVersionString, SC.FinalCookPlatform);
+					}
+					else
+					{
+						// add the first pak file as needing deployment and convert to lower case again if needed
+						SC.FilesToStage.UFSFiles.Add(OutputRelativeLocation, OutputLocation);
+					}
 				}
 				else
 				{
-					// add the first pak file as needing deployment and convert to lower case again if needed
+					// add the pak file as needing deployment and convert to lower case again if needed
 					SC.FilesToStage.UFSFiles.Add(OutputRelativeLocation, OutputLocation);
-				}
-			}
-			else
-			{
-				// add the pak file as needing deployment and convert to lower case again if needed
-				SC.FilesToStage.UFSFiles.Add(OutputRelativeLocation, OutputLocation);
 
-				// add the base pak files to deployment as well
-				if (bShouldGeneratePatch && !ShouldSkipGeneratingPatch(PlatformGameConfig, PakName))
-				{
-					string ExistingPatchSearchPath = SC.StageTargetPlatform.GetReleasePakFilePath(SC, Params, null);
-					if (Directory.Exists(ExistingPatchSearchPath))
+					// add the base pak files to deployment as well
+					if (bShouldGeneratePatch && !ShouldSkipGeneratingPatch(PlatformGameConfig, PakName))
 					{
-						HashSet<string> IncludedExtensions = new HashSet<string>();
-						IncludedExtensions.Add(OutputFilenameExtension);
-						if (ShouldCreateIoStoreContainerFiles(Params, SC.StageTargetPlatform))
-						{
-							IncludedExtensions.Add(".ucas");
-							IncludedExtensions.Add(".utoc");
-						}
-						IEnumerable<string> PakFileSet = Directory.EnumerateFiles(ExistingPatchSearchPath, PakName + "-" + SC.FinalCookPlatform + "*.*");
-						foreach (string PakFilePath in PakFileSet)
+						string ExistingPatchSearchPath = SC.StageTargetPlatform.GetReleasePakFilePath(SC, Params, null);
+						if (Directory.Exists(ExistingPatchSearchPath))
 						{
-							if (!IncludedExtensions.Contains(Path.GetExtension(PakFilePath).ToLower()))
+							HashSet<string> IncludedExtensions = new HashSet<string>();
+							IncludedExtensions.Add(OutputFilenameExtension);
+							if (ShouldCreateIoStoreContainerFiles(Params, SC.StageTargetPlatform))
 							{
-								continue;
+								IncludedExtensions.Add(".ucas");
+								IncludedExtensions.Add(".utoc");
 							}
-							FileReference OutputDestinationPath = FileReference.Combine(OutputLocation.Directory, Path.GetFileName(PakFilePath));
-							if (!File.Exists(OutputDestinationPath.FullName))
+							IEnumerable<string> PakFileSet = Directory.EnumerateFiles(ExistingPatchSearchPath, PakName + "-" + SC.FinalCookPlatform + "*.*");
+							foreach (string PakFilePath in PakFileSet)
 							{
-								InternalUtils.SafeCopyFile(PakFilePath, OutputDestinationPath.FullName);
-								StagedFileReference OutputDestinationRelativeLocation = StagedFileReference.Combine(SC.RelativeProjectRootForStage, "Content/Paks/", Path.GetFileName(PakFilePath));
-								if (SC.StageTargetPlatform.DeployLowerCaseFilenames(StagedFileType.UFS))
+								if (!IncludedExtensions.Contains(Path.GetExtension(PakFilePath).ToLower()))
+								{
+									continue;
+								}
+								FileReference OutputDestinationPath = FileReference.Combine(OutputLocation.Directory, Path.GetFileName(PakFilePath));
+								if (!File.Exists(OutputDestinationPath.FullName))
 								{
-									OutputDestinationRelativeLocation = OutputDestinationRelativeLocation.ToLowerInvariant();
+									InternalUtils.SafeCopyFile(PakFilePath, OutputDestinationPath.FullName);
+									StagedFileReference OutputDestinationRelativeLocation = StagedFileReference.Combine(SC.RelativeProjectRootForStage, "Content/Paks/", Path.GetFileName(PakFilePath));
+									if (SC.StageTargetPlatform.DeployLowerCaseFilenames(StagedFileType.UFS))
+									{
+										OutputDestinationRelativeLocation = OutputDestinationRelativeLocation.ToLowerInvariant();
+									}
+									OutputDestinationRelativeLocation = SC.StageTargetPlatform.Remap(OutputDestinationRelativeLocation);
+									SC.FilesToStage.UFSFiles.Add(OutputDestinationRelativeLocation, OutputDestinationPath);
 								}
-								OutputDestinationRelativeLocation = SC.StageTargetPlatform.Remap(OutputDestinationRelativeLocation);
-								SC.FilesToStage.UFSFiles.Add(OutputDestinationRelativeLocation, OutputDestinationPath);
 							}
 						}
 					}
 				}
 			}
 		}
-	}
 
-	private static void RunIoStore(ProjectParams Params, DeploymentContext SC, string CommandsFileName, List<OrderFile> OrderFiles, string AdditionalArgs)
-	{
-		StagedFileReference GlobalContainerOutputRelativeLocation;
-		GlobalContainerOutputRelativeLocation = StagedFileReference.Combine(SC.RelativeProjectRootForStage, "Content", "Paks", "global.utoc");
-		if (SC.StageTargetPlatform.DeployLowerCaseFilenames(StagedFileType.UFS))
+		private static void RunIoStore(ProjectParams Params, DeploymentContext SC, string CommandsFileName, List<OrderFile> OrderFiles, string AdditionalArgs)
 		{
-			GlobalContainerOutputRelativeLocation = GlobalContainerOutputRelativeLocation.ToLowerInvariant();
-		}
-		GlobalContainerOutputRelativeLocation = SC.StageTargetPlatform.Remap(GlobalContainerOutputRelativeLocation);
-		FileReference GlobalContainerOutputLocation = FileReference.Combine(SC.RuntimeRootDir, GlobalContainerOutputRelativeLocation.Name);
+			StagedFileReference GlobalContainerOutputRelativeLocation;
+			GlobalContainerOutputRelativeLocation = StagedFileReference.Combine(SC.RelativeProjectRootForStage, "Content", "Paks", "global.utoc");
+			if (SC.StageTargetPlatform.DeployLowerCaseFilenames(StagedFileType.UFS))
+			{
+				GlobalContainerOutputRelativeLocation = GlobalContainerOutputRelativeLocation.ToLowerInvariant();
+			}
+			GlobalContainerOutputRelativeLocation = SC.StageTargetPlatform.Remap(GlobalContainerOutputRelativeLocation);
+			FileReference GlobalContainerOutputLocation = FileReference.Combine(SC.RuntimeRootDir, GlobalContainerOutputRelativeLocation.Name);
 
-		string CommandletParams = String.Empty;
+			string CommandletParams = String.Empty;
 
-		if (Params.HasDLCName)
-		{
-			CommandletParams += String.Format("-DLCFile={0}", MakePathSafeToUseWithCommandLine(Params.DLCFile.FullName));
+			if (Params.HasDLCName)
+			{
+				CommandletParams += String.Format("-DLCFile={0}", MakePathSafeToUseWithCommandLine(Params.DLCFile.FullName));
 
-			DirectoryReference DLCRoot = Params.DLCFile.Directory;
-			string DLCName = Params.DLCFile.GetFileNameWithoutExtension();
+				DirectoryReference DLCRoot = Params.DLCFile.Directory;
+				string DLCName = Params.DLCFile.GetFileNameWithoutExtension();
 
-			//TODO: Find a better way. Create Plugin ConfigType and interate all ini files in plugin?
-			bool bRemapPluginContentToGame = false;
-			FileReference PluginConfigFile = FileReference.Combine(DLCRoot, "Config", String.Format("Default{0}.ini", DLCName));
-			if (FileReference.Exists(PluginConfigFile))
-			{
-				ConfigFile File = new ConfigFile(PluginConfigFile);
-				ConfigFileSection PluginSettings;
-				if (File.TryGetSection("PluginSettings", out PluginSettings))
+				//TODO: Find a better way. Create Plugin ConfigType and interate all ini files in plugin?
+				bool bRemapPluginContentToGame = false;
+				FileReference PluginConfigFile = FileReference.Combine(DLCRoot, "Config", String.Format("Default{0}.ini", DLCName));
+				if (FileReference.Exists(PluginConfigFile))
 				{
-					foreach (ConfigLine Line in PluginSettings.Lines)
+					ConfigFile File = new ConfigFile(PluginConfigFile);
+					ConfigFileSection PluginSettings;
+					if (File.TryGetSection("PluginSettings", out PluginSettings))
 					{
-						if (Line.Key == "RemapPluginContentToGame")
+						foreach (ConfigLine Line in PluginSettings.Lines)
 						{
-							bool.TryParse(Line.Value, out bRemapPluginContentToGame);
-							break;
+							if (Line.Key == "RemapPluginContentToGame")
+							{
+								bool.TryParse(Line.Value, out bRemapPluginContentToGame);
+								break;
+							}
 						}
 					}
 				}
+
+				if (bRemapPluginContentToGame)
+				{
+					CommandletParams += " -RemapPluginContentToGame";
+				}
+			}
+			else
+			{
+				CommandletParams += String.Format("-CreateGlobalContainer={0}", MakePathSafeToUseWithCommandLine(GlobalContainerOutputLocation.FullName));
 			}
 
-			if (bRemapPluginContentToGame)
+			CommandletParams += String.Format(" -CookedDirectory={0} -Commands={1}", MakePathSafeToUseWithCommandLine(SC.PlatformCookDir.ToString()), MakePathSafeToUseWithCommandLine(CommandsFileName));
+			if (OrderFiles != null && OrderFiles.Count() > 0)
 			{
-				CommandletParams += " -RemapPluginContentToGame";
+				CommandletParams += String.Format(" -Order={0}", MakePathSafeToUseWithCommandLine(string.Join(",", OrderFiles.Select(u => u.File.FullName).ToArray())));
+				CommandletParams += String.Format(" -OrderPriority={0}", MakePathSafeToUseWithCommandLine(string.Join(",", OrderFiles.Select(u => u.Priority).ToArray())));
+			}
+			if (!string.IsNullOrWhiteSpace(AdditionalArgs))
+			{
+				CommandletParams += AdditionalArgs;
 			}
-		}
-		else
-		{
-			CommandletParams += String.Format("-CreateGlobalContainer={0}", MakePathSafeToUseWithCommandLine(GlobalContainerOutputLocation.FullName));
-		}
 
-		CommandletParams += String.Format(" -CookedDirectory={0} -Commands={1}", MakePathSafeToUseWithCommandLine(SC.PlatformCookDir.ToString()), MakePathSafeToUseWithCommandLine(CommandsFileName));
-		if (OrderFiles != null && OrderFiles.Count() > 0)
-		{
-			CommandletParams += String.Format(" -Order={0}", MakePathSafeToUseWithCommandLine(string.Join(",", OrderFiles.Select(u => u.File.FullName).ToArray())));
-			CommandletParams += String.Format(" -OrderPriority={0}", MakePathSafeToUseWithCommandLine(string.Join(",", OrderFiles.Select(u => u.Priority).ToArray())));
-		} 
-		if (!string.IsNullOrWhiteSpace(AdditionalArgs))
-		{
-			CommandletParams += AdditionalArgs;
-		}
+			CommandletParams += String.Format(" -TargetPlatform={0}", SC.StageTargetPlatform.GetCookPlatform(Params.DedicatedServer, Params.Client));
 
-		CommandletParams += String.Format(" -TargetPlatform={0}", SC.StageTargetPlatform.GetCookPlatform(Params.DedicatedServer, Params.Client));
+			if (Params.HasBasedOnReleaseVersion)
+			{
+				CommandletParams += String.Format(" -BasedOnReleaseVersionPath={0}", Params.GetBasedOnReleaseVersionPath(SC, Params.Client));
+			}
 
-		if (Params.HasBasedOnReleaseVersion)
-		{
-			CommandletParams += String.Format(" -BasedOnReleaseVersionPath={0}", Params.GetBasedOnReleaseVersionPath(SC, Params.Client));
+			LogInformation("Running IoStore commandlet with arguments: {0}", CommandletParams);
+			RunCommandlet(SC.RawProjectPath, Params.UE4Exe, "IoStore", CommandletParams);
 		}
 
-		LogInformation("Running IoStore commandlet with arguments: {0}", CommandletParams);
-		RunCommandlet(SC.RawProjectPath, Params.UE4Exe, "IoStore", CommandletParams);
-	}
-
-	private static void RunUnrealPakInParallel(List<string> Commands, List<string> LogNames, string AdditionalCompressionOptionsOnCommandLine)
-	{
-		LogInformation("Executing {0} UnrealPak command{1}...", Commands.Count, (Commands.Count > 1)? "s" : "");
-
-		// Spawn tasks for each command
-		IProcessResult[] Results = new IProcessResult[Commands.Count];
-		using(ThreadPoolWorkQueue Queue = new ThreadPoolWorkQueue())
+		private static void RunUnrealPakInParallel(List<string> Commands, List<string> LogNames, string AdditionalCompressionOptionsOnCommandLine)
 		{
-			string UnrealPakExe = GetUnrealPakLocation().FullName;
-			for(int Idx = 0; Idx < Commands.Count; Idx++)
+			LogInformation("Executing {0} UnrealPak command{1}...", Commands.Count, (Commands.Count > 1) ? "s" : "");
+
+			// Spawn tasks for each command
+			IProcessResult[] Results = new IProcessResult[Commands.Count];
+			using (ThreadPoolWorkQueue Queue = new ThreadPoolWorkQueue())
 			{
-				string LogFile = LogUtils.GetUniqueLogName(CombinePaths(CmdEnv.LogFolder, String.Format("UnrealPak-{0}", LogNames[Idx])));
-				string Arguments = Commands[Idx] = String.Format("{0} -multiprocess -abslog={1} {2}", Commands[Idx], MakePathSafeToUseWithCommandLine(LogFile), AdditionalCompressionOptionsOnCommandLine);
+				string UnrealPakExe = GetUnrealPakLocation().FullName;
+				for (int Idx = 0; Idx < Commands.Count; Idx++)
+				{
+					string LogFile = LogUtils.GetUniqueLogName(CombinePaths(CmdEnv.LogFolder, String.Format("UnrealPak-{0}", LogNames[Idx])));
+					string Arguments = Commands[Idx] = String.Format("{0} -multiprocess -abslog={1} {2}", Commands[Idx], MakePathSafeToUseWithCommandLine(LogFile), AdditionalCompressionOptionsOnCommandLine);
 
-				int LocalIdx = Idx;
-				Queue.Enqueue(() => Results[LocalIdx] = Run(UnrealPakExe, Arguments, Options: ERunOptions.AppMustExist | ERunOptions.NoLoggingOfRunCommand));
-			}
+					int LocalIdx = Idx;
+					Queue.Enqueue(() => Results[LocalIdx] = Run(UnrealPakExe, Arguments, Options: ERunOptions.AppMustExist | ERunOptions.NoLoggingOfRunCommand));
+				}
 
-			bool bComplete = false;
-			while(!bComplete)
-			{
-				bComplete = Queue.Wait(10 * 1000);
-				LogInformation("Waiting for child processes to complete ({0}/{1})", Results.Count(x => x != null), Results.Length);
+				bool bComplete = false;
+				while (!bComplete)
+				{
+					bComplete = Queue.Wait(10 * 1000);
+					LogInformation("Waiting for child processes to complete ({0}/{1})", Results.Count(x => x != null), Results.Length);
+				}
 			}
-		}
 
-		// Output all the results
-		int NumFailed = 0;
-		for(int Idx = 0; Idx < Results.Length; Idx++)
-		{
-			LogInformation("Output from: {0}", Commands[Idx]);
-			using(new LogIndentScope("  "))
+			// Output all the results
+			int NumFailed = 0;
+			for (int Idx = 0; Idx < Results.Length; Idx++)
 			{
-				foreach(string Line in Results[Idx].Output.TrimEnd().Split('\n'))
+				LogInformation("Output from: {0}", Commands[Idx]);
+				using (new LogIndentScope("  "))
+				{
+					foreach (string Line in Results[Idx].Output.TrimEnd().Split('\n'))
+					{
+						LogInformation(Line);
+					}
+				}
+				LogInformation("UnrealPak terminated with exit code {0}", Results[Idx].ExitCode);
+
+				if (Results[Idx].ExitCode > 0)
 				{
-					LogInformation(Line);
+					NumFailed++;
 				}
 			}
-			LogInformation("UnrealPak terminated with exit code {0}", Results[Idx].ExitCode);
 
-			if(Results[Idx].ExitCode > 0)
+			// Abort if any instance failed
+			if (NumFailed > 0)
 			{
-				NumFailed++;
+				throw new AutomationException("UnrealPak failed");
 			}
 		}
-			
-		// Abort if any instance failed
-		if(NumFailed > 0)
-		{
-			throw new AutomationException("UnrealPak failed");
-		}
-	}
-
-	/// <summary>
-	/// Creates pak files using streaming install chunk manifests.
-	/// </summary>
-	/// <param name="Params"></param>
-	/// <param name="SC"></param>
-	private static void CopyPaksFromNetwork(ProjectParams Params, DeploymentContext SC)
-	{
-		LogInformation("Copying paks from network.");
 
-		if (!CommandUtils.P4Enabled)
-		{
-			throw new AutomationException("-PrePak requires -P4");
-		}
-		if (CommandUtils.P4Env.Changelist < 1000)
+		/// <summary>
+		/// Creates pak files using streaming install chunk manifests.
+		/// </summary>
+		/// <param name="Params"></param>
+		/// <param name="SC"></param>
+		private static void CopyPaksFromNetwork(ProjectParams Params, DeploymentContext SC)
 		{
-			throw new AutomationException("-PrePak requires a CL from P4 and we have {0}", CommandUtils.P4Env.Changelist);
-		}
+			LogInformation("Copying paks from network.");
+
+			if (!CommandUtils.P4Enabled)
+			{
+				throw new AutomationException("-PrePak requires -P4");
+			}
+			if (CommandUtils.P4Env.Changelist < 1000)
+			{
+				throw new AutomationException("-PrePak requires a CL from P4 and we have {0}", CommandUtils.P4Env.Changelist);
+			}
 
-		string BuildRoot = CombinePaths(CommandUtils.RootBuildStorageDirectory());
-		string CachePath = InternalUtils.GetEnvironmentVariable("UE-BuildCachePath", "");
+			string BuildRoot = CombinePaths(CommandUtils.RootBuildStorageDirectory());
+			string CachePath = InternalUtils.GetEnvironmentVariable("UE-BuildCachePath", "");
 
-		string SrcBuildPath = CombinePaths(BuildRoot, Params.ShortProjectName);
-		string SrcBuildPath2 = CombinePaths(BuildRoot, Params.ShortProjectName.Replace("Game", "").Replace("game", ""));
+			string SrcBuildPath = CombinePaths(BuildRoot, Params.ShortProjectName);
+			string SrcBuildPath2 = CombinePaths(BuildRoot, Params.ShortProjectName.Replace("Game", "").Replace("game", ""));
 
-		string SrcBuildPath_Cache = CombinePaths(CachePath, Params.ShortProjectName);
-		string SrcBuildPath2_Cache = CombinePaths(CachePath, Params.ShortProjectName.Replace("Game", "").Replace("game", ""));
+			string SrcBuildPath_Cache = CombinePaths(CachePath, Params.ShortProjectName);
+			string SrcBuildPath2_Cache = CombinePaths(CachePath, Params.ShortProjectName.Replace("Game", "").Replace("game", ""));
 
-		if (!InternalUtils.SafeDirectoryExists(SrcBuildPath))
-		{
-			if (!InternalUtils.SafeDirectoryExists(SrcBuildPath2))
+			if (!InternalUtils.SafeDirectoryExists(SrcBuildPath))
 			{
-				throw new AutomationException("PrePak: Neither {0} nor {1} exists.", SrcBuildPath, SrcBuildPath2);
+				if (!InternalUtils.SafeDirectoryExists(SrcBuildPath2))
+				{
+					throw new AutomationException("PrePak: Neither {0} nor {1} exists.", SrcBuildPath, SrcBuildPath2);
+				}
+				SrcBuildPath = SrcBuildPath2;
+				SrcBuildPath_Cache = SrcBuildPath2_Cache;
+			}
+			string SrcCLPath = CombinePaths(SrcBuildPath, CommandUtils.EscapePath(CommandUtils.P4Env.Branch) + "-CL-" + CommandUtils.P4Env.Changelist.ToString());
+			string SrcCLPath_Cache = CombinePaths(SrcBuildPath_Cache, CommandUtils.EscapePath(CommandUtils.P4Env.Branch) + "-CL-" + CommandUtils.P4Env.Changelist.ToString());
+			if (!InternalUtils.SafeDirectoryExists(SrcCLPath))
+			{
+				throw new AutomationException("PrePak: {0} does not exist.", SrcCLPath);
+			}
+			string PlatformPath = CombinePaths(SrcCLPath, SC.FinalCookPlatform);
+			string PlatformPath_Cache = CombinePaths(SrcCLPath_Cache, SC.FinalCookPlatform);
+			if (!InternalUtils.SafeDirectoryExists(PlatformPath))
+			{
+				throw new AutomationException("PrePak: {0} does not exist.", PlatformPath);
+			}
+			string PakPath = CombinePaths(PlatformPath, "Staged", SC.ShortProjectName, "Content", "Paks");
+			string PakPath_Cache = CombinePaths(PlatformPath_Cache, "Staged", SC.ShortProjectName, "Content", "Paks");
+			if (!InternalUtils.SafeDirectoryExists(PakPath))
+			{
+				throw new AutomationException("PrePak: {0} does not exist.", PakPath);
 			}
-			SrcBuildPath = SrcBuildPath2;
-			SrcBuildPath_Cache = SrcBuildPath2_Cache;
-		}
-		string SrcCLPath = CombinePaths(SrcBuildPath, CommandUtils.EscapePath(CommandUtils.P4Env.Branch) + "-CL-" + CommandUtils.P4Env.Changelist.ToString());
-		string SrcCLPath_Cache = CombinePaths(SrcBuildPath_Cache, CommandUtils.EscapePath(CommandUtils.P4Env.Branch) + "-CL-" + CommandUtils.P4Env.Changelist.ToString());
-		if (!InternalUtils.SafeDirectoryExists(SrcCLPath))
-		{
-			throw new AutomationException("PrePak: {0} does not exist.", SrcCLPath);
-		}
-		string PlatformPath = CombinePaths(SrcCLPath, SC.FinalCookPlatform);
-		string PlatformPath_Cache = CombinePaths(SrcCLPath_Cache, SC.FinalCookPlatform);
-		if (!InternalUtils.SafeDirectoryExists(PlatformPath))
-		{
-			throw new AutomationException("PrePak: {0} does not exist.", PlatformPath);
-		}
-		string PakPath = CombinePaths(PlatformPath, "Staged", SC.ShortProjectName, "Content", "Paks");
-		string PakPath_Cache = CombinePaths(PlatformPath_Cache, "Staged", SC.ShortProjectName, "Content", "Paks");
-		if (!InternalUtils.SafeDirectoryExists(PakPath))
-		{
-			throw new AutomationException("PrePak: {0} does not exist.", PakPath);
-		}
 
-		string DestPath = CombinePaths(SC.StageDirectory.FullName, SC.ShortProjectName, "Content", "Paks");
+			string DestPath = CombinePaths(SC.StageDirectory.FullName, SC.ShortProjectName, "Content", "Paks");
 
-		{
-			var PakFiles = CommandUtils.FindFiles("*.pak", false, PakPath);
-			var SigFiles = CommandUtils.FindFiles("*.sig", false, PakPath);
+			{
+				var PakFiles = CommandUtils.FindFiles("*.pak", false, PakPath);
+				var SigFiles = CommandUtils.FindFiles("*.sig", false, PakPath);
 
-			var Files = new List<string>();
+				var Files = new List<string>();
 
-			Files.AddRange(PakFiles);
-			Files.AddRange(SigFiles);
+				Files.AddRange(PakFiles);
+				Files.AddRange(SigFiles);
 
-			if (Files.Count < 1)
-			{
-				throw new AutomationException("PrePak: {0} exists but does not have any paks in it.", PakPath);
-			}
+				if (Files.Count < 1)
+				{
+					throw new AutomationException("PrePak: {0} exists but does not have any paks in it.", PakPath);
+				}
 
 
-			foreach (var SrcFile in Files)
-			{
-				string DestFileName = CombinePaths(DestPath, new FileReference(SrcFile).GetFileName());
-				if (!string.IsNullOrEmpty(CachePath))
+				foreach (var SrcFile in Files)
 				{
-					string CacheSrcFile = CombinePaths(PakPath_Cache, new FileReference(SrcFile).GetFileName());
-					try
+					string DestFileName = CombinePaths(DestPath, new FileReference(SrcFile).GetFileName());
+					if (!string.IsNullOrEmpty(CachePath))
 					{
-						if (InternalUtils.SafeFileExists(CacheSrcFile))
+						string CacheSrcFile = CombinePaths(PakPath_Cache, new FileReference(SrcFile).GetFileName());
+						try
 						{
-							var Info = new System.IO.FileInfo(SrcFile);
-							var InfoCache = new System.IO.FileInfo(CacheSrcFile);
-							if (Info.Exists && InfoCache.Exists && Info.Length == InfoCache.Length)
+							if (InternalUtils.SafeFileExists(CacheSrcFile))
 							{
-								LogInformation("Copying from cache {0} -> {1}", CacheSrcFile, DestFileName);
-								CopyFileIncremental(new FileReference(CacheSrcFile), new FileReference(DestFileName));
-								continue;
+								var Info = new System.IO.FileInfo(SrcFile);
+								var InfoCache = new System.IO.FileInfo(CacheSrcFile);
+								if (Info.Exists && InfoCache.Exists && Info.Length == InfoCache.Length)
+								{
+									LogInformation("Copying from cache {0} -> {1}", CacheSrcFile, DestFileName);
+									CopyFileIncremental(new FileReference(CacheSrcFile), new FileReference(DestFileName));
+									continue;
+								}
 							}
 						}
-					}
-					catch (Exception)
-					{
+						catch (Exception)
+						{
 
+						}
 					}
+					LogInformation("Copying {0} -> {1}", SrcFile, DestFileName);
+					CopyFileIncremental(new FileReference(SrcFile), new FileReference(DestFileName));
 				}
-				LogInformation("Copying {0} -> {1}", SrcFile, DestFileName);
-				CopyFileIncremental(new FileReference(SrcFile), new FileReference(DestFileName));
 			}
 		}
-	}
 
-	// Defines contents of a specific chunk/pak file
-	private class ChunkDefinition
-	{
-		public ChunkDefinition(string InChunkName)
+		// Defines contents of a specific chunk/pak file
+		private class ChunkDefinition
 		{
-			ChunkName = InChunkName;
-			ResponseFile = new ConcurrentDictionary<string, string>();
-			Manifest = null;
-			bCompressed = false;
-		}
-
-		// Name of pak file without extension, ie pakchunk0
-		public string ChunkName;
-
-		// List of files to include in this chunk
-		public ConcurrentDictionary<string, string> ResponseFile;
+			public ChunkDefinition(string InChunkName)
+			{
+				ChunkName = InChunkName;
+				ResponseFile = new ConcurrentDictionary<string, string>();
+				Manifest = null;
+				bCompressed = false;
+			}
 
-		// Parsed copy of pakchunk*.txt
-		public HashSet<string> Manifest;
+			// Name of pak file without extension, ie pakchunk0
+			public string ChunkName;
 
-		// Data for compression and encryption
-		public bool bCompressed;
-		public string EncryptionKeyGuid;
-		public string RequestedEncryptionKeyGuid;
-	}
+			// List of files to include in this chunk
+			public ConcurrentDictionary<string, string> ResponseFile;
 
-	/// <summary>
-	/// Creates pak files using streaming install chunk manifests.
-	/// </summary>
-	/// <param name="Params"></param>
-	/// <param name="SC"></param>
-	private static void CreatePaksUsingChunkManifests(ProjectParams Params, DeploymentContext SC)
-	{
-		LogInformation("Creating pak using streaming install manifests.");
-		DumpManifest(SC, CombinePaths(CmdEnv.LogFolder, "PrePak" + (SC.DedicatedServer ? "_Server" : "") + "_" + SC.CookPlatform));
+			// Parsed copy of pakchunk*.txt
+			public HashSet<string> Manifest;
 
-		ConfigHierarchy PlatformGameConfig = ConfigCache.ReadHierarchy(ConfigHierarchyType.Game, DirectoryReference.FromFile(Params.RawProjectPath), SC.StageTargetPlatform.IniPlatformType);
+			// Data for compression and encryption
+			public bool bCompressed;
+			public string EncryptionKeyGuid;
+			public string RequestedEncryptionKeyGuid;
+		}
 
-		bool bShouldGenerateEarlyDownloaderPakFile = false, bForceOneChunkPerFile = false;
-		string EarlyDownloaderPakFilePrefix = string.Empty;
-		if (!Params.HasDLCName)
+		/// <summary>
+		/// Creates pak files using streaming install chunk manifests.
+		/// </summary>
+		/// <param name="Params"></param>
+		/// <param name="SC"></param>
+		private static void CreatePaksUsingChunkManifests(ProjectParams Params, DeploymentContext SC)
 		{
-			PlatformGameConfig.GetBool("/Script/UnrealEd.ProjectPackagingSettings", "GenerateEarlyDownloaderPakFile", out bShouldGenerateEarlyDownloaderPakFile);
-		}
-		PlatformGameConfig.GetBool("/Script/UnrealEd.ProjectPackagingSettings", "bForceOneChunkPerFile", out bForceOneChunkPerFile);
-		PlatformGameConfig.GetString("/Script/UnrealEd.ProjectPackagingSettings", "EarlyDownloaderPakFilePrefix", out EarlyDownloaderPakFilePrefix);
+			LogInformation("Creating pak using streaming install manifests.");
+			DumpManifest(SC, CombinePaths(CmdEnv.LogFolder, "PrePak" + (SC.DedicatedServer ? "_Server" : "") + "_" + SC.CookPlatform));
 
-		List<ChunkDefinition> ChunkDefinitions = new List<ChunkDefinition>();
-		List<PakFileRules> PakRulesList = GetPakFileRules(Params, SC);
+			ConfigHierarchy PlatformGameConfig = ConfigCache.ReadHierarchy(ConfigHierarchyType.Game, DirectoryReference.FromFile(Params.RawProjectPath), SC.StageTargetPlatform.IniPlatformType);
 
-		var TmpPackagingPath = GetTmpPackagingPath(Params, SC);
+			bool bShouldGenerateEarlyDownloaderPakFile = false, bForceOneChunkPerFile = false;
+			string EarlyDownloaderPakFilePrefix = string.Empty;
+			if (!Params.HasDLCName)
+			{
+				PlatformGameConfig.GetBool("/Script/UnrealEd.ProjectPackagingSettings", "GenerateEarlyDownloaderPakFile", out bShouldGenerateEarlyDownloaderPakFile);
+			}
+			PlatformGameConfig.GetBool("/Script/UnrealEd.ProjectPackagingSettings", "bForceOneChunkPerFile", out bForceOneChunkPerFile);
+			PlatformGameConfig.GetString("/Script/UnrealEd.ProjectPackagingSettings", "EarlyDownloaderPakFilePrefix", out EarlyDownloaderPakFilePrefix);
 
-		// Parse and cache crypto settings from INI file
-		EncryptionAndSigning.CryptoSettings PakCryptoSettings = EncryptionAndSigning.ParseCryptoSettings(DirectoryReference.FromFile(Params.RawProjectPath), SC.StageTargetPlatform.IniPlatformType);
-		FileReference CryptoKeysCacheFilename = FileReference.Combine(SC.MetadataDir, "Crypto.json");
-		PakCryptoSettings.Save(CryptoKeysCacheFilename);
+			List<ChunkDefinition> ChunkDefinitions = new List<ChunkDefinition>();
+			List<PakFileRules> PakRulesList = GetPakFileRules(Params, SC);
 
-		// We still want to have a list of all files to stage. We will use the chunk manifests
-		// to put the files from staging manifest into the right chunk
-		var StagingManifestResponseFile = CreatePakResponseFileFromStagingManifest(SC, SC.FilesToStage.UFSFiles);
+			var TmpPackagingPath = GetTmpPackagingPath(Params, SC);
 
-		{
-			// DefaultChunkIndex assumes 0 is the 'base' chunk
-			const int DefaultChunkIndex = 0;
+			// Parse and cache crypto settings from INI file
+			EncryptionAndSigning.CryptoSettings PakCryptoSettings = EncryptionAndSigning.ParseCryptoSettings(DirectoryReference.FromFile(Params.RawProjectPath), SC.StageTargetPlatform.IniPlatformType);
+			FileReference CryptoKeysCacheFilename = FileReference.Combine(SC.MetadataDir, "Crypto.json");
+			PakCryptoSettings.Save(CryptoKeysCacheFilename);
 
-			var ChunkListFilename = GetChunkPakManifestListFilename(Params, SC);
-			List<string> ChunkList = new List<string>(ReadAllLines(ChunkListFilename));
-			Log.TraceInformation("Reading chunk list file {0} which contains {1} entries", ChunkListFilename, ChunkList.Count);
+			// We still want to have a list of all files to stage. We will use the chunk manifests
+			// to put the files from staging manifest into the right chunk
+			var StagingManifestResponseFile = CreatePakResponseFileFromStagingManifest(SC, SC.FilesToStage.UFSFiles);
 
-			for (int Index = 0; Index < ChunkList.Count; ++Index)
 			{
-				string[] ChunkOptions = ChunkList[Index].Split(' ');
+				// DefaultChunkIndex assumes 0 is the 'base' chunk
+				const int DefaultChunkIndex = 0;
 
-				// Set chunk name to string like "pakchunk0"
-				var ChunkManifestFilename = CombinePaths(TmpPackagingPath, ChunkOptions[0]);
-				ChunkDefinition CD = new ChunkDefinition(Path.GetFileNameWithoutExtension(ChunkOptions[0]));
-				for (int I = 1; I < ChunkOptions.Length; ++I)
+				var ChunkListFilename = GetChunkPakManifestListFilename(Params, SC);
+				List<string> ChunkList = new List<string>(ReadAllLines(ChunkListFilename));
+				Log.TraceInformation("Reading chunk list file {0} which contains {1} entries", ChunkListFilename, ChunkList.Count);
+
+				for (int Index = 0; Index < ChunkList.Count; ++Index)
 				{
-					if (string.Compare(ChunkOptions[I], "compressed", true) == 0)
-					{
-						CD.bCompressed = true;
-					}
-					else if (ChunkOptions[I].StartsWith("encryptionkeyguid="))
-					{
-						CD.RequestedEncryptionKeyGuid = ChunkOptions[I].Substring(ChunkOptions[I].IndexOf('=') + 1); ;
+					string[] ChunkOptions = ChunkList[Index].Split(' ');
 
-						if (PakCryptoSettings.SecondaryEncryptionKeys != null)
+					// Set chunk name to string like "pakchunk0"
+					var ChunkManifestFilename = CombinePaths(TmpPackagingPath, ChunkOptions[0]);
+					ChunkDefinition CD = new ChunkDefinition(Path.GetFileNameWithoutExtension(ChunkOptions[0]));
+					for (int I = 1; I < ChunkOptions.Length; ++I)
+					{
+						if (string.Compare(ChunkOptions[I], "compressed", true) == 0)
+						{
+							CD.bCompressed = true;
+						}
+						else if (ChunkOptions[I].StartsWith("encryptionkeyguid="))
 						{
-							foreach (EncryptionAndSigning.EncryptionKey Key in PakCryptoSettings.SecondaryEncryptionKeys)
+							CD.RequestedEncryptionKeyGuid = ChunkOptions[I].Substring(ChunkOptions[I].IndexOf('=') + 1); ;
+
+							if (PakCryptoSettings.SecondaryEncryptionKeys != null)
 							{
-								if (string.Compare(Key.Guid, CD.RequestedEncryptionKeyGuid, true) == 0)
+								foreach (EncryptionAndSigning.EncryptionKey Key in PakCryptoSettings.SecondaryEncryptionKeys)
 								{
-									CD.EncryptionKeyGuid = CD.RequestedEncryptionKeyGuid;
-									break;
+									if (string.Compare(Key.Guid, CD.RequestedEncryptionKeyGuid, true) == 0)
+									{
+										CD.EncryptionKeyGuid = CD.RequestedEncryptionKeyGuid;
+										break;
+									}
 								}
 							}
 						}
 					}
+					CD.Manifest = ReadPakChunkManifest(ChunkManifestFilename);
+					Log.TraceInformation("Reading chunk manifest {0} which contains {1} entries", ChunkManifestFilename, CD.Manifest.Count);
+					ChunkDefinitions.Add(CD);
 				}
-				CD.Manifest = ReadPakChunkManifest(ChunkManifestFilename);
-				Log.TraceInformation("Reading chunk manifest {0} which contains {1} entries", ChunkManifestFilename, CD.Manifest.Count);
-				ChunkDefinitions.Add(CD);
-			}
 
-			const string OptionalBulkDataFileExtension = ".uptnl";
-			ConcurrentDictionary<string, ChunkDefinition> OptionalChunks = new ConcurrentDictionary<string, ChunkDefinition>(StringComparer.InvariantCultureIgnoreCase);
-			ChunkDefinition DefaultChunk = ChunkDefinitions[DefaultChunkIndex];
+				const string OptionalBulkDataFileExtension = ".uptnl";
+				ConcurrentDictionary<string, ChunkDefinition> OptionalChunks = new ConcurrentDictionary<string, ChunkDefinition>(StringComparer.InvariantCultureIgnoreCase);
+				ChunkDefinition DefaultChunk = ChunkDefinitions[DefaultChunkIndex];
 
-			Dictionary<string, List<ChunkDefinition>> FileNameToChunks = new Dictionary<string, List<ChunkDefinition>>(StringComparer.InvariantCultureIgnoreCase);
-			foreach (ChunkDefinition Chunk in ChunkDefinitions)
-			{
-				foreach (string FileName in Chunk.Manifest)
+				Dictionary<string, List<ChunkDefinition>> FileNameToChunks = new Dictionary<string, List<ChunkDefinition>>(StringComparer.InvariantCultureIgnoreCase);
+				foreach (ChunkDefinition Chunk in ChunkDefinitions)
 				{
-					List<ChunkDefinition> Chunks;
-					if (!FileNameToChunks.TryGetValue(FileName, out Chunks))
+					foreach (string FileName in Chunk.Manifest)
 					{
-						Chunks = new List<ChunkDefinition>();
-						FileNameToChunks.Add(FileName, Chunks);
+						List<ChunkDefinition> Chunks;
+						if (!FileNameToChunks.TryGetValue(FileName, out Chunks))
+						{
+							Chunks = new List<ChunkDefinition>();
+							FileNameToChunks.Add(FileName, Chunks);
+						}
+						Chunks.Add(Chunk);
 					}
-					Chunks.Add(Chunk);
 				}
-			}
-
-			ConcurrentDictionary<string, ChunkDefinition> ChunkNameToDefinition = new ConcurrentDictionary<string, ChunkDefinition>(StringComparer.InvariantCultureIgnoreCase);
-			foreach (ChunkDefinition Chunk in ChunkDefinitions)
-			{
-				ChunkNameToDefinition.TryAdd(Chunk.ChunkName, Chunk);
-			}
 
-			Parallel.ForEach(StagingManifestResponseFile, StagingFile =>
-			{
-				bool bAddedToChunk = false;
-				bool bExcludeFromPaks = false;
-				HashSet<ChunkDefinition> PakList = new HashSet<ChunkDefinition>();
+				ConcurrentDictionary<string, ChunkDefinition> ChunkNameToDefinition = new ConcurrentDictionary<string, ChunkDefinition>(StringComparer.InvariantCultureIgnoreCase);
+				foreach (ChunkDefinition Chunk in ChunkDefinitions)
+				{
+					ChunkNameToDefinition.TryAdd(Chunk.ChunkName, Chunk);
+				}
 
-				string OriginalFilename = StagingFile.Key;
-				string NoExtension = CombinePaths(Path.GetDirectoryName(OriginalFilename), Path.GetFileNameWithoutExtension(OriginalFilename));
-				if(Path.GetExtension(NoExtension) == ".m")
+				Parallel.ForEach(StagingManifestResponseFile, StagingFile =>
 				{
+					bool bAddedToChunk = false;
+					bool bExcludeFromPaks = false;
+					HashSet<ChunkDefinition> PakList = new HashSet<ChunkDefinition>();
+
+					string OriginalFilename = StagingFile.Key;
+					string NoExtension = CombinePaths(Path.GetDirectoryName(OriginalFilename), Path.GetFileNameWithoutExtension(OriginalFilename));
+					if (Path.GetExtension(NoExtension) == ".m")
+					{
 					// Hack around .m.ubulk files having a double extension
 					NoExtension = CombinePaths(Path.GetDirectoryName(OriginalFilename), Path.GetFileNameWithoutExtension(NoExtension));
-				}
-				string OriginalReplaceSlashes = OriginalFilename.Replace('/', '\\');
-				string NoExtensionReplaceSlashes = NoExtension.Replace('/', '\\');
+					}
+					string OriginalReplaceSlashes = OriginalFilename.Replace('/', '\\');
+					string NoExtensionReplaceSlashes = NoExtension.Replace('/', '\\');
 
 				// First read manifest
 				List<ChunkDefinition> Chunks;
-				if (FileNameToChunks.TryGetValue(OriginalFilename, out Chunks))
-				{
-					PakList.UnionWith(Chunks);
-				}
-				if (FileNameToChunks.TryGetValue(OriginalReplaceSlashes, out Chunks))
-				{
-					PakList.UnionWith(Chunks);
-				}
-				if (FileNameToChunks.TryGetValue(NoExtension, out Chunks))
-				{
-					PakList.UnionWith(Chunks);
-				}
-				if (FileNameToChunks.TryGetValue(NoExtensionReplaceSlashes, out Chunks))
-				{
-					PakList.UnionWith(Chunks);
-				}
+					if (FileNameToChunks.TryGetValue(OriginalFilename, out Chunks))
+					{
+						PakList.UnionWith(Chunks);
+					}
+					if (FileNameToChunks.TryGetValue(OriginalReplaceSlashes, out Chunks))
+					{
+						PakList.UnionWith(Chunks);
+					}
+					if (FileNameToChunks.TryGetValue(NoExtension, out Chunks))
+					{
+						PakList.UnionWith(Chunks);
+					}
+					if (FileNameToChunks.TryGetValue(NoExtensionReplaceSlashes, out Chunks))
+					{
+						PakList.UnionWith(Chunks);
+					}
 
 				// Now run through the pak rules which may override things
 				ApplyPakFileRules(PakRulesList, StagingFile, PakList, ChunkNameToDefinition, out bExcludeFromPaks);
 
-				if (bExcludeFromPaks)
-				{
-					return;
-				}
+					if (bExcludeFromPaks)
+					{
+						return;
+					}
 
 				// Actually add to chunk
 				foreach (ChunkDefinition Chunk in PakList)
-				{
-					ChunkDefinition TargetChunk = Chunk;
-
-					string OrigExt = Path.GetExtension(OriginalFilename);
-					if (OrigExt.Equals(OptionalBulkDataFileExtension))
 					{
+						ChunkDefinition TargetChunk = Chunk;
+
+						string OrigExt = Path.GetExtension(OriginalFilename);
+						if (OrigExt.Equals(OptionalBulkDataFileExtension))
+						{
 						// any optional files encountered we want to put in a separate pak file
 						string OptionalChunkName = Chunk.ChunkName + "optional";
-						if (!OptionalChunks.ContainsKey(OptionalChunkName))
-						{
-							ChunkDefinition OptionalChunk = new ChunkDefinition(OptionalChunkName);
-							OptionalChunk.RequestedEncryptionKeyGuid = Chunk.RequestedEncryptionKeyGuid;
-							OptionalChunk.EncryptionKeyGuid = Chunk.EncryptionKeyGuid;
-							OptionalChunks.TryAdd(OptionalChunkName, OptionalChunk);
+							if (!OptionalChunks.ContainsKey(OptionalChunkName))
+							{
+								ChunkDefinition OptionalChunk = new ChunkDefinition(OptionalChunkName);
+								OptionalChunk.RequestedEncryptionKeyGuid = Chunk.RequestedEncryptionKeyGuid;
+								OptionalChunk.EncryptionKeyGuid = Chunk.EncryptionKeyGuid;
+								OptionalChunks.TryAdd(OptionalChunkName, OptionalChunk);
+							}
+							TargetChunk = OptionalChunks[OptionalChunkName];
 						}
-						TargetChunk = OptionalChunks[OptionalChunkName];
-					}
 
-					TargetChunk.ResponseFile.TryAdd(StagingFile.Key, StagingFile.Value);
-					bAddedToChunk = true;
+						TargetChunk.ResponseFile.TryAdd(StagingFile.Key, StagingFile.Value);
+						bAddedToChunk = true;
 
-					if (bForceOneChunkPerFile)
-					{
+						if (bForceOneChunkPerFile)
+						{
 						// Files are only allowed to be in a single chunk
 						break;
+						}
 					}
-				}
 
-				if (!bAddedToChunk)
-				{
+					if (!bAddedToChunk)
+					{
 					//LogInformation("No chunk assigned found for {0}. Using default chunk.", StagingFile.Key);
 					DefaultChunk.ResponseFile.TryAdd(StagingFile.Key, StagingFile.Value);
+					}
+				});
+
+				foreach (var OptionalChunkIt in OptionalChunks)
+				{
+					ChunkDefinitions.Add(OptionalChunkIt.Value);
 				}
-			});
+			}
 
-			foreach (var OptionalChunkIt in OptionalChunks)
+			// chunk downloader pak file is a minimal pak file which contains no content.  It can be provided with as a minimal download so that the game can download all the content from another source.
+			// Deprecated in favor of DefaultPakFileRules.ini
+			if (bShouldGenerateEarlyDownloaderPakFile)
 			{
-				ChunkDefinitions.Add(OptionalChunkIt.Value);
-			}
-		}
+				string EarlyChunkName = EarlyDownloaderPakFilePrefix;
+				ChunkDefinition EarlyChunk = new ChunkDefinition(EarlyChunkName);
 
-		// chunk downloader pak file is a minimal pak file which contains no content.  It can be provided with as a minimal download so that the game can download all the content from another source.
-		// Deprecated in favor of DefaultPakFileRules.ini
-		if (bShouldGenerateEarlyDownloaderPakFile)
-		{
-			string EarlyChunkName = EarlyDownloaderPakFilePrefix;
-			ChunkDefinition EarlyChunk = new ChunkDefinition(EarlyChunkName);
+				EarlyChunk.bCompressed = true;
 
-			EarlyChunk.bCompressed = true;
+				ConcurrentDictionary<string, string> EarlyPakFile = EarlyChunk.ResponseFile;
 
-			ConcurrentDictionary<string,string> EarlyPakFile = EarlyChunk.ResponseFile;
+				// find the list of files to put in the early downloader pak file
+				List<string> FilesInEarlyPakFile = new List<string>();
+				PlatformGameConfig.GetArray("/Script/UnrealEd.ProjectPackagingSettings", "EarlyDownloaderPakFileFiles", out FilesInEarlyPakFile);
 
-			// find the list of files to put in the early downloader pak file
-			List<string> FilesInEarlyPakFile = new List<string>();
-			PlatformGameConfig.GetArray("/Script/UnrealEd.ProjectPackagingSettings", "EarlyDownloaderPakFileFiles", out FilesInEarlyPakFile);
+				FileFilter EarlyPakFileFilter = new FileFilter();
+				foreach (string FileFilter in FilesInEarlyPakFile)
+				{
+					EarlyPakFileFilter.AddRule(FileFilter);
+				}
 
-			FileFilter EarlyPakFileFilter = new FileFilter();
-			foreach (string FileFilter in FilesInEarlyPakFile)
-			{
-				EarlyPakFileFilter.AddRule(FileFilter);
-			}
+				List<string> FilesToFilter = new List<string>();
+				foreach (var ResponseFile in StagingManifestResponseFile)
+				{
+					FilesToFilter.Add(ResponseFile.Key);
+				}
 
-			List<string> FilesToFilter = new List<string>();
-			foreach (var ResponseFile in StagingManifestResponseFile)
-			{
-				FilesToFilter.Add(ResponseFile.Key);
-			}
+				foreach (string FilteredFile in EarlyPakFileFilter.ApplyTo(FilesToFilter))
+				{
+					EarlyPakFile.TryAdd(FilteredFile, StagingManifestResponseFile[FilteredFile]);
+				}
 
-			foreach ( string FilteredFile in EarlyPakFileFilter.ApplyTo(FilesToFilter) )
-			{
-				EarlyPakFile.TryAdd(FilteredFile, StagingManifestResponseFile[FilteredFile]);
+				ChunkDefinitions.Add(EarlyChunk);
 			}
 
-			ChunkDefinitions.Add(EarlyChunk);
-		}
-
-		if (Params.CreateChunkInstall)
-		{
-			string ManifestDir = CombinePaths(Params.ChunkInstallDirectory, SC.FinalCookPlatform, "ManifestDir");
-			if (InternalUtils.SafeDirectoryExists(ManifestDir))
+			if (Params.CreateChunkInstall)
 			{
-				foreach (string ManifestFile in Directory.GetFiles(ManifestDir, "*.manifest"))
+				string ManifestDir = CombinePaths(Params.ChunkInstallDirectory, SC.FinalCookPlatform, "ManifestDir");
+				if (InternalUtils.SafeDirectoryExists(ManifestDir))
 				{
-					InternalUtils.SafeDeleteFile(ManifestFile, true);
+					foreach (string ManifestFile in Directory.GetFiles(ManifestDir, "*.manifest"))
+					{
+						InternalUtils.SafeDeleteFile(ManifestFile, true);
+					}
+				}
+				string DestDir = CombinePaths(Params.ChunkInstallDirectory, SC.FinalCookPlatform, Params.ChunkInstallVersionString);
+				if (InternalUtils.SafeDirectoryExists(DestDir))
+				{
+					InternalUtils.SafeDeleteDirectory(DestDir);
 				}
 			}
-			string DestDir = CombinePaths(Params.ChunkInstallDirectory, SC.FinalCookPlatform, Params.ChunkInstallVersionString);
-			if (InternalUtils.SafeDirectoryExists(DestDir))
-			{
-				InternalUtils.SafeDeleteDirectory(DestDir);
-			}
-		}
 
-		System.Threading.Tasks.ParallelOptions Options = new System.Threading.Tasks.ParallelOptions();
+			System.Threading.Tasks.ParallelOptions Options = new System.Threading.Tasks.ParallelOptions();
 
-		LogInformation("Creating Pak files utilizing {0} cores", Environment.ProcessorCount);
-		Options.MaxDegreeOfParallelism = Environment.ProcessorCount;
+			LogInformation("Creating Pak files utilizing {0} cores", Environment.ProcessorCount);
+			Options.MaxDegreeOfParallelism = Environment.ProcessorCount;
 
-		// Check for chunks that requested an encryption key that we don't have. This can happen for Epic contractors who don't have access to our keychains. Emit a warning though, just in case
-		// this somehow happens when the keychain is available
-		foreach (ChunkDefinition CD in ChunkDefinitions)
-		{
-			if (!string.IsNullOrEmpty(CD.RequestedEncryptionKeyGuid) && (CD.RequestedEncryptionKeyGuid != CD.EncryptionKeyGuid))
+			// Check for chunks that requested an encryption key that we don't have. This can happen for Epic contractors who don't have access to our keychains. Emit a warning though, just in case
+			// this somehow happens when the keychain is available
+			foreach (ChunkDefinition CD in ChunkDefinitions)
 			{
-				LogWarning("Chunk '" + CD.ChunkName + "' requested encryption key '" + CD.RequestedEncryptionKeyGuid + "' but it couldn't be found in the project keychain.");
+				if (!string.IsNullOrEmpty(CD.RequestedEncryptionKeyGuid) && (CD.RequestedEncryptionKeyGuid != CD.EncryptionKeyGuid))
+				{
+					LogWarning("Chunk '" + CD.ChunkName + "' requested encryption key '" + CD.RequestedEncryptionKeyGuid + "' but it couldn't be found in the project keychain.");
+				}
 			}
-		}
 
-		List<CreatePakParams> PakInputs = new List<CreatePakParams>();
-		foreach (ChunkDefinition Chunk in ChunkDefinitions)
-		{
-			if (Chunk.ResponseFile.Count > 0)
+			List<CreatePakParams> PakInputs = new List<CreatePakParams>();
+			foreach (ChunkDefinition Chunk in ChunkDefinitions)
 			{
-				string EncryptionKeyToUse = Params.SkipEncryption ? "" : Chunk.EncryptionKeyGuid;
-				PakInputs.Add(new CreatePakParams(Chunk.ChunkName, 
-					Chunk.ResponseFile.ToDictionary(entry => entry.Key, entry => entry.Value),
-					Params.Compressed || Chunk.bCompressed, EncryptionKeyToUse));
+				if (Chunk.ResponseFile.Count > 0)
+				{
+					string EncryptionKeyToUse = Params.SkipEncryption ? "" : Chunk.EncryptionKeyGuid;
+					PakInputs.Add(new CreatePakParams(Chunk.ChunkName,
+						Chunk.ResponseFile.ToDictionary(entry => entry.Key, entry => entry.Value),
+						Params.Compressed || Chunk.bCompressed, EncryptionKeyToUse));
+				}
 			}
-		}
 
-		CreatePaks(Params, SC, PakInputs, PakCryptoSettings, CryptoKeysCacheFilename);
+			CreatePaks(Params, SC, PakInputs, PakCryptoSettings, CryptoKeysCacheFilename);
 
-		String ChunkLayerFilename = CombinePaths(GetTmpPackagingPath(Params, SC), GetChunkPakLayerListName());
-		String OutputChunkLayerFilename = Path.Combine(SC.ProjectRoot.FullName, "Build", SC.FinalCookPlatform, "ChunkLayerInfo", GetChunkPakLayerListName());
-		Directory.CreateDirectory(Path.GetDirectoryName(OutputChunkLayerFilename));
-		File.Copy(ChunkLayerFilename, OutputChunkLayerFilename, true);
-	}
-
-	private static void GenerateMasterChunkManifest(string Dir, string Version, string PlatformStr)
-	{
-		//Create the directory if it doesn't exist
-		InternalUtils.SafeCreateDirectory(Dir);
+			String ChunkLayerFilename = CombinePaths(GetTmpPackagingPath(Params, SC), GetChunkPakLayerListName());
+			String OutputChunkLayerFilename = Path.Combine(SC.ProjectRoot.FullName, "Build", SC.FinalCookPlatform, "ChunkLayerInfo", GetChunkPakLayerListName());
+			Directory.CreateDirectory(Path.GetDirectoryName(OutputChunkLayerFilename));
+			File.Copy(ChunkLayerFilename, OutputChunkLayerFilename, true);
+		}
 
-		string FileName = CombinePaths(Dir, PlatformStr.ToLower() + ".manifestlist");
-		using (JsonWriter Writer = new JsonWriter(FileName))
+		private static void GenerateMasterChunkManifest(string Dir, string Version, string PlatformStr)
 		{
-			Writer.WriteObjectStart();
-			Writer.WriteValue("ClientVersion", Version);
-			Writer.WriteValue("BuildUrl", Version + "/" + PlatformStr);
-			Writer.WriteArrayStart("files");
-			// iterate over all of the files in the directory
-			DirectoryInfo di = new DirectoryInfo(Dir);
-			foreach (var fi in di.EnumerateFiles("*.manifest"))
-			{
-				if (fi.Name == PlatformStr.ToLower() + ".manifest")
-					continue;
-				FileStream fs = fi.OpenRead();
-				byte[] hash = System.Security.Cryptography.SHA1.Create().ComputeHash(fs);
-				fs.Seek(0, SeekOrigin.Begin);
-				byte[] hash256 = System.Security.Cryptography.SHA256.Create().ComputeHash(fs);
-				fs.Close();
+			//Create the directory if it doesn't exist
+			InternalUtils.SafeCreateDirectory(Dir);
+
+			string FileName = CombinePaths(Dir, PlatformStr.ToLower() + ".manifestlist");
+			using (JsonWriter Writer = new JsonWriter(FileName))
+			{
 				Writer.WriteObjectStart();
-				Writer.WriteValue("filename", fi.Name);
-				Writer.WriteValue("uniqueFilename", fi.Name);
-				Writer.WriteValue("length", fi.Length);
-				Writer.WriteValue("URL", fi.Name);
-				Writer.WriteValue("hash", BitConverter.ToString(hash).Replace("-", ""));
-				Writer.WriteValue("hash256", BitConverter.ToString(hash256).Replace("-", ""));
+				Writer.WriteValue("ClientVersion", Version);
+				Writer.WriteValue("BuildUrl", Version + "/" + PlatformStr);
+				Writer.WriteArrayStart("files");
+				// iterate over all of the files in the directory
+				DirectoryInfo di = new DirectoryInfo(Dir);
+				foreach (var fi in di.EnumerateFiles("*.manifest"))
+				{
+					if (fi.Name == PlatformStr.ToLower() + ".manifest")
+						continue;
+					FileStream fs = fi.OpenRead();
+					byte[] hash = System.Security.Cryptography.SHA1.Create().ComputeHash(fs);
+					fs.Seek(0, SeekOrigin.Begin);
+					byte[] hash256 = System.Security.Cryptography.SHA256.Create().ComputeHash(fs);
+					fs.Close();
+					Writer.WriteObjectStart();
+					Writer.WriteValue("filename", fi.Name);
+					Writer.WriteValue("uniqueFilename", fi.Name);
+					Writer.WriteValue("length", fi.Length);
+					Writer.WriteValue("URL", fi.Name);
+					Writer.WriteValue("hash", BitConverter.ToString(hash).Replace("-", ""));
+					Writer.WriteValue("hash256", BitConverter.ToString(hash256).Replace("-", ""));
+					Writer.WriteObjectEnd();
+				}
+				Writer.WriteArrayEnd();
 				Writer.WriteObjectEnd();
 			}
-			Writer.WriteArrayEnd();
-			Writer.WriteObjectEnd();
 		}
-	}
-
-	private static bool DoesChunkPakManifestExist(ProjectParams Params, DeploymentContext SC)
-	{
-		return FileExists_NoExceptions(GetChunkPakManifestListFilename(Params, SC));
-	}
-
-	private static string GetChunkPakManifestListFilename(ProjectParams Params, DeploymentContext SC)
-	{
-		return CombinePaths(GetTmpPackagingPath(Params, SC), "pakchunklist.txt");
-	}
-
-	private static string GetChunkPakLayerListName()
-	{
-		return "pakchunklayers.txt";
-	}
 
-	private static string GetTmpPackagingPath(ProjectParams Params, DeploymentContext SC)
-	{
-		string TmpPackagingPath = CombinePaths(Path.GetDirectoryName(Params.RawProjectPath.FullName), "Saved", "TmpPackaging", SC.CookPlatform);
-		if (Params.bUseExtraFlavor)
+		private static bool DoesChunkPakManifestExist(ProjectParams Params, DeploymentContext SC)
 		{
-			TmpPackagingPath = CombinePaths(TmpPackagingPath, "ExtraFlavor");
+			return FileExists_NoExceptions(GetChunkPakManifestListFilename(Params, SC));
 		}
 
-		return TmpPackagingPath;
-	}
-
-	private static bool ShouldCreateIoStoreContainerFiles(ProjectParams Params, Platform StageTargetPlatform)
-	{
-		if (Params.CookOnTheFly)
+		private static string GetChunkPakManifestListFilename(ProjectParams Params, DeploymentContext SC)
 		{
-			return false;
+			return CombinePaths(GetTmpPackagingPath(Params, SC), "pakchunklist.txt");
 		}
 
-		if (Params.SkipIoStore)
+		private static string GetChunkPakLayerListName()
 		{
-			return false;
+			return "pakchunklayers.txt";
 		}
 
-		if (Params.IoStore)
+		private static string GetTmpPackagingPath(ProjectParams Params, DeploymentContext SC)
 		{
-			return true;
-		}
+			string TmpPackagingPath = CombinePaths(Path.GetDirectoryName(Params.RawProjectPath.FullName), "Saved", "TmpPackaging", SC.CookPlatform);
+			if (Params.bUseExtraFlavor)
+			{
+				TmpPackagingPath = CombinePaths(TmpPackagingPath, "ExtraFlavor");
+			}
 
-		if (Params.Stage && !Params.SkipStage)
-		{
-			ConfigHierarchy PlatformGameConfig = ConfigCache.ReadHierarchy(ConfigHierarchyType.Game, DirectoryReference.FromFile(Params.RawProjectPath), StageTargetPlatform.IniPlatformType);
-			bool bUseIoStore = false;
-			PlatformGameConfig.GetBool("/Script/UnrealEd.ProjectPackagingSettings", "bUseIoStore", out bUseIoStore);
-			return bUseIoStore;
+			return TmpPackagingPath;
 		}
 
-		return false;
-	}
-
-	private static bool ShouldCreatePak(ProjectParams Params, DeploymentContext SC)
-	{
-		if (Params.CookOnTheFly)
+		private static bool ShouldCreateIoStoreContainerFiles(ProjectParams Params, Platform StageTargetPlatform)
 		{
-			return false;
-		}
+			if (Params.CookOnTheFly)
+			{
+				return false;
+			}
 
-		Platform.PakType Pak = SC.StageTargetPlatform.RequiresPak(Params);
+			if (Params.SkipIoStore)
+			{
+				return false;
+			}
 
-		// we may care but we don't want. 
-		if (Params.SkipPak)
-			return false;
+			if (Params.IoStore)
+			{
+				return true;
+			}
+
+			if (Params.Stage && !Params.SkipStage)
+			{
+				ConfigHierarchy PlatformGameConfig = ConfigCache.ReadHierarchy(ConfigHierarchyType.Game, DirectoryReference.FromFile(Params.RawProjectPath), StageTargetPlatform.IniPlatformType);
+				bool bUseIoStore = false;
+				PlatformGameConfig.GetBool("/Script/UnrealEd.ProjectPackagingSettings", "bUseIoStore", out bUseIoStore);
+				return bUseIoStore;
+			}
 
-		if (Pak == Platform.PakType.Always)
-		{
-			return true;
-		}
-		else if (Pak == Platform.PakType.Never)
-		{
 			return false;
 		}
-		else // DontCare
+
+		private static bool ShouldCreatePak(ProjectParams Params, DeploymentContext SC)
 		{
-			return (Params.Pak);
-		}
-	}
+			if (Params.CookOnTheFly)
+			{
+				return false;
+			}
 
-	private static bool ShouldCreatePak(ProjectParams Params)
-	{
-		Platform.PakType Pak = Params.ClientTargetPlatformInstances[0].RequiresPak(Params);
+			Platform.PakType Pak = SC.StageTargetPlatform.RequiresPak(Params);
 
-		// we may care but we don't want. 
-		if (Params.SkipPak)
-			return false;
+			// we may care but we don't want. 
+			if (Params.SkipPak)
+				return false;
 
-		if (Pak == Platform.PakType.Always)
-		{
-			return true;
-		}
-		else if (Pak == Platform.PakType.Never)
-		{
-			return false;
+			if (Pak == Platform.PakType.Always)
+			{
+				return true;
+			}
+			else if (Pak == Platform.PakType.Never)
+			{
+				return false;
+			}
+			else // DontCare
+			{
+				return (Params.Pak);
+			}
 		}
-		else // DontCare
+
+		private static bool ShouldCreatePak(ProjectParams Params)
 		{
-			return (Params.Pak);
-		}
-	}
+			Platform.PakType Pak = Params.ClientTargetPlatformInstances[0].RequiresPak(Params);
 
-	protected static void DeletePakFiles(string StagingDirectory)
-	{
-		var StagedFilesDir = new DirectoryInfo(StagingDirectory);
-		StagedFilesDir.GetFiles("*.pak", SearchOption.AllDirectories).ToList().ForEach(File => File.Delete());
-		StagedFilesDir.GetFiles("*.ucas", SearchOption.AllDirectories).ToList().ForEach(File => File.Delete());
-		StagedFilesDir.GetFiles("*.utoc", SearchOption.AllDirectories).ToList().ForEach(File => File.Delete());
-	}
+			// we may care but we don't want. 
+			if (Params.SkipPak)
+				return false;
 
-	protected static void CleanDirectoryExcludingPakFiles(DirectoryInfo StagingDirectory)
-	{
-		foreach (DirectoryInfo SubDir in StagingDirectory.EnumerateDirectories())
-		{
-			if (!SubDir.Name.Equals("Paks", StringComparison.OrdinalIgnoreCase))
+			if (Pak == Platform.PakType.Always)
+			{
+				return true;
+			}
+			else if (Pak == Platform.PakType.Never)
+			{
+				return false;
+			}
+			else // DontCare
 			{
-				CleanDirectoryExcludingPakFiles(SubDir);
-				try { SubDir.Delete(); } catch { }
+				return (Params.Pak);
 			}
 		}
 
-		foreach (System.IO.FileInfo File in StagingDirectory.EnumerateFiles())
+		protected static void DeletePakFiles(string StagingDirectory)
 		{
-			try { File.Delete(); } catch { }
+			var StagedFilesDir = new DirectoryInfo(StagingDirectory);
+			StagedFilesDir.GetFiles("*.pak", SearchOption.AllDirectories).ToList().ForEach(File => File.Delete());
+			StagedFilesDir.GetFiles("*.ucas", SearchOption.AllDirectories).ToList().ForEach(File => File.Delete());
+			StagedFilesDir.GetFiles("*.utoc", SearchOption.AllDirectories).ToList().ForEach(File => File.Delete());
 		}
-	}
 
-	public static void CleanStagingDirectory(ProjectParams Params, DeploymentContext SC)
-	{
-		if (Params.NoCleanStage)
+		protected static void CleanDirectoryExcludingPakFiles(DirectoryInfo StagingDirectory)
 		{
-			LogInformation("Skipping clean of staging directory due to -NoCleanStage argument.");
-		}
-		else if (SC.Stage && !Params.SkipStage)
-		{
-			if (Params.SkipPak)
-			{
-				LogInformation("Cleaning Stage Directory (exluding PAK files): {0}", SC.StageDirectory.FullName);
-				CleanDirectoryExcludingPakFiles(new DirectoryInfo(SC.StageDirectory.FullName));
-			}
-			else
+			foreach (DirectoryInfo SubDir in StagingDirectory.EnumerateDirectories())
 			{
-				LogInformation("Cleaning Stage Directory: {0}", SC.StageDirectory.FullName);
-				try
-				{
-					DeleteDirectory(SC.StageDirectory.FullName);
-				}
-				catch (Exception Ex)
+				if (!SubDir.Name.Equals("Paks", StringComparison.OrdinalIgnoreCase))
 				{
-					// Delete cooked data (if any) as it may be incomplete / corrupted.
-					throw new AutomationException(ExitCode.Error_FailedToDeleteStagingDirectory, Ex, "Stage Failed. Failed to delete staging directory " + SC.StageDirectory.FullName);
+					CleanDirectoryExcludingPakFiles(SubDir);
+					try { SubDir.Delete(); } catch { }
 				}
 			}
-		}
-		else
-		{
-			LogInformation("Cleaning PAK files in stage directory: {0}", SC.StageDirectory.FullName);
-			try
-			{
-				// delete old pak files
-				DeletePakFiles(SC.StageDirectory.FullName);
-			}
-			catch (Exception Ex)
+
+			foreach (System.IO.FileInfo File in StagingDirectory.EnumerateFiles())
 			{
-				// Delete cooked data (if any) as it may be incomplete / corrupted.
-				throw new AutomationException(ExitCode.Error_FailedToDeleteStagingDirectory, Ex, "Stage Failed. Failed to delete pak files in " + SC.StageDirectory.FullName);
+				try { File.Delete(); } catch { }
 			}
 		}
-	}
 
-	public static void ApplyStagingManifest(ProjectParams Params, DeploymentContext SC)
-	{
-		if (ShouldCreatePak(Params, SC))
+		public static void CleanStagingDirectory(ProjectParams Params, DeploymentContext SC)
 		{
-			if (SC.CrashReporterUFSFiles.Count > 0)
-			{
-				CreatePakForCrashReporter(Params, SC);
-			}
-
-			if (Params.PrePak)
+			if (Params.NoCleanStage)
 			{
-				CopyPaksFromNetwork(Params, SC);
+				LogInformation("Skipping clean of staging directory due to -NoCleanStage argument.");
 			}
-			else if (SC.PlatformUsesChunkManifests && DoesChunkPakManifestExist(Params, SC))
+			else if (SC.Stage && !Params.SkipStage)
 			{
-				CreatePaksUsingChunkManifests(Params, SC);
+				if (Params.SkipPak)
+				{
+					LogInformation("Cleaning Stage Directory (exluding PAK files): {0}", SC.StageDirectory.FullName);
+					CleanDirectoryExcludingPakFiles(new DirectoryInfo(SC.StageDirectory.FullName));
+				}
+				else
+				{
+					LogInformation("Cleaning Stage Directory: {0}", SC.StageDirectory.FullName);
+					try
+					{
+						DeleteDirectory(SC.StageDirectory.FullName);
+					}
+					catch (Exception Ex)
+					{
+						// Delete cooked data (if any) as it may be incomplete / corrupted.
+						throw new AutomationException(ExitCode.Error_FailedToDeleteStagingDirectory, Ex, "Stage Failed. Failed to delete staging directory " + SC.StageDirectory.FullName);
+					}
+				}
 			}
 			else
 			{
-				CreatePakUsingStagingManifest(Params, SC);
+				LogInformation("Cleaning PAK files in stage directory: {0}", SC.StageDirectory.FullName);
+				try
+				{
+					// delete old pak files
+					DeletePakFiles(SC.StageDirectory.FullName);
+				}
+				catch (Exception Ex)
+				{
+					// Delete cooked data (if any) as it may be incomplete / corrupted.
+					throw new AutomationException(ExitCode.Error_FailedToDeleteStagingDirectory, Ex, "Stage Failed. Failed to delete pak files in " + SC.StageDirectory.FullName);
+				}
 			}
 		}
 
-		string BaseManifestFileName = CombinePaths(CmdEnv.LogFolder, "FinalCopy" + (SC.DedicatedServer ? "_Server" : ""));
-		DumpManifest(SC.FilesToStage.UFSFiles, BaseManifestFileName + "_UFSFiles.txt");
-
-		if (!SC.Stage || Params.SkipStage)
+		public static void ApplyStagingManifest(ProjectParams Params, DeploymentContext SC)
 		{
-			return;
-		}
+			if (ShouldCreatePak(Params, SC))
+			{
+				if (SC.CrashReporterUFSFiles.Count > 0)
+				{
+					CreatePakForCrashReporter(Params, SC);
+				}
 
-		DumpManifest(SC.FilesToStage.NonUFSFiles, BaseManifestFileName + "_NonUFSFiles.txt");
-		DumpManifest(SC.FilesToStage.NonUFSDebugFiles, BaseManifestFileName + "_NonUFSFilesDebug.txt");
+				if (Params.PrePak)
+				{
+					CopyPaksFromNetwork(Params, SC);
+				}
+				else if (SC.PlatformUsesChunkManifests && DoesChunkPakManifestExist(Params, SC))
+				{
+					CreatePaksUsingChunkManifests(Params, SC);
+				}
+				else
+				{
+					CreatePakUsingStagingManifest(Params, SC);
+				}
+			}
 
-		CopyUsingStagingManifest(Params, SC);
+			string BaseManifestFileName = CombinePaths(CmdEnv.LogFolder, "FinalCopy" + (SC.DedicatedServer ? "_Server" : ""));
+			DumpManifest(SC.FilesToStage.UFSFiles, BaseManifestFileName + "_UFSFiles.txt");
 
-		var ThisPlatform = SC.StageTargetPlatform;
-		ThisPlatform.PostStagingFileCopy(Params, SC);
-	}
+			if (!SC.Stage || Params.SkipStage)
+			{
+				return;
+			}
 
-	private static DirectoryReference GetIntermediateCommandlineDir(DeploymentContext SC)
-	{
-		return DirectoryReference.Combine(SC.ProjectRoot, "Intermediate", "UAT", SC.FinalCookPlatform);
-	}
+			DumpManifest(SC.FilesToStage.NonUFSFiles, BaseManifestFileName + "_NonUFSFiles.txt");
+			DumpManifest(SC.FilesToStage.NonUFSDebugFiles, BaseManifestFileName + "_NonUFSFilesDebug.txt");
 
-	public static void WriteStageCommandline(FileReference IntermediateCmdLineFile, ProjectParams Params, DeploymentContext SC)
-	{
-		// this file needs to be treated as a UFS file for casing, but NonUFS for being put into the .pak file. 
-		// @todo: Maybe there should be a new category - UFSNotForPak
-		if (SC.StageTargetPlatform.DeployLowerCaseFilenames(StagedFileType.SystemNonUFS))
-		{
-			IntermediateCmdLineFile = new FileReference(CombinePaths(Path.GetDirectoryName(IntermediateCmdLineFile.FullName), Path.GetFileName(IntermediateCmdLineFile.FullName).ToLowerInvariant()));
-		}
-		if (FileReference.Exists(IntermediateCmdLineFile))
-		{
-			FileReference.Delete(IntermediateCmdLineFile);
+			CopyUsingStagingManifest(Params, SC);
+
+			var ThisPlatform = SC.StageTargetPlatform;
+			ThisPlatform.PostStagingFileCopy(Params, SC);
 		}
 
-		if (!SC.StageTargetPlatform.ShouldStageCommandLine(Params, SC))
+		private static DirectoryReference GetIntermediateCommandlineDir(DeploymentContext SC)
 		{
-			return;
+			return DirectoryReference.Combine(SC.ProjectRoot, "Intermediate", "UAT", SC.FinalCookPlatform);
 		}
 
-		LogInformation("Creating UECommandLine.txt");
-		if (!string.IsNullOrEmpty(Params.StageCommandline) || !string.IsNullOrEmpty(Params.RunCommandline))
+		public static void WriteStageCommandline(FileReference IntermediateCmdLineFile, ProjectParams Params, DeploymentContext SC)
 		{
-			string FileHostParams = " ";
-			if (Params.CookOnTheFly || Params.FileServer)
+			// this file needs to be treated as a UFS file for casing, but NonUFS for being put into the .pak file. 
+			// @todo: Maybe there should be a new category - UFSNotForPak
+			if (SC.StageTargetPlatform.DeployLowerCaseFilenames(StagedFileType.SystemNonUFS))
+			{
+				IntermediateCmdLineFile = new FileReference(CombinePaths(Path.GetDirectoryName(IntermediateCmdLineFile.FullName), Path.GetFileName(IntermediateCmdLineFile.FullName).ToLowerInvariant()));
+			}
+			if (FileReference.Exists(IntermediateCmdLineFile))
+			{
+				FileReference.Delete(IntermediateCmdLineFile);
+			}
+
+			if (!SC.StageTargetPlatform.ShouldStageCommandLine(Params, SC))
+			{
+				return;
+			}
+
+			LogInformation("Creating UECommandLine.txt");
+			if (!string.IsNullOrEmpty(Params.StageCommandline) || !string.IsNullOrEmpty(Params.RunCommandline))
 			{
-				FileHostParams += Params.ZenStore ? "-cookonthefly -zenstorehost=" : "-filehostip=";
-				// add localhost first for platforms using redirection
-				const string LocalHost = "127.0.0.1";
-				if (!IsNullOrEmpty(Params.Port))
+				string FileHostParams = " ";
+				if (Params.CookOnTheFly || Params.FileServer)
 				{
-					bool FirstParam = true;
-					foreach (var Port in Params.Port)
+					FileHostParams += Params.ZenStore ? "-cookonthefly -zenstorehost=" : "-filehostip=";
+					// add localhost first for platforms using redirection
+					const string LocalHost = "127.0.0.1";
+					if (!IsNullOrEmpty(Params.Port))
 					{
-						if (!FirstParam)
-						{
-							FileHostParams += "+";
-						}
-						FirstParam = false;
-						string[] PortProtocol = Port.Split(new char[] { ':' });
-						if (PortProtocol.Length > 1)
-						{
-							FileHostParams += String.Format("{0}://{1}:{2}", PortProtocol[0], LocalHost, PortProtocol[1]);
-						}
-						else
+						bool FirstParam = true;
+						foreach (var Port in Params.Port)
 						{
-							FileHostParams += LocalHost;
-							FileHostParams += ":";
-							FileHostParams += Params.Port;
-						}
+							if (!FirstParam)
+							{
+								FileHostParams += "+";
+							}
+							FirstParam = false;
+							string[] PortProtocol = Port.Split(new char[] { ':' });
+							if (PortProtocol.Length > 1)
+							{
+								FileHostParams += String.Format("{0}://{1}:{2}", PortProtocol[0], LocalHost, PortProtocol[1]);
+							}
+							else
+							{
+								FileHostParams += LocalHost;
+								FileHostParams += ":";
+								FileHostParams += Params.Port;
+							}
 
+						}
 					}
-				}
-				else
-				{
-					// use default port
-					FileHostParams += LocalHost;
-				}
-				if (UnrealBuildTool.BuildHostPlatform.Current.Platform == UnrealTargetPlatform.Mac)
-				{
-					NetworkInterface[] Interfaces = NetworkInterface.GetAllNetworkInterfaces();
-					foreach (NetworkInterface adapter in Interfaces)
+					else
+					{
+						// use default port
+						FileHostParams += LocalHost;
+					}
+					if (UnrealBuildTool.BuildHostPlatform.Current.Platform == UnrealTargetPlatform.Mac)
 					{
-						if (adapter.NetworkInterfaceType != NetworkInterfaceType.Loopback)
+						NetworkInterface[] Interfaces = NetworkInterface.GetAllNetworkInterfaces();
+						foreach (NetworkInterface adapter in Interfaces)
 						{
-							IPInterfaceProperties IP = adapter.GetIPProperties();
-							for (int Index = 0; Index < IP.UnicastAddresses.Count; ++Index)
+							if (adapter.NetworkInterfaceType != NetworkInterfaceType.Loopback)
 							{
-								if (InternalUtils.IsDnsEligible(IP.UnicastAddresses[Index]) && IP.UnicastAddresses[Index].Address.AddressFamily == System.Net.Sockets.AddressFamily.InterNetwork)
+								IPInterfaceProperties IP = adapter.GetIPProperties();
+								for (int Index = 0; Index < IP.UnicastAddresses.Count; ++Index)
 								{
-									if (!IsNullOrEmpty(Params.Port))
+									if (InternalUtils.IsDnsEligible(IP.UnicastAddresses[Index]) && IP.UnicastAddresses[Index].Address.AddressFamily == System.Net.Sockets.AddressFamily.InterNetwork)
 									{
-										foreach (var Port in Params.Port)
+										if (!IsNullOrEmpty(Params.Port))
 										{
-											FileHostParams += "+";
-											string[] PortProtocol = Port.Split(new char[] { ':' });
-											if (PortProtocol.Length > 1)
-											{
-												FileHostParams += String.Format("{0}://{1}:{2}", PortProtocol[0], IP.UnicastAddresses[Index].Address.ToString(), PortProtocol[1]);
-											}
-											else
+											foreach (var Port in Params.Port)
 											{
-												FileHostParams += IP.UnicastAddresses[Index].Address.ToString();
-												FileHostParams += ":";
-												FileHostParams += Params.Port;
-											}
+												FileHostParams += "+";
+												string[] PortProtocol = Port.Split(new char[] { ':' });
+												if (PortProtocol.Length > 1)
+												{
+													FileHostParams += String.Format("{0}://{1}:{2}", PortProtocol[0], IP.UnicastAddresses[Index].Address.ToString(), PortProtocol[1]);
+												}
+												else
+												{
+													FileHostParams += IP.UnicastAddresses[Index].Address.ToString();
+													FileHostParams += ":";
+													FileHostParams += Params.Port;
+												}
 
+											}
+										}
+										else
+										{
+											FileHostParams += "+";
+											// use default port
+											FileHostParams += IP.UnicastAddresses[Index].Address.ToString();
 										}
-									}
-									else
-									{
-										FileHostParams += "+";
-										// use default port
-										FileHostParams += IP.UnicastAddresses[Index].Address.ToString();
-									}
 
+									}
 								}
 							}
 						}
 					}
-				}
-				else
-				{
-					NetworkInterface[] Interfaces = NetworkInterface.GetAllNetworkInterfaces();
-					foreach (NetworkInterface adapter in Interfaces)
+					else
 					{
-						if (adapter.OperationalStatus == OperationalStatus.Up)
+						NetworkInterface[] Interfaces = NetworkInterface.GetAllNetworkInterfaces();
+						foreach (NetworkInterface adapter in Interfaces)
 						{
-							IPInterfaceProperties IP = adapter.GetIPProperties();
-							for (int Index = 0; Index < IP.UnicastAddresses.Count; ++Index)
+							if (adapter.OperationalStatus == OperationalStatus.Up)
 							{
-								if (InternalUtils.IsDnsEligible(IP.UnicastAddresses[Index]))
+								IPInterfaceProperties IP = adapter.GetIPProperties();
+								for (int Index = 0; Index < IP.UnicastAddresses.Count; ++Index)
 								{
-									if (!IsNullOrEmpty(Params.Port))
+									if (InternalUtils.IsDnsEligible(IP.UnicastAddresses[Index]))
 									{
-										foreach (var Port in Params.Port)
+										if (!IsNullOrEmpty(Params.Port))
 										{
-											FileHostParams += "+";
-											string[] PortProtocol = Port.Split(new char[] { ':' });
-											if (PortProtocol.Length > 1)
-											{
-												FileHostParams += String.Format("{0}://{1}:{2}", PortProtocol[0], IP.UnicastAddresses[Index].Address.ToString(), PortProtocol[1]);
-											}
-											else
+											foreach (var Port in Params.Port)
 											{
-												FileHostParams += IP.UnicastAddresses[Index].Address.ToString();
-												FileHostParams += ":";
-												FileHostParams += Params.Port;
+												FileHostParams += "+";
+												string[] PortProtocol = Port.Split(new char[] { ':' });
+												if (PortProtocol.Length > 1)
+												{
+													FileHostParams += String.Format("{0}://{1}:{2}", PortProtocol[0], IP.UnicastAddresses[Index].Address.ToString(), PortProtocol[1]);
+												}
+												else
+												{
+													FileHostParams += IP.UnicastAddresses[Index].Address.ToString();
+													FileHostParams += ":";
+													FileHostParams += Params.Port;
+												}
 											}
 										}
-									}
-									else
-									{
-										FileHostParams += "+";
-										// use default port
-										FileHostParams += IP.UnicastAddresses[Index].Address.ToString();
+										else
+										{
+											FileHostParams += "+";
+											// use default port
+											FileHostParams += IP.UnicastAddresses[Index].Address.ToString();
+										}
 									}
 								}
 							}
 						}
 					}
-				}
 
-				FileHostParams += " ";
-			}
+					FileHostParams += " ";
+				}
 
-			String ProjectFile = String.Format("{0} ", SC.ProjectArgForCommandLines);
-			if (SC.StageTargetPlatform.PlatformType == UnrealTargetPlatform.Mac || SC.StageTargetPlatform.PlatformType == UnrealTargetPlatform.Win64 || SC.StageTargetPlatform.PlatformType == UnrealTargetPlatform.Linux)
-			{
-				ProjectFile = "";
+				String ProjectFile = String.Format("{0} ", SC.ProjectArgForCommandLines);
+				if (SC.StageTargetPlatform.PlatformType == UnrealTargetPlatform.Mac || SC.StageTargetPlatform.PlatformType == UnrealTargetPlatform.Win64 || SC.StageTargetPlatform.PlatformType == UnrealTargetPlatform.Linux)
+				{
+					ProjectFile = "";
+				}
+				DirectoryReference.CreateDirectory(GetIntermediateCommandlineDir(SC));
+				string CommandLine = String.Format("{0} {1} {2} {3}\n", ProjectFile, Params.StageCommandline.Trim(new char[] { '\"' }), Params.RunCommandline.Trim(new char[] { '\"' }), FileHostParams).Trim();
+				if (Params.IterativeDeploy)
+				{
+					CommandLine += " -iterative";
+				}
+				File.WriteAllText(IntermediateCmdLineFile.FullName, CommandLine);
 			}
-			DirectoryReference.CreateDirectory(GetIntermediateCommandlineDir(SC));
-			string CommandLine = String.Format("{0} {1} {2} {3}\n", ProjectFile, Params.StageCommandline.Trim(new char[] { '\"' }), Params.RunCommandline.Trim(new char[] { '\"' }), FileHostParams).Trim();
-			if (Params.IterativeDeploy)
+			else
 			{
-				CommandLine += " -iterative";
+				String ProjectFile = String.Format("{0} ", SC.ProjectArgForCommandLines);
+				if (SC.StageTargetPlatform.PlatformType == UnrealTargetPlatform.Mac || SC.StageTargetPlatform.PlatformType == UnrealTargetPlatform.Win64 || SC.StageTargetPlatform.PlatformType == UnrealTargetPlatform.Linux)
+				{
+					ProjectFile = "";
+				}
+				DirectoryReference.CreateDirectory(GetIntermediateCommandlineDir(SC));
+				File.WriteAllText(IntermediateCmdLineFile.FullName, ProjectFile);
 			}
-			File.WriteAllText(IntermediateCmdLineFile.FullName, CommandLine);
 		}
-		else
+
+		private static void WriteStageCommandline(ProjectParams Params, DeploymentContext SC)
 		{
-			String ProjectFile = String.Format("{0} ", SC.ProjectArgForCommandLines);
-			if (SC.StageTargetPlatform.PlatformType == UnrealTargetPlatform.Mac || SC.StageTargetPlatform.PlatformType == UnrealTargetPlatform.Win64 || SC.StageTargetPlatform.PlatformType == UnrealTargetPlatform.Linux)
-			{
-				ProjectFile = "";
-			}
-			DirectoryReference.CreateDirectory(GetIntermediateCommandlineDir(SC));
-			File.WriteAllText(IntermediateCmdLineFile.FullName, ProjectFile);
+			// always delete the existing commandline text file, so it doesn't reuse an old one
+			FileReference IntermediateCmdLineFile = FileReference.Combine(GetIntermediateCommandlineDir(SC), "UECommandLine.txt");
+			WriteStageCommandline(IntermediateCmdLineFile, Params, SC);
 		}
-	}
-
-	private static void WriteStageCommandline(ProjectParams Params, DeploymentContext SC)
-	{
-		// always delete the existing commandline text file, so it doesn't reuse an old one
-		FileReference IntermediateCmdLineFile = FileReference.Combine(GetIntermediateCommandlineDir(SC), "UECommandLine.txt");
-		WriteStageCommandline(IntermediateCmdLineFile, Params, SC);
-	}
 
-	private static Dictionary<StagedFileReference, string> ReadDeployedManifest(ProjectParams Params, DeploymentContext SC, List<string> ManifestList)
-	{
-		Dictionary<StagedFileReference, string> DeployedFiles = new Dictionary<StagedFileReference, string>();
-		HashSet<StagedFileReference> CRCFiles = SC.StageTargetPlatform.GetFilesForCRCCheck();
-
-		// read the manifest
-		bool bContinueSearch = true;
-		foreach (string Manifest in ManifestList)
+		private static Dictionary<StagedFileReference, string> ReadDeployedManifest(ProjectParams Params, DeploymentContext SC, List<string> ManifestList)
 		{
-			int FilesAdded = 0;
-			if (bContinueSearch)
+			Dictionary<StagedFileReference, string> DeployedFiles = new Dictionary<StagedFileReference, string>();
+			HashSet<StagedFileReference> CRCFiles = SC.StageTargetPlatform.GetFilesForCRCCheck();
+
+			// read the manifest
+			bool bContinueSearch = true;
+			foreach (string Manifest in ManifestList)
 			{
-				string Data = File.ReadAllText(Manifest);
-				string[] Lines = Data.Split('\n');
-				foreach (string Line in Lines)
+				int FilesAdded = 0;
+				if (bContinueSearch)
 				{
-					string[] Pair = Line.Split('\t');
-					if (Pair.Length > 1)
+					string Data = File.ReadAllText(Manifest);
+					string[] Lines = Data.Split('\n');
+					foreach (string Line in Lines)
 					{
-						StagedFileReference Filename = new StagedFileReference(Pair[0]);
-						string TimeStamp = Pair[1];
-						FilesAdded++;
-						if (DeployedFiles.ContainsKey(Filename))
+						string[] Pair = Line.Split('\t');
+						if (Pair.Length > 1)
 						{
-							if ((CRCFiles.Contains(Filename) && DeployedFiles[Filename] != TimeStamp) || (!CRCFiles.Contains(Filename) && DateTime.Parse(DeployedFiles[Filename]) > DateTime.Parse(TimeStamp)))
+							StagedFileReference Filename = new StagedFileReference(Pair[0]);
+							string TimeStamp = Pair[1];
+							FilesAdded++;
+							if (DeployedFiles.ContainsKey(Filename))
 							{
-								DeployedFiles[Filename] = TimeStamp;
+								if ((CRCFiles.Contains(Filename) && DeployedFiles[Filename] != TimeStamp) || (!CRCFiles.Contains(Filename) && DateTime.Parse(DeployedFiles[Filename]) > DateTime.Parse(TimeStamp)))
+								{
+									DeployedFiles[Filename] = TimeStamp;
+								}
+							}
+							else
+							{
+								DeployedFiles.Add(Filename, TimeStamp);
 							}
-						}
-						else
-						{
-							DeployedFiles.Add(Filename, TimeStamp);
 						}
 					}
 				}
-			}
-			File.Delete(Manifest);
+				File.Delete(Manifest);
 
-			if (FilesAdded == 0 && bContinueSearch)
-			{
-				// no files have been deployed at all to this guy, so remove all previously added files and exit the loop as this means we need to deploy everything
-				DeployedFiles.Clear();
-				bContinueSearch = false;
+				if (FilesAdded == 0 && bContinueSearch)
+				{
+					// no files have been deployed at all to this guy, so remove all previously added files and exit the loop as this means we need to deploy everything
+					DeployedFiles.Clear();
+					bContinueSearch = false;
+				}
 			}
-		}
 
-		return DeployedFiles;
-	}
-
-	protected static Dictionary<StagedFileReference, string> ReadStagedManifest(ProjectParams Params, DeploymentContext SC, string Manifest)
-	{
-		Dictionary<StagedFileReference, string> StagedFiles = new Dictionary<StagedFileReference, string>();
-		HashSet<StagedFileReference> CRCFiles = SC.StageTargetPlatform.GetFilesForCRCCheck();
+			return DeployedFiles;
+		}
 
-		// get the staged manifest from staged directory
-		FileReference ManifestFile = FileReference.Combine(SC.StageDirectory, Manifest);
-		if (FileReference.Exists(ManifestFile))
+		protected static Dictionary<StagedFileReference, string> ReadStagedManifest(ProjectParams Params, DeploymentContext SC, string Manifest)
 		{
-			string[] Lines = FileReference.ReadAllLines(ManifestFile);
-			foreach (string Line in Lines)
+			Dictionary<StagedFileReference, string> StagedFiles = new Dictionary<StagedFileReference, string>();
+			HashSet<StagedFileReference> CRCFiles = SC.StageTargetPlatform.GetFilesForCRCCheck();
+
+			// get the staged manifest from staged directory
+			FileReference ManifestFile = FileReference.Combine(SC.StageDirectory, Manifest);
+			if (FileReference.Exists(ManifestFile))
 			{
-				string[] Pair = Line.Split('\t');
-				if (Pair.Length > 1)
+				string[] Lines = FileReference.ReadAllLines(ManifestFile);
+				foreach (string Line in Lines)
 				{
-					StagedFileReference Filename = new StagedFileReference(Pair[0]);
-					string TimeStamp = Pair[1];
-					if (!StagedFiles.ContainsKey(Filename))
-					{
-						StagedFiles.Add(Filename, TimeStamp);
-					}
-					else if ((CRCFiles.Contains(Filename) && StagedFiles[Filename] != TimeStamp) || (!CRCFiles.Contains(Filename) && DateTime.Parse(StagedFiles[Filename]) > DateTime.Parse(TimeStamp)))
+					string[] Pair = Line.Split('\t');
+					if (Pair.Length > 1)
 					{
-						StagedFiles[Filename] = TimeStamp;
+						StagedFileReference Filename = new StagedFileReference(Pair[0]);
+						string TimeStamp = Pair[1];
+						if (!StagedFiles.ContainsKey(Filename))
+						{
+							StagedFiles.Add(Filename, TimeStamp);
+						}
+						else if ((CRCFiles.Contains(Filename) && StagedFiles[Filename] != TimeStamp) || (!CRCFiles.Contains(Filename) && DateTime.Parse(StagedFiles[Filename]) > DateTime.Parse(TimeStamp)))
+						{
+							StagedFiles[Filename] = TimeStamp;
+						}
 					}
 				}
 			}
+			return StagedFiles;
 		}
-		return StagedFiles;
-	}
-
-	protected static void WriteObsoleteManifest(ProjectParams Params, DeploymentContext SC, Dictionary<StagedFileReference, string> DeployedFiles, Dictionary<StagedFileReference, string> StagedFiles, string ObsoleteManifest)
-	{
-		List<StagedFileReference> ObsoleteFiles = new List<StagedFileReference>();
 
-		// any file that has been deployed, but is no longer in the staged files is obsolete and should be deleted.
-		foreach (KeyValuePair<StagedFileReference, string> File in DeployedFiles)
+		protected static void WriteObsoleteManifest(ProjectParams Params, DeploymentContext SC, Dictionary<StagedFileReference, string> DeployedFiles, Dictionary<StagedFileReference, string> StagedFiles, string ObsoleteManifest)
 		{
-			if (!StagedFiles.ContainsKey(File.Key))
+			List<StagedFileReference> ObsoleteFiles = new List<StagedFileReference>();
+
+			// any file that has been deployed, but is no longer in the staged files is obsolete and should be deleted.
+			foreach (KeyValuePair<StagedFileReference, string> File in DeployedFiles)
 			{
-				ObsoleteFiles.Add(File.Key);
+				if (!StagedFiles.ContainsKey(File.Key))
+				{
+					ObsoleteFiles.Add(File.Key);
+				}
 			}
-		}
 
-		// write out to the deltamanifest.json
-		using (var Writer = File.CreateText(ObsoleteManifest))
-		{
-			foreach (StagedFileReference ObsoleteFile in ObsoleteFiles)
+			// write out to the deltamanifest.json
+			using (var Writer = File.CreateText(ObsoleteManifest))
 			{
-				Writer.WriteLine(ObsoleteFile);
+				foreach (StagedFileReference ObsoleteFile in ObsoleteFiles)
+				{
+					Writer.WriteLine(ObsoleteFile);
+				}
 			}
 		}
-	}
 
-	protected static void WriteDeltaManifest(ProjectParams Params, DeploymentContext SC, Dictionary<StagedFileReference, string> DeployedFiles, Dictionary<StagedFileReference, string> StagedFiles, string DeltaManifest)
-	{
-		HashSet<StagedFileReference> CRCFiles = SC.StageTargetPlatform.GetFilesForCRCCheck();
-		List<string> DeltaFiles = new List<string>();
-		foreach (KeyValuePair<StagedFileReference, string> File in StagedFiles)
+		protected static void WriteDeltaManifest(ProjectParams Params, DeploymentContext SC, Dictionary<StagedFileReference, string> DeployedFiles, Dictionary<StagedFileReference, string> StagedFiles, string DeltaManifest)
 		{
-			bool bNeedsDeploy = true;
-			if (DeployedFiles.ContainsKey(File.Key))
+			HashSet<StagedFileReference> CRCFiles = SC.StageTargetPlatform.GetFilesForCRCCheck();
+			List<string> DeltaFiles = new List<string>();
+			foreach (KeyValuePair<StagedFileReference, string> File in StagedFiles)
 			{
-				if (CRCFiles.Contains(File.Key))
+				bool bNeedsDeploy = true;
+				if (DeployedFiles.ContainsKey(File.Key))
 				{
-					bNeedsDeploy = (File.Value != DeployedFiles[File.Key]);
+					if (CRCFiles.Contains(File.Key))
+					{
+						bNeedsDeploy = (File.Value != DeployedFiles[File.Key]);
+					}
+					else
+					{
+						DateTime Staged = DateTime.Parse(File.Value);
+						DateTime Deployed = DateTime.Parse(DeployedFiles[File.Key]);
+						bNeedsDeploy = (Staged > Deployed);
+					}
 				}
-				else
+
+				if (bNeedsDeploy)
 				{
-					DateTime Staged = DateTime.Parse(File.Value);
-					DateTime Deployed = DateTime.Parse(DeployedFiles[File.Key]);
-					bNeedsDeploy = (Staged > Deployed);
+					DeltaFiles.Add(File.Key.Name);
 				}
 			}
 
-			if (bNeedsDeploy)
+			// add the manifest
+			if (!DeltaManifest.Contains("NonUFS"))
 			{
-				DeltaFiles.Add(File.Key.Name);
+				DeltaFiles.Add(SC.GetNonUFSDeployedManifestFileName(null));
+				DeltaFiles.Add(SC.GetUFSDeployedManifestFileName(null));
 			}
-		}
 
-		// add the manifest
-		if (!DeltaManifest.Contains("NonUFS"))
-		{
-			DeltaFiles.Add(SC.GetNonUFSDeployedManifestFileName(null));
-			DeltaFiles.Add(SC.GetUFSDeployedManifestFileName(null));
-		}
+			// TODO: determine files which need to be removed
 
-		// TODO: determine files which need to be removed
-
-		// write out to the deltamanifest.json
-		FileReference ManifestFile = FileReference.Combine(SC.StageDirectory, DeltaManifest);
-		FileReference.WriteAllLines(ManifestFile, DeltaFiles);
-	}
+			// write out to the deltamanifest.json
+			FileReference ManifestFile = FileReference.Combine(SC.StageDirectory, DeltaManifest);
+			FileReference.WriteAllLines(ManifestFile, DeltaFiles);
+		}
 
-	//@todo move this
-	public static List<DeploymentContext> CreateDeploymentContext(ProjectParams Params, bool InDedicatedServer, bool DoCleanStage = false)
-	{
-		ParamList<string> ListToProcess = InDedicatedServer && (Params.Cook || Params.CookOnTheFly) ? Params.ServerCookedTargets : Params.ClientCookedTargets;
-		var ConfigsToProcess = InDedicatedServer && (Params.Cook || Params.CookOnTheFly) ? Params.ServerConfigsToBuild : Params.ClientConfigsToBuild;
+		//@todo move this
+		public static List<DeploymentContext> CreateDeploymentContext(ProjectParams Params, bool InDedicatedServer, bool DoCleanStage = false)
+		{
+			ParamList<string> ListToProcess = InDedicatedServer && (Params.Cook || Params.CookOnTheFly) ? Params.ServerCookedTargets : Params.ClientCookedTargets;
+			var ConfigsToProcess = InDedicatedServer && (Params.Cook || Params.CookOnTheFly) ? Params.ServerConfigsToBuild : Params.ClientConfigsToBuild;
 
-		List<Tuple<string, UnrealTargetConfiguration>> TargetAndConfigPairs = new List<Tuple<string, UnrealTargetConfiguration>>();
+			List<Tuple<string, UnrealTargetConfiguration>> TargetAndConfigPairs = new List<Tuple<string, UnrealTargetConfiguration>>();
 
-		foreach (var Target in ListToProcess)
-		{
-			// If we are staging a client and have been asked to include editor targets, we currently only want to
-			// include a single Development editor target. Ideally, we should have shipping editor configs and then
-			// just include the requested configs for all targets
-			if (Params.EditorTargets.Contains(Target) || Params.ProgramTargets.Contains(Target))
+			foreach (var Target in ListToProcess)
 			{
-				TargetAndConfigPairs.Add(new Tuple<string, UnrealTargetConfiguration>(Target, UnrealTargetConfiguration.Development));
-			}
-			else
-			{
-				foreach (var Config in ConfigsToProcess)
+				// If we are staging a client and have been asked to include editor targets, we currently only want to
+				// include a single Development editor target. Ideally, we should have shipping editor configs and then
+				// just include the requested configs for all targets
+				if (Params.EditorTargets.Contains(Target) || Params.ProgramTargets.Contains(Target))
+				{
+					TargetAndConfigPairs.Add(new Tuple<string, UnrealTargetConfiguration>(Target, UnrealTargetConfiguration.Development));
+				}
+				else
 				{
-					TargetAndConfigPairs.Add(new Tuple<string, UnrealTargetConfiguration>(Target, Config));
+					foreach (var Config in ConfigsToProcess)
+					{
+						TargetAndConfigPairs.Add(new Tuple<string, UnrealTargetConfiguration>(Target, Config));
+					}
 				}
 			}
-		}
-
-		List<TargetPlatformDescriptor> PlatformsToStage = Params.ClientTargetPlatforms;
-		if (InDedicatedServer && (Params.Cook || Params.CookOnTheFly))
-		{
-			PlatformsToStage = Params.ServerTargetPlatforms;
-		}
-
-		List<DeploymentContext> DeploymentContexts = new List<DeploymentContext>();
-		foreach (var StagePlatform in PlatformsToStage)
-		{
-			// Get the platform to get cooked data from, may differ from the stage platform
-			TargetPlatformDescriptor CookedDataPlatform = Params.GetCookedDataPlatformForClientTarget(StagePlatform);
 
+			List<TargetPlatformDescriptor> PlatformsToStage = Params.ClientTargetPlatforms;
 			if (InDedicatedServer && (Params.Cook || Params.CookOnTheFly))
 			{
-				CookedDataPlatform = Params.GetCookedDataPlatformForServerTarget(StagePlatform);
+				PlatformsToStage = Params.ServerTargetPlatforms;
 			}
 
-			List<string> ExecutablesToStage = new List<string>();
-
-			string PlatformName = StagePlatform.ToString();
-			string StageArchitecture = !String.IsNullOrEmpty(Params.SpecifiedArchitecture) ? Params.SpecifiedArchitecture : "";
-			foreach (var TargetAndConfig in TargetAndConfigPairs)
+			List<DeploymentContext> DeploymentContexts = new List<DeploymentContext>();
+			foreach (var StagePlatform in PlatformsToStage)
 			{
-				string Target = TargetAndConfig.Item1;
-				UnrealTargetConfiguration Config = TargetAndConfig.Item2;
-				string Exe = Target;
-				if (Config != UnrealTargetConfiguration.Development)
+				// Get the platform to get cooked data from, may differ from the stage platform
+				TargetPlatformDescriptor CookedDataPlatform = Params.GetCookedDataPlatformForClientTarget(StagePlatform);
+
+				if (InDedicatedServer && (Params.Cook || Params.CookOnTheFly))
 				{
-					Exe = Target + "-" + PlatformName + "-" + Config.ToString() + StageArchitecture;
+					CookedDataPlatform = Params.GetCookedDataPlatformForServerTarget(StagePlatform);
 				}
-				ExecutablesToStage.Add(Exe);
-			}
-
-			string StageDirectory = ((ShouldCreatePak(Params) || (Params.Stage)) || !String.IsNullOrEmpty(Params.StageDirectoryParam)) ? Params.BaseStageDirectory : "";
-			string ArchiveDirectory = (Params.Archive || !String.IsNullOrEmpty(Params.ArchiveDirectoryParam)) ? Params.BaseArchiveDirectory : "";
-			DirectoryReference EngineDir = DirectoryReference.Combine(Unreal.RootDirectory, "Engine");
-			DirectoryReference ProjectDir = DirectoryReference.FromFile(Params.RawProjectPath);
 
-			List<StageTarget> TargetsToStage = new List<StageTarget>();
-			foreach (var TargetAndConfig in TargetAndConfigPairs)
-			{
-				string Target = TargetAndConfig.Item1;
-				UnrealTargetConfiguration Config = TargetAndConfig.Item2;
-				DirectoryReference ReceiptBaseDir = Params.IsCodeBasedProject ? ProjectDir : EngineDir;
+				List<string> ExecutablesToStage = new List<string>();
 
-				Platform PlatformInstance = Platform.Platforms[StagePlatform];
-				UnrealTargetPlatform[] SubPlatformsToStage = PlatformInstance.GetStagePlatforms();
+				string PlatformName = StagePlatform.ToString();
+				string StageArchitecture = !String.IsNullOrEmpty(Params.SpecifiedArchitecture) ? Params.SpecifiedArchitecture : "";
+				foreach (var TargetAndConfig in TargetAndConfigPairs)
+				{
+					string Target = TargetAndConfig.Item1;
+					UnrealTargetConfiguration Config = TargetAndConfig.Item2;
+					string Exe = Target;
+					if (Config != UnrealTargetConfiguration.Development)
+					{
+						Exe = Target + "-" + PlatformName + "-" + Config.ToString() + StageArchitecture;
+					}
+					ExecutablesToStage.Add(Exe);
+				}
 
-				// if we are attempting to gathering multiple platforms, the files aren't required
-				bool bJustPackaging = Params.SkipStage && Params.Package;
-				bool bIsIterativeSharedCooking = Params.HasIterateSharedCookedBuild;
-				bool bRequireStagedFilesToExist = SubPlatformsToStage.Length == 1 && PlatformsToStage.Count == 1 && !bJustPackaging && !bIsIterativeSharedCooking && !Params.HasDLCName;
+				string StageDirectory = ((ShouldCreatePak(Params) || (Params.Stage)) || !String.IsNullOrEmpty(Params.StageDirectoryParam)) ? Params.BaseStageDirectory : "";
+				string ArchiveDirectory = (Params.Archive || !String.IsNullOrEmpty(Params.ArchiveDirectoryParam)) ? Params.BaseArchiveDirectory : "";
+				DirectoryReference EngineDir = DirectoryReference.Combine(Unreal.RootDirectory, "Engine");
+				DirectoryReference ProjectDir = DirectoryReference.FromFile(Params.RawProjectPath);
 
-				foreach (UnrealTargetPlatform ReceiptPlatform in SubPlatformsToStage)
+				List<StageTarget> TargetsToStage = new List<StageTarget>();
+				foreach (var TargetAndConfig in TargetAndConfigPairs)
 				{
-					string Architecture = Params.SpecifiedArchitecture;
-					if (string.IsNullOrEmpty(Architecture))
+					string Target = TargetAndConfig.Item1;
+					UnrealTargetConfiguration Config = TargetAndConfig.Item2;
+					DirectoryReference ReceiptBaseDir = Params.IsCodeBasedProject ? ProjectDir : EngineDir;
+
+					Platform PlatformInstance = Platform.Platforms[StagePlatform];
+					UnrealTargetPlatform[] SubPlatformsToStage = PlatformInstance.GetStagePlatforms();
+
+					// if we are attempting to gathering multiple platforms, the files aren't required
+					bool bJustPackaging = Params.SkipStage && Params.Package;
+					bool bIsIterativeSharedCooking = Params.HasIterateSharedCookedBuild;
+					bool bRequireStagedFilesToExist = SubPlatformsToStage.Length == 1 && PlatformsToStage.Count == 1 && !bJustPackaging && !bIsIterativeSharedCooking && !Params.HasDLCName;
+
+					foreach (UnrealTargetPlatform ReceiptPlatform in SubPlatformsToStage)
 					{
-						Architecture = "";
-						if (PlatformExports.IsPlatformAvailable(ReceiptPlatform))
+						string Architecture = Params.SpecifiedArchitecture;
+						if (string.IsNullOrEmpty(Architecture))
 						{
-							Architecture = PlatformExports.GetDefaultArchitecture(ReceiptPlatform, Params.RawProjectPath);
+							Architecture = "";
+							if (PlatformExports.IsPlatformAvailable(ReceiptPlatform))
+							{
+								Architecture = PlatformExports.GetDefaultArchitecture(ReceiptPlatform, Params.RawProjectPath);
+							}
 						}
-					}
 
-					if (Params.IterateSharedBuildUsePrecompiledExe)
-					{
-						continue;
-					}
+						if (Params.IterateSharedBuildUsePrecompiledExe)
+						{
+							continue;
+						}
 
-					FileReference ReceiptFileName = TargetReceipt.GetDefaultPath(ReceiptBaseDir, Target, ReceiptPlatform, Config, Architecture);
-					if (!FileReference.Exists(ReceiptFileName))
-					{
-						if (bRequireStagedFilesToExist)
+						FileReference ReceiptFileName = TargetReceipt.GetDefaultPath(ReceiptBaseDir, Target, ReceiptPlatform, Config, Architecture);
+						if (!FileReference.Exists(ReceiptFileName))
 						{
-							// if we aren't collecting multiple platforms, then it is expected to exist
-							throw new AutomationException(ExitCode.Error_MissingExecutable, "Stage Failed. Missing receipt '{0}'. Check that this target has been built.", ReceiptFileName);
+							if (bRequireStagedFilesToExist)
+							{
+								// if we aren't collecting multiple platforms, then it is expected to exist
+								throw new AutomationException(ExitCode.Error_MissingExecutable, "Stage Failed. Missing receipt '{0}'. Check that this target has been built.", ReceiptFileName);
+							}
+							else
+							{
+								// if it's multiple platforms, then allow missing receipts
+								continue;
+							}
+
 						}
-						else
+
+						// Read the receipt for this target
+						TargetReceipt Receipt;
+						if (!TargetReceipt.TryRead(ReceiptFileName, out Receipt))
 						{
-							// if it's multiple platforms, then allow missing receipts
-							continue;
+							throw new AutomationException("Missing or invalid target receipt ({0})", ReceiptFileName);
 						}
 
+						// Convert the paths to absolute
+						TargetsToStage.Add(new StageTarget { Receipt = Receipt, RequireFilesExist = bRequireStagedFilesToExist });
 					}
+				}
 
-					// Read the receipt for this target
-					TargetReceipt Receipt;
-					if (!TargetReceipt.TryRead(ReceiptFileName, out Receipt))
-					{
-						throw new AutomationException("Missing or invalid target receipt ({0})", ReceiptFileName);
-					}
+				//@todo should pull StageExecutables from somewhere else if not cooked
+				var SC = new DeploymentContext(Params.RawProjectPath, Unreal.RootDirectory,
+					String.IsNullOrEmpty(StageDirectory) ? null : new DirectoryReference(StageDirectory),
+					String.IsNullOrEmpty(ArchiveDirectory) ? null : new DirectoryReference(ArchiveDirectory),
+					Platform.Platforms[CookedDataPlatform],
+					Platform.Platforms[StagePlatform],
+					ConfigsToProcess,
+					TargetsToStage,
+					ExecutablesToStage,
+					InDedicatedServer,
+					Params.Cook || Params.CookOnTheFly,
+					Params.CrashReporter,
+					Params.Stage,
+					Params.CookOnTheFly,
+					Params.Archive,
+					Params.IsProgramTarget,
+					Params.Client,
+					Params.Manifests,
+					Params.SeparateDebugInfo,
+					Params.HasDLCName
+					);
+				LogDeploymentContext(SC);
 
-					// Convert the paths to absolute
-					TargetsToStage.Add(new StageTarget { Receipt = Receipt, RequireFilesExist = bRequireStagedFilesToExist });
+				if (Params.PreModifyDeploymentContextCallback != null)
+				{
+					Params.PreModifyDeploymentContextCallback(Params, SC);
 				}
-			}
 
-			//@todo should pull StageExecutables from somewhere else if not cooked
-			var SC = new DeploymentContext(Params.RawProjectPath, Unreal.RootDirectory,
-				String.IsNullOrEmpty(StageDirectory) ? null : new DirectoryReference(StageDirectory),
-				String.IsNullOrEmpty(ArchiveDirectory) ? null : new DirectoryReference(ArchiveDirectory),
-				Platform.Platforms[CookedDataPlatform],
-				Platform.Platforms[StagePlatform],
-				ConfigsToProcess,
-				TargetsToStage,
-				ExecutablesToStage,
-				InDedicatedServer,
-				Params.Cook || Params.CookOnTheFly,
-				Params.CrashReporter,
-				Params.Stage,
-				Params.CookOnTheFly,
-				Params.Archive,
-				Params.IsProgramTarget,
-				Params.Client,
-				Params.Manifests,
-				Params.SeparateDebugInfo,
-				Params.HasDLCName
-				);
-			LogDeploymentContext(SC);
+				// If we're a derived platform make sure we're at the end, otherwise make sure we're at the front
 
-			if (Params.PreModifyDeploymentContextCallback != null)
-			{
-				Params.PreModifyDeploymentContextCallback(Params, SC);
+				if (!CookedDataPlatform.Equals(StagePlatform))
+				{
+					DeploymentContexts.Add(SC);
+				}
+				else
+				{
+					DeploymentContexts.Insert(0, SC);
+				}
 			}
 
-			// If we're a derived platform make sure we're at the end, otherwise make sure we're at the front
-
-			if (!CookedDataPlatform.Equals(StagePlatform))
-			{
-				DeploymentContexts.Add(SC);
-			}
-			else
-			{
-				DeploymentContexts.Insert(0, SC);
-			}
+			return DeploymentContexts;
 		}
 
-		return DeploymentContexts;
-	}
-
-	public static void CopyBuildToStagingDirectory(ProjectParams Params)
-	{
-		if (ShouldCreatePak(Params) || (Params.Stage && !Params.SkipStage))
+		public static void CopyBuildToStagingDirectory(ProjectParams Params)
 		{
-			Params.ValidateAndLog();
-
-			LogInformation("********** STAGE COMMAND STARTED **********");
-
-			if (!Params.NoClient)
+			if (ShouldCreatePak(Params) || (Params.Stage && !Params.SkipStage))
 			{
-				var DeployContextList = CreateDeploymentContext(Params, false, true);
+				Params.ValidateAndLog();
 
-				// clean the staging directories first
-				foreach (var SC in DeployContextList)
-				{
-					SC.StageTargetPlatform.PreStage(Params, SC);
+				LogInformation("********** STAGE COMMAND STARTED **********");
 
-					// write out the commandline file now so it can go into the manifest
-					WriteStageCommandline(Params, SC);
-					CreateStagingManifest(Params, SC);
-					CleanStagingDirectory(Params, SC);
-				}
-				foreach (var SC in DeployContextList)
+				if (!Params.NoClient)
 				{
-					ApplyStagingManifest(Params, SC);
+					var DeployContextList = CreateDeploymentContext(Params, false, true);
 
-					if (Params.Deploy && SC.StageTargetPlatform.RequiresManifestFiles)
+					// clean the staging directories first
+					foreach (var SC in DeployContextList)
 					{
-						List<string> UFSManifests;
-						List<string> NonUFSManifests;
+						SC.StageTargetPlatform.PreStage(Params, SC);
 
-						// get the staged file data
-						Dictionary<StagedFileReference, string> StagedUFSFiles = ReadStagedManifest(Params, SC, SC.GetUFSDeployedManifestFileName(null));
-						Dictionary<StagedFileReference, string> StagedNonUFSFiles = ReadStagedManifest(Params, SC, SC.GetNonUFSDeployedManifestFileName(null));
+						// write out the commandline file now so it can go into the manifest
+						WriteStageCommandline(Params, SC);
+						CreateStagingManifest(Params, SC);
+						CleanStagingDirectory(Params, SC);
+					}
+					foreach (var SC in DeployContextList)
+					{
+						ApplyStagingManifest(Params, SC);
 
-						foreach (var DeviceName in Params.DeviceNames)
+						if (Params.Deploy && SC.StageTargetPlatform.RequiresManifestFiles)
 						{
-							// get the deployed file data
-							Dictionary<StagedFileReference, string> DeployedUFSFiles = new Dictionary<StagedFileReference, string>();
-							Dictionary<StagedFileReference, string> DeployedNonUFSFiles = new Dictionary<StagedFileReference, string>();
+							List<string> UFSManifests;
+							List<string> NonUFSManifests;
 
-							if (SC.StageTargetPlatform.RetrieveDeployedManifests(Params, SC, DeviceName, out UFSManifests, out NonUFSManifests))
+							// get the staged file data
+							Dictionary<StagedFileReference, string> StagedUFSFiles = ReadStagedManifest(Params, SC, SC.GetUFSDeployedManifestFileName(null));
+							Dictionary<StagedFileReference, string> StagedNonUFSFiles = ReadStagedManifest(Params, SC, SC.GetNonUFSDeployedManifestFileName(null));
+
+							foreach (var DeviceName in Params.DeviceNames)
 							{
-								DeployedUFSFiles = ReadDeployedManifest(Params, SC, UFSManifests);
-								DeployedNonUFSFiles = ReadDeployedManifest(Params, SC, NonUFSManifests);
-							}
+								// get the deployed file data
+								Dictionary<StagedFileReference, string> DeployedUFSFiles = new Dictionary<StagedFileReference, string>();
+								Dictionary<StagedFileReference, string> DeployedNonUFSFiles = new Dictionary<StagedFileReference, string>();
+
+								if (SC.StageTargetPlatform.RetrieveDeployedManifests(Params, SC, DeviceName, out UFSManifests, out NonUFSManifests))
+								{
+									DeployedUFSFiles = ReadDeployedManifest(Params, SC, UFSManifests);
+									DeployedNonUFSFiles = ReadDeployedManifest(Params, SC, NonUFSManifests);
+								}
 
-							WriteObsoleteManifest(Params, SC, DeployedUFSFiles, StagedUFSFiles, SC.GetUFSDeploymentObsoletePath(DeviceName));
-							WriteObsoleteManifest(Params, SC, DeployedNonUFSFiles, StagedNonUFSFiles, SC.GetNonUFSDeploymentObsoletePath(DeviceName));
+								WriteObsoleteManifest(Params, SC, DeployedUFSFiles, StagedUFSFiles, SC.GetUFSDeploymentObsoletePath(DeviceName));
+								WriteObsoleteManifest(Params, SC, DeployedNonUFSFiles, StagedNonUFSFiles, SC.GetNonUFSDeploymentObsoletePath(DeviceName));
 
-							if (Params.IterativeDeploy)
-							{
-								// write out the delta file data
-								WriteDeltaManifest(Params, SC, DeployedUFSFiles, StagedUFSFiles, SC.GetUFSDeploymentDeltaPath(DeviceName));
-								WriteDeltaManifest(Params, SC, DeployedNonUFSFiles, StagedNonUFSFiles, SC.GetNonUFSDeploymentDeltaPath(DeviceName));
+								if (Params.IterativeDeploy)
+								{
+									// write out the delta file data
+									WriteDeltaManifest(Params, SC, DeployedUFSFiles, StagedUFSFiles, SC.GetUFSDeploymentDeltaPath(DeviceName));
+									WriteDeltaManifest(Params, SC, DeployedNonUFSFiles, StagedNonUFSFiles, SC.GetNonUFSDeploymentDeltaPath(DeviceName));
+								}
 							}
 						}
-					}
 
-					if (Params.bCodeSign && !Params.SkipStage)
-					{
-						SC.StageTargetPlatform.SignExecutables(SC, Params);
+						if (Params.bCodeSign && !Params.SkipStage)
+						{
+							SC.StageTargetPlatform.SignExecutables(SC, Params);
+						}
 					}
 				}
-			}
 
-			if (Params.DedicatedServer)
-			{
-				var DeployContextList = CreateDeploymentContext(Params, true, true);
-				foreach (var SC in DeployContextList)
+				if (Params.DedicatedServer)
 				{
-					CreateStagingManifest(Params, SC);
-					CleanStagingDirectory(Params, SC);
-				}
+					var DeployContextList = CreateDeploymentContext(Params, true, true);
+					foreach (var SC in DeployContextList)
+					{
+						CreateStagingManifest(Params, SC);
+						CleanStagingDirectory(Params, SC);
+					}
 
-				foreach (var SC in DeployContextList)
-				{
-					ApplyStagingManifest(Params, SC);
+					foreach (var SC in DeployContextList)
+					{
+						ApplyStagingManifest(Params, SC);
+					}
 				}
+				LogInformation("********** STAGE COMMAND COMPLETED **********");
 			}
-			LogInformation("********** STAGE COMMAND COMPLETED **********");
 		}
 	}
 }
