AutomationTool:

diff --git a/Engine/Source/Programs/AutomationTool/Scripts/DeployCommand.Automation.cs b/Engine/Source/Programs/AutomationTool/Scripts/DeployCommand.Automation.cs
--- a/Engine/Source/Programs/AutomationTool/Scripts/DeployCommand.Automation.cs
+++ b/Engine/Source/Programs/AutomationTool/Scripts/DeployCommand.Automation.cs
@@ -7,65 +7,68 @@ using System.Reflection;
 using AutomationTool;
 using UnrealBuildTool;
 
-public partial class Project : CommandUtils
+namespace AutomationScripts
 {
-	public static void Deploy(ProjectParams Params)
+	public partial class Project : CommandUtils
 	{
-		Params.ValidateAndLog();
-		if (!Params.Deploy)
+		public static void Deploy(ProjectParams Params)
 		{
-			return;
-		}
+			Params.ValidateAndLog();
+			if (!Params.Deploy)
+			{
+				return;
+			}
 
-		LogInformation("********** DEPLOY COMMAND STARTED **********");
-		
-		if (!Params.NoClient)
-		{
-			var DeployContextList = CreateDeploymentContext(Params, false, false);
-			foreach ( var SC in DeployContextList )
+			LogInformation("********** DEPLOY COMMAND STARTED **********");
+
+			if (!Params.NoClient)
 			{
-				if (SC.StageTargetPlatform.DeployViaUFE)
+				var DeployContextList = CreateDeploymentContext(Params, false, false);
+				foreach (var SC in DeployContextList)
 				{
-					string ClientCmdLine = "-run=Deploy ";
-					ClientCmdLine += "-Device=" + string.Join("+", Params.Devices) + " ";
-					ClientCmdLine += "-Targetplatform=" + SC.StageTargetPlatform.PlatformType.ToString() + " ";
-					ClientCmdLine += "-SourceDir=\"" + CombinePaths(Params.BaseStageDirectory, SC.StageTargetPlatform.PlatformType.ToString()) + "\" ";
-					string ClientApp = CombinePaths(CmdEnv.LocalRoot, "Engine/Binaries/Win64/UnrealFrontend.exe");
+					if (SC.StageTargetPlatform.DeployViaUFE)
+					{
+						string ClientCmdLine = "-run=Deploy ";
+						ClientCmdLine += "-Device=" + string.Join("+", Params.Devices) + " ";
+						ClientCmdLine += "-Targetplatform=" + SC.StageTargetPlatform.PlatformType.ToString() + " ";
+						ClientCmdLine += "-SourceDir=\"" + CombinePaths(Params.BaseStageDirectory, SC.StageTargetPlatform.PlatformType.ToString()) + "\" ";
+						string ClientApp = CombinePaths(CmdEnv.LocalRoot, "Engine/Binaries/Win64/UnrealFrontend.exe");
 
-					LogInformation("Deploying via UFE:");
-					LogInformation("\tClientCmdLine: " + ClientCmdLine + "");
+						LogInformation("Deploying via UFE:");
+						LogInformation("\tClientCmdLine: " + ClientCmdLine + "");
 
-					//@todo UAT: Consolidate running of external applications like UFE (See 'RunProjectCommand' for other instances)
-					PushDir(Path.GetDirectoryName(ClientApp));
-					// Always start client process and don't wait for exit.
-					IProcessResult ClientProcess = Run(ClientApp, ClientCmdLine, null, ERunOptions.NoWaitForExit);
-					PopDir();
-					if (ClientProcess != null)
-					{
-						do
+						//@todo UAT: Consolidate running of external applications like UFE (See 'RunProjectCommand' for other instances)
+						PushDir(Path.GetDirectoryName(ClientApp));
+						// Always start client process and don't wait for exit.
+						IProcessResult ClientProcess = Run(ClientApp, ClientCmdLine, null, ERunOptions.NoWaitForExit);
+						PopDir();
+						if (ClientProcess != null)
 						{
-							Thread.Sleep(100);
+							do
+							{
+								Thread.Sleep(100);
+							}
+							while (ClientProcess.HasExited == false);
 						}
-						while (ClientProcess.HasExited == false);
 					}
-				}
-				else
-				{
-					SC.StageTargetPlatform.Deploy(Params, SC);
+					else
+					{
+						SC.StageTargetPlatform.Deploy(Params, SC);
+					}
 				}
 			}
-		}
-		if (Params.DedicatedServer)
-		{
-			ProjectParams ServerParams = new ProjectParams(Params);
-			ServerParams.Devices = new ParamList<string>(ServerParams.ServerDevice);
-			var DeployContextList = CreateDeploymentContext(ServerParams, true, false);
-			foreach ( var SC in DeployContextList )
+			if (Params.DedicatedServer)
 			{
-				SC.StageTargetPlatform.Deploy(ServerParams, SC);
+				ProjectParams ServerParams = new ProjectParams(Params);
+				ServerParams.Devices = new ParamList<string>(ServerParams.ServerDevice);
+				var DeployContextList = CreateDeploymentContext(ServerParams, true, false);
+				foreach (var SC in DeployContextList)
+				{
+					SC.StageTargetPlatform.Deploy(ServerParams, SC);
+				}
 			}
-		}
 
-		LogInformation("********** DEPLOY COMMAND COMPLETED **********");
+			LogInformation("********** DEPLOY COMMAND COMPLETED **********");
+		}
 	}
 }
diff --git a/Engine/Source/Programs/AutomationTool/Scripts/GetFileCommand.Automation.cs b/Engine/Source/Programs/AutomationTool/Scripts/GetFileCommand.Automation.cs
--- a/Engine/Source/Programs/AutomationTool/Scripts/GetFileCommand.Automation.cs
+++ b/Engine/Source/Programs/AutomationTool/Scripts/GetFileCommand.Automation.cs
@@ -9,29 +9,32 @@ using AutomationTool;
 using UnrealBuildTool;
 using EpicGames.Core;
 
-public partial class Project : CommandUtils
+namespace AutomationScripts
 {
-	public static void GetFile(ProjectParams Params)
+	public partial class Project : CommandUtils
 	{
-		Params.ValidateAndLog();
-		if (string.IsNullOrEmpty(Params.GetFile))
+		public static void GetFile(ProjectParams Params)
 		{
-			return;
-		}
+			Params.ValidateAndLog();
+			if (string.IsNullOrEmpty(Params.GetFile))
+			{
+				return;
+			}
 
-		LogInformation("********** GETFILE COMMAND STARTED **********");
+			LogInformation("********** GETFILE COMMAND STARTED **********");
 
-		var FileName = Path.GetFileName(Params.GetFile);
-		var LocalFile = CombinePaths(CmdEnv.EngineSavedFolder, FileName);
+			var FileName = Path.GetFileName(Params.GetFile);
+			var LocalFile = CombinePaths(CmdEnv.EngineSavedFolder, FileName);
 
-		var SC = CreateDeploymentContext(Params, false);
-		if (SC.Count == 0)
-		{
-			throw new AutomationException("Failed to create deployment context");
-		}
+			var SC = CreateDeploymentContext(Params, false);
+			if (SC.Count == 0)
+			{
+				throw new AutomationException("Failed to create deployment context");
+			}
 
-		SC[0].StageTargetPlatform.GetTargetFile(Params.GetFile, LocalFile, Params);
+			SC[0].StageTargetPlatform.GetTargetFile(Params.GetFile, LocalFile, Params);
 
-		LogInformation("********** GETFILE COMMAND COMPLETED **********");
+			LogInformation("********** GETFILE COMMAND COMPLETED **********");
+		}
 	}
 }
diff --git a/Engine/Source/Programs/AutomationTool/Scripts/PackageCommand.Automation.cs b/Engine/Source/Programs/AutomationTool/Scripts/PackageCommand.Automation.cs
--- a/Engine/Source/Programs/AutomationTool/Scripts/PackageCommand.Automation.cs
+++ b/Engine/Source/Programs/AutomationTool/Scripts/PackageCommand.Automation.cs
@@ -7,36 +7,39 @@ using System.Reflection;
 using AutomationTool;
 using UnrealBuildTool;
 
-public partial class Project : CommandUtils
+namespace AutomationScripts
 {
-	public static void Package(ProjectParams Params, int WorkingCL=-1)
+	public partial class Project : CommandUtils
 	{
-		if ((!Params.SkipStage || Params.Package) && !Params.SkipPackage)
+		public static void Package(ProjectParams Params, int WorkingCL = -1)
 		{
-			Params.ValidateAndLog();
-			List<DeploymentContext> DeployContextList = new List<DeploymentContext>();
-			if (!Params.NoClient)
+			if ((!Params.SkipStage || Params.Package) && !Params.SkipPackage)
 			{
-				DeployContextList.AddRange(CreateDeploymentContext(Params, false, false));
-			}
-			if (Params.DedicatedServer)
-			{
-				DeployContextList.AddRange(CreateDeploymentContext(Params, true, false));
-			}
-
-			if (DeployContextList.Count > 0 )
-			{
-				LogInformation("********** PACKAGE COMMAND STARTED **********");
+				Params.ValidateAndLog();
+				List<DeploymentContext> DeployContextList = new List<DeploymentContext>();
+				if (!Params.NoClient)
+				{
+					DeployContextList.AddRange(CreateDeploymentContext(Params, false, false));
+				}
+				if (Params.DedicatedServer)
+				{
+					DeployContextList.AddRange(CreateDeploymentContext(Params, true, false));
+				}
 
-				foreach (var SC in DeployContextList)
+				if (DeployContextList.Count > 0)
 				{
-					if (Params.Package || (SC.StageTargetPlatform.RequiresPackageToDeploy && Params.Deploy))
+					LogInformation("********** PACKAGE COMMAND STARTED **********");
+
+					foreach (var SC in DeployContextList)
 					{
-						SC.StageTargetPlatform.Package(Params, SC, WorkingCL);
+						if (Params.Package || (SC.StageTargetPlatform.RequiresPackageToDeploy && Params.Deploy))
+						{
+							SC.StageTargetPlatform.Package(Params, SC, WorkingCL);
+						}
 					}
-				}
 
-				LogInformation("********** PACKAGE COMMAND COMPLETED **********");
+					LogInformation("********** PACKAGE COMMAND COMPLETED **********");
+				}
 			}
 		}
 	}
diff --git a/Engine/Source/Programs/AutomationTool/Scripts/RunProjectCommand.Automation.cs b/Engine/Source/Programs/AutomationTool/Scripts/RunProjectCommand.Automation.cs
--- a/Engine/Source/Programs/AutomationTool/Scripts/RunProjectCommand.Automation.cs
+++ b/Engine/Source/Programs/AutomationTool/Scripts/RunProjectCommand.Automation.cs
@@ -14,501 +14,504 @@ using UnrealBuildTool;
 using EpicGames.Core;
 using UnrealBuildBase;
 
-/// <summary>
-/// Helper command to run a game.
-/// </summary>
-/// <remarks>
-/// Uses the following command line params:
-/// -cooked
-/// -cookonthefly
-/// -dedicatedserver
-/// -win32
-/// -noclient
-/// -logwindow
-/// </remarks>
-public partial class Project : CommandUtils
+namespace AutomationScripts
 {
 	/// <summary>
-	/// Thread used to read client log file.
+	/// Helper command to run a game.
 	/// </summary>
-	private static Thread ClientLogReaderThread = null;
-
-	/// <summary>
-	/// Process for the server, can be set by the cook command when a cook on the fly server is used
-	/// </summary>
-	public static IProcessResult ServerProcess;
-
-	// debug commands for the engine to crash
-	public static string[] CrashCommands = 
-    {
-        "crash",
-        "CHECK",
-        "GPF",
-        "ASSERT",
-        "ENSURE",
-        "RENDERCRASH",
-        "RENDERCHECK",
-        "RENDERGPF",
-        "THREADCRASH",
-        "THREADCHECK",
-        "THREADGPF",
-    };
-
-	/// <summary>
-	/// For not-installed runs, returns a temp log folder to make sure it doesn't fall into sandbox paths
-	/// </summary>
-	/// <returns></returns>
-	private static string GetLogFolderOutsideOfSandbox()
+	/// <remarks>
+	/// Uses the following command line params:
+	/// -cooked
+	/// -cookonthefly
+	/// -dedicatedserver
+	/// -win32
+	/// -noclient
+	/// -logwindow
+	/// </remarks>
+	public partial class Project : CommandUtils
 	{
-		return Unreal.IsEngineInstalled() ?
-			CmdEnv.LogFolder :
-			CombinePaths(Path.GetTempPath(), CommandUtils.EscapePath(CmdEnv.LocalRoot), "Logs");
-	}
-
-	/// <summary>
-	/// Fot not-installed runs, copies all logs from the temp log folder back to the UAT log folder.
-	/// </summary>
-	private static void CopyLogsBackToLogFolder()
-	{
-		if (!Unreal.IsEngineInstalled())
+		/// <summary>
+		/// Thread used to read client log file.
+		/// </summary>
+		private static Thread ClientLogReaderThread = null;
+
+		/// <summary>
+		/// Process for the server, can be set by the cook command when a cook on the fly server is used
+		/// </summary>
+		public static IProcessResult ServerProcess;
+
+		// debug commands for the engine to crash
+		public static string[] CrashCommands =
 		{
-			var LogFolderOutsideOfSandbox = GetLogFolderOutsideOfSandbox();
-			var TempLogFiles = FindFiles_NoExceptions("*", false, LogFolderOutsideOfSandbox);
-			foreach (var LogFilename in TempLogFiles)
-			{
-				var DestFilename = CombinePaths(CmdEnv.LogFolder, Path.GetFileName(LogFilename));
-				CopyFile_NoExceptions(LogFilename, DestFilename);
-			}
-		}
-	}
-
-	public static void Run(ProjectParams Params)
-	{
-		Params.ValidateAndLog();
-		if (!Params.Run)
+		"crash",
+		"CHECK",
+		"GPF",
+		"ASSERT",
+		"ENSURE",
+		"RENDERCRASH",
+		"RENDERCHECK",
+		"RENDERGPF",
+		"THREADCRASH",
+		"THREADCHECK",
+		"THREADGPF",
+	};
+
+		/// <summary>
+		/// For not-installed runs, returns a temp log folder to make sure it doesn't fall into sandbox paths
+		/// </summary>
+		/// <returns></returns>
+		private static string GetLogFolderOutsideOfSandbox()
 		{
-			return;
+			return Unreal.IsEngineInstalled() ?
+				CmdEnv.LogFolder :
+				CombinePaths(Path.GetTempPath(), CommandUtils.EscapePath(CmdEnv.LocalRoot), "Logs");
 		}
 
-		LogInformation("********** RUN COMMAND STARTED **********");
-
-		var LogFolderOutsideOfSandbox = GetLogFolderOutsideOfSandbox();
-		if (!Unreal.IsEngineInstalled() && ServerProcess == null)
+		/// <summary>
+		/// Fot not-installed runs, copies all logs from the temp log folder back to the UAT log folder.
+		/// </summary>
+		private static void CopyLogsBackToLogFolder()
 		{
-			// In the installed runs, this is the same folder as CmdEnv.LogFolder so delete only in not-installed
-			DeleteDirectory(LogFolderOutsideOfSandbox);
-			CreateDirectory(LogFolderOutsideOfSandbox);
+			if (!Unreal.IsEngineInstalled())
+			{
+				var LogFolderOutsideOfSandbox = GetLogFolderOutsideOfSandbox();
+				var TempLogFiles = FindFiles_NoExceptions("*", false, LogFolderOutsideOfSandbox);
+				foreach (var LogFilename in TempLogFiles)
+				{
+					var DestFilename = CombinePaths(CmdEnv.LogFolder, Path.GetFileName(LogFilename));
+					CopyFile_NoExceptions(LogFilename, DestFilename);
+				}
+			}
 		}
-		var ServerLogFile = CombinePaths(LogFolderOutsideOfSandbox, "Server.log");
-		var ClientLogFile = CombinePaths(LogFolderOutsideOfSandbox, Params.EditorTest ? "Editor.log" : "Client.log");
 
-		try
-		{
-			RunInternal(Params, ServerLogFile, ClientLogFile);
-		}
-		catch
+		public static void Run(ProjectParams Params)
 		{
-			throw;
-		}
-		finally
-		{
-			CopyLogsBackToLogFolder();
-		}
+			Params.ValidateAndLog();
+			if (!Params.Run)
+			{
+				return;
+			}
 
-		LogInformation("********** RUN COMMAND COMPLETED **********");
-	}
+			LogInformation("********** RUN COMMAND STARTED **********");
 
-	private static void RunInternal(ProjectParams Params, string ServerLogFile, string ClientLogFile)
-	{
-		// Start the UnrealTrace
-		StartUnrealTrace();
+			var LogFolderOutsideOfSandbox = GetLogFolderOutsideOfSandbox();
+			if (!Unreal.IsEngineInstalled() && ServerProcess == null)
+			{
+				// In the installed runs, this is the same folder as CmdEnv.LogFolder so delete only in not-installed
+				DeleteDirectory(LogFolderOutsideOfSandbox);
+				CreateDirectory(LogFolderOutsideOfSandbox);
+			}
+			var ServerLogFile = CombinePaths(LogFolderOutsideOfSandbox, "Server.log");
+			var ClientLogFile = CombinePaths(LogFolderOutsideOfSandbox, Params.EditorTest ? "Editor.log" : "Client.log");
 
-		// Setup server process if required.
-		if (Params.DedicatedServer && !Params.SkipServer)
-		{
-			if (Params.ServerTargetPlatforms.Count > 0)
+			try
 			{
-				TargetPlatformDescriptor ServerPlatformDesc = Params.ServerTargetPlatforms[0];
-				ServerProcess = RunDedicatedServer(Params, ServerLogFile, Params.RunCommandline);
-
-                // With dedicated server, the client connects to local host to load a map, unless client parameters are already specified
-                if (String.IsNullOrEmpty(Params.ClientCommandline))
-                {
-                    if (!String.IsNullOrEmpty(Params.ServerDeviceAddress))
-                    {
-                        Params.ClientCommandline = Params.ServerDeviceAddress;
-                    }
-                    else
-                    {
-                        Params.ClientCommandline = "127.0.0.1";
-                    }
-                }
+				RunInternal(Params, ServerLogFile, ClientLogFile);
 			}
-			else
+			catch
 			{
-				throw new AutomationException("Failed to run, server target platform not specified");
+				throw;
+			}
+			finally
+			{
+				CopyLogsBackToLogFolder();
 			}
-		}
-		else if (Params.FileServer && !Params.SkipServer)
-		{
-			ServerProcess = RunFileServer(Params, ServerLogFile, Params.RunCommandline);
-		}
 
-		if (ServerProcess != null)
-		{
-			LogInformation("Waiting a few seconds for the server to start...");
-			Thread.Sleep(5000);
+			LogInformation("********** RUN COMMAND COMPLETED **********");
 		}
 
-		if (!Params.NoClient)
+		private static void RunInternal(ProjectParams Params, string ServerLogFile, string ClientLogFile)
 		{
-			LogInformation("Starting Client....");
-
-			var SC = CreateDeploymentContext(Params, false);
+			// Start the UnrealTrace
+			StartUnrealTrace();
 
-			ERunOptions ClientRunFlags;
-			string ClientApp;
-			string ClientCmdLine;
-			SetupClientParams(SC, Params, ClientLogFile, out ClientRunFlags, out ClientApp, out ClientCmdLine);
-
-			// Run the client.
-			if (ServerProcess != null)
+			// Setup server process if required.
+			if (Params.DedicatedServer && !Params.SkipServer)
 			{
-				RunClientWithServer(SC, ServerLogFile, ServerProcess, ClientApp, ClientCmdLine, ClientRunFlags, ClientLogFile, Params);
+				if (Params.ServerTargetPlatforms.Count > 0)
+				{
+					TargetPlatformDescriptor ServerPlatformDesc = Params.ServerTargetPlatforms[0];
+					ServerProcess = RunDedicatedServer(Params, ServerLogFile, Params.RunCommandline);
+
+					// With dedicated server, the client connects to local host to load a map, unless client parameters are already specified
+					if (String.IsNullOrEmpty(Params.ClientCommandline))
+					{
+						if (!String.IsNullOrEmpty(Params.ServerDeviceAddress))
+						{
+							Params.ClientCommandline = Params.ServerDeviceAddress;
+						}
+						else
+						{
+							Params.ClientCommandline = "127.0.0.1";
+						}
+					}
+				}
+				else
+				{
+					throw new AutomationException("Failed to run, server target platform not specified");
+				}
 			}
-			else
+			else if (Params.FileServer && !Params.SkipServer)
 			{
-				RunStandaloneClient(SC, ClientLogFile, ClientRunFlags, ClientApp, ClientCmdLine, Params);
+				ServerProcess = RunFileServer(Params, ServerLogFile, Params.RunCommandline);
 			}
-		}
-	}
 
-	private static void RunStandaloneClient(List<DeploymentContext> DeployContextList, string ClientLogFile, ERunOptions ClientRunFlags, string ClientApp, string ClientCmdLine, ProjectParams Params)
-	{
-		if (Params.Unattended)
-		{
-			string LookFor = "Bringing up level for play took";
-			bool bCommandlet = false;
-
-			if (Params.RunAutomationTest != "" || Params.RunAutomationTests)
-			{
-				LookFor = "Automation Test Queue Empty";
-			}
-			else if (Params.EditorTest)
+			if (ServerProcess != null)
 			{
-				LookFor = "Asset discovery search completed in";
+				LogInformation("Waiting a few seconds for the server to start...");
+				Thread.Sleep(5000);
 			}
-			// If running a commandlet, just detect a normal exit
-			else if (ClientCmdLine.IndexOf("-run=", StringComparison.InvariantCultureIgnoreCase) >= 0)
+
+			if (!Params.NoClient)
 			{
-				LookFor = "Game engine shut down";
-				bCommandlet = true;
+				LogInformation("Starting Client....");
+
+				var SC = CreateDeploymentContext(Params, false);
+
+				ERunOptions ClientRunFlags;
+				string ClientApp;
+				string ClientCmdLine;
+				SetupClientParams(SC, Params, ClientLogFile, out ClientRunFlags, out ClientApp, out ClientCmdLine);
+
+				// Run the client.
+				if (ServerProcess != null)
+				{
+					RunClientWithServer(SC, ServerLogFile, ServerProcess, ClientApp, ClientCmdLine, ClientRunFlags, ClientLogFile, Params);
+				}
+				else
+				{
+					RunStandaloneClient(SC, ClientLogFile, ClientRunFlags, ClientApp, ClientCmdLine, Params);
+				}
 			}
+		}
 
+		private static void RunStandaloneClient(List<DeploymentContext> DeployContextList, string ClientLogFile, ERunOptions ClientRunFlags, string ClientApp, string ClientCmdLine, ProjectParams Params)
+		{
+			if (Params.Unattended)
 			{
+				string LookFor = "Bringing up level for play took";
+				bool bCommandlet = false;
 
-				string AllClientOutput = "";
-				int LastAutoFailIndex = -1;
-				IProcessResult ClientProcess = null;
-				FileStream ClientProcessLog = null;
-				StreamReader ClientLogReader = null;
-				LogInformation("Starting Client for unattended test....");
-				ClientProcess = Run(ClientApp, ClientCmdLine + " -testexit=\"" + LookFor + "\"", null, ClientRunFlags | ERunOptions.NoWaitForExit);
-				while (!FileExists(ClientLogFile) && !ClientProcess.HasExited)
+				if (Params.RunAutomationTest != "" || Params.RunAutomationTests)
 				{
-					LogInformation("Waiting for client logging process to start...{0}", ClientLogFile);
-					Thread.Sleep(2000);
+					LookFor = "Automation Test Queue Empty";
 				}
-				if (FileExists(ClientLogFile))
+				else if (Params.EditorTest)
 				{
-					Thread.Sleep(2000);
-					LogInformation("Client logging process started...{0}", ClientLogFile);
-					ClientProcessLog = File.Open(ClientLogFile, FileMode.Open, FileAccess.Read, FileShare.ReadWrite);
-					ClientLogReader = new StreamReader(ClientProcessLog);
+					LookFor = "Asset discovery search completed in";
 				}
-				if (ClientLogReader == null)
+				// If running a commandlet, just detect a normal exit
+				else if (ClientCmdLine.IndexOf("-run=", StringComparison.InvariantCultureIgnoreCase) >= 0)
 				{
-					throw new AutomationException("Client exited without creating a log file.");
+					LookFor = "Game engine shut down";
+					bCommandlet = true;
 				}
-				bool bKeepReading = true;
-				bool WelcomedCorrectly = false;
-				bool bClientExited = false;
-				DateTime ExitTime = DateTime.UtcNow;
 
-				while (bKeepReading)
 				{
-					if (!bClientExited && ClientProcess.HasExited)
+
+					string AllClientOutput = "";
+					int LastAutoFailIndex = -1;
+					IProcessResult ClientProcess = null;
+					FileStream ClientProcessLog = null;
+					StreamReader ClientLogReader = null;
+					LogInformation("Starting Client for unattended test....");
+					ClientProcess = Run(ClientApp, ClientCmdLine + " -testexit=\"" + LookFor + "\"", null, ClientRunFlags | ERunOptions.NoWaitForExit);
+					while (!FileExists(ClientLogFile) && !ClientProcess.HasExited)
 					{
-						ExitTime = DateTime.UtcNow;
-						bClientExited = true;
+						LogInformation("Waiting for client logging process to start...{0}", ClientLogFile);
+						Thread.Sleep(2000);
 					}
-					string ClientOutput = ClientLogReader.ReadToEnd();
-					if (!String.IsNullOrEmpty(ClientOutput))
+					if (FileExists(ClientLogFile))
 					{
-						if (bClientExited)
-						{
-							ExitTime = DateTime.UtcNow; // as long as it is spewing, we reset the timer
-						}
-						AllClientOutput += ClientOutput;
-						Console.Write(ClientOutput);
+						Thread.Sleep(2000);
+						LogInformation("Client logging process started...{0}", ClientLogFile);
+						ClientProcessLog = File.Open(ClientLogFile, FileMode.Open, FileAccess.Read, FileShare.ReadWrite);
+						ClientLogReader = new StreamReader(ClientProcessLog);
+					}
+					if (ClientLogReader == null)
+					{
+						throw new AutomationException("Client exited without creating a log file.");
+					}
+					bool bKeepReading = true;
+					bool WelcomedCorrectly = false;
+					bool bClientExited = false;
+					DateTime ExitTime = DateTime.UtcNow;
 
-						if (AllClientOutput.LastIndexOf(LookFor) > AllClientOutput.IndexOf(LookFor))
+					while (bKeepReading)
+					{
+						if (!bClientExited && ClientProcess.HasExited)
 						{
-							WelcomedCorrectly = true;
-							LogInformation("Test complete...");
-							bKeepReading = false;
+							ExitTime = DateTime.UtcNow;
+							bClientExited = true;
 						}
-						else if (Params.RunAutomationTests)
+						string ClientOutput = ClientLogReader.ReadToEnd();
+						if (!String.IsNullOrEmpty(ClientOutput))
 						{
-							int FailIndex = AllClientOutput.LastIndexOf("Automation Test Failed");
-							int ParenIndex = AllClientOutput.LastIndexOf(")");
-							if (FailIndex >= 0 && ParenIndex > FailIndex && FailIndex > LastAutoFailIndex)
+							if (bClientExited)
 							{
-								string Tail = AllClientOutput.Substring(FailIndex);
-								int CloseParenIndex = Tail.IndexOf(")");
-								int OpenParenIndex = Tail.IndexOf("(");
-								string Test = "";
-								if (OpenParenIndex >= 0 && CloseParenIndex > OpenParenIndex)
-								{
-									Test = Tail.Substring(OpenParenIndex + 1, CloseParenIndex - OpenParenIndex - 1);
-									LogError("Automated test failed ({0}).", Test);
-									LastAutoFailIndex = FailIndex;
-								}
+								ExitTime = DateTime.UtcNow; // as long as it is spewing, we reset the timer
 							}
-						}
-						// Detect commandlet failure
-						else if (bCommandlet)
-						{
-							const string ResultLog = "Commandlet->Main return this error code: ";
-
-							int ResultStart = AllClientOutput.LastIndexOf(ResultLog);
-							int ResultValIdx = ResultStart + ResultLog.Length;
+							AllClientOutput += ClientOutput;
+							Console.Write(ClientOutput);
 
-							if (ResultStart >= 0 && ResultValIdx < AllClientOutput.Length &&
-								AllClientOutput.Substring(ResultValIdx, 1) == "1")
+							if (AllClientOutput.LastIndexOf(LookFor) > AllClientOutput.IndexOf(LookFor))
+							{
+								WelcomedCorrectly = true;
+								LogInformation("Test complete...");
+								bKeepReading = false;
+							}
+							else if (Params.RunAutomationTests)
 							{
-								// Parse the full commandlet warning/error summary
-								string FullSummary = "";
-								int SummaryStart = AllClientOutput.LastIndexOf("Warning/Error Summary");
-								
-								if (SummaryStart >= 0 && SummaryStart < ResultStart)
+								int FailIndex = AllClientOutput.LastIndexOf("Automation Test Failed");
+								int ParenIndex = AllClientOutput.LastIndexOf(")");
+								if (FailIndex >= 0 && ParenIndex > FailIndex && FailIndex > LastAutoFailIndex)
 								{
-									FullSummary = AllClientOutput.Substring(SummaryStart, ResultStart - SummaryStart);
+									string Tail = AllClientOutput.Substring(FailIndex);
+									int CloseParenIndex = Tail.IndexOf(")");
+									int OpenParenIndex = Tail.IndexOf("(");
+									string Test = "";
+									if (OpenParenIndex >= 0 && CloseParenIndex > OpenParenIndex)
+									{
+										Test = Tail.Substring(OpenParenIndex + 1, CloseParenIndex - OpenParenIndex - 1);
+										LogError("Automated test failed ({0}).", Test);
+										LastAutoFailIndex = FailIndex;
+									}
 								}
+							}
+							// Detect commandlet failure
+							else if (bCommandlet)
+							{
+								const string ResultLog = "Commandlet->Main return this error code: ";
 
+								int ResultStart = AllClientOutput.LastIndexOf(ResultLog);
+								int ResultValIdx = ResultStart + ResultLog.Length;
 
-								if (FullSummary.Length > 0)
-								{
-									LogError("Commandlet failed, summary:" + Environment.NewLine +
-																					FullSummary);
-								}
-								else
+								if (ResultStart >= 0 && ResultValIdx < AllClientOutput.Length &&
+									AllClientOutput.Substring(ResultValIdx, 1) == "1")
 								{
-									LogError("Commandlet failed.");
+									// Parse the full commandlet warning/error summary
+									string FullSummary = "";
+									int SummaryStart = AllClientOutput.LastIndexOf("Warning/Error Summary");
+
+									if (SummaryStart >= 0 && SummaryStart < ResultStart)
+									{
+										FullSummary = AllClientOutput.Substring(SummaryStart, ResultStart - SummaryStart);
+									}
+
+
+									if (FullSummary.Length > 0)
+									{
+										LogError("Commandlet failed, summary:" + Environment.NewLine +
+																						FullSummary);
+									}
+									else
+									{
+										LogError("Commandlet failed.");
+									}
 								}
 							}
 						}
+						else if (bClientExited && (DateTime.UtcNow - ExitTime).TotalSeconds > 30)
+						{
+							LogInformation("Client exited and has been quiet for 30 seconds...exiting");
+							bKeepReading = false;
+						}
 					}
-					else if (bClientExited && (DateTime.UtcNow - ExitTime).TotalSeconds > 30)
+					if (ClientProcess != null && !ClientProcess.HasExited)
 					{
-						LogInformation("Client exited and has been quiet for 30 seconds...exiting");
-						bKeepReading = false;
+						LogInformation("Client is supposed to exit, lets wait a while for it to exit naturally...");
+						for (int i = 0; i < 120 && !ClientProcess.HasExited; i++)
+						{
+							Thread.Sleep(1000);
+						}
 					}
-				}
-				if (ClientProcess != null && !ClientProcess.HasExited)
-				{
-					LogInformation("Client is supposed to exit, lets wait a while for it to exit naturally...");
-					for (int i = 0; i < 120 && !ClientProcess.HasExited; i++)
+					if (ClientProcess != null && !ClientProcess.HasExited)
 					{
-						Thread.Sleep(1000);
+						LogInformation("Stopping client...");
+						ClientProcess.StopProcess();
+						Thread.Sleep(10000);
 					}
-				}
-				if (ClientProcess != null && !ClientProcess.HasExited)
-				{
-					LogInformation("Stopping client...");
-					ClientProcess.StopProcess();
-					Thread.Sleep(10000);
-				}
-				while (!ClientLogReader.EndOfStream)
-				{
-					string ClientOutput = ClientLogReader.ReadToEnd();
-					if (!String.IsNullOrEmpty(ClientOutput))
+					while (!ClientLogReader.EndOfStream)
 					{
-						Console.Write(ClientOutput);
+						string ClientOutput = ClientLogReader.ReadToEnd();
+						if (!String.IsNullOrEmpty(ClientOutput))
+						{
+							Console.Write(ClientOutput);
+						}
 					}
-				}
 
-				if (!WelcomedCorrectly)
-				{
-					throw new AutomationException("Client exited before we asked it to.");
+					if (!WelcomedCorrectly)
+					{
+						throw new AutomationException("Client exited before we asked it to.");
+					}
 				}
 			}
-		}
-		else
-		{
-			var SC = DeployContextList[0];
-			IProcessResult ClientProcess = SC.StageTargetPlatform.RunClient(ClientRunFlags, ClientApp, ClientCmdLine, Params);
-			if (ClientProcess != null)
+			else
 			{
-				// If the client runs without StdOut redirect we're going to read the log output directly from log file on
-				// a separate thread.
-				if ((ClientRunFlags & ERunOptions.NoStdOutRedirect) == ERunOptions.NoStdOutRedirect)
+				var SC = DeployContextList[0];
+				IProcessResult ClientProcess = SC.StageTargetPlatform.RunClient(ClientRunFlags, ClientApp, ClientCmdLine, Params);
+				if (ClientProcess != null)
 				{
-					ClientLogReaderThread = new System.Threading.Thread(ClientLogReaderProc);
-					ClientLogReaderThread.Start(new object[] { ClientLogFile, ClientProcess });
-				}
+					// If the client runs without StdOut redirect we're going to read the log output directly from log file on
+					// a separate thread.
+					if ((ClientRunFlags & ERunOptions.NoStdOutRedirect) == ERunOptions.NoStdOutRedirect)
+					{
+						ClientLogReaderThread = new System.Threading.Thread(ClientLogReaderProc);
+						ClientLogReaderThread.Start(new object[] { ClientLogFile, ClientProcess });
+					}
 
-				do
-				{
-					Thread.Sleep(100);
-				}
-				while (ClientProcess.HasExited == false);
+					do
+					{
+						Thread.Sleep(100);
+					}
+					while (ClientProcess.HasExited == false);
 
-				SC.StageTargetPlatform.PostRunClient(ClientProcess, Params);
+					SC.StageTargetPlatform.PostRunClient(ClientProcess, Params);
 
-                // any non-zero exit code should propagate an exception. The Virtual function above may have
-                // already thrown a more specific exception or given a more specific ErrorCode, but this catches the rest.
-				if (ClientProcess.ExitCode != 0)
-				{
-					throw new AutomationException("Client exited with error code: " + ClientProcess.ExitCode);
+					// any non-zero exit code should propagate an exception. The Virtual function above may have
+					// already thrown a more specific exception or given a more specific ErrorCode, but this catches the rest.
+					if (ClientProcess.ExitCode != 0)
+					{
+						throw new AutomationException("Client exited with error code: " + ClientProcess.ExitCode);
+					}
 				}
 			}
 		}
-	}
 
-	private static void RunClientWithServer(List<DeploymentContext> DeployContextList, string ServerLogFile, IProcessResult ServerProcess, string ClientApp, string ClientCmdLine, ERunOptions ClientRunFlags, string ClientLogFile, ProjectParams Params)
-	{
-		IProcessResult ClientProcess = null;
-		var OtherClients = new List<IProcessResult>();
+		private static void RunClientWithServer(List<DeploymentContext> DeployContextList, string ServerLogFile, IProcessResult ServerProcess, string ClientApp, string ClientCmdLine, ERunOptions ClientRunFlags, string ClientLogFile, ProjectParams Params)
+		{
+			IProcessResult ClientProcess = null;
+			var OtherClients = new List<IProcessResult>();
 
-		bool WelcomedCorrectly = false;
-		int NumClients = Params.NumClients;
-		string AllClientOutput = "";
-		int LastAutoFailIndex = -1;
+			bool WelcomedCorrectly = false;
+			int NumClients = Params.NumClients;
+			string AllClientOutput = "";
+			int LastAutoFailIndex = -1;
 
-		if (Params.Unattended)
-		{
-			string LookFor = "Bringing up level for play took";
-			if (Params.DedicatedServer)
-			{
-				LookFor = "Welcomed by server";
-			}
-			else if (Params.RunAutomationTest != "" || Params.RunAutomationTests)
-			{
-				LookFor = "Automation Test Queue Empty";
-			}
+			if (Params.Unattended)
 			{
-				while (!FileExists(ServerLogFile) && !ServerProcess.HasExited)
+				string LookFor = "Bringing up level for play took";
+				if (Params.DedicatedServer)
 				{
-					LogInformation("Waiting for logging process to start...");
-					Thread.Sleep(2000);
+					LookFor = "Welcomed by server";
 				}
-				Thread.Sleep(1000);
-
-				string AllServerOutput = "";
-				using (FileStream ProcessLog = File.Open(ServerLogFile, FileMode.Open, FileAccess.Read, FileShare.ReadWrite))
+				else if (Params.RunAutomationTest != "" || Params.RunAutomationTests)
 				{
-					StreamReader LogReader = new StreamReader(ProcessLog);
-					bool bKeepReading = true;
-
-					FileStream ClientProcessLog = null;
-					StreamReader ClientLogReader = null;
+					LookFor = "Automation Test Queue Empty";
+				}
+				{
+					while (!FileExists(ServerLogFile) && !ServerProcess.HasExited)
+					{
+						LogInformation("Waiting for logging process to start...");
+						Thread.Sleep(2000);
+					}
+					Thread.Sleep(1000);
 
-					// Read until the process has exited.
-					while (!ServerProcess.HasExited && bKeepReading)
+					string AllServerOutput = "";
+					using (FileStream ProcessLog = File.Open(ServerLogFile, FileMode.Open, FileAccess.Read, FileShare.ReadWrite))
 					{
-						while (!LogReader.EndOfStream && bKeepReading && ClientProcess == null)
+						StreamReader LogReader = new StreamReader(ProcessLog);
+						bool bKeepReading = true;
+
+						FileStream ClientProcessLog = null;
+						StreamReader ClientLogReader = null;
+
+						// Read until the process has exited.
+						while (!ServerProcess.HasExited && bKeepReading)
 						{
-							string Output = LogReader.ReadToEnd();
-							if (!String.IsNullOrEmpty(Output))
+							while (!LogReader.EndOfStream && bKeepReading && ClientProcess == null)
 							{
-								AllServerOutput += Output;
-								if (ClientProcess == null &&
-										(AllServerOutput.Contains("Game Engine Initialized") || AllServerOutput.Contains("Unreal Network File Server is ready")))
+								string Output = LogReader.ReadToEnd();
+								if (!String.IsNullOrEmpty(Output))
 								{
-									LogInformation("Starting Client for unattended test....");
-									ClientProcess = Run(ClientApp, ClientCmdLine + " -testexit=\"" + LookFor + "\"", null, ClientRunFlags | ERunOptions.NoWaitForExit);
-									//@todo no testing is done on these
-									if (NumClients > 1 && NumClients < 9)
+									AllServerOutput += Output;
+									if (ClientProcess == null &&
+											(AllServerOutput.Contains("Game Engine Initialized") || AllServerOutput.Contains("Unreal Network File Server is ready")))
 									{
-										for (int i = 1; i < NumClients; i++)
+										LogInformation("Starting Client for unattended test....");
+										ClientProcess = Run(ClientApp, ClientCmdLine + " -testexit=\"" + LookFor + "\"", null, ClientRunFlags | ERunOptions.NoWaitForExit);
+										//@todo no testing is done on these
+										if (NumClients > 1 && NumClients < 9)
+										{
+											for (int i = 1; i < NumClients; i++)
+											{
+												LogInformation("Starting Extra Client....");
+												OtherClients.Add(Run(ClientApp, ClientCmdLine, null, ClientRunFlags | ERunOptions.NoWaitForExit));
+											}
+										}
+										while (!FileExists(ClientLogFile) && !ClientProcess.HasExited)
+										{
+											LogInformation("Waiting for client logging process to start...{0}", ClientLogFile);
+											Thread.Sleep(2000);
+										}
+										if (!ClientProcess.HasExited)
 										{
-											LogInformation("Starting Extra Client....");
-											OtherClients.Add(Run(ClientApp, ClientCmdLine, null, ClientRunFlags | ERunOptions.NoWaitForExit));
+											Thread.Sleep(2000);
+											LogInformation("Client logging process started...{0}", ClientLogFile);
+											ClientProcessLog = File.Open(ClientLogFile, FileMode.Open, FileAccess.Read, FileShare.ReadWrite);
+											ClientLogReader = new StreamReader(ClientProcessLog);
 										}
 									}
-									while (!FileExists(ClientLogFile) && !ClientProcess.HasExited)
+									else if (ClientProcess == null && !ServerProcess.HasExited)
 									{
-										LogInformation("Waiting for client logging process to start...{0}", ClientLogFile);
+										LogInformation("Waiting for server to start....");
 										Thread.Sleep(2000);
 									}
-									if (!ClientProcess.HasExited)
+									if (ClientProcess != null && ClientProcess.HasExited)
 									{
-										Thread.Sleep(2000);
-										LogInformation("Client logging process started...{0}", ClientLogFile);
-										ClientProcessLog = File.Open(ClientLogFile, FileMode.Open, FileAccess.Read, FileShare.ReadWrite);
-										ClientLogReader = new StreamReader(ClientProcessLog);
+										ServerProcess.StopProcess();
+										throw new AutomationException("Client exited before we asked it to.");
 									}
 								}
-								else if (ClientProcess == null && !ServerProcess.HasExited)
-								{
-									LogInformation("Waiting for server to start....");
-									Thread.Sleep(2000);
-								}
-								if (ClientProcess != null && ClientProcess.HasExited)
+							}
+							if (ClientLogReader != null)
+							{
+								if (ClientProcess.HasExited)
 								{
 									ServerProcess.StopProcess();
-									throw new AutomationException("Client exited before we asked it to.");
+									throw new AutomationException("Client exited or closed the log before we asked it to.");
 								}
-							}
-						}
-						if (ClientLogReader != null)
-						{
-							if (ClientProcess.HasExited)
-							{
-								ServerProcess.StopProcess();
-								throw new AutomationException("Client exited or closed the log before we asked it to.");
-							}
-							while (!ClientProcess.HasExited && !ServerProcess.HasExited && bKeepReading)
-							{
-								while (!ClientLogReader.EndOfStream && bKeepReading && !ServerProcess.HasExited && !ClientProcess.HasExited)
+								while (!ClientProcess.HasExited && !ServerProcess.HasExited && bKeepReading)
 								{
-									string ClientOutput = ClientLogReader.ReadToEnd();
-									if (!String.IsNullOrEmpty(ClientOutput))
+									while (!ClientLogReader.EndOfStream && bKeepReading && !ServerProcess.HasExited && !ClientProcess.HasExited)
 									{
-										AllClientOutput += ClientOutput;
-										Console.Write(ClientOutput);
-
-										if (AllClientOutput.LastIndexOf(LookFor) > AllClientOutput.IndexOf(LookFor))
+										string ClientOutput = ClientLogReader.ReadToEnd();
+										if (!String.IsNullOrEmpty(ClientOutput))
 										{
-											if (Params.FakeClient)
-											{
-												LogInformation("Welcomed by server or client loaded, lets wait ten minutes...");
-												Thread.Sleep(60000 * 10);
-											}
-											else
+											AllClientOutput += ClientOutput;
+											Console.Write(ClientOutput);
+
+											if (AllClientOutput.LastIndexOf(LookFor) > AllClientOutput.IndexOf(LookFor))
 											{
-												LogInformation("Welcomed by server or client loaded, lets wait 30 seconds...");
-												Thread.Sleep(30000);
+												if (Params.FakeClient)
+												{
+													LogInformation("Welcomed by server or client loaded, lets wait ten minutes...");
+													Thread.Sleep(60000 * 10);
+												}
+												else
+												{
+													LogInformation("Welcomed by server or client loaded, lets wait 30 seconds...");
+													Thread.Sleep(30000);
+												}
+												WelcomedCorrectly = true;
+												bKeepReading = false;
 											}
-											WelcomedCorrectly = true;
-											bKeepReading = false;
-										}
-										else if (Params.RunAutomationTests)
-										{
-											int FailIndex = AllClientOutput.LastIndexOf("Automation Test Failed");
-											int ParenIndex = AllClientOutput.LastIndexOf(")");
-											if (FailIndex >= 0 && ParenIndex > FailIndex && FailIndex > LastAutoFailIndex)
+											else if (Params.RunAutomationTests)
 											{
-												string Tail = AllClientOutput.Substring(FailIndex);
-												int CloseParenIndex = Tail.IndexOf(")");
-												int OpenParenIndex = Tail.IndexOf("(");
-												string Test = "";
-												if (OpenParenIndex >= 0 && CloseParenIndex > OpenParenIndex)
+												int FailIndex = AllClientOutput.LastIndexOf("Automation Test Failed");
+												int ParenIndex = AllClientOutput.LastIndexOf(")");
+												if (FailIndex >= 0 && ParenIndex > FailIndex && FailIndex > LastAutoFailIndex)
 												{
-													Test = Tail.Substring(OpenParenIndex + 1, CloseParenIndex - OpenParenIndex - 1);
-													LogError("Automated test failed ({0}).", Test);
-													LastAutoFailIndex = FailIndex;
+													string Tail = AllClientOutput.Substring(FailIndex);
+													int CloseParenIndex = Tail.IndexOf(")");
+													int OpenParenIndex = Tail.IndexOf("(");
+													string Test = "";
+													if (OpenParenIndex >= 0 && CloseParenIndex > OpenParenIndex)
+													{
+														Test = Tail.Substring(OpenParenIndex + 1, CloseParenIndex - OpenParenIndex - 1);
+														LogError("Automated test failed ({0}).", Test);
+														LastAutoFailIndex = FailIndex;
+													}
 												}
 											}
 										}
@@ -519,512 +522,511 @@ public partial class Project : CommandUtils
 					}
 				}
 			}
-		}
-		else
-		{
-			LogFileReaderProcess(ServerLogFile, ServerProcess, (string Output) =>
+			else
 			{
-				bool bKeepReading = true;
-				if (ClientProcess == null && !String.IsNullOrEmpty(Output))
+				LogFileReaderProcess(ServerLogFile, ServerProcess, (string Output) =>
 				{
-					AllClientOutput += Output;
-					if (AllClientOutput.Contains("Game Engine Initialized") || AllClientOutput.Contains("Unreal Network File Server is ready") || AllClientOutput.Contains("COTF server is ready"))
+					bool bKeepReading = true;
+					if (ClientProcess == null && !String.IsNullOrEmpty(Output))
 					{
-						LogInformation("Starting Client....");
-						var SC = DeployContextList[0];
-						ClientProcess = SC.StageTargetPlatform.RunClient(ClientRunFlags | ERunOptions.NoWaitForExit, ClientApp, ClientCmdLine, Params);
-//						ClientProcess = Run(ClientApp, ClientCmdLine, null, ClientRunFlags | ERunOptions.NoWaitForExit);
-						if (NumClients > 1 && NumClients < 9)
+						AllClientOutput += Output;
+						if (AllClientOutput.Contains("Game Engine Initialized") || AllClientOutput.Contains("Unreal Network File Server is ready") || AllClientOutput.Contains("COTF server is ready"))
 						{
-							for (int i = 1; i < NumClients; i++)
+							LogInformation("Starting Client....");
+							var SC = DeployContextList[0];
+							ClientProcess = SC.StageTargetPlatform.RunClient(ClientRunFlags | ERunOptions.NoWaitForExit, ClientApp, ClientCmdLine, Params);
+						//						ClientProcess = Run(ClientApp, ClientCmdLine, null, ClientRunFlags | ERunOptions.NoWaitForExit);
+						if (NumClients > 1 && NumClients < 9)
 							{
-								LogInformation("Starting Extra Client....");
-								IProcessResult NewClient = SC.StageTargetPlatform.RunClient(ClientRunFlags | ERunOptions.NoWaitForExit, ClientApp, ClientCmdLine, Params);
-								OtherClients.Add(NewClient);
+								for (int i = 1; i < NumClients; i++)
+								{
+									LogInformation("Starting Extra Client....");
+									IProcessResult NewClient = SC.StageTargetPlatform.RunClient(ClientRunFlags | ERunOptions.NoWaitForExit, ClientApp, ClientCmdLine, Params);
+									OtherClients.Add(NewClient);
+								}
 							}
 						}
 					}
-				}
-				else if (ClientProcess == null && !ServerProcess.HasExited)
-				{
-					LogInformation("Waiting for server to start....");
-					Thread.Sleep(2000);
-				}
-
-				if (String.IsNullOrEmpty(Output) == false)
-				{
-					Console.Write(Output);
-				}
+					else if (ClientProcess == null && !ServerProcess.HasExited)
+					{
+						LogInformation("Waiting for server to start....");
+						Thread.Sleep(2000);
+					}
 
-				if (ClientProcess != null && ClientProcess.HasExited)
-				{
+					if (String.IsNullOrEmpty(Output) == false)
+					{
+						Console.Write(Output);
+					}
 
-					LogInformation("Client exited, stopping server....");
-					if (!GlobalCommandLine.NoKill)
+					if (ClientProcess != null && ClientProcess.HasExited)
 					{
-						ServerProcess.StopProcess();
+
+						LogInformation("Client exited, stopping server....");
+						if (!GlobalCommandLine.NoKill)
+						{
+							ServerProcess.StopProcess();
+						}
+						bKeepReading = false;
 					}
-					bKeepReading = false;
-				}
 
-				return bKeepReading; // Keep reading
+					return bKeepReading; // Keep reading
 			});
-		}
-		LogInformation("Server exited....");
-		if (ClientProcess != null && !ClientProcess.HasExited)
-		{
-			ClientProcess.StopProcess();
-		}
-		foreach (var OtherClient in OtherClients)
-		{
-			if (OtherClient != null && !OtherClient.HasExited)
+			}
+			LogInformation("Server exited....");
+			if (ClientProcess != null && !ClientProcess.HasExited)
 			{
-				OtherClient.StopProcess();
+				ClientProcess.StopProcess();
 			}
-		}
-		if (Params.Unattended)
-		{
-			if (!WelcomedCorrectly)
+			foreach (var OtherClient in OtherClients)
 			{
-				throw new AutomationException("Server or client exited before we asked it to.");
+				if (OtherClient != null && !OtherClient.HasExited)
+				{
+					OtherClient.StopProcess();
+				}
+			}
+			if (Params.Unattended)
+			{
+				if (!WelcomedCorrectly)
+				{
+					throw new AutomationException("Server or client exited before we asked it to.");
+				}
 			}
 		}
-	}
 
-	private static void SetupClientParams(List<DeploymentContext> DeployContextList, ProjectParams Params, string ClientLogFile, out ERunOptions ClientRunFlags, out string ClientApp, out string ClientCmdLine)
-	{
-		if (Params.ClientTargetPlatforms.Count == 0)
+		private static void SetupClientParams(List<DeploymentContext> DeployContextList, ProjectParams Params, string ClientLogFile, out ERunOptions ClientRunFlags, out string ClientApp, out string ClientCmdLine)
 		{
-			throw new AutomationException("No ClientTargetPlatform set for SetupClientParams.");
-		}
+			if (Params.ClientTargetPlatforms.Count == 0)
+			{
+				throw new AutomationException("No ClientTargetPlatform set for SetupClientParams.");
+			}
 
-		//		var DeployContextList = CreateDeploymentContext(Params, false);
+			//		var DeployContextList = CreateDeploymentContext(Params, false);
 
-		if (DeployContextList.Count == 0)
-		{
-			throw new AutomationException("No DeployContextList for SetupClientParams.");
-		}
+			if (DeployContextList.Count == 0)
+			{
+				throw new AutomationException("No DeployContextList for SetupClientParams.");
+			}
 
-		var SC = DeployContextList[0];
+			var SC = DeployContextList[0];
 
-		// Get client app name and command line.
-		ClientRunFlags = ERunOptions.AllowSpew | ERunOptions.AppMustExist;
-		ClientApp = "";
-		ClientCmdLine = "";
-		string TempCmdLine = SC.ProjectArgForCommandLines + " ";
-		var PlatformName = Params.ClientTargetPlatforms[0].ToString();
-		if (Params.Cook || Params.CookOnTheFly)
-		{
-			List<FileReference> Exes = SC.StageTargetPlatform.GetExecutableNames(SC);
-			ClientApp = Exes[0].FullName;
-
-            if (!String.IsNullOrEmpty(Params.ClientCommandline))
-            {
-                TempCmdLine += Params.ClientCommandline + " ";
-            }
-            else
-            {
-                TempCmdLine += Params.MapToRun + " ";
-            }
-
-			if (Params.CookOnTheFly || Params.FileServer)
+			// Get client app name and command line.
+			ClientRunFlags = ERunOptions.AllowSpew | ERunOptions.AppMustExist;
+			ClientApp = "";
+			ClientCmdLine = "";
+			string TempCmdLine = SC.ProjectArgForCommandLines + " ";
+			var PlatformName = Params.ClientTargetPlatforms[0].ToString();
+			if (Params.Cook || Params.CookOnTheFly)
 			{
-				if (Params.ZenStore)
+				List<FileReference> Exes = SC.StageTargetPlatform.GetExecutableNames(SC);
+				ClientApp = Exes[0].FullName;
+
+				if (!String.IsNullOrEmpty(Params.ClientCommandline))
 				{
-					TempCmdLine += "-cookonthefly -zenstorehost=";
+					TempCmdLine += Params.ClientCommandline + " ";
 				}
 				else
 				{
-					TempCmdLine += "-filehostip=";
+					TempCmdLine += Params.MapToRun + " ";
 				}
 
-				// add localhost first for platforms using redirection
-				const string LocalHost = "127.0.0.1";
-				if (!IsNullOrEmpty(Params.Port))
+				if (Params.CookOnTheFly || Params.FileServer)
 				{
-					bool FirstParam = true;
-					foreach (var Port in Params.Port)
+					if (Params.ZenStore)
 					{
-						if (!FirstParam)
-						{
-							TempCmdLine += "+";
-						}
-						FirstParam = false;
-						string[] PortProtocol = Port.Split(new char[] { ':' });
-						if (PortProtocol.Length > 1)
-						{
-							TempCmdLine += String.Format("{0}://{1}:{2}", PortProtocol[0], LocalHost, PortProtocol[1]);
-						}
-						else
+						TempCmdLine += "-cookonthefly -zenstorehost=";
+					}
+					else
+					{
+						TempCmdLine += "-filehostip=";
+					}
+
+					// add localhost first for platforms using redirection
+					const string LocalHost = "127.0.0.1";
+					if (!IsNullOrEmpty(Params.Port))
+					{
+						bool FirstParam = true;
+						foreach (var Port in Params.Port)
 						{
-							TempCmdLine += LocalHost;
-							TempCmdLine += ":";
-							TempCmdLine += Params.Port;
-						}
+							if (!FirstParam)
+							{
+								TempCmdLine += "+";
+							}
+							FirstParam = false;
+							string[] PortProtocol = Port.Split(new char[] { ':' });
+							if (PortProtocol.Length > 1)
+							{
+								TempCmdLine += String.Format("{0}://{1}:{2}", PortProtocol[0], LocalHost, PortProtocol[1]);
+							}
+							else
+							{
+								TempCmdLine += LocalHost;
+								TempCmdLine += ":";
+								TempCmdLine += Params.Port;
+							}
 
+						}
+					}
+					else
+					{
+						// use default port
+						TempCmdLine += LocalHost;
 					}
-				}
-				else
-				{
-					// use default port
-					TempCmdLine += LocalHost;
-				}
 
-				if (UnrealBuildTool.BuildHostPlatform.Current.Platform == UnrealTargetPlatform.Mac)
-				{
-					NetworkInterface[] Interfaces = NetworkInterface.GetAllNetworkInterfaces();
-					foreach (NetworkInterface adapter in Interfaces)
+					if (UnrealBuildTool.BuildHostPlatform.Current.Platform == UnrealTargetPlatform.Mac)
 					{
-						if (adapter.NetworkInterfaceType != NetworkInterfaceType.Loopback)
+						NetworkInterface[] Interfaces = NetworkInterface.GetAllNetworkInterfaces();
+						foreach (NetworkInterface adapter in Interfaces)
 						{
-							IPInterfaceProperties IP = adapter.GetIPProperties();
-							for (int Index = 0; Index < IP.UnicastAddresses.Count; ++Index)
+							if (adapter.NetworkInterfaceType != NetworkInterfaceType.Loopback)
 							{
-								if (InternalUtils.IsDnsEligible(IP.UnicastAddresses[Index]) && IP.UnicastAddresses[Index].Address.AddressFamily == System.Net.Sockets.AddressFamily.InterNetwork)
+								IPInterfaceProperties IP = adapter.GetIPProperties();
+								for (int Index = 0; Index < IP.UnicastAddresses.Count; ++Index)
 								{
-									if (!IsNullOrEmpty(Params.Port))
+									if (InternalUtils.IsDnsEligible(IP.UnicastAddresses[Index]) && IP.UnicastAddresses[Index].Address.AddressFamily == System.Net.Sockets.AddressFamily.InterNetwork)
 									{
-										foreach (var Port in Params.Port)
+										if (!IsNullOrEmpty(Params.Port))
 										{
-											TempCmdLine += "+";
-											string[] PortProtocol = Port.Split(new char[] { ':' });
-											if (PortProtocol.Length > 1)
-											{
-												TempCmdLine += String.Format("{0}://{1}:{2}", PortProtocol[0], IP.UnicastAddresses[Index].Address.ToString(), PortProtocol[1]);
-											}
-											else
+											foreach (var Port in Params.Port)
 											{
-												TempCmdLine += IP.UnicastAddresses[Index].Address.ToString();
-												TempCmdLine += ":";
-												TempCmdLine += Params.Port;
+												TempCmdLine += "+";
+												string[] PortProtocol = Port.Split(new char[] { ':' });
+												if (PortProtocol.Length > 1)
+												{
+													TempCmdLine += String.Format("{0}://{1}:{2}", PortProtocol[0], IP.UnicastAddresses[Index].Address.ToString(), PortProtocol[1]);
+												}
+												else
+												{
+													TempCmdLine += IP.UnicastAddresses[Index].Address.ToString();
+													TempCmdLine += ":";
+													TempCmdLine += Params.Port;
+												}
 											}
 										}
+										else
+										{
+											TempCmdLine += "+";
+											// use default port
+											TempCmdLine += IP.UnicastAddresses[Index].Address.ToString();
+										}
 									}
-                                    else
-                                    {
-										TempCmdLine += "+";
-										// use default port
-                                        TempCmdLine += IP.UnicastAddresses[Index].Address.ToString();
-                                    }
 								}
 							}
 						}
 					}
-				}
-				else
-				{
-					NetworkInterface[] Interfaces = NetworkInterface.GetAllNetworkInterfaces();
-					foreach (NetworkInterface adapter in Interfaces)
+					else
 					{
-						if (adapter.OperationalStatus == OperationalStatus.Up)
+						NetworkInterface[] Interfaces = NetworkInterface.GetAllNetworkInterfaces();
+						foreach (NetworkInterface adapter in Interfaces)
 						{
-							IPInterfaceProperties IP = adapter.GetIPProperties();
-							for (int Index = 0; Index < IP.UnicastAddresses.Count; ++Index)
+							if (adapter.OperationalStatus == OperationalStatus.Up)
 							{
-								if (InternalUtils.IsDnsEligible(IP.UnicastAddresses[Index]))
+								IPInterfaceProperties IP = adapter.GetIPProperties();
+								for (int Index = 0; Index < IP.UnicastAddresses.Count; ++Index)
 								{
-									if (!IsNullOrEmpty(Params.Port))
+									if (InternalUtils.IsDnsEligible(IP.UnicastAddresses[Index]))
 									{
-										foreach (var Port in Params.Port)
+										if (!IsNullOrEmpty(Params.Port))
 										{
-											TempCmdLine += "+";
-											string[] PortProtocol = Port.Split(new char[] { ':' });
-											if (PortProtocol.Length > 1)
-											{
-												TempCmdLine += String.Format("{0}://{1}:{2}", PortProtocol[0], IP.UnicastAddresses[Index].Address.ToString(), PortProtocol[1]);
-											}
-											else
+											foreach (var Port in Params.Port)
 											{
-												TempCmdLine += IP.UnicastAddresses[Index].Address.ToString();
-												TempCmdLine += ":";
-												TempCmdLine += Params.Port;
+												TempCmdLine += "+";
+												string[] PortProtocol = Port.Split(new char[] { ':' });
+												if (PortProtocol.Length > 1)
+												{
+													TempCmdLine += String.Format("{0}://{1}:{2}", PortProtocol[0], IP.UnicastAddresses[Index].Address.ToString(), PortProtocol[1]);
+												}
+												else
+												{
+													TempCmdLine += IP.UnicastAddresses[Index].Address.ToString();
+													TempCmdLine += ":";
+													TempCmdLine += Params.Port;
+												}
 											}
 										}
+										else
+										{
+											TempCmdLine += "+";
+											// use default port
+											TempCmdLine += IP.UnicastAddresses[Index].Address.ToString();
+										}
 									}
-                                    else
-                                    {
-										TempCmdLine += "+";
-										// use default port
-                                        TempCmdLine += IP.UnicastAddresses[Index].Address.ToString();
-                                    }
 								}
 							}
 						}
 					}
-				}
 
-				TempCmdLine += " ";
+					TempCmdLine += " ";
 
-				if (Params.CookOnTheFlyStreaming)
-				{
-					TempCmdLine += "-streaming ";
+					if (Params.CookOnTheFlyStreaming)
+					{
+						TempCmdLine += "-streaming ";
+					}
+					else if (SC.StageTargetPlatform.PlatformType != UnrealTargetPlatform.IOS)
+					{
+						// per josh, allowcaching is deprecated/doesn't make sense for iOS.
+						TempCmdLine += "-allowcaching ";
+					}
 				}
-				else if (SC.StageTargetPlatform.PlatformType != UnrealTargetPlatform.IOS)
+				else if (Params.UsePak(SC.StageTargetPlatform))
 				{
-					// per josh, allowcaching is deprecated/doesn't make sense for iOS.
-					TempCmdLine += "-allowcaching ";
+					if (Params.SignedPak)
+					{
+						TempCmdLine += "-signedpak ";
+					}
 				}
-			}
-			else if (Params.UsePak(SC.StageTargetPlatform))
-			{
-				if (Params.SignedPak)
+				else if (!Params.Stage)
 				{
-					TempCmdLine += "-signedpak ";
+					var SandboxPath = CombinePaths(SC.RuntimeProjectRootDir.FullName, "Saved", "Cooked", SC.CookPlatform);
+					if (!SC.StageTargetPlatform.LaunchViaUFE)
+					{
+						TempCmdLine += "-sandbox=" + CommandUtils.MakePathSafeToUseWithCommandLine(SandboxPath) + " ";
+					}
+					else
+					{
+						TempCmdLine += "-sandbox=\'" + SandboxPath + "\' ";
+					}
 				}
 			}
-			else if (!Params.Stage)
+			else
 			{
-				var SandboxPath = CombinePaths(SC.RuntimeProjectRootDir.FullName, "Saved", "Cooked", SC.CookPlatform);
-				if (!SC.StageTargetPlatform.LaunchViaUFE)
+				ClientApp = CombinePaths(CmdEnv.LocalRoot, "Engine/Binaries", PlatformName, "UnrealEditor.exe");
+				if (!Params.EditorTest)
 				{
-					TempCmdLine += "-sandbox=" + CommandUtils.MakePathSafeToUseWithCommandLine(SandboxPath) + " ";
+					TempCmdLine += "-game " + Params.MapToRun + " ";
 				}
 				else
 				{
-					TempCmdLine += "-sandbox=\'" + SandboxPath + "\' ";
+					TempCmdLine += Params.MapToRun + " ";
 				}
 			}
-		}
-		else
-		{
-			ClientApp = CombinePaths(CmdEnv.LocalRoot, "Engine/Binaries", PlatformName, "UnrealEditor.exe");
-			if (!Params.EditorTest)
+			if (Params.LogWindow)
 			{
-				TempCmdLine += "-game " + Params.MapToRun + " ";
+				// Without NoStdOutRedirect '-log' doesn't log anything to the window
+				ClientRunFlags |= ERunOptions.NoStdOutRedirect;
+				TempCmdLine += "-log ";
+			}
+			else
+			{
+				TempCmdLine += "-stdout ";
+			}
+			if (Params.Unattended)
+			{
+				TempCmdLine += "-unattended ";
+			}
+			if (IsBuildMachine || Params.Unattended)
+			{
+				TempCmdLine += "-buildmachine ";
+			}
+			if (Params.HasDDCGraph)
+			{
+				TempCmdLine += "-ddc=" + Params.DDCGraph + " ";
 			}
-            else
-            {
-                TempCmdLine += Params.MapToRun + " ";
-            }
-		}
-		if (Params.LogWindow)
-		{
-			// Without NoStdOutRedirect '-log' doesn't log anything to the window
-			ClientRunFlags |= ERunOptions.NoStdOutRedirect;
-			TempCmdLine += "-log ";
-		}
-		else
-		{
-			TempCmdLine += "-stdout ";
-		}
-		if (Params.Unattended)
-		{
-			TempCmdLine += "-unattended ";
-		}
-		if (IsBuildMachine || Params.Unattended)
-		{
-			TempCmdLine += "-buildmachine ";
-		}
-		if (Params.HasDDCGraph)
-		{
-			TempCmdLine += "-ddc=" + Params.DDCGraph +" ";
-		}
 
-		if (Params.CrashIndex > 0)
-		{
-			int RealIndex = Params.CrashIndex - 1;
-			if (RealIndex >= CrashCommands.Count())
+			if (Params.CrashIndex > 0)
 			{
-				throw new AutomationException("CrashIndex {0} is out of range...max={1}", Params.CrashIndex, CrashCommands.Count());
+				int RealIndex = Params.CrashIndex - 1;
+				if (RealIndex >= CrashCommands.Count())
+				{
+					throw new AutomationException("CrashIndex {0} is out of range...max={1}", Params.CrashIndex, CrashCommands.Count());
+				}
+				TempCmdLine += String.Format("-execcmds=\"debug {0}\" ", CrashCommands[RealIndex]);
 			}
-			TempCmdLine += String.Format("-execcmds=\"debug {0}\" ", CrashCommands[RealIndex]);
-		}
-		else if (Params.RunAutomationTest != "")
-		{
-			TempCmdLine += "-execcmds=\"automation list;runtests " + Params.RunAutomationTest + "\" ";
-		}
-		else if (Params.RunAutomationTests)
-		{
-			TempCmdLine += "-execcmds=\"automation list;runall\" ";
-		}
-		if (SC.StageTargetPlatform.UseAbsLog)
-		{
-			TempCmdLine += "-abslog=" + CommandUtils.MakePathSafeToUseWithCommandLine(ClientLogFile) + " ";
-		}
-		if (SC.StageTargetPlatform.PlatformType == UnrealTargetPlatform.Win64)
-		{
-			TempCmdLine += "-Messaging -Windowed ";
-		}
-		else
-		{
-			TempCmdLine += "-Messaging ";
-		}
-		if (Params.NullRHI && SC.StageTargetPlatform.PlatformType != UnrealTargetPlatform.Mac) // all macs have GPUs, and currently the mac dies with nullrhi
-		{
-			TempCmdLine += "-nullrhi ";
-		}
-		TempCmdLine += SC.StageTargetPlatform.GetLaunchExtraCommandLine(Params);
+			else if (Params.RunAutomationTest != "")
+			{
+				TempCmdLine += "-execcmds=\"automation list;runtests " + Params.RunAutomationTest + "\" ";
+			}
+			else if (Params.RunAutomationTests)
+			{
+				TempCmdLine += "-execcmds=\"automation list;runall\" ";
+			}
+			if (SC.StageTargetPlatform.UseAbsLog)
+			{
+				TempCmdLine += "-abslog=" + CommandUtils.MakePathSafeToUseWithCommandLine(ClientLogFile) + " ";
+			}
+			if (SC.StageTargetPlatform.PlatformType == UnrealTargetPlatform.Win64)
+			{
+				TempCmdLine += "-Messaging -Windowed ";
+			}
+			else
+			{
+				TempCmdLine += "-Messaging ";
+			}
+			if (Params.NullRHI && SC.StageTargetPlatform.PlatformType != UnrealTargetPlatform.Mac) // all macs have GPUs, and currently the mac dies with nullrhi
+			{
+				TempCmdLine += "-nullrhi ";
+			}
+			TempCmdLine += SC.StageTargetPlatform.GetLaunchExtraCommandLine(Params);
 
-		TempCmdLine += "-CrashForUAT ";
-		TempCmdLine += Params.RunCommandline;
+			TempCmdLine += "-CrashForUAT ";
+			TempCmdLine += Params.RunCommandline;
 
-		// todo: move this into the platform
-		if (SC.StageTargetPlatform.LaunchViaUFE)
-		{
-			ClientCmdLine = "-run=Launch ";
-            ClientCmdLine += "-Device=" + Params.Devices[0];
-            for (int DeviceIndex = 1; DeviceIndex < Params.Devices.Count; DeviceIndex++)
-            {
-                ClientCmdLine += "+" + Params.Devices[DeviceIndex];
-            }
-            ClientCmdLine += " ";
-			ClientCmdLine += "-Exe=\"" + ClientApp + "\" ";
-			ClientCmdLine += "-Targetplatform=" + Params.ClientTargetPlatforms[0].ToString() + " ";
-			ClientCmdLine += "-Params=\"" + TempCmdLine + "\"";
-			ClientApp = CombinePaths(CmdEnv.LocalRoot, "Engine/Binaries/Win64/UnrealFrontend.exe");
-
-			LogInformation("Launching via UFE:");
-			LogInformation("\tClientCmdLine: " + ClientCmdLine + "");
-		}
-		else
-		{
-			ClientCmdLine = TempCmdLine;
+			// todo: move this into the platform
+			if (SC.StageTargetPlatform.LaunchViaUFE)
+			{
+				ClientCmdLine = "-run=Launch ";
+				ClientCmdLine += "-Device=" + Params.Devices[0];
+				for (int DeviceIndex = 1; DeviceIndex < Params.Devices.Count; DeviceIndex++)
+				{
+					ClientCmdLine += "+" + Params.Devices[DeviceIndex];
+				}
+				ClientCmdLine += " ";
+				ClientCmdLine += "-Exe=\"" + ClientApp + "\" ";
+				ClientCmdLine += "-Targetplatform=" + Params.ClientTargetPlatforms[0].ToString() + " ";
+				ClientCmdLine += "-Params=\"" + TempCmdLine + "\"";
+				ClientApp = CombinePaths(CmdEnv.LocalRoot, "Engine/Binaries/Win64/UnrealFrontend.exe");
+
+				LogInformation("Launching via UFE:");
+				LogInformation("\tClientCmdLine: " + ClientCmdLine + "");
+			}
+			else
+			{
+				ClientCmdLine = TempCmdLine;
+			}
 		}
-	}
 
-	private static void ClientLogReaderProc(object ArgsContainer)
-	{
-		var Args = ArgsContainer as object[];
-		var ClientLogFile = (string)Args[0];
-		var ClientProcess = (IProcessResult)Args[1];
-		LogFileReaderProcess(ClientLogFile, ClientProcess, (string Output) =>
+		private static void ClientLogReaderProc(object ArgsContainer)
 		{
-			if (String.IsNullOrEmpty(Output) == false)
+			var Args = ArgsContainer as object[];
+			var ClientLogFile = (string)Args[0];
+			var ClientProcess = (IProcessResult)Args[1];
+			LogFileReaderProcess(ClientLogFile, ClientProcess, (string Output) =>
 			{
-				LogInformation(Output);
-			}
-			return true;
-		});
-	}
-
-	private static IProcessResult RunDedicatedServer(ProjectParams Params, string ServerLogFile, string AdditionalCommandLine)
-	{
-		ProjectParams ServerParams = new ProjectParams(Params);
-		ServerParams.Devices = new ParamList<string>(Params.ServerDevice);
+				if (String.IsNullOrEmpty(Output) == false)
+				{
+					LogInformation(Output);
+				}
+				return true;
+			});
+		}
 
-		if (ServerParams.ServerTargetPlatforms.Count == 0)
+		private static IProcessResult RunDedicatedServer(ProjectParams Params, string ServerLogFile, string AdditionalCommandLine)
 		{
-			throw new AutomationException("No ServerTargetPlatform set for RunDedicatedServer.");
-		}
+			ProjectParams ServerParams = new ProjectParams(Params);
+			ServerParams.Devices = new ParamList<string>(Params.ServerDevice);
 
-		var DeployContextList = CreateDeploymentContext(ServerParams, true);
+			if (ServerParams.ServerTargetPlatforms.Count == 0)
+			{
+				throw new AutomationException("No ServerTargetPlatform set for RunDedicatedServer.");
+			}
 
-		if (DeployContextList.Count == 0)
-		{
-			throw new AutomationException("No DeployContextList for RunDedicatedServer.");
-		}
+			var DeployContextList = CreateDeploymentContext(ServerParams, true);
 
-		var SC = DeployContextList[0];
+			if (DeployContextList.Count == 0)
+			{
+				throw new AutomationException("No DeployContextList for RunDedicatedServer.");
+			}
 
-		var ServerApp = CombinePaths(CmdEnv.LocalRoot, "Engine/Binaries/Win64/UnrealEditor.exe");
-		if (ServerParams.Cook)
-		{
-			List<FileReference> Exes = SC.StageTargetPlatform.GetExecutableNames(SC);
-			ServerApp = Exes[0].FullName;
-		}
-		var Args = ServerParams.Cook ? "" : (SC.ProjectArgForCommandLines + " ");
-		Console.WriteLine(Params.ServerDeviceAddress);
-		TargetPlatformDescriptor ServerPlatformDesc = ServerParams.ServerTargetPlatforms[0];
-		if (ServerParams.Cook && ServerPlatformDesc.Type == UnrealTargetPlatform.Linux && !String.IsNullOrEmpty(ServerParams.ServerDeviceAddress))
-		{
-			ServerApp = @"C:\Windows\system32\cmd.exe";
+			var SC = DeployContextList[0];
 
-			string plinkPath = CombinePaths(Environment.GetEnvironmentVariable("LINUX_ROOT"), "bin/PLINK.exe ");
-			string exePath = CombinePaths(SC.ShortProjectName, "Binaries", ServerPlatformDesc.Type.ToString(), SC.ShortProjectName + "Server");
-			if (ServerParams.ServerConfigsToBuild[0] != UnrealTargetConfiguration.Development)
+			var ServerApp = CombinePaths(CmdEnv.LocalRoot, "Engine/Binaries/Win64/UnrealEditor.exe");
+			if (ServerParams.Cook)
 			{
-				exePath += "-" + ServerPlatformDesc.Type.ToString() + "-" + ServerParams.ServerConfigsToBuild[0].ToString();
+				List<FileReference> Exes = SC.StageTargetPlatform.GetExecutableNames(SC);
+				ServerApp = Exes[0].FullName;
 			}
-			exePath = CombinePaths("LinuxServer", exePath.ToLower()).Replace("\\", "/");
-			Args = String.Format("/k {0} -batch -ssh -t -i {1} {2}@{3} {4} {5} {6} -server -Messaging", plinkPath, ServerParams.DevicePassword, ServerParams.DeviceUsername, ServerParams.ServerDeviceAddress, exePath, Args, ServerParams.MapToRun);
-		}
-		else
-		{
-			var Map = ServerParams.MapToRun;
-			if (!String.IsNullOrEmpty(ServerParams.AdditionalServerMapParams))
+			var Args = ServerParams.Cook ? "" : (SC.ProjectArgForCommandLines + " ");
+			Console.WriteLine(Params.ServerDeviceAddress);
+			TargetPlatformDescriptor ServerPlatformDesc = ServerParams.ServerTargetPlatforms[0];
+			if (ServerParams.Cook && ServerPlatformDesc.Type == UnrealTargetPlatform.Linux && !String.IsNullOrEmpty(ServerParams.ServerDeviceAddress))
 			{
-				Map += ServerParams.AdditionalServerMapParams;
+				ServerApp = @"C:\Windows\system32\cmd.exe";
+
+				string plinkPath = CombinePaths(Environment.GetEnvironmentVariable("LINUX_ROOT"), "bin/PLINK.exe ");
+				string exePath = CombinePaths(SC.ShortProjectName, "Binaries", ServerPlatformDesc.Type.ToString(), SC.ShortProjectName + "Server");
+				if (ServerParams.ServerConfigsToBuild[0] != UnrealTargetConfiguration.Development)
+				{
+					exePath += "-" + ServerPlatformDesc.Type.ToString() + "-" + ServerParams.ServerConfigsToBuild[0].ToString();
+				}
+				exePath = CombinePaths("LinuxServer", exePath.ToLower()).Replace("\\", "/");
+				Args = String.Format("/k {0} -batch -ssh -t -i {1} {2}@{3} {4} {5} {6} -server -Messaging", plinkPath, ServerParams.DevicePassword, ServerParams.DeviceUsername, ServerParams.ServerDeviceAddress, exePath, Args, ServerParams.MapToRun);
 			}
-			if (Params.FakeClient)
+			else
 			{
-				Map += "?fake";
+				var Map = ServerParams.MapToRun;
+				if (!String.IsNullOrEmpty(ServerParams.AdditionalServerMapParams))
+				{
+					Map += ServerParams.AdditionalServerMapParams;
+				}
+				if (Params.FakeClient)
+				{
+					Map += "?fake";
+				}
+				Args += String.Format("{0} -server -abslog={1}  -log -Messaging", Map, CommandUtils.MakePathSafeToUseWithCommandLine(ServerLogFile));
+				if (Params.Unattended)
+				{
+					Args += " -unattended";
+				}
+
+				if (Params.ServerCommandline.Length > 0)
+				{
+					Args += " " + Params.ServerCommandline;
+				}
 			}
-			Args += String.Format("{0} -server -abslog={1}  -log -Messaging", Map, CommandUtils.MakePathSafeToUseWithCommandLine(ServerLogFile));
-			if (Params.Unattended)
+
+			if (ServerParams.UsePak(SC.StageTargetPlatform))
 			{
-				Args += " -unattended";
+				if (ServerParams.SignedPak)
+				{
+					Args += " -signedpak";
+				}
 			}
-			
-			if (Params.ServerCommandline.Length > 0)
+
+			if (Params.HasDDCGraph)
 			{
-				Args += " " + Params.ServerCommandline;
+				Args += " -ddc=" + Params.DDCGraph;
 			}
-		}
 
-		if (ServerParams.UsePak(SC.StageTargetPlatform))
-		{
-			if (ServerParams.SignedPak)
+			if (IsBuildMachine || Params.Unattended)
 			{
-				Args += " -signedpak";
+				Args += " -buildmachine";
 			}
-		}
+			Args += " -CrashForUAT";
+			Args += " " + AdditionalCommandLine;
 
-		if (Params.HasDDCGraph)
-		{
-			Args += " -ddc=" + Params.DDCGraph;
-		}
 
-		if (IsBuildMachine || Params.Unattended)
-		{
-			Args += " -buildmachine";
-		}
-		Args += " -CrashForUAT";
-		Args += " " + AdditionalCommandLine;
+			if (ServerParams.Cook && ServerPlatformDesc.Type == UnrealTargetPlatform.Linux && !String.IsNullOrEmpty(ServerParams.ServerDeviceAddress))
+			{
+				Args += String.Format(" 2>&1 > {0}", ServerLogFile);
+			}
 
+			PushDir(Path.GetDirectoryName(ServerApp));
+			var Result = Run(ServerApp, Args, null, ERunOptions.AllowSpew | ERunOptions.NoWaitForExit | ERunOptions.AppMustExist | ERunOptions.NoStdOutRedirect);
+			PopDir();
 
-		if (ServerParams.Cook && ServerPlatformDesc.Type == UnrealTargetPlatform.Linux && !String.IsNullOrEmpty(ServerParams.ServerDeviceAddress))
-		{
-			Args += String.Format(" 2>&1 > {0}", ServerLogFile);
+			return Result;
 		}
 
-		PushDir(Path.GetDirectoryName(ServerApp));
-		var Result = Run(ServerApp, Args, null, ERunOptions.AllowSpew | ERunOptions.NoWaitForExit | ERunOptions.AppMustExist | ERunOptions.NoStdOutRedirect);
-		PopDir();
-
-		return Result;
-	}
-
-	private static IProcessResult RunCookOnTheFlyServer(FileReference ProjectName, string ServerLogFile, string TargetPlatform, string AdditionalCommandLine)
-	{
-		var ServerApp = HostPlatform.Current.GetUE4ExePath("UnrealEditor.exe");
-		var Args = String.Format("{0} -run=cook -cookonthefly -unattended -CrashForUAT -log",
-			CommandUtils.MakePathSafeToUseWithCommandLine(ProjectName.FullName));
-		if (!String.IsNullOrEmpty(ServerLogFile))
+		private static IProcessResult RunCookOnTheFlyServer(FileReference ProjectName, string ServerLogFile, string TargetPlatform, string AdditionalCommandLine)
 		{
-			// Issue with dotnet not allowing any files with an exclusive advisory lock to be opened for read-only or copied
-			// [REDACTED_PROFILE_URL]/runtime/issues/34126
-			if (TargetPlatform == "Linux")
+			var ServerApp = HostPlatform.Current.GetUE4ExePath("UnrealEditor.exe");
+			var Args = String.Format("{0} -run=cook -cookonthefly -unattended -CrashForUAT -log",
+				CommandUtils.MakePathSafeToUseWithCommandLine(ProjectName.FullName));
+			if (!String.IsNullOrEmpty(ServerLogFile))
 			{
-				Args += " -noexclusivelockonwrite";
+				// Issue with dotnet not allowing any files with an exclusive advisory lock to be opened for read-only or copied
+				// [REDACTED_PROFILE_URL]/runtime/issues/34126
+				if (TargetPlatform == "Linux")
+				{
+					Args += " -noexclusivelockonwrite";
+				}
+
+				Args += " -abslog=" + CommandUtils.MakePathSafeToUseWithCommandLine(ServerLogFile);
+			}
+			if (IsBuildMachine)
+			{
+				Args += " -buildmachine";
 			}
+			Args += " " + AdditionalCommandLine;
 
-			Args += " -abslog=" + CommandUtils.MakePathSafeToUseWithCommandLine(ServerLogFile);
+			// Run the server (Without NoStdOutRedirect -log doesn't log anything to the window)
+			PushDir(Path.GetDirectoryName(ServerApp));
+			var Result = Run(ServerApp, Args, null, ERunOptions.AllowSpew | ERunOptions.NoWaitForExit | ERunOptions.AppMustExist | ERunOptions.NoStdOutRedirect);
+			PopDir();
+			return Result;
 		}
-		if (IsBuildMachine)
-		{
-			Args += " -buildmachine";
-		}
-		Args += " " + AdditionalCommandLine;
-
-		// Run the server (Without NoStdOutRedirect -log doesn't log anything to the window)
-		PushDir(Path.GetDirectoryName(ServerApp));
-		var Result = Run(ServerApp, Args, null, ERunOptions.AllowSpew | ERunOptions.NoWaitForExit | ERunOptions.AppMustExist | ERunOptions.NoStdOutRedirect);
-		PopDir();
-		return Result;
-	}
 
-	private static IProcessResult RunFileServer(ProjectParams Params, string ServerLogFile, string AdditionalCommandLine)
-	{
+		private static IProcessResult RunFileServer(ProjectParams Params, string ServerLogFile, string AdditionalCommandLine)
+		{
 #if false
 		// this section of code would provide UFS with a more accurate file mapping
 		var SC = new StagingContext(Params, false);
@@ -1046,51 +1048,52 @@ public partial class Project : CommandUtils
 		string UnrealFileServerResponseFileName = CombinePaths(CmdEnv.LogFolder, "UnrealFileServerList.txt");
 		File.WriteAllLines(UnrealFileServerResponseFileName, UnrealFileServerResponseFile);
 #endif
-		var UnrealFileServerExe = HostPlatform.Current.GetUE4ExePath("UnrealFileServer.exe");
-
-		LogInformation("Running UnrealFileServer *******");
-		var Args = String.Format("{0} -abslog={1} -unattended -CrashForUAT -log {2}",
-						CommandUtils.MakePathSafeToUseWithCommandLine(Params.RawProjectPath.FullName),
-						CommandUtils.MakePathSafeToUseWithCommandLine(ServerLogFile),
-						AdditionalCommandLine);
-		if (IsBuildMachine)
-		{
-			Args += " -buildmachine";
+			var UnrealFileServerExe = HostPlatform.Current.GetUE4ExePath("UnrealFileServer.exe");
+
+			LogInformation("Running UnrealFileServer *******");
+			var Args = String.Format("{0} -abslog={1} -unattended -CrashForUAT -log {2}",
+							CommandUtils.MakePathSafeToUseWithCommandLine(Params.RawProjectPath.FullName),
+							CommandUtils.MakePathSafeToUseWithCommandLine(ServerLogFile),
+							AdditionalCommandLine);
+			if (IsBuildMachine)
+			{
+				Args += " -buildmachine";
+			}
+			PushDir(Path.GetDirectoryName(UnrealFileServerExe));
+			var Result = Run(UnrealFileServerExe, Args, null, ERunOptions.AllowSpew | ERunOptions.NoWaitForExit | ERunOptions.AppMustExist | ERunOptions.NoStdOutRedirect);
+			PopDir();
+			return Result;
 		}
-		PushDir(Path.GetDirectoryName(UnrealFileServerExe));
-		var Result = Run(UnrealFileServerExe, Args, null, ERunOptions.AllowSpew | ERunOptions.NoWaitForExit | ERunOptions.AppMustExist | ERunOptions.NoStdOutRedirect);
-		PopDir();
-		return Result;
-	}
 
-	private static bool StartUnrealTrace()
-	{
-		// [TEMPORARY] - UnrealTrace server is currently only available on Windows
-		if (!RuntimePlatform.IsWindows)
+		private static bool StartUnrealTrace()
 		{
-			return true;
-		}
-		// [/TEMPORARY]
+			// [TEMPORARY] - UnrealTrace server is currently only available on Windows
+			if (!RuntimePlatform.IsWindows)
+			{
+				return true;
+			}
+			// [/TEMPORARY]
 
-		LogInformation("UnrealTrace: Starting server");
+			LogInformation("UnrealTrace: Starting server");
 
-		// Locate the UnrealTrace binary
-		var UnrealTracePath = HostPlatform.Current.GetUE4ExePath("UnrealTraceServer.exe");
-		if (!File.Exists(UnrealTracePath))
-		{
-			LogWarning("UnrealTrace: Unable to locate binary at " + UnrealTracePath);
-			return false;
-		}
+			// Locate the UnrealTrace binary
+			var UnrealTracePath = HostPlatform.Current.GetUE4ExePath("UnrealTraceServer.exe");
+			if (!File.Exists(UnrealTracePath))
+			{
+				LogWarning("UnrealTrace: Unable to locate binary at " + UnrealTracePath);
+				return false;
+			}
 
-		// Launch UnrealTrace and wait for it to fork and return
-		Process Proc = Process.Start(UnrealTracePath);
-		Proc.WaitForExit();
-		if (Proc.ExitCode != 0)
-		{
-			LogWarning("UnrealTrace: Failed to start server; ExitCode=" + Proc.ExitCode);
-			return false;
-		}
+			// Launch UnrealTrace and wait for it to fork and return
+			Process Proc = Process.Start(UnrealTracePath);
+			Proc.WaitForExit();
+			if (Proc.ExitCode != 0)
+			{
+				LogWarning("UnrealTrace: Failed to start server; ExitCode=" + Proc.ExitCode);
+				return false;
+			}
 
-		return true;
+			return true;
+		}
 	}
 }
