UAT: Fix some UAT build products being copied over to network share multiple times.

diff --git a/Engine/Source/Programs/AutomationTool/Scripts/CopyUAT.Automation.cs b/Engine/Source/Programs/AutomationTool/Scripts/CopyUAT.Automation.cs
--- a/Engine/Source/Programs/AutomationTool/Scripts/CopyUAT.Automation.cs
+++ b/Engine/Source/Programs/AutomationTool/Scripts/CopyUAT.Automation.cs
@@ -30,7 +30,7 @@ public class CopyUAT : BuildCommand
 		Build.AddUBTFilesToBuildProducts();
 
 		// Get a list of all the input files
-		List<FileReference> SourceFiles = new List<FileReference>();
+		HashSet<FileReference> SourceFiles = new HashSet<FileReference>();
 		foreach(string BuildProductFile in Build.BuildProductFiles)
 		{
 			FileReference SourceFile = new FileReference(BuildProductFile);
