UAT: Use ThreadedCopyFiles to get UAT onto network share.

diff --git a/Engine/Source/Programs/AutomationTool/Scripts/CopyUAT.Automation.cs b/Engine/Source/Programs/AutomationTool/Scripts/CopyUAT.Automation.cs
--- a/Engine/Source/Programs/AutomationTool/Scripts/CopyUAT.Automation.cs
+++ b/Engine/Source/Programs/AutomationTool/Scripts/CopyUAT.Automation.cs
@@ -30,33 +30,29 @@ public class CopyUAT : BuildCommand
 		Build.AddUBTFilesToBuildProducts();
 
 		// Get a list of all the input files
-		HashSet<FileReference> SourceFiles = new HashSet<FileReference>();
+		HashSet<FileReference> SourceFilesSet = new HashSet<FileReference>();
 		foreach(string BuildProductFile in Build.BuildProductFiles)
 		{
 			FileReference SourceFile = new FileReference(BuildProductFile);
-			SourceFiles.Add(SourceFile);
+			SourceFilesSet.Add(SourceFile);
 
 			FileReference SourceSymbolFile = SourceFile.ChangeExtension(".pdb");
 			if(FileReference.Exists(SourceSymbolFile))
 			{
-				SourceFiles.Add(SourceSymbolFile);
+				SourceFilesSet.Add(SourceSymbolFile);
 			}
 
 			FileReference DocumentationFile = SourceFile.ChangeExtension(".xml");
 			if(FileReference.Exists(DocumentationFile))
 			{
-				SourceFiles.Add(DocumentationFile);
+				SourceFilesSet.Add(DocumentationFile);
 			}
 		}
 
 		// Copy all the files over
 		DirectoryReference TargetDir = new DirectoryReference(TargetDirParam);
-		foreach(FileReference SourceFile in SourceFiles)
-		{
-			FileReference TargetFile = FileReference.Combine(TargetDir, SourceFile.MakeRelativeTo(Unreal.RootDirectory));
-			DirectoryReference.CreateDirectory(TargetFile.Directory);
-			CommandUtils.CopyFile(SourceFile.FullName, TargetFile.FullName);
-		}
+		List<FileReference> SourceFiles = SourceFilesSet.OrderBy(x => x.FullName).ToList();
+		CommandUtils.ThreadedCopyFiles(SourceFiles, Unreal.RootDirectory, TargetDir);
 
 		LogInformation("Copied {0} files to {1}", SourceFiles.Count, TargetDir);
 		File.WriteAllLines(Path.Combine(TargetDirParam, "CopiedFiles.txt"), SourceFiles.Select(F => F.FullName));
