Add Horde Job Step link in telemetry results

#rnx

diff --git a/Engine/Source/Programs/AutomationTool/Gauntlet/Framework/Utils/Gauntlet.HordeReport.cs b/Engine/Source/Programs/AutomationTool/Gauntlet/Framework/Utils/Gauntlet.HordeReport.cs
--- a/Engine/Source/Programs/AutomationTool/Gauntlet/Framework/Utils/Gauntlet.HordeReport.cs
+++ b/Engine/Source/Programs/AutomationTool/Gauntlet/Framework/Utils/Gauntlet.HordeReport.cs
@@ -56,6 +56,27 @@ namespace Gauntlet
 				return string.Empty;
 			}
 		}
+		/// <summary>
+		/// Default Job Step Horde link if running under Horde agent, otherwise return an empty string
+		/// </summary>
+		public static string DefaultHordeJobStepLink
+		{
+			get
+			{
+				string HordeJobStepId = Environment.GetEnvironmentVariable("UE_HORDE_STEPID");
+				string HordeJobLink = DefaultHordeJobLink;
+				if (!string.IsNullOrEmpty(HordeJobStepId))
+				{
+					return string.Format("{0}?step={1}", HordeJobLink, HordeJobStepId);
+				}
+				else if(!string.IsNullOrEmpty(HordeJobLink))
+				{
+					return HordeJobLink;
+				}
+
+				return string.Empty;
+			}
+		}		
 
 		public abstract class BaseHordeReport : BaseTestReport
 		{
diff --git a/Engine/Source/Programs/AutomationTool/Gauntlet/Unreal/Utils/Gauntlet.UnrealTelemetry.cs b/Engine/Source/Programs/AutomationTool/Gauntlet/Unreal/Utils/Gauntlet.UnrealTelemetry.cs
--- a/Engine/Source/Programs/AutomationTool/Gauntlet/Unreal/Utils/Gauntlet.UnrealTelemetry.cs
+++ b/Engine/Source/Programs/AutomationTool/Gauntlet/Unreal/Utils/Gauntlet.UnrealTelemetry.cs
@@ -194,7 +194,7 @@ namespace Gauntlet
 			if (CSVMappingStrings.Count > 0)
 			{
 				CSVMapping = new Dictionary<string, string>();
-				foreach(string StringLine in CSVMappingStrings)
+				foreach (string StringLine in CSVMappingStrings)
 				{
 					var SplitLine = StringLine.Split(":");
 					if (SplitLine.Length != 2)
@@ -248,14 +248,14 @@ namespace Gauntlet
 
 			if (string.IsNullOrEmpty(JobLink))
 			{
-				JobLink = HordeReport.DefaultHordeJobLink;
+				JobLink = HordeReport.DefaultHordeJobStepLink;
 			}
 			Context.SetProperty("JobLink", JobLink);
 
 			// Create a report
 			HordeReport.SimpleTestReport Report = new HordeReport.SimpleTestReport();
 
-			if(Report is ITelemetryReport TelemetryReport)
+			if (Report is ITelemetryReport TelemetryReport)
 			{
 				// Parse CSV
 				if (!string.IsNullOrEmpty(CSVDirectory))
@@ -295,4 +295,4 @@ namespace Gauntlet
 			return 0;
 		}
 	}
-}
\ No newline at end of file
+}
