Preserve slashes and other regular expression operations when throwing text at markdown.

This fixes the handling of \ and + characters in file paths when communicating the logdir through to the horde web pages.
For now remove the markdown link decoration, as whatever markdown processor is running on this .md file, it doesn't yield a file:// anchor tag that works in Chrome or Firefox, at least with Windows network paths (\\computer\share\folder\file.txt).

diff --git a/Engine/Source/Programs/AutomationTool/Gauntlet/Framework/Horde/Gauntlet.Horde.cs b/Engine/Source/Programs/AutomationTool/Gauntlet/Framework/Horde/Gauntlet.Horde.cs
--- a/Engine/Source/Programs/AutomationTool/Gauntlet/Framework/Horde/Gauntlet.Horde.cs
+++ b/Engine/Source/Programs/AutomationTool/Gauntlet/Framework/Horde/Gauntlet.Horde.cs
@@ -47,7 +47,9 @@ namespace Gauntlet
 
 				string MarkdownFilename = "GauntletStepDetails.md";
 
-				string Markdown = $"Gauntlet Artifacts: [{Globals.LogDir}](file://{Globals.LogDir})";
+				// file:// is attractive here, and file://\\computer\share\folder\file.txt does work ok in Chrome, but it does not Firefox
+				// so, lacking a uniform way to link to a network filesystem, just put the path here as static text and escaped for markdown
+				string Markdown = System.Text.RegularExpressions.Regex.Replace($"Gauntlet Artifacts: {Globals.LogDir}", @"([|\\*\+])", @"\$1");
 
 				File.WriteAllText(Path.Combine(LogFolder, MarkdownFilename), Markdown);
 
