IterationProfile: Add -dx11 to side step UE-117842

#rnx

diff --git a/Engine/Source/Programs/AutomationTool/Gauntlet/Unreal/Editor/EditorTests.cs b/Engine/Source/Programs/AutomationTool/Gauntlet/Unreal/Editor/EditorTests.cs
--- a/Engine/Source/Programs/AutomationTool/Gauntlet/Unreal/Editor/EditorTests.cs
+++ b/Engine/Source/Programs/AutomationTool/Gauntlet/Unreal/Editor/EditorTests.cs
@@ -135,6 +135,9 @@ namespace UnrealEditor
 				AppConfig.CommandLineParams.Add("BENCHMARK");
 				AppConfig.CommandLineParams.Add("Deterministic");
 			}
+
+			// Workaround for UE-117842
+			AppConfig.CommandLineParams.Add("dx11");
 		}
 	}
 
