Setting CustomConfig=Foo in a Target.cs file will cause it to overlay config files from Project/Config/Custom/Foo on top of the defaults, allowing easy override of things like OSS settings to support multiple stores
In development, -CustomConfig=foo can be specified in both C++ and C# to enable the same functionality, which can be used to select between multiple different stage/deploy configurations in a build script

diff --git a/Engine/Source/Programs/AutomationTool/AutomationUtils/DeploymentContext.cs b/Engine/Source/Programs/AutomationTool/AutomationUtils/DeploymentContext.cs
--- a/Engine/Source/Programs/AutomationTool/AutomationUtils/DeploymentContext.cs
+++ b/Engine/Source/Programs/AutomationTool/AutomationUtils/DeploymentContext.cs
@@ -136,6 +136,11 @@ public class DeploymentContext //: ProjectParams
 	/// </summary>
 	public List<StageTarget> StageTargets;
 
+	/// <summary>
+	/// Extra subdirectory to load config files out of, for making multiple types of builds with the same platform
+	/// </summary>
+	public string CustomConfig;
+
 	/// <summary>
 	/// This is the root directory that contains the engine: d:\a\UE4\
 	/// </summary>
@@ -443,8 +448,30 @@ public class DeploymentContext //: ProjectParams
 		ProjectArgForCommandLines = CommandUtils.MakePathSafeToUseWithCommandLine(RawProjectPath.FullName);
 		CookSourceRuntimeRootDir = RuntimeRootDir = LocalRoot;
 		RuntimeProjectRootDir = ProjectRoot;
-       
-        if (Stage)
+
+		// Parse the custom config dir out of the receipts
+		foreach (StageTarget Target in StageTargets)
+		{
+			var Results = Target.Receipt.AdditionalProperties.Where(x => x.Name == "CustomConfig");
+			foreach (var Property in Results)
+			{
+				string FoundCustomConfig = Property.Value;
+				if (String.IsNullOrEmpty(FoundCustomConfig))
+				{
+					continue;
+				}
+				else if (String.IsNullOrEmpty(CustomConfig))
+				{
+					CustomConfig = FoundCustomConfig;
+				}
+				else if (CustomConfig != FoundCustomConfig)
+				{
+					throw new AutomationException("Cannot deploy targts with conflicting CustomConfig values! {0} does not match {1}", FoundCustomConfig, CustomConfig);
+				}
+			}
+		}
+
+		if (Stage)
 		{
 			CommandUtils.CreateDirectory(StageDirectory.FullName);
 
@@ -470,7 +497,7 @@ public class DeploymentContext //: ProjectParams
 		RestrictedFolderNames.Remove(StageTargetPlatform.IniPlatformType.ToString());
 
 		// Read the game config files
-		ConfigHierarchy GameConfig = ConfigCache.ReadHierarchy(ConfigHierarchyType.Game, ProjectRoot, InTargetPlatform.PlatformType);
+		ConfigHierarchy GameConfig = ConfigCache.ReadHierarchy(ConfigHierarchyType.Game, ProjectRoot, InTargetPlatform.PlatformType, CustomConfig);
 
 		// Read the list of directories to remap when staging
 		List<string> RemapDirectoriesList;
@@ -524,12 +551,11 @@ public class DeploymentContext //: ProjectParams
 		ReadConfigFileList(GameConfig, "Staging", "BlacklistConfigFiles", BlacklistConfigFiles);
 
 		// Grab the game ini data
-		ConfigHierarchy GameIni = ConfigCache.ReadHierarchy(ConfigHierarchyType.Game, ProjectRoot, InTargetPlatform.PlatformType);
-		String IniPath = "/Script/UnrealEd.ProjectPackagingSettings";
+		String PackagingIniPath = "/Script/UnrealEd.ProjectPackagingSettings";
 
 		// Read the config blacklists
-		GameIni.GetArray(IniPath, "IniKeyBlacklist", out IniKeyBlacklist);
-		GameIni.GetArray(IniPath, "IniSectionBlacklist", out IniSectionBlacklist);
+		GameConfig.GetArray(PackagingIniPath, "IniKeyBlacklist", out IniKeyBlacklist);
+		GameConfig.GetArray(PackagingIniPath, "IniSectionBlacklist", out IniSectionBlacklist);
 
 		// TODO: Drive these lists from a config file
 		IniSuffixWhitelist = new List<string>
@@ -576,7 +602,7 @@ public class DeploymentContext //: ProjectParams
 		else
 		{
 			bool bSetting = false;
-			if (GameIni.GetBool(IniPath, "bGenerateChunks", out bSetting))
+			if (GameConfig.GetBool(PackagingIniPath, "bGenerateChunks", out bSetting))
 			{
 				PlatformUsesChunkManifests = bSetting;
 			}
diff --git a/Engine/Source/Programs/AutomationTool/Scripts/CopyBuildToStagingDirectory.Automation.cs b/Engine/Source/Programs/AutomationTool/Scripts/CopyBuildToStagingDirectory.Automation.cs
--- a/Engine/Source/Programs/AutomationTool/Scripts/CopyBuildToStagingDirectory.Automation.cs
+++ b/Engine/Source/Programs/AutomationTool/Scripts/CopyBuildToStagingDirectory.Automation.cs
@@ -638,7 +638,7 @@ namespace AutomationScripts
 
 			LogInformation("Creating Staging Manifest...");
 
-			ConfigHierarchy PlatformGameConfig = ConfigCache.ReadHierarchy(ConfigHierarchyType.Game, DirectoryReference.FromFile(Params.RawProjectPath), SC.StageTargetPlatform.IniPlatformType);
+			ConfigHierarchy PlatformGameConfig = ConfigCache.ReadHierarchy(ConfigHierarchyType.Game, DirectoryReference.FromFile(Params.RawProjectPath), SC.StageTargetPlatform.IniPlatformType, SC.CustomConfig);
 
 			if (Params.HasIterateSharedCookedBuild)
 			{
@@ -1129,7 +1129,7 @@ namespace AutomationScripts
 					}
 					// check if the game will be verifying ssl connections - if not, we can skip staging files that won't be needed
 					bool bStageSSLCertificates = false;
-					ConfigHierarchy PlatformEngineConfig = ConfigCache.ReadHierarchy(ConfigHierarchyType.Engine, DirectoryReference.FromFile(Params.RawProjectPath), SC.StageTargetPlatform.IniPlatformType);
+					ConfigHierarchy PlatformEngineConfig = ConfigCache.ReadHierarchy(ConfigHierarchyType.Engine, DirectoryReference.FromFile(Params.RawProjectPath), SC.StageTargetPlatform.IniPlatformType, SC.CustomConfig);
 					if (PlatformEngineConfig != null)
 					{
 						PlatformEngineConfig.GetBool("/Script/Engine.NetworkSettings", "n.VerifyPeer", out bStageSSLCertificates);
@@ -1397,6 +1397,7 @@ namespace AutomationScripts
 		/// <param name="SC">The staging context</param>
 		/// <param name="ConfigDir">Directory containing the config files</param>
 		/// <param name="ConfigFile">The config file to check</param>
+		/// <param name="PlatformExtenionName">Name of platform to scan outside the main directory</param>
 		/// <returns>True if the file should be staged, false otherwise</returns>
 		static Nullable<bool> ShouldStageConfigFile(DeploymentContext SC, DirectoryReference ConfigDir, FileReference ConfigFile, string PlatformExtensionName)
 		{
@@ -1410,8 +1411,35 @@ namespace AutomationScripts
 				return false;
 			}
 
+			const string CustomPrefix = "custom/";
 			string NormalizedPath = ConfigFile.MakeRelativeTo(ConfigDir).ToLowerInvariant().Replace('\\', '/');
 
+			if (NormalizedPath.StartsWith(CustomPrefix))
+			{
+				// Strip custom prefix, then check if it matches our specified custom config
+				NormalizedPath = NormalizedPath.Substring(CustomPrefix.Length);
+
+				if (!String.IsNullOrEmpty(SC.CustomConfig))
+				{
+					string PrefixToStrip = SC.CustomConfig.ToLowerInvariant() + "/";
+					if (NormalizedPath.StartsWith(PrefixToStrip))
+					{
+						// Strip custom config path off the front
+						NormalizedPath = NormalizedPath.Substring(PrefixToStrip.Length);
+					}
+					else
+					{
+						// If custom config is specified, we want to ignore others automatically
+						return false;
+					}
+				}
+				else
+				{
+					// Ignore all custom config directories
+					return false;
+				}
+			}
+
 			int DirectoryIdx = NormalizedPath.IndexOf('/');
 			if (DirectoryIdx == -1)
 			{
@@ -1821,7 +1849,7 @@ namespace AutomationScripts
 
 
 
-			ConfigHierarchy PakRulesConfig = ConfigCache.ReadHierarchy(ConfigHierarchyType.PakFileRules, DirectoryReference.FromFile(Params.RawProjectPath), SC.StageTargetPlatform.IniPlatformType);
+			ConfigHierarchy PakRulesConfig = ConfigCache.ReadHierarchy(ConfigHierarchyType.PakFileRules, DirectoryReference.FromFile(Params.RawProjectPath), SC.StageTargetPlatform.IniPlatformType, SC.CustomConfig);
 
 			bool bChunkedBuild = SC.PlatformUsesChunkManifests && DoesChunkPakManifestExist(Params, SC);
 
@@ -2287,7 +2315,7 @@ namespace AutomationScripts
 
 			// read some compression settings from the project (once, shared across all pak commands)
 
-			ConfigHierarchy PlatformGameConfig = ConfigCache.ReadHierarchy(ConfigHierarchyType.Game, DirectoryReference.FromFile(Params.RawProjectPath), SC.StageTargetPlatform.IniPlatformType);
+			ConfigHierarchy PlatformGameConfig = ConfigCache.ReadHierarchy(ConfigHierarchyType.Game, DirectoryReference.FromFile(Params.RawProjectPath), SC.StageTargetPlatform.IniPlatformType, SC.CustomConfig);
 
 			// in standard runs, ProjectPackagingSettings/bCompressed is read by the program invoking this script
 			//	and used to pass "-compressed" on the command line
@@ -2771,7 +2799,7 @@ namespace AutomationScripts
 						AdditionalArgs += " -platform=" + ConfigHierarchy.GetIniPlatformName(SC.StageTargetPlatform.IniPlatformType);
 
 						Dictionary<string, string> UnrealPakResponseFile = PakParams.UnrealPakResponseFile;
-						if (ShouldCreateIoStoreContainerFiles(Params, SC.StageTargetPlatform))
+						if (ShouldCreateIoStoreContainerFiles(Params, SC.StageTargetPlatform, SC.CustomConfig))
 						{
 							bool bAllowBulkDataInIoStore = true;
 							if (!PlatformEngineConfig.GetBool("Core.System", "AllowBulkDataInIoStore", out bAllowBulkDataInIoStore))
@@ -2950,7 +2978,7 @@ namespace AutomationScripts
 
 					InternalUtils.SafeCreateDirectory(Path.GetDirectoryName(ReleaseVersionPath));
 					InternalUtils.SafeCopyFile(OutputLocation.FullName, ReleaseVersionPath);
-					if (ShouldCreateIoStoreContainerFiles(Params, SC.StageTargetPlatform))
+					if (ShouldCreateIoStoreContainerFiles(Params, SC.StageTargetPlatform, SC.CustomConfig))
 					{
 						InternalUtils.SafeCopyFile(Path.ChangeExtension(OutputLocation.FullName, ".utoc"), Path.ChangeExtension(ReleaseVersionPath, ".utoc"));
 						InternalUtils.SafeCopyFile(Path.ChangeExtension(OutputLocation.FullName, ".ucas"), Path.ChangeExtension(ReleaseVersionPath, ".ucas"));
@@ -3060,7 +3088,7 @@ namespace AutomationScripts
 						{
 							HashSet<string> IncludedExtensions = new HashSet<string>();
 							IncludedExtensions.Add(OutputFilenameExtension);
-							if (ShouldCreateIoStoreContainerFiles(Params, SC.StageTargetPlatform))
+							if (ShouldCreateIoStoreContainerFiles(Params, SC.StageTargetPlatform, SC.CustomConfig))
 							{
 								IncludedExtensions.Add(".ucas");
 								IncludedExtensions.Add(".utoc");
@@ -3355,7 +3383,7 @@ namespace AutomationScripts
 			LogInformation("Creating pak using streaming install manifests.");
 			DumpManifest(SC, CombinePaths(CmdEnv.LogFolder, "PrePak" + (SC.DedicatedServer ? "_Server" : "") + "_" + SC.CookPlatform));
 
-			ConfigHierarchy PlatformGameConfig = ConfigCache.ReadHierarchy(ConfigHierarchyType.Game, DirectoryReference.FromFile(Params.RawProjectPath), SC.StageTargetPlatform.IniPlatformType);
+			ConfigHierarchy PlatformGameConfig = ConfigCache.ReadHierarchy(ConfigHierarchyType.Game, DirectoryReference.FromFile(Params.RawProjectPath), SC.StageTargetPlatform.IniPlatformType, SC.CustomConfig);
 
 			bool bShouldGenerateEarlyDownloaderPakFile = false, bForceOneChunkPerFile = false;
 			string EarlyDownloaderPakFilePrefix = string.Empty;
@@ -3684,7 +3712,7 @@ namespace AutomationScripts
 			return TmpPackagingPath;
 		}
 
-		private static bool ShouldCreateIoStoreContainerFiles(ProjectParams Params, Platform StageTargetPlatform)
+		private static bool ShouldCreateIoStoreContainerFiles(ProjectParams Params, Platform StageTargetPlatform, String CustomConfig)
 		{
 			if (Params.CookOnTheFly)
 			{
@@ -3703,7 +3731,7 @@ namespace AutomationScripts
 
 			if (Params.Stage && !Params.SkipStage)
 			{
-				ConfigHierarchy PlatformGameConfig = ConfigCache.ReadHierarchy(ConfigHierarchyType.Game, DirectoryReference.FromFile(Params.RawProjectPath), StageTargetPlatform.IniPlatformType);
+				ConfigHierarchy PlatformGameConfig = ConfigCache.ReadHierarchy(ConfigHierarchyType.Game, DirectoryReference.FromFile(Params.RawProjectPath), StageTargetPlatform.IniPlatformType, CustomConfig);
 				bool bUseIoStore = false;
 				PlatformGameConfig.GetBool("/Script/UnrealEd.ProjectPackagingSettings", "bUseIoStore", out bUseIoStore);
 				return bUseIoStore;
diff --git a/Engine/Source/Programs/UnrealBuildTool/Configuration/TargetRules.cs b/Engine/Source/Programs/UnrealBuildTool/Configuration/TargetRules.cs
--- a/Engine/Source/Programs/UnrealBuildTool/Configuration/TargetRules.cs
+++ b/Engine/Source/Programs/UnrealBuildTool/Configuration/TargetRules.cs
@@ -397,6 +397,13 @@ namespace UnrealBuildTool
 		[CommandLine("-CompileAsDll")]
 		public bool bShouldCompileAsDLL = false;
 
+		/// <summary>
+		/// Extra subdirectory to load config files out of, for making multiple types of builds with the same platform
+		/// This will get baked into the game executable as CUSTOM_CONFIG and used when staging to filter files and settings
+		/// </summary>
+		[RequiresUniqueBuildEnvironment]
+		public string CustomConfig = String.Empty;
+
 		/// <summary>
 		/// Subfolder to place executables in, relative to the default location.
 		/// </summary>
@@ -2065,6 +2072,11 @@ namespace UnrealBuildTool
 			get { return Inner.SolutionDirectory; }
 		}
 
+		public string CustomConfig
+		{
+			get { return Inner.CustomConfig; }
+		}
+
 		public bool? bBuildInSolutionByDefault
 		{
 			get { return Inner.bBuildInSolutionByDefault; }
diff --git a/Engine/Source/Programs/UnrealBuildTool/Configuration/UEBuildTarget.cs b/Engine/Source/Programs/UnrealBuildTool/Configuration/UEBuildTarget.cs
--- a/Engine/Source/Programs/UnrealBuildTool/Configuration/UEBuildTarget.cs
+++ b/Engine/Source/Programs/UnrealBuildTool/Configuration/UEBuildTarget.cs
@@ -1523,6 +1523,12 @@ namespace UnrealBuildTool
 			// add the SDK used by the tool chain
 			Receipt.AdditionalProperties.Add(new ReceiptProperty("SDK", ToolChain.GetSDKVersion()));
 
+			if (!String.IsNullOrEmpty(Rules.CustomConfig))
+			{
+				// Pass through the custom config string
+				Receipt.AdditionalProperties.Add(new ReceiptProperty("CustomConfig", Rules.CustomConfig));
+			}
+
 			return Receipt;
 		}
 
@@ -3752,6 +3758,12 @@ namespace UnrealBuildTool
 				GlobalCompileEnvironment.Definitions.Add("WITH_CPP_MODULES=0");
 			}
 
+			// Define the custom config if specified, this should not be set in an editor build
+			if (!bUseSharedBuildEnvironment && !String.IsNullOrEmpty(Rules.CustomConfig))
+			{
+				GlobalCompileEnvironment.Definitions.Add(String.Format("CUSTOM_CONFIG=\"{0}\"", Rules.CustomConfig));
+			}
+
 			// Compile in the names of the module manifests
 			GlobalCompileEnvironment.Definitions.Add(String.Format("UBT_MODULE_MANIFEST=\"{0}\"", ModuleManifest.GetStandardFileName(AppName, Platform, Configuration, Architecture, false)));
 			GlobalCompileEnvironment.Definitions.Add(String.Format("UBT_MODULE_MANIFEST_DEBUGGAME=\"{0}\"", ModuleManifest.GetStandardFileName(AppName, Platform, UnrealTargetConfiguration.DebugGame, Architecture, true)));
diff --git a/Engine/Source/Programs/UnrealBuildTool/System/ConfigCache.cs b/Engine/Source/Programs/UnrealBuildTool/System/ConfigCache.cs
--- a/Engine/Source/Programs/UnrealBuildTool/System/ConfigCache.cs
+++ b/Engine/Source/Programs/UnrealBuildTool/System/ConfigCache.cs
@@ -81,17 +81,24 @@ namespace UnrealBuildTool
 			/// </summary>
 			public UnrealTargetPlatform Platform;
 
+			/// <summary>
+			/// Custom config subdirectory to load
+			/// </summary>
+			public string CustomConfig;
+
 			/// <summary>
 			/// Constructor
 			/// </summary>
 			/// <param name="Type">The hierarchy type</param>
 			/// <param name="ProjectDir">The project directory to read from</param>
 			/// <param name="Platform">Which platform-specific files to read</param>
-			public ConfigHierarchyKey(ConfigHierarchyType Type, DirectoryReference? ProjectDir, UnrealTargetPlatform Platform)
+			/// <param name="CustomConfig">Custom config subdirectory to load</param>
+			public ConfigHierarchyKey(ConfigHierarchyType Type, DirectoryReference? ProjectDir, UnrealTargetPlatform Platform, string CustomConfig)
 			{
 				this.Type = Type;
 				this.ProjectDir = ProjectDir;
 				this.Platform = Platform;
+				this.CustomConfig = CustomConfig;
 			}
 
 			/// <summary>
@@ -102,7 +109,7 @@ namespace UnrealBuildTool
 			public override bool Equals(object? Other)
 			{
 				ConfigHierarchyKey? OtherKey = Other as ConfigHierarchyKey;
-				return (OtherKey != null && OtherKey.Type == Type && OtherKey.ProjectDir == ProjectDir && OtherKey.Platform == Platform);
+				return (OtherKey != null && OtherKey.Type == Type && OtherKey.ProjectDir == ProjectDir && OtherKey.Platform == Platform && OtherKey.CustomConfig == CustomConfig);
 			}
 
 			/// <summary>
@@ -111,7 +118,7 @@ namespace UnrealBuildTool
 			/// <returns>Hash value for this object</returns>
 			public override int GetHashCode()
 			{
-				return ((ProjectDir != null) ? ProjectDir.GetHashCode() : 0) + ((int)Type * 123) + (Platform.GetHashCode() * 345);
+				return ((ProjectDir != null) ? ProjectDir.GetHashCode() : 0) + ((int)Type * 123) + (Platform.GetHashCode() * 345) + (CustomConfig.GetHashCode() * 789);
 			}
 		}
 
@@ -163,11 +170,34 @@ namespace UnrealBuildTool
 		/// <param name="Type">The type of hierarchy to read</param>
 		/// <param name="ProjectDir">The project directory to read the hierarchy for</param>
 		/// <param name="Platform">Which platform to read platform-specific config files for</param>
+		/// <param name="CustomConfig">Optional override config directory to search, for support of multiple target types</param>
 		/// <returns>The requested config hierarchy</returns>
-		public static ConfigHierarchy ReadHierarchy(ConfigHierarchyType Type, DirectoryReference? ProjectDir, UnrealTargetPlatform Platform)
+		public static ConfigHierarchy ReadHierarchy(ConfigHierarchyType Type, DirectoryReference? ProjectDir, UnrealTargetPlatform Platform, string CustomConfig = "")
 		{
+			// Handle command line overrides
+			List<String> OverrideStrings = new List<String>();
+			string[] CmdLine = Environment.GetCommandLineArgs();
+			string IniConfigArgPrefix = "-ini:" + Enum.GetName(typeof(ConfigHierarchyType), Type) + ":";
+			string CustomConfigPrefix = "-CustomConfig=";
+			foreach (string CmdLineArg in CmdLine)
+			{
+				if (CmdLineArg.StartsWith(IniConfigArgPrefix, StringComparison.InvariantCultureIgnoreCase))
+				{
+					OverrideStrings.Add(CmdLineArg.Substring(IniConfigArgPrefix.Length));
+				}
+				if (CmdLineArg.StartsWith(CustomConfigPrefix, StringComparison.InvariantCultureIgnoreCase))
+				{
+					CustomConfig = CmdLineArg.Substring(CustomConfigPrefix.Length);
+				}
+			}
+
+			if (CustomConfig == null)
+			{
+				CustomConfig = String.Empty;
+			}
+
 			// Get the key to use for the cache. It cannot be null, so we use the engine directory if a project directory is not given.
-			ConfigHierarchyKey Key = new ConfigHierarchyKey(Type, ProjectDir, Platform);
+			ConfigHierarchyKey Key = new ConfigHierarchyKey(Type, ProjectDir, Platform, CustomConfig);
 
 			// Try to get the cached hierarchy with this key
 			ConfigHierarchy? Hierarchy;
@@ -177,7 +207,7 @@ namespace UnrealBuildTool
 				{
 					// Find all the input files
 					List<ConfigFile> Files = new List<ConfigFile>();
-					foreach (FileReference IniFileName in ConfigHierarchy.EnumerateConfigFileLocations(Type, ProjectDir, Platform))
+					foreach (FileReference IniFileName in ConfigHierarchy.EnumerateConfigFileLocations(Type, ProjectDir, Platform, CustomConfig))
 					{
 						ConfigFile? File;
 						if (TryReadFile(IniFileName, out File))
@@ -186,16 +216,10 @@ namespace UnrealBuildTool
 						}
 					}
 
-					// Handle command line overrides
-					string[] CmdLine = Environment.GetCommandLineArgs();
-					string IniConfigArgPrefix = "-ini:" + Enum.GetName(typeof(ConfigHierarchyType), Type) + ":";
-					foreach (string CmdLineArg in CmdLine)
+					foreach (string OverrideString in OverrideStrings)
 					{
-						if (CmdLineArg.StartsWith(IniConfigArgPrefix))
-						{
-							ConfigFile OverrideFile = new ConfigFile(CmdLineArg.Substring(IniConfigArgPrefix.Length));
-							Files.Add(OverrideFile);
-						}
+						ConfigFile OverrideFile = new ConfigFile(OverrideString);
+						Files.Add(OverrideFile);
 					}
 
 					// Create the hierarchy
diff --git a/Engine/Source/Programs/UnrealBuildTool/System/ConfigHierarchy.cs b/Engine/Source/Programs/UnrealBuildTool/System/ConfigHierarchy.cs
--- a/Engine/Source/Programs/UnrealBuildTool/System/ConfigHierarchy.cs
+++ b/Engine/Source/Programs/UnrealBuildTool/System/ConfigHierarchy.cs
@@ -906,12 +906,16 @@ namespace UnrealBuildTool
 			"{PROJECT}/Config/Default{TYPE}.ini",
 			// Project/Generated*.ini  this is reserved for files which are generated by buildmachine processes (i.e. should never be checked in)
 			"{PROJECT}/Config/Generated{TYPE}.ini",
+			// Project/Config/Custom/CustomConfig/Default*.ini only if CustomConfig is defined
+			"{PROJECT}/Config/Custom/{CUSTOMCONFIG}/Default{TYPE}.ini",
 			// Engine/Platform/Platform*.ini
 			"{ENGINE}/Config/{PLATFORM}/{PLATFORM}{TYPE}.ini",
 			// Project/Platform/Platform*.ini
 			"{PROJECT}/Config/{PLATFORM}/{PLATFORM}{TYPE}.ini",
 			// Project/Platform/GeneratedPlatform*.ini  this is reserved for files which are generated by buildmachine processes (i.e. should never be checked in)
 			"{PROJECT}/Config/{PLATFORM}/Generated{PLATFORM}{TYPE}.ini",
+			// Project/Config/Custom/CustomConfig/Platform/Platform*.ini only if CustomConfig is defined
+			"{PROJECT}/Config/Custom/{CUSTOMCONFIG}/{PLATFORM}/{PLATFORM}{TYPE}.ini",
 			// UserSettings/.../User*.ini
 			"{USERSETTINGS}/Unreal Engine/Engine/Config/User{TYPE}.ini",
 			// UserDir/.../User*.ini
@@ -953,7 +957,7 @@ namespace UnrealBuildTool
 		}
 
 
-		private static string PerformBasicReplacements(string InString, string BaseIniName)
+		private static string PerformBasicReplacements(string InString, string BaseIniName, string CustomConfig)
 		{
 			string OutString = InString.Replace("{TYPE}", BaseIniName);
 
@@ -964,6 +968,7 @@ namespace UnrealBuildTool
 			}
 
 			OutString = OutString.Replace("{USER}", GetUserDir());
+			OutString = OutString.Replace("{CUSTOMCONFIG}", CustomConfig);
 
 			return OutString;
 		}
@@ -1029,7 +1034,7 @@ namespace UnrealBuildTool
 		/// <summary>
 		/// Returns a list of INI filenames for the given project
 		/// </summary>
-		public static IEnumerable<FileReference> EnumerateConfigFileLocations(ConfigHierarchyType Type, DirectoryReference? ProjectDir, UnrealTargetPlatform Platform)
+		public static IEnumerable<FileReference> EnumerateConfigFileLocations(ConfigHierarchyType Type, DirectoryReference? ProjectDir, UnrealTargetPlatform Platform, string CustomConfig = "")
 		{
 			string BaseIniName = Enum.GetName(typeof(ConfigHierarchyType), Type) ?? String.Empty;
 			string PlatformName = GetIniPlatformName(Platform);
@@ -1039,16 +1044,18 @@ namespace UnrealBuildTool
 				bool bHasPlatformTag = Layer.Contains("{PLATFORM}");
 				bool bHasProjectTag = Layer.Contains("{PROJECT}");
 				bool bHasUserTag = Layer.Contains("{USER}");
+				bool bHasCustomConfigTag = Layer.Contains("{CUSTOMCONFIG}");
 
 				// skip certain layers if we are platform-less, project-less, or userdir-less
 				if ((bHasPlatformTag && PlatformName == "None") ||
 					(bHasProjectTag && ProjectDir == null) ||
-					(bHasUserTag && GetUserDir() == null))
+					(bHasUserTag && GetUserDir() == null) ||
+					(bHasCustomConfigTag && String.IsNullOrEmpty(CustomConfig)))
 				{
 					continue;
 				}
 
-				string LayerPath = PerformBasicReplacements(Layer, BaseIniName);
+				string LayerPath = PerformBasicReplacements(Layer, BaseIniName, CustomConfig);
 
 				// we only expand engine/project inis
 				if (Layer.Contains("{ENGINE}") || Layer.Contains("{PROJECT}"))
diff --git a/Engine/Source/Runtime/Core/Private/Misc/ConfigCacheIni.cpp b/Engine/Source/Runtime/Core/Private/Misc/ConfigCacheIni.cpp
--- a/Engine/Source/Runtime/Core/Private/Misc/ConfigCacheIni.cpp
+++ b/Engine/Source/Runtime/Core/Private/Misc/ConfigCacheIni.cpp
@@ -32,6 +32,10 @@
 	#define INI_CACHE 0
 #endif
 
+#ifndef CUSTOM_CONFIG
+#define CUSTOM_CONFIG ""
+#endif
+
 #ifndef DISABLE_GENERATED_INI_WHEN_COOKED
 #define DISABLE_GENERATED_INI_WHEN_COOKED 0
 #endif
@@ -547,6 +551,7 @@ enum class EConfigLayerFlags : int32
 	DedicatedServerOnly			= (1 << 2), // replaces Default, Base, and (NOT {PLATFORM} yet) with an empty string
 	GenerateCacheKey			= (1 << 3),
 	NoExpand					= (1 << 4),
+	RequiresCustomConfig		= (1 << 5), // disabled if no custom config specified
 };
 ENUM_CLASS_FLAGS(EConfigLayerFlags);
 
@@ -582,12 +587,16 @@ struct FConfigLayer
 	{ TEXT("ProjectDefault"),			TEXT("{PROJECT}/Config/Default{TYPE}.ini"), EConfigLayerFlags::AllowCommandLineOverride | EConfigLayerFlags::GenerateCacheKey },
 	// Project/Generated*.ini Reserved for files generated by build process and should never be checked in 
 	{ TEXT("ProjectGenerated"),			TEXT("{PROJECT}/Config/Generated{TYPE}.ini"), EConfigLayerFlags::GenerateCacheKey },
+	// Project/Custom/CustomConfig/Default*.ini only if CustomConfig is defined
+	{ TEXT("CustomConfig"),				TEXT("{PROJECT}/Config/Custom/{CUSTOMCONFIG}/Default{TYPE}.ini"), EConfigLayerFlags::RequiresCustomConfig },
 	// Engine/Platform/Platform*.ini
 	{ TEXT("EnginePlatform"),			TEXT("{ENGINE}/Config/{PLATFORM}/{PLATFORM}{TYPE}.ini") },
 	// Project/Platform/Platform*.ini
 	{ TEXT("ProjectPlatform"),			TEXT("{PROJECT}/Config/{PLATFORM}/{PLATFORM}{TYPE}.ini") },
-	// Project/Generated*.ini Reserved for files generated by build process and should never be checked in 
+	// Project/Platform/GeneratedPlatform*.ini Reserved for files generated by build process and should never be checked in 
 	{ TEXT("ProjectPlatformGenerated"),	TEXT("{PROJECT}/Config/{PLATFORM}/Generated{PLATFORM}{TYPE}.ini") },
+	// Project/Custom/CustomConfig/Platform/Platform*.ini only if CustomConfig is defined
+	{ TEXT("CustomConfigPlatform"),		TEXT("{PROJECT}/Config/Custom/{CUSTOMCONFIG}/{PLATFORM}/{PLATFORM}{TYPE}.ini"), EConfigLayerFlags::RequiresCustomConfig },
 	// UserSettings/.../User*.ini
 	{ TEXT("UserSettingsDir"),			TEXT("{USERSETTINGS}Unreal Engine/Engine/Config/User{TYPE}.ini"), EConfigLayerFlags::NoExpand },
 	// UserDir/.../User*.ini
@@ -1182,6 +1191,7 @@ namespace CommandlineOverrideSpecifiers
 	const TCHAR SectionStartIdentifier[]  = TEXT("[");
 	const TCHAR PropertyStartIdentifier[] = TEXT("]:");
 	const TCHAR PropertySeperator[]       = TEXT(",");
+	const TCHAR CustomConfigIdentifier[] = TEXT("-CustomConfig=");
 }
 
 #endif
@@ -3769,6 +3779,7 @@ static FString PerformBasicReplacements(const FString& InString, const TCHAR* Ba
 	FString OutString = InString.Replace(TEXT("{TYPE}"), BaseIniName, ESearchCase::CaseSensitive);
 	OutString = OutString.Replace(TEXT("{USERSETTINGS}"), FPlatformProcess::UserSettingsDir(), ESearchCase::CaseSensitive);
 	OutString = OutString.Replace(TEXT("{USER}"), FPlatformProcess::UserDir(), ESearchCase::CaseSensitive);
+	OutString = OutString.Replace(TEXT("{CUSTOMCONFIG}"), *FConfigCacheIni::GetCustomConfigString(), ESearchCase::CaseSensitive);
 
 	return OutString;
 }
@@ -3869,12 +3880,19 @@ void FConfigFile::AddStaticLayersToHierarchy(const TCHAR* InBaseIniName, const T
 	const FString DedicatedServerString = IsRunningDedicatedServer() ? TEXT("DedicatedServer") : TEXT("");
 
 	// cache some platform extension information that can be used inside the loops
+	const bool bHasCustomConfig = !FConfigCacheIni::GetCustomConfigString().IsEmpty();
 
 	// go over all the config layers
 	for (int32 LayerIndex = 0; LayerIndex < UE_ARRAY_COUNT(GConfigLayers); LayerIndex++)
 	{
 		const FConfigLayer& Layer = GConfigLayers[LayerIndex];
 
+		// skip optional layers
+		if (EnumHasAnyFlags(Layer.Flag, EConfigLayerFlags::RequiresCustomConfig) && !bHasCustomConfig)
+		{
+			continue;
+		}
+
 		// start replacing basic variables
 		FString LayerPath = PerformBasicReplacements(Layer.Path, InBaseIniName);
 		bool bHasPlatformTag = LayerPath.Contains(TEXT("{PLATFORM}"));
@@ -4299,6 +4317,27 @@ void FConfigCacheIni::InitializeConfigSystem()
 #endif
 }
 
+const FString& FConfigCacheIni::GetCustomConfigString()
+{
+	static FString CustomConfigString;
+	static bool bInitialized = false;
+	if (!bInitialized)
+	{
+		bInitialized = true;
+
+		// Set to compiled in value, then possibly override
+		CustomConfigString = TEXT(CUSTOM_CONFIG);
+
+#if ALLOW_INI_OVERRIDE_FROM_COMMANDLINE
+		if (FParse::Value(FCommandLine::Get(), CommandlineOverrideSpecifiers::CustomConfigIdentifier, CustomConfigString))
+		{
+			UE_LOG(LogConfig, Log, TEXT("Overriding CustomConfig from %s to %s"), TEXT(CUSTOM_CONFIG), *CustomConfigString);
+		}
+#endif
+	}
+	return CustomConfigString;
+}
+
 bool FConfigCacheIni::LoadGlobalIniFile(FString& OutFinalIniFilename, const TCHAR* BaseIniName, const TCHAR* Platform, bool bForceReload, bool bRequireDefaultIni, bool bAllowGeneratedIniWhenCooked, bool bAllowRemoteConfig, const TCHAR* GeneratedConfigDir, FConfigCacheIni* ConfigSystem)
 {
 	// figure out where the end ini file is
diff --git a/Engine/Source/Runtime/Core/Public/Misc/ConfigCacheIni.h b/Engine/Source/Runtime/Core/Public/Misc/ConfigCacheIni.h
--- a/Engine/Source/Runtime/Core/Public/Misc/ConfigCacheIni.h
+++ b/Engine/Source/Runtime/Core/Public/Misc/ConfigCacheIni.h
@@ -1040,6 +1040,13 @@ public:
 	 */
 	static void InitializeConfigSystem();
 
+	/**
+	 * Returns the Custom Config string, which if set will load additional config files from Config/Custom/{CustomConfig}/DefaultX.ini to allow different types of builds.
+	 * It can be set from a game Target.cs file with CustomConfig = "Name".
+	 * Or in development, it can be overridden with a -CustomConfig=Name command line parameter.
+	 */
+	static const FString& GetCustomConfigString();
+
 	/**
 	 * Calculates the name of a dest (generated) .ini file for a given base (ie Engine, Game, etc)
 	 *
