P4WriteConfig:

Write p4config.txt with a trailing newline.

#trivial

diff --git a/Engine/Source/Programs/AutomationTool/AutomationUtils/P4WriteConfig.cs b/Engine/Source/Programs/AutomationTool/AutomationUtils/P4WriteConfig.cs
--- a/Engine/Source/Programs/AutomationTool/AutomationUtils/P4WriteConfig.cs
+++ b/Engine/Source/Programs/AutomationTool/AutomationUtils/P4WriteConfig.cs
@@ -15,7 +15,7 @@ namespace AutomationTool
 	[Help("path=<path>", "Write to a path other than the current directory")]
 	[Help("p4port=<server:port>", "Optional hint/override of the server to use during lookup")]
 	[Help("p4user=<server:port>", "Optional hint/override of the username to use during lookup")]
-	public class P4WriteConfig: BuildCommand
+	public class P4WriteConfig : BuildCommand
 	{
 		public override ExitCode Execute()
 		{
@@ -48,28 +48,27 @@ namespace AutomationTool
 				CommandUtils.InitDefaultP4Connection();
 			}
 			catch (Exception Ex)
-			{ 		
-				Log.TraceError("Unable to find matching Perforce info. If the below does not help try P4CreateConfig -p4port=<server:port> and -p4user=<username> to supply more info");
+			{
+				Log.TraceError(
+					"Unable to find matching Perforce info. If the below does not help try P4CreateConfig -p4port=<server:port> and -p4user=<username> to supply more info");
 				Log.TraceError("{0}", Ex.Message);
 				return ExitCode.Error_Arguments;
 			}
 
 			// store all our settings
-			Dictionary<string, string> P4Config = new Dictionary<string, string>();
+			StringBuilder P4Config = new StringBuilder();
 
-			P4Config["P4PORT"] = P4Env.ServerAndPort;
-			P4Config["P4USER"] = P4Env.User;
-			P4Config["P4CLIENT"] = P4Env.Client;
+			P4Config.AppendLine($"P4PORT={P4Env.ServerAndPort}");
+			P4Config.AppendLine($"P4USER={P4Env.User}");
+			P4Config.AppendLine($"P4CLIENT={P4Env.Client}");
 
 			if (SetIgnore)
 			{
 				string IgnorePath = Path.Combine(Unreal.EngineDirectory.ToString(), "Extras", "Perforce", "p4ignore");
-				P4Config["P4IGNORE"] = IgnorePath;
+				P4Config.AppendLine($"P4IGNORE={IgnorePath}");
 			}
 
-			string P4Settings = string.Join("\n", P4Config.Keys.Select(K => string.Format("{0}={1}", K, P4Config[K])));
-
-			string DefaultPath = Environment.CurrentDirectory;
+			string P4Settings = P4Config.ToString();
 
 			if (!string.IsNullOrEmpty(OutputPath))
 			{
@@ -77,7 +76,6 @@ namespace AutomationTool
 				{
 					throw new AutomationException("Path {0} does not exist.", OutputPath);
 				}
-				DefaultPath = OutputPath;
 			}
 			else
 			{
@@ -89,7 +87,7 @@ namespace AutomationTool
 				OutputPath = Path.Combine(OutputPath, "p4config.txt");
 			}
 
-			Console.WriteLine("***\nWriting_\n{0}\nto - {1}\n***", P4Settings, OutputPath);
+			Console.WriteLine("***\nWriting\n{0}to - {1}\n***", P4Settings, OutputPath);
 
 			if (!ListOnly)
 			{
@@ -107,8 +105,7 @@ namespace AutomationTool
 				Log.TraceInformation("Skipped write");
 			}
 
-
 			return ExitCode.Success;
-		}		
+		}
 	}
 }
