Copy the ImageMagick native Dlls to Gauntlet bin output

#rnx

diff --git a/Engine/Source/Programs/AutomationTool/Gauntlet/Gauntlet.Automation.csproj b/Engine/Source/Programs/AutomationTool/Gauntlet/Gauntlet.Automation.csproj
--- a/Engine/Source/Programs/AutomationTool/Gauntlet/Gauntlet.Automation.csproj
+++ b/Engine/Source/Programs/AutomationTool/Gauntlet/Gauntlet.Automation.csproj
@@ -47,9 +47,9 @@
   </ItemGroup>
 
   <ItemGroup>
-    <None Include="ThirdParty\Magick.NET\Magick.NET-Q16-HDRI-AnyCPU\runtimes\**\*.*" Link="..\..\runtimes\%(RecursiveDir)%(Filename)%(Extension)">
+    <Content Include="ThirdParty\Magick.NET\Magick.NET-Q16-HDRI-AnyCPU\runtimes\**\*.*" Link="%(Filename)%(Extension)">
       <CopyToOutputDirectory>PreserveNewest</CopyToOutputDirectory>
-    </None>
+    </Content>
   </ItemGroup>
 
     <!-- The following will glob all platform extensions' Gauntlet source
