Gauntlet: Fixing issue with reservation renewal, adding Horde job and step ids to reservations.  Cleaning up some legacy workers configuration
#rnx

diff --git a/Engine/Source/Programs/AutomationTool/AutomationUtils/DeviceReservation.cs b/Engine/Source/Programs/AutomationTool/AutomationUtils/DeviceReservation.cs
--- a/Engine/Source/Programs/AutomationTool/AutomationUtils/DeviceReservation.cs
+++ b/Engine/Source/Programs/AutomationTool/AutomationUtils/DeviceReservation.cs
@@ -201,12 +201,26 @@ namespace AutomationTool.DeviceReservation
 			}
 		}
 
-		public static void InvokeAPI(Uri UriToRequest, string Method)
+		public static void InvokeAPI(Uri UriToRequest, string Method, object ObjectToSerialize = null)
 		{
 			var Request = (HttpWebRequest)WebRequest.Create(UriToRequest);
 			Request.UseDefaultCredentials = true;
 			Request.Method = Method;
-			if (Method == "PUT")
+
+			if (ObjectToSerialize != null)
+			{
+				Request.ContentType = "application/json";
+
+				using (var RequestStream = Request.GetRequestStream())
+				{
+					var JsonString = fastJSON.JSON.Instance.ToJSON(ObjectToSerialize, new fastJSON.JSONParameters() { UseExtensions = false });
+					var Writer = new StreamWriter(RequestStream);
+					Writer.Write(JsonString);
+					Writer.Flush();
+					RequestStream.Flush();
+				}
+			}
+			else if (Method == "PUT")
 			{
 				Request.ContentLength = 0;
 			}
@@ -228,22 +242,15 @@ namespace AutomationTool.DeviceReservation
 		public Guid Guid { get; set; }
 		public static string ReservationDetails = "";
 
-		private void Copy(Reservation Other)
-		{
-			DeviceNames = Other.DeviceNames.ToArray();
-			HostName = Other.HostName;
-			StartDateTime = Other.StartDateTime;
-			Duration = Other.Duration;
-			Guid = Other.Guid;
-		}
-
 		private sealed class CreateReservationData
 		{
 			public string[] DeviceTypes;
 			public string Hostname;
 			public TimeSpan Duration;
 			public string ReservationDetails;
-			public string PoolID;
+			public string PoolId;
+			public string JobId;
+			public string StepId;
 		}
 
 		public static Reservation Create(Uri BaseUri, string[] DeviceTypes, TimeSpan Duration, int RetryMax = 5, string PoolID = "")
@@ -273,7 +280,9 @@ namespace AutomationTool.DeviceReservation
 						Hostname = Environment.MachineName,
 						Duration = Duration,
 						ReservationDetails = ReservationDetails,
-						PoolID = PoolID
+						PoolId = PoolID,
+						JobId = Environment.GetEnvironmentVariable("UE_HORDE_JOBID"),
+						StepId = Environment.GetEnvironmentVariable("UE_HORDE_STEPID")
 					});
 				}
 				catch (WebException WebEx)
@@ -323,8 +332,7 @@ namespace AutomationTool.DeviceReservation
 		{
 			try
 			{
-				var NewReservation = Utils.InvokeAPI<Reservation>(BaseUri.AppendPath("api/v1/reservations/" + Guid.ToString()), "PUT", NewDuration);
-				Copy(NewReservation);
+				Utils.InvokeAPI(BaseUri.AppendPath("api/v1/reservations/" + Guid.ToString()), "PUT", NewDuration);				
 			}
 			catch (Exception ex)
 			{
diff --git a/Engine/Source/Programs/AutomationTool/Gauntlet/Framework/Devices/Gauntlet.DevicePool.cs b/Engine/Source/Programs/AutomationTool/Gauntlet/Framework/Devices/Gauntlet.DevicePool.cs
--- a/Engine/Source/Programs/AutomationTool/Gauntlet/Framework/Devices/Gauntlet.DevicePool.cs
+++ b/Engine/Source/Programs/AutomationTool/Gauntlet/Framework/Devices/Gauntlet.DevicePool.cs
@@ -625,7 +625,7 @@ namespace Gauntlet
 			{
 				DeviceReservationAutoRenew DeviceReservation = null;
 
-				string PoolID = Globals.WorkerPoolID != -1 ? Globals.WorkerPoolID.ToString() : "";
+				string PoolID = Globals.DevicePoolId;
 
 				try
 				{
@@ -1042,12 +1042,6 @@ namespace Gauntlet
 				UnrealTargetPlatform.PS4, UnrealTargetPlatform.Switch
 			};
 
-			// support Android over wifi, though not on workers
-			if (!Globals.IsWorker)
-			{
-				ServicePlatforms.Add(UnrealTargetPlatform.Android);
-			}
-
 			var Devices = TooFewTotalDevices.Concat(TooFewCurrentDevices);
 			var UnsupportedPlatforms = Devices.Where(D => !ServicePlatforms.Contains(D.Platform.Value) && (!D.Platform.Value.ToString().StartsWith("XboxOne")) && (D.Platform.Value.ToString() != "XSX") && (D.Platform.Value.ToString() != "PS5"));
 
diff --git a/Engine/Source/Programs/AutomationTool/Gauntlet/Framework/Gauntlet.Utils.cs b/Engine/Source/Programs/AutomationTool/Gauntlet/Framework/Gauntlet.Utils.cs
--- a/Engine/Source/Programs/AutomationTool/Gauntlet/Framework/Gauntlet.Utils.cs
+++ b/Engine/Source/Programs/AutomationTool/Gauntlet/Framework/Gauntlet.Utils.cs
@@ -35,43 +35,15 @@ namespace Gauntlet
 		static List<Action> InnerPostAbortHandlers = new List<Action>();
 		public static bool CancelSignalled { get; private set; }
 
-		/// <summary>
-		/// Get the worker id of this Gauntlet instance
-		/// returns -1 if instance is not a member of a worker group
-		/// </summary>
-		public static int WorkerID
-		{
-			get
-			{
-				int Default = -1;
-				return Params.ParseValue("workerid", Default);
-			}
-
-		}
-
-		/// <summary>
-		/// Get the worker pool id of the host worker, pools are assigned to teams such as QA, Automation, etc
-		/// returns -1 if instance is not running on a worker pool
-		/// </summary>
-		public static int WorkerPoolID
-		{
-			get
-			{
-				int Default = -1;
-				return Params.ParseValue("workerpoolid", Default);
-			}
-
-		}
-
 
 		/// <summary>
-		/// Returns true if Gauntlet instance is a member of a worker group
+		/// Returns the device pool id 
 		/// </summary>
-		public static bool IsWorker
+		public static string DevicePoolId
 		{
 			get
 			{
-				return WorkerID != -1;
+				return Params.ParseValue("devicepool", "");
 			}
 		}
 
diff --git a/Engine/Source/Programs/AutomationTool/Gauntlet/Platform/IOS/Gauntlet.TargetDeviceIOS.cs b/Engine/Source/Programs/AutomationTool/Gauntlet/Platform/IOS/Gauntlet.TargetDeviceIOS.cs
--- a/Engine/Source/Programs/AutomationTool/Gauntlet/Platform/IOS/Gauntlet.TargetDeviceIOS.cs
+++ b/Engine/Source/Programs/AutomationTool/Gauntlet/Platform/IOS/Gauntlet.TargetDeviceIOS.cs
@@ -747,7 +747,7 @@ namespace Gauntlet
 		void KillZombies()
 		{
 
-			if (Globals.IsWorker || ZombiesKilled)
+			if (ZombiesKilled)
 			{
 				return;
 			}
@@ -765,7 +765,7 @@ namespace Gauntlet
 		{
 			get
 			{	
-				return Path.Combine(Globals.TempDir, string.Format("IOSAppCache{0}", Globals.WorkerID == -1 ? "" : Globals.WorkerID.ToString()));
+				return Path.Combine(Globals.TempDir, "IOSAppCache");
 			}
 		}
 
diff --git a/Engine/Source/Programs/AutomationTool/Gauntlet/Unreal/Base/Gauntlet.UnrealTestNode.cs b/Engine/Source/Programs/AutomationTool/Gauntlet/Unreal/Base/Gauntlet.UnrealTestNode.cs
--- a/Engine/Source/Programs/AutomationTool/Gauntlet/Unreal/Base/Gauntlet.UnrealTestNode.cs
+++ b/Engine/Source/Programs/AutomationTool/Gauntlet/Unreal/Base/Gauntlet.UnrealTestNode.cs
@@ -1001,12 +1001,6 @@ namespace Gauntlet
 				Message = string.Format("See Gauntlet.log for details");
 			}
 
-			if (Globals.IsWorker)
-			{
-				// log for worker to parse context
-				Log.Info("GauntletWorker:CreateReportFailure:{0}", Context.WorkerJobID);
-			}
-
 			Log.Warning("CreateReport Failed: {0}", Message);
 		}
 
