Change CheckPerforceCase to print warnings instead of errors

#rnx

diff --git a/Engine/Source/Programs/AutomationTool/Scripts/CheckPerforceCase.cs b/Engine/Source/Programs/AutomationTool/Scripts/CheckPerforceCase.cs
--- a/Engine/Source/Programs/AutomationTool/Scripts/CheckPerforceCase.cs
+++ b/Engine/Source/Programs/AutomationTool/Scripts/CheckPerforceCase.cs
@@ -89,7 +89,7 @@ namespace BuildScripts.Automation
 
 							Message.AppendFormat("\n    ({0} {1})", Pair.Value.Count, (Pair.Value.Count == 1) ? "file" : "files");
 						}
-						LogError(Message.ToString());
+						LogWarning(Message.ToString());
 						NumIssues++;
 					}
 
