Fix device usage comment calls for any UnrealTests that don't have an AppInstance

diff --git a/Engine/Source/Programs/AutomationTool/Gauntlet/Unreal/Base/Gauntlet.UnrealTestNode.cs b/Engine/Source/Programs/AutomationTool/Gauntlet/Unreal/Base/Gauntlet.UnrealTestNode.cs
--- a/Engine/Source/Programs/AutomationTool/Gauntlet/Unreal/Base/Gauntlet.UnrealTestNode.cs
+++ b/Engine/Source/Programs/AutomationTool/Gauntlet/Unreal/Base/Gauntlet.UnrealTestNode.cs
@@ -734,8 +734,11 @@ namespace Gauntlet
 			// Add info from test context to device usage log
 			foreach(IAppInstance AppInstance in TestInstance.ClientApps)
 			{
-				IDeviceUsageReporter.RecordComment(AppInstance.Device.Name, (UnrealTargetPlatform)AppInstance.Device.Platform, IDeviceUsageReporter.EventType.Device, Context.Options.JobDetails);
-				IDeviceUsageReporter.RecordComment(AppInstance.Device.Name, (UnrealTargetPlatform)AppInstance.Device.Platform, IDeviceUsageReporter.EventType.Test, this.GetType().Name);
+				if (AppInstance != null)
+				{
+					IDeviceUsageReporter.RecordComment(AppInstance.Device.Name, (UnrealTargetPlatform)AppInstance.Device.Platform, IDeviceUsageReporter.EventType.Device, Context.Options.JobDetails);
+					IDeviceUsageReporter.RecordComment(AppInstance.Device.Name, (UnrealTargetPlatform)AppInstance.Device.Platform, IDeviceUsageReporter.EventType.Test, this.GetType().Name);
+				}
 			}
 
 			
