verify directory access permission on Windows before expanding wildcards in Turnkey.

#rnx

diff --git a/Engine/Source/Programs/AutomationTool/Turnkey/Providers/NullCopyProvider.cs b/Engine/Source/Programs/AutomationTool/Turnkey/Providers/NullCopyProvider.cs
--- a/Engine/Source/Programs/AutomationTool/Turnkey/Providers/NullCopyProvider.cs
+++ b/Engine/Source/Programs/AutomationTool/Turnkey/Providers/NullCopyProvider.cs
@@ -5,6 +5,10 @@ using System.Collections.Generic;
 using System.Linq;
 using System.Text.RegularExpressions;
 using EpicGames.Core;
+#if WINDOWS
+using System.Security.AccessControl;
+using System.Security.Principal;
+#endif
 
 namespace Turnkey
 {
@@ -118,6 +122,51 @@ namespace Turnkey
 			return OutputPath;
 		}
 
+		private bool VerifyPathAccess( string PathString )
+		{
+#if WINDOWS
+			bool bResult = false;
+
+			DirectoryInfo DirInfo = new DirectoryInfo(PathString);
+			while (DirInfo != null)
+			{
+				try
+				{
+					WindowsIdentity Identity = WindowsIdentity.GetCurrent();
+					DirectorySecurity AccessControl = DirInfo.GetAccessControl(AccessControlSections.Access);
+					foreach (FileSystemAccessRule AccessRule in AccessControl.GetAccessRules(true, true, typeof(SecurityIdentifier)))
+					{
+						if (Identity.Owner == AccessRule.IdentityReference || Identity.Groups.Contains(AccessRule.IdentityReference))
+						{
+							bResult |= AccessRule.FileSystemRights.HasFlag(FileSystemRights.ReadData|FileSystemRights.ListDirectory);
+						}
+					}
+
+					break;
+				}
+				catch(System.UnauthorizedAccessException)
+				{
+					bResult = false;
+					break;
+				}
+				catch(System.Exception)
+				{
+					// something went wrong - look further up the path to check permissions on the parent folder
+					DirInfo = DirInfo.Parent;
+				}
+			}
+
+			if (!bResult)
+			{
+				Turnkey.LogWarning($"You do not have permission to access {PathString}");
+			}
+			return bResult;
+#else
+			// WindowsIdentity not supported on Mac/Linux
+			return true;
+#endif
+		}
+
 		private void ExpandWildcards(string Prefix, string PathString, Dictionary<string, List<string>> Output, List<string> ExpansionSet)
 		{
 			char Slash = Path.DirectorySeparatorChar;
@@ -129,6 +178,10 @@ namespace Turnkey
 			if (StarLocation == -1)
 			{
 				PathString = Prefix + PathString;
+				if (!VerifyPathAccess(PathString))
+				{
+					return;
+				}
 				if (Directory.Exists(PathString))
 				{
 					// make sure we end with a single slash
@@ -153,6 +206,10 @@ namespace Turnkey
 			// get the wildcard path component
 			string FullPathBeforeWildcard = Prefix + (PrevSlash >= 0 ? PathString.Substring(0, PrevSlash) : "");
 
+			if (!VerifyPathAccess(FullPathBeforeWildcard))
+			{
+				return;
+			}
 			if (Directory.Exists(FullPathBeforeWildcard))
 			{
 				// get the path component that has a wildcard
diff --git a/Engine/Source/Programs/AutomationTool/Turnkey/Turnkey.Automation.csproj b/Engine/Source/Programs/AutomationTool/Turnkey/Turnkey.Automation.csproj
--- a/Engine/Source/Programs/AutomationTool/Turnkey/Turnkey.Automation.csproj
+++ b/Engine/Source/Programs/AutomationTool/Turnkey/Turnkey.Automation.csproj
@@ -70,5 +70,6 @@
     <PackageReference Include="System.Data.DataSetExtensions" Version="4.5.0" />
     <PackageReference Include="System.Net.Http" Version="4.3.4" />
     <PackageReference Include="System.Drawing.Common" Version="4.7.0" />
-  </ItemGroup>
+    <PackageReference Include="System.IO.FileSystem.AccessControl" Version="4.7.0" />
+ </ItemGroup>
 </Project>
\ No newline at end of file
