AutomationTool: make script-local copies of Newtonsoft.Json.dll

Neither AutomationTool, AutomationUtils, nor their dependencies have a reference to Newtonsoft.Json.dll, which can result in the dependency not being available.

Change the csproj files to make local copies of the dll (<Private>True</Private> means "Copy Local")

diff --git a/Engine/Source/Programs/AutomationTool/Gauntlet/Gauntlet.Automation.csproj b/Engine/Source/Programs/AutomationTool/Gauntlet/Gauntlet.Automation.csproj
--- a/Engine/Source/Programs/AutomationTool/Gauntlet/Gauntlet.Automation.csproj
+++ b/Engine/Source/Programs/AutomationTool/Gauntlet/Gauntlet.Automation.csproj
@@ -46,7 +46,7 @@
     </Reference>
     <Reference Include="Newtonsoft.JSON">
       <HintPath>..\..\..\..\Binaries\ThirdParty\Newtonsoft\Newtonsoft.Json.dll</HintPath>
-      <Private>False</Private>
+      <Private>True</Private>
     </Reference>
   </ItemGroup>
 
