Gauntlet:

Skip the AutomationTool assembly: if Microsoft.Build.Framework can't be loaded, trying to do so will throw an exception.

diff --git a/Engine/Source/Programs/AutomationTool/Gauntlet/Framework/Gauntlet.Utils.cs b/Engine/Source/Programs/AutomationTool/Gauntlet/Framework/Gauntlet.Utils.cs
--- a/Engine/Source/Programs/AutomationTool/Gauntlet/Framework/Gauntlet.Utils.cs
+++ b/Engine/Source/Programs/AutomationTool/Gauntlet/Framework/Gauntlet.Utils.cs
@@ -514,6 +514,13 @@ namespace Gauntlet
 
 				foreach (var Assembly in AppDomain.CurrentDomain.GetAssemblies())
 				{
+					if (String.Equals("AutomationTool", Assembly.GetName().Name))
+					{
+						// Skip the AutomationTool driver assembly - it contains no relevant types, and trying to process it here
+						// can result in an exception if Microsoft.Build.Framework is not able to be loaded.
+						continue;
+					}
+
 					foreach (var Type in Assembly.GetTypes())
 					{
 						if (typeof(ITestNode).IsAssignableFrom(Type))
