BenchmarkBuild: Record build time when task fails

#rnx

diff --git a/Engine/Source/Programs/AutomationTool/Scripts/BenchmarkBuild/BenchmarkBuild.Automation.cs b/Engine/Source/Programs/AutomationTool/Scripts/BenchmarkBuild/BenchmarkBuild.Automation.cs
--- a/Engine/Source/Programs/AutomationTool/Scripts/BenchmarkBuild/BenchmarkBuild.Automation.cs
+++ b/Engine/Source/Programs/AutomationTool/Scripts/BenchmarkBuild/BenchmarkBuild.Automation.cs
@@ -209,6 +209,12 @@ namespace AutomationTool.Benchmark
 			}
 		}
 
+		struct BenchmarkResult
+		{
+			public TimeSpan TaskTime { get; set; }
+			public bool Failed { get; set; }
+		}
+
 		public BenchmarkBuild()
 		{
 		}
@@ -220,7 +226,7 @@ namespace AutomationTool.Benchmark
 
 			List<BenchmarkTaskBase> Tasks = new List<BenchmarkTaskBase>();
 
-			Dictionary<BenchmarkTaskBase, List<TimeSpan>> Results = new Dictionary<BenchmarkTaskBase, List<TimeSpan>>();
+			Dictionary<BenchmarkTaskBase, List<BenchmarkResult>> Results = new Dictionary<BenchmarkTaskBase, List<BenchmarkResult>>();
 
 			for (int ProjectIndex = 0; ProjectIndex < Options.ProjectsToTest.Count(); ProjectIndex++)
 			{
@@ -275,7 +281,7 @@ namespace AutomationTool.Benchmark
 				// create results lists
 				foreach (var Task in Tasks)
 				{
-					Results.Add(Task, new List<TimeSpan>());
+					Results.Add(Task, new List<BenchmarkResult>());
 				}
 
 				DateTime StartTime = DateTime.Now;
@@ -290,7 +296,16 @@ namespace AutomationTool.Benchmark
 
 						Log.TraceInformation("Task {0} took {1}", Task.GetFullTaskName(), Task.TaskTime.ToString(@"hh\:mm\:ss"));
 
-						Results[Task].Add(Task.TaskTime);
+						if (Task.Failed)
+						{
+							Log.TraceError("Task failed! Benchmark time may be inaccurate.");
+						}
+
+						Results[Task].Add(new BenchmarkResult
+						{
+							TaskTime = Task.TaskTime,
+							Failed = Task.Failed
+						});
 
 						// write results so far
 						WriteCSVResults(Options.FileName, Tasks, Results);
@@ -306,30 +321,28 @@ namespace AutomationTool.Benchmark
 				{
 					string TimeString = "";
 
-					IEnumerable<TimeSpan> TaskTimes = Results[Task];
+					IEnumerable<BenchmarkResult> TaskResults = Results[Task];
 
-					foreach (var TaskTime in TaskTimes)
+					foreach (var Result in TaskResults)
 					{
 						if (TimeString.Length > 0)
 						{
 							TimeString += ", ";
 						}
 
-						if (TaskTime == TimeSpan.Zero)
-						{
-							TimeString += "Failed";
-						}
-						else
+						if (Result.Failed)
 						{
-							TimeString += TaskTime.ToString(@"hh\:mm\:ss");
+							TimeString += "Failed ";
 						}
+
+						TimeString += Result.TaskTime.ToString(@"hh\:mm\:ss");
 					}
 
 					var AvgTimeString = "";
 
-					if (TaskTimes.Count() > 1)
+					if (TaskResults.Count() > 1)
 					{
-						var AvgTime = new TimeSpan(TaskTimes.Sum(T => T.Ticks) / TaskTimes.Count());
+						var AvgTime = new TimeSpan(TaskResults.Select(R => R.TaskTime).Sum(T => T.Ticks) / TaskResults.Count());
 
 						AvgTimeString = string.Format(" (Avg: {0})", AvgTime.ToString(@"hh\:mm\:ss"));
 					}
@@ -487,7 +500,7 @@ namespace AutomationTool.Benchmark
 		/// Writes our current result to a CSV file. It's expected that this function is called multiple times so results are
 		/// updated as we go
 		/// </summary>
-		void WriteCSVResults(string InFileName, List<BenchmarkTaskBase> InTasks, Dictionary<BenchmarkTaskBase, List<TimeSpan>> InResults)
+		void WriteCSVResults(string InFileName, List<BenchmarkTaskBase> InTasks, Dictionary<BenchmarkTaskBase, List<BenchmarkResult>> InResults)
 		{
 			Log.TraceInformation("Writing results to {0}", InFileName);
 
@@ -520,17 +533,14 @@ namespace AutomationTool.Benchmark
 					string Line = string.Format("{0},{1}", Task.GetFullTaskName(), Task.StartTime.ToString("yyyy-dd-MM HH:mm:ss"));
 
 					// now append all iteration times
-					foreach (TimeSpan TaskTime in InResults[Task])
+					foreach (BenchmarkResult Result in InResults[Task])
 					{
 						Line += ",";
-						if (TaskTime == TimeSpan.Zero)
-						{
-							Line += "FAILED";
-						}
-						else
+						if (Result.Failed)
 						{
-							Line += TaskTime.ToString(@"hh\:mm\:ss");
+							Line += "FAILED ";
 						}
+						Line += Result.TaskTime.ToString(@"hh\:mm\:ss");
 					}
 
 					Lines.Add(Line);
diff --git a/Engine/Source/Programs/AutomationTool/Scripts/BenchmarkBuild/BenchmarkTaskBase.cs b/Engine/Source/Programs/AutomationTool/Scripts/BenchmarkBuild/BenchmarkTaskBase.cs
--- a/Engine/Source/Programs/AutomationTool/Scripts/BenchmarkBuild/BenchmarkTaskBase.cs
+++ b/Engine/Source/Programs/AutomationTool/Scripts/BenchmarkBuild/BenchmarkTaskBase.cs
@@ -81,15 +81,12 @@ namespace AutomationTool.Benchmark
 				{
 					StartTime = DateTime.Now;
 
-					if (PerformTask())
-					{
-						TaskTime = DateTime.Now - StartTime;
-					}
-					else
+					if (!PerformTask())
 					{
 						FailureString = "Task Failed";
 						Failed = true;
 					}
+					TaskTime = DateTime.Now - StartTime;
 				}
 				else
 				{
@@ -130,7 +127,7 @@ namespace AutomationTool.Benchmark
 			}
 			else
 			{
-				Log.TraceInformation("Task {0}::\t\t\t\tFailed. {1}", GetFullTaskName(), FailureString);
+				Log.TraceInformation("Task {0}::\t\t\t\t{1} Failed. {2}", GetFullTaskName(), TaskTime.ToString(@"hh\:mm\:ss"), FailureString);
 			}
 		}
 
