AutomationTool:

Process "-Telemetry=path" explicitly as a global option.

diff --git a/Engine/Source/Programs/AutomationTool/Program.cs b/Engine/Source/Programs/AutomationTool/Program.cs
--- a/Engine/Source/Programs/AutomationTool/Program.cs
+++ b/Engine/Source/Programs/AutomationTool/Program.cs
@@ -30,8 +30,10 @@ namespace AutomationToolDriver
 		{
 			bool bGlobalParam = AutomationToolCommandLine.TrySetGlobal(CurrentParam);
 
+			// Global value parameters, handled explicitly
 			string Option_ScriptsForProject = "-ScriptsForProject";
 			string Option_ScriptDir = "-ScriptDir";
+			string Option_Telemetry = "-Telemetry";
 
 			// The parameter was not found in the list of global parameters, continue looking...
 			if (CurrentParam.StartsWith(Option_ScriptsForProject + "=", StringComparison.InvariantCultureIgnoreCase))
@@ -62,6 +64,11 @@ namespace AutomationToolDriver
 					throw new Exception($"Specified ScriptDir doesn't exist: {ScriptDir}");
 				}
 			}
+			else if (CurrentParam.StartsWith(Option_Telemetry + "=", StringComparison.InvariantCultureIgnoreCase))
+            {
+				string TelemetryPath = CurrentParam.Substring(CurrentParam.IndexOf('=') + 1);
+				AutomationToolCommandLine.SetUnchecked(Option_Telemetry, TelemetryPath);
+            }
 			else if (CurrentParam.StartsWith("-"))
 			{
 				if (CurrentCommand != null)
