AutomationTool:

Update references to DotNETCommon/MetaData.cs to Shared/MetaData.cs

#trivial

diff --git a/Engine/Source/Programs/AutomationTool/AutomationUtils/UE4Build.cs b/Engine/Source/Programs/AutomationTool/AutomationUtils/UE4Build.cs
--- a/Engine/Source/Programs/AutomationTool/AutomationUtils/UE4Build.cs
+++ b/Engine/Source/Programs/AutomationTool/AutomationUtils/UE4Build.cs
@@ -382,7 +382,7 @@ namespace AutomationTool
 
             {
                 // Use Version.h data to update MetaData.cs so the assemblies match the engine version.
-				FileReference MetaDataFile = FileReference.Combine(Unreal.EngineDirectory, "Source", "Programs", "DotNETCommon", "MetaData.cs");
+				FileReference MetaDataFile = FileReference.Combine(Unreal.EngineDirectory, "Source", "Programs", "Shared", "MetaData.cs");
 
 				if (bDoUpdateVersionFiles)
                 {
diff --git a/Engine/Source/Programs/AutomationTool/Scripts/SyncProject.Automation.cs b/Engine/Source/Programs/AutomationTool/Scripts/SyncProject.Automation.cs
--- a/Engine/Source/Programs/AutomationTool/Scripts/SyncProject.Automation.cs
+++ b/Engine/Source/Programs/AutomationTool/Scripts/SyncProject.Automation.cs
@@ -114,7 +114,7 @@ class SyncProject : SyncProjectBase
 		string[] ForceSyncFiles = new string[]
 		{
 			"Engine/Build/Build.version",
-			"Engine/Source/Programs/DotNETCommon/MetaData.cs"
+			"Engine/Source/Programs/Shared/MetaData.cs"
 		};
 
 		// Parse the project filename (as a local path)
