BenchmarkBuild: Save TaskTime in the event of an exception

#rnx

diff --git a/Engine/Source/Programs/AutomationTool/Scripts/BenchmarkBuild/BenchmarkTaskBase.cs b/Engine/Source/Programs/AutomationTool/Scripts/BenchmarkBuild/BenchmarkTaskBase.cs
--- a/Engine/Source/Programs/AutomationTool/Scripts/BenchmarkBuild/BenchmarkTaskBase.cs
+++ b/Engine/Source/Programs/AutomationTool/Scripts/BenchmarkBuild/BenchmarkTaskBase.cs
@@ -86,7 +86,6 @@ namespace AutomationTool.Benchmark
 						FailureString = "Task Failed";
 						Failed = true;
 					}
-					TaskTime = DateTime.Now - StartTime;
 				}
 				else
 				{
@@ -99,7 +98,12 @@ namespace AutomationTool.Benchmark
 			{
 				FailureString = string.Format("Exception: {0}", Ex.ToString());
 				Failed = true;
-			}		
+			}
+
+			if (StartTime != DateTime.MinValue)
+			{
+				TaskTime = DateTime.Now - StartTime;
+			}
 			
 			if (Failed)
 			{
