AutomationTool:

When compiling script modules, log additional information if a project fails to load.
Correct path case for command line arguments

diff --git a/Engine/Source/Programs/AutomationTool/CompileScriptModules.cs b/Engine/Source/Programs/AutomationTool/CompileScriptModules.cs
--- a/Engine/Source/Programs/AutomationTool/CompileScriptModules.cs
+++ b/Engine/Source/Programs/AutomationTool/CompileScriptModules.cs
@@ -431,12 +431,16 @@ namespace AutomationToolDriver
 				// Project automation scripts currently require source engine builds
 				if (!Unreal.IsEngineInstalled())
 				{
-					AllGameFolders = new List<DirectoryReference> { new DirectoryReference(Path.GetDirectoryName(ScriptsForProjectFileName)) };
+					DirectoryReference ScriptsDir = new DirectoryReference(Path.GetDirectoryName(ScriptsForProjectFileName));
+					ScriptsDir = DirectoryReference.FindCorrectCase(ScriptsDir);
+					AllGameFolders = new List<DirectoryReference> { ScriptsDir };
 				}
 			}
 
 			AdditionalScriptsFolders = AdditionalScriptsFolders ?? new List<string>();
-			List<DirectoryReference> AllAdditionalScriptFolders = AdditionalScriptsFolders.Select(x => new DirectoryReference(x)).ToList();
+			List<DirectoryReference> AllAdditionalScriptFolders = AdditionalScriptsFolders.Select(
+				x => DirectoryReference.FindCorrectCase(new DirectoryReference(x))).ToList();
+
 			foreach (DirectoryReference Folder in AllGameFolders)
 			{
 				DirectoryReference GameBuildFolder = DirectoryReference.Combine(Folder, "Build");
@@ -581,7 +585,29 @@ namespace AutomationToolDriver
 						ProjectPath = Path.GetFullPath(ProjectPath);
 						if (!Projects.ContainsKey(ProjectPath))
 						{
-							Project Project = new Project(ProjectPath, GlobalProperties, toolsVersion: null);
+							Project Project;
+
+							// Microsoft.Build.Evaluation.Project doesn't give a lot of useful information if this fails,
+							// so make sure to print our own diagnostic info if something goes wrong
+							try
+							{
+								Project = new Project(ProjectPath, GlobalProperties, toolsVersion: null);
+							}
+							catch (Exception Ex)
+							{
+								Log.TraceError($"Error while loading project {ProjectPath}");
+								Log.TraceError(ExceptionUtils.FormatException(Ex));
+								if (Projects.Count > 0)
+								{
+									Log.TraceError("After loading other projects:");
+									foreach (string Path in Projects.Keys)
+                                    {
+										Log.TraceError($"  {Path}");
+                                    }
+								}
+								throw Ex;
+							}
+
 							Projects.Add(ProjectPath, Project);
 							foreach (string ReferencedProject in Project.GetItems("ProjectReference").
 								Select(I => I.EvaluatedInclude))
