AutomationTool BuildAllScriptPlugins():

Clean up error text a little and add information so that when a project file fails to load it is easier to see who/what caused the project to load.

diff --git a/Engine/Source/Programs/AutomationTool/CompileScriptModules.cs b/Engine/Source/Programs/AutomationTool/CompileScriptModules.cs
--- a/Engine/Source/Programs/AutomationTool/CompileScriptModules.cs
+++ b/Engine/Source/Programs/AutomationTool/CompileScriptModules.cs
@@ -580,7 +580,7 @@ namespace AutomationToolDriver
 				// Load all found automation projects, and any other referenced projects.
 				foreach (FileReference ProjectPath in FoundAutomationProjects)
 				{
-					void LoadProjectAndReferences(string ProjectPath)
+					void LoadProjectAndReferences(string ProjectPath, string ReferencedBy)
 					{
 						ProjectPath = Path.GetFullPath(ProjectPath);
 						if (!Projects.ContainsKey(ProjectPath))
@@ -593,30 +593,37 @@ namespace AutomationToolDriver
 							{
 								Project = new Project(ProjectPath, GlobalProperties, toolsVersion: null);
 							}
-							catch (Exception Ex)
+							catch (Microsoft.Build.Exceptions.InvalidProjectFileException IPFEx)
 							{
-								Log.TraceError($"Error while loading project {ProjectPath}");
-								Log.TraceError(ExceptionUtils.FormatException(Ex));
+								Log.TraceError($"Could not load project file {ProjectPath}");
+								Log.TraceError(IPFEx.BaseMessage);
+								
+								if (!String.IsNullOrEmpty(ReferencedBy))
+								{
+									Log.TraceError($"Referenced by: {ReferencedBy}");
+								}
 								if (Projects.Count > 0)
 								{
-									Log.TraceError("After loading other projects:");
-									foreach (string Path in Projects.Keys)
+									Log.TraceError("See the log file for the list of previously loaded projects.");
+									Log.TraceLog("Loaded projects (most recently loaded first):");
+									foreach (string Path in Projects.Keys.Reverse())
                                     {
-										Log.TraceError($"  {Path}");
+										Log.TraceLog($"  {Path}");
                                     }
 								}
-								throw Ex;
+								throw IPFEx;
 							}
 
 							Projects.Add(ProjectPath, Project);
+							ReferencedBy = String.IsNullOrEmpty(ReferencedBy) ? ProjectPath : $"{ProjectPath}{Environment.NewLine}{ReferencedBy}";
 							foreach (string ReferencedProject in Project.GetItems("ProjectReference").
 								Select(I => I.EvaluatedInclude))
 							{
-								LoadProjectAndReferences(Path.Combine(Project.DirectoryPath, ReferencedProject));
+								LoadProjectAndReferences(Path.Combine(Project.DirectoryPath, ReferencedProject), ReferencedBy);
 							}
 						}
 					}
-					LoadProjectAndReferences(ProjectPath.FullName);
+					LoadProjectAndReferences(ProjectPath.FullName, null);
 				}
 				
 				// generate a BuildRecord for each loaded project - the gathered information will be used to determine if the project is
