Update for changed MetaData.cs location

#trivial

diff --git a/Engine/Source/Programs/AutomationTool/BuildGraph/Tasks/SetVersionTask.cs b/Engine/Source/Programs/AutomationTool/BuildGraph/Tasks/SetVersionTask.cs
--- a/Engine/Source/Programs/AutomationTool/BuildGraph/Tasks/SetVersionTask.cs
+++ b/Engine/Source/Programs/AutomationTool/BuildGraph/Tasks/SetVersionTask.cs
@@ -66,7 +66,7 @@ namespace AutomationTool.Tasks
 	}
 
 	/// <summary>
-	/// Updates the local version files (Engine/Source/Runtime/Launch/Resources/Version.h, Engine/Build/Build.version, and Engine/Source/Programs/DotNETCommon/Metadata.cs) with the given version information.
+	/// Updates the local version files (Engine/Source/Runtime/Launch/Resources/Version.h, Engine/Build/Build.version, and Engine/Source/Programs/Shared/Metadata.cs) with the given version information.
 	/// </summary>
 	[TaskElement("SetVersion", typeof(SetVersionTaskParameters))]
 	public class SetVersionTask : CustomTask
diff --git a/Engine/Source/Programs/GitDependencies/GitDependencies.csproj b/Engine/Source/Programs/GitDependencies/GitDependencies.csproj
--- a/Engine/Source/Programs/GitDependencies/GitDependencies.csproj
+++ b/Engine/Source/Programs/GitDependencies/GitDependencies.csproj
@@ -42,7 +42,7 @@
     <Reference Include="System.Xml" />
   </ItemGroup>
   <ItemGroup>
-    <Compile Include="..\DotNETCommon\MetaData.cs">
+    <Compile Include="..\Shared\MetaData.cs">
       <Link>Properties\MetaData.cs</Link>
     </Compile>
     <Compile Include="DependencyManifest.cs" />
diff --git a/Engine/Source/Programs/IOS/MobileDeviceInterface/MobileDeviceInterface.csproj b/Engine/Source/Programs/IOS/MobileDeviceInterface/MobileDeviceInterface.csproj
--- a/Engine/Source/Programs/IOS/MobileDeviceInterface/MobileDeviceInterface.csproj
+++ b/Engine/Source/Programs/IOS/MobileDeviceInterface/MobileDeviceInterface.csproj
@@ -80,7 +80,7 @@
     <Reference Include="System.Xml" />
   </ItemGroup>
   <ItemGroup>
-    <Compile Include="..\..\DotNETCommon\MetaData.cs">
+    <Compile Include="..\..\Shared\MetaData.cs">
       <Link>Properties\MetaData.cs</Link>
     </Compile>
     <Compile Include="ConnectEventArgs.cs" />
diff --git a/Engine/Source/Programs/IOS/iPhonePackager/iPhonePackager.csproj b/Engine/Source/Programs/IOS/iPhonePackager/iPhonePackager.csproj
--- a/Engine/Source/Programs/IOS/iPhonePackager/iPhonePackager.csproj
+++ b/Engine/Source/Programs/IOS/iPhonePackager/iPhonePackager.csproj
@@ -93,7 +93,7 @@
     </Reference>
   </ItemGroup>
   <ItemGroup>
-    <Compile Include="..\..\DotNETCommon\MetaData.cs">
+    <Compile Include="..\..\Shared\MetaData.cs">
       <Link>Properties\MetaData.cs</Link>
     </Compile>
     <Compile Include="CodeSigning.cs" />
diff --git a/Engine/Source/Programs/NetworkProfiler/NetworkProfiler/NetworkProfiler.csproj b/Engine/Source/Programs/NetworkProfiler/NetworkProfiler/NetworkProfiler.csproj
--- a/Engine/Source/Programs/NetworkProfiler/NetworkProfiler/NetworkProfiler.csproj
+++ b/Engine/Source/Programs/NetworkProfiler/NetworkProfiler/NetworkProfiler.csproj
@@ -101,7 +101,7 @@
     <Reference Include="System.Xml" />
   </ItemGroup>
   <ItemGroup>
-    <Compile Include="..\..\DotNETCommon\MetaData.cs">
+    <Compile Include="..\..\Shared\MetaData.cs">
       <Link>Properties\MetaData.cs</Link>
     </Compile>
     <Compile Include="BinaryReaderBigEndian.cs" />
diff --git a/Engine/Source/Programs/UnrealControls/UnrealControls.csproj b/Engine/Source/Programs/UnrealControls/UnrealControls.csproj
--- a/Engine/Source/Programs/UnrealControls/UnrealControls.csproj
+++ b/Engine/Source/Programs/UnrealControls/UnrealControls.csproj
@@ -69,7 +69,7 @@
     <Reference Include="System.Xml" />
   </ItemGroup>
   <ItemGroup>
-    <Compile Include="..\DotNETCommon\MetaData.cs">
+    <Compile Include="..\Shared\MetaData.cs">
       <Link>Properties\MetaData.cs</Link>
     </Compile>
     <Compile Include="ColoredDocumentLineSegment.cs" />
diff --git a/Engine/Source/Programs/UnrealControls/UnrealControls_Mono.csproj b/Engine/Source/Programs/UnrealControls/UnrealControls_Mono.csproj
--- a/Engine/Source/Programs/UnrealControls/UnrealControls_Mono.csproj
+++ b/Engine/Source/Programs/UnrealControls/UnrealControls_Mono.csproj
@@ -65,7 +65,7 @@
     <Reference Include="System.Xml" />
   </ItemGroup>
   <ItemGroup>
-    <Compile Include="..\DotNETCommon\MetaData.cs">
+    <Compile Include="..\Shared\MetaData.cs">
       <Link>Properties\MetaData.cs</Link>
     </Compile>
     <Compile Include="ColoredDocumentLineSegment.cs" />
diff --git a/Engine/Source/Programs/UnrealSwarm/Agent/Agent.csproj b/Engine/Source/Programs/UnrealSwarm/Agent/Agent.csproj
--- a/Engine/Source/Programs/UnrealSwarm/Agent/Agent.csproj
+++ b/Engine/Source/Programs/UnrealSwarm/Agent/Agent.csproj
@@ -97,7 +97,7 @@
     <Reference Include="System.Xml" />
   </ItemGroup>
   <ItemGroup>
-    <Compile Include="..\..\DotNETCommon\MetaData.cs">
+    <Compile Include="..\..\Shared\MetaData.cs">
       <Link>Properties\MetaData.cs</Link>
     </Compile>
     <Compile Include="Channels.cs" />
diff --git a/Engine/Source/Programs/UnrealSwarm/Agent/Agent_Mono.csproj b/Engine/Source/Programs/UnrealSwarm/Agent/Agent_Mono.csproj
--- a/Engine/Source/Programs/UnrealSwarm/Agent/Agent_Mono.csproj
+++ b/Engine/Source/Programs/UnrealSwarm/Agent/Agent_Mono.csproj
@@ -87,7 +87,7 @@
     <Reference Include="System.Xml" />
   </ItemGroup>
   <ItemGroup>
-    <Compile Include="..\..\DotNETCommon\MetaData.cs">
+    <Compile Include="..\..\Shared\MetaData.cs">
       <Link>Properties\MetaData.cs</Link>
     </Compile>
     <Compile Include="Channels.cs" />
diff --git a/Engine/Source/Programs/UnrealSwarm/AgentInterface/AgentInterface.csproj b/Engine/Source/Programs/UnrealSwarm/AgentInterface/AgentInterface.csproj
--- a/Engine/Source/Programs/UnrealSwarm/AgentInterface/AgentInterface.csproj
+++ b/Engine/Source/Programs/UnrealSwarm/AgentInterface/AgentInterface.csproj
@@ -71,7 +71,7 @@
     <Reference Include="System.Xml" />
   </ItemGroup>
   <ItemGroup>
-    <Compile Include="..\..\DotNETCommon\MetaData.cs">
+    <Compile Include="..\..\Shared\MetaData.cs">
       <Link>Properties\MetaData.cs</Link>
     </Compile>
     <Compile Include="AgentInterface.cs" />
diff --git a/Engine/Source/Programs/UnrealSwarm/AgentInterface/AgentInterface_Mono.csproj b/Engine/Source/Programs/UnrealSwarm/AgentInterface/AgentInterface_Mono.csproj
--- a/Engine/Source/Programs/UnrealSwarm/AgentInterface/AgentInterface_Mono.csproj
+++ b/Engine/Source/Programs/UnrealSwarm/AgentInterface/AgentInterface_Mono.csproj
@@ -67,7 +67,7 @@
     <Reference Include="System.Xml" />
   </ItemGroup>
   <ItemGroup>
-    <Compile Include="..\..\DotNETCommon\MetaData.cs">
+    <Compile Include="..\..\Shared\MetaData.cs">
       <Link>Properties\MetaData.cs</Link>
     </Compile>
     <Compile Include="AgentInterface.cs" />
diff --git a/Engine/Source/Programs/UnrealSwarm/SwarmCoordinator/SwarmCoordinator.csproj b/Engine/Source/Programs/UnrealSwarm/SwarmCoordinator/SwarmCoordinator.csproj
--- a/Engine/Source/Programs/UnrealSwarm/SwarmCoordinator/SwarmCoordinator.csproj
+++ b/Engine/Source/Programs/UnrealSwarm/SwarmCoordinator/SwarmCoordinator.csproj
@@ -87,7 +87,7 @@
     <Reference Include="System.XML" />
   </ItemGroup>
   <ItemGroup>
-    <Compile Include="..\..\DotNETCommon\MetaData.cs">
+    <Compile Include="..\..\Shared\MetaData.cs">
       <Link>Properties\MetaData.cs</Link>
     </Compile>
     <Compile Include="Coordinator.cs" />
diff --git a/Engine/Source/Programs/UnrealSwarm/SwarmCoordinator/SwarmCoordinator_Mono.csproj b/Engine/Source/Programs/UnrealSwarm/SwarmCoordinator/SwarmCoordinator_Mono.csproj
--- a/Engine/Source/Programs/UnrealSwarm/SwarmCoordinator/SwarmCoordinator_Mono.csproj
+++ b/Engine/Source/Programs/UnrealSwarm/SwarmCoordinator/SwarmCoordinator_Mono.csproj
@@ -78,7 +78,7 @@
     <Reference Include="System.Xml" />
   </ItemGroup>
   <ItemGroup>
-    <Compile Include="..\..\DotNETCommon\MetaData.cs">
+    <Compile Include="..\..\Shared\MetaData.cs">
       <Link>Properties\MetaData.cs</Link>
     </Compile>
     <Compile Include="Coordinator.cs" />
diff --git a/Engine/Source/Programs/UnrealSwarm/SwarmCoordinatorInterface/SwarmCoordinatorInterface.csproj b/Engine/Source/Programs/UnrealSwarm/SwarmCoordinatorInterface/SwarmCoordinatorInterface.csproj
--- a/Engine/Source/Programs/UnrealSwarm/SwarmCoordinatorInterface/SwarmCoordinatorInterface.csproj
+++ b/Engine/Source/Programs/UnrealSwarm/SwarmCoordinatorInterface/SwarmCoordinatorInterface.csproj
@@ -69,7 +69,7 @@
     <Reference Include="System.Xml" />
   </ItemGroup>
   <ItemGroup>
-    <Compile Include="..\..\DotNETCommon\MetaData.cs">
+    <Compile Include="..\..\Shared\MetaData.cs">
       <Link>Properties\MetaData.cs</Link>
     </Compile>
     <Compile Include="SwarmCoordinatorInterface.cs" />
diff --git a/Engine/Source/Programs/UnrealSwarm/SwarmCoordinatorInterface/SwarmCoordinatorInterface_Mono.csproj b/Engine/Source/Programs/UnrealSwarm/SwarmCoordinatorInterface/SwarmCoordinatorInterface_Mono.csproj
--- a/Engine/Source/Programs/UnrealSwarm/SwarmCoordinatorInterface/SwarmCoordinatorInterface_Mono.csproj
+++ b/Engine/Source/Programs/UnrealSwarm/SwarmCoordinatorInterface/SwarmCoordinatorInterface_Mono.csproj
@@ -63,7 +63,7 @@
     <Reference Include="System.Xml" />
   </ItemGroup>
   <ItemGroup>
-    <Compile Include="..\..\DotNETCommon\MetaData.cs">
+    <Compile Include="..\..\Shared\MetaData.cs">
       <Link>Properties\MetaData.cs</Link>
     </Compile>
     <Compile Include="SwarmCoordinatorInterface.cs" />
