Resume UE tests pass using json report as tracking support after test pass stoped because of critical failure.
* Implement loading Json report and resume from last know failure
* Make Gauntlet uses the retry state when critical failure occured on UE automation test.
* Capture the callstack if any and update the json report
* Attached passed log before retrying
* Use enum in json serialization and desirialization
* Use one common Event class for json report and parsing from the log
* Use an unique path for artifact at each runs/retries
* Use Json report to generate the pass summary in Gauntlet

diff --git a/Engine/Source/Developer/AutomationController/Private/AutomationControllerManager.h b/Engine/Source/Developer/AutomationController/Private/AutomationControllerManager.h
--- a/Engine/Source/Developer/AutomationController/Private/AutomationControllerManager.h
+++ b/Engine/Source/Developer/AutomationController/Private/AutomationControllerManager.h
@@ -50,6 +50,23 @@ public:
 		Errors = InErrors;
 	}
 
+	void AddEvent(EAutomationEventType EvenType, const FString& InMessage)
+	{
+		Entries.Add(FAutomationExecutionEntry(FAutomationEvent(EvenType, InMessage)));
+
+		switch (EvenType)
+		{
+		case EAutomationEventType::Warning:
+			Warnings++;
+			break;
+		case EAutomationEventType::Error:
+			Errors++;
+			break;
+		default:
+			break;
+		}
+	}
+
 	void SetArtifacts(const TArray<FAutomationArtifact>& InArtifacts)
 	{
 		Artifacts = InArtifacts;
@@ -152,58 +169,15 @@ public:
 		return Succeeded + SucceededWithWarnings + Failed + NotRun + InProcess;
 	}
 
-	void AddTestResult(const IAutomationReportPtr& TestReport)
-	{
-		FAutomatedTestResult TestResult;
-		TestResult.Test = TestReport;
-		TestResult.TestDisplayName = TestReport->GetDisplayName();
-		TestResult.FullTestPath = TestReport->GetFullTestPath();
+	void AddTestResult(const IAutomationReportPtr& TestReport);
 
-		TestsMapIndex.Add(TestReport->GetFullTestPath(), Tests.Num());
-		Tests.Add(TestResult);
-		NotRun++;
-	}
+	FAutomatedTestResult& GetTestResult(const IAutomationReportPtr& TestReport);
 
-	FAutomatedTestResult& GetTestResult(const IAutomationReportPtr& TestReport)
-	{
-		const FString& FullTestPath = TestReport->GetFullTestPath();
-		check(TestsMapIndex.Contains(FullTestPath));
-		return Tests[TestsMapIndex[FullTestPath]];
-	}
+	void ReBuildTestsMapIndex();
 
-	void UpdateTestResultStatus(const IAutomationReportPtr TestReport, EAutomationState State, bool bHasWarning = false)
-	{
-		FAutomatedTestResult& TestResult = GetTestResult(TestReport);
-		TestResult.State = State;
+	bool ReflectResultStateToReport(IAutomationReportPtr& TestReport);
 
-		// Book keeping
-		switch (State)
-		{
-		case EAutomationState::Success:
-			if (bHasWarning)
-			{
-				SucceededWithWarnings++;
-			}
-			else
-			{
-				Succeeded++;
-			}
-			InProcess--;
-			break;
-		case EAutomationState::Fail:
-			Failed++;
-			InProcess--;
-			break;
-		case EAutomationState::InProcess:
-			NotRun--;
-			InProcess++;
-			break;
-		default:
-			NotRun++;
-			InProcess--;
-			break;
-		}
-	}
+	void UpdateTestResultStatus(const IAutomationReportPtr& TestReport, EAutomationState State, bool bHasWarning = false);
 
 	void ClearAllEntries()
 	{
@@ -422,6 +396,11 @@ protected:
 	 */
 	bool GenerateTestPassHtmlIndex();
 
+	/**
+	 * Load test results from previous json test pass summary file and reflect results on reports
+	 */
+	bool LoadJsonTestPassSummary(FString& ReportFilePath, TArray<IAutomationReportPtr> TestReports);
+
 	/**
 	* Gather all info, warning, and error lines generated over the course of a test.
 	*
@@ -621,6 +600,8 @@ private:
 
 	FString DeveloperReportUrl;
 
+	bool bResumeRunTest;
+
 private:
 
 	/** Holds a delegate that is invoked when the controller shuts down. */
diff --git a/Engine/Source/Developer/AutomationController/Private/AutomationControllerManger.cpp b/Engine/Source/Developer/AutomationController/Private/AutomationControllerManger.cpp
--- a/Engine/Source/Developer/AutomationController/Private/AutomationControllerManger.cpp
+++ b/Engine/Source/Developer/AutomationController/Private/AutomationControllerManger.cpp
@@ -45,6 +45,94 @@ DEFINE_LOG_CATEGORY_STATIC(LogAutomationController, Log, All)
 
 #define LOCTEXT_NAMESPACE "AutomationTesting"
 
+void FAutomatedTestPassResults::AddTestResult(const IAutomationReportPtr& TestReport)
+{
+	const FString& FullTestPath = TestReport->GetFullTestPath();
+	if (!TestsMapIndex.Contains(FullTestPath))
+	{
+		FAutomatedTestResult TestResult;
+		TestResult.Test = TestReport;
+		TestResult.TestDisplayName = TestReport->GetDisplayName();
+		TestResult.FullTestPath = FullTestPath;
+
+		TestsMapIndex.Add(FullTestPath, Tests.Num());
+		Tests.Add(TestResult);
+		NotRun++;
+	}
+}
+
+FAutomatedTestResult& FAutomatedTestPassResults::GetTestResult(const IAutomationReportPtr& TestReport)
+{
+	const FString& FullTestPath = TestReport->GetFullTestPath();
+	check(TestsMapIndex.Contains(FullTestPath));
+	return Tests[TestsMapIndex[FullTestPath]];
+}
+
+void FAutomatedTestPassResults::ReBuildTestsMapIndex()
+{
+	TestsMapIndex.Empty();
+	uint32 Index = 0;
+	for (const FAutomatedTestResult& TestResult : Tests)
+	{
+		TestsMapIndex.Add(TestResult.FullTestPath, Index++);
+	}
+}
+
+bool FAutomatedTestPassResults::ReflectResultStateToReport(IAutomationReportPtr& TestReport)
+{
+	const FString& FullTestPath = TestReport->GetFullTestPath();
+	if (TestsMapIndex.Contains(FullTestPath))
+	{
+		FAutomatedTestResult& TestResult = Tests[TestsMapIndex[FullTestPath]];
+		if (TestResult.State == EAutomationState::InProcess)
+		{
+			// Marking incomplete test from previous pass as failure.
+			TestResult.AddEvent(EAutomationEventType::Error, TEXT("Test did not run until completion. The test exited prematurely."));
+			TestResult.State = EAutomationState::Fail;
+			InProcess--;
+			Failed++;
+		}
+		TestReport->SetState(TestResult.State);
+		TestResult.Test = TestReport;
+		return true;
+	}
+	return false;
+}
+
+void FAutomatedTestPassResults::UpdateTestResultStatus(const IAutomationReportPtr& TestReport, EAutomationState State, bool bHasWarning)
+{
+	FAutomatedTestResult& TestResult = GetTestResult(TestReport);
+	TestResult.State = State;
+
+	// Book keeping
+	switch (State)
+	{
+	case EAutomationState::Success:
+		if (bHasWarning)
+		{
+			SucceededWithWarnings++;
+		}
+		else
+		{
+			Succeeded++;
+		}
+		InProcess--;
+		break;
+	case EAutomationState::Fail:
+		Failed++;
+		InProcess--;
+		break;
+	case EAutomationState::InProcess:
+		NotRun--;
+		InProcess++;
+		break;
+	default:
+		NotRun++;
+		InProcess--;
+		break;
+	}
+}
+
 FAutomationControllerManager::FAutomationControllerManager()
 {
 
@@ -94,6 +182,8 @@ FAutomationControllerManager::FAutomationControllerManager()
 	{
 		DeveloperReportUrl = ReportURLPath / TEXT("dev") / FString(FPlatformProcess::UserName()).ToLower() / TEXT("index.html");
 	}
+
+	bResumeRunTest = FParse::Param(FCommandLine::Get(), TEXT("ResumeRunTest"));
 }
 
 void FAutomationControllerManager::RequestAvailableWorkers(const FGuid& SessionId)
@@ -180,6 +270,17 @@ void FAutomationControllerManager::RunTests(const bool bInIsLocalSession)
 	{
 		JsonTestPassResults.IsRequired = true;
 		TArray<IAutomationReportPtr> TestsToRun = ReportManager.GetEnabledTestReports();
+
+		// Load previous test results
+		if (bResumeRunTest)
+		{
+			FString ReportFileName = FString::Printf(TEXT("%s/index.json"), *ReportExportPath);
+			if (FPaths::FileExists(ReportFileName))
+			{
+				LoadJsonTestPassSummary(ReportFileName, TestsToRun);
+			}
+		}
+
 		for (IAutomationReportPtr& TestReport : TestsToRun)
 		{
 			JsonTestPassResults.AddTestResult(TestReport);
@@ -510,6 +611,35 @@ bool FAutomationControllerManager::GenerateTestPassHtmlIndex()
 	return false;
 }
 
+bool FAutomationControllerManager::LoadJsonTestPassSummary(FString& ReportFilePath, TArray<IAutomationReportPtr> TestReports)
+{
+	FString Json;
+	if (!FFileHelper::LoadFileToString(Json, *ReportFilePath))
+	{
+		UE_LOG(LogAutomationController, Error, TEXT("Failed to read json results file '%s'."), *ReportFilePath);
+		return false;
+	}
+	if (!FJsonObjectConverter::JsonObjectStringToUStruct(Json, &JsonTestPassResults))
+	{
+		UE_LOG(LogAutomationController, Error, TEXT("Failed to load json results file '%s'."), *ReportFilePath);
+		return false;
+	}
+	UE_LOG(LogAutomationController, Display, TEXT("Successfully loaded json results file."));
+
+	JsonTestPassResults.ReBuildTestsMapIndex();
+
+	// Reflect Test Result Status on Report Status
+	for (IAutomationReportPtr& TestReport : TestReports)
+	{
+		if (!JsonTestPassResults.ReflectResultStateToReport(TestReport))
+		{
+			UE_LOG(LogAutomationController, Warning, TEXT("No match to test '%s' in json results."), *TestReport->GetFullTestPath());
+		}
+	}
+
+	return true;
+}
+
 FString FAutomationControllerManager::SlugString(const FString& DisplayString) const
 {
 	FString GeneratedName = DisplayString;
diff --git a/Engine/Source/Developer/AutomationController/Private/AutomationReport.cpp b/Engine/Source/Developer/AutomationController/Private/AutomationReport.cpp
--- a/Engine/Source/Developer/AutomationController/Private/AutomationReport.cpp
+++ b/Engine/Source/Developer/AutomationController/Private/AutomationReport.cpp
@@ -467,6 +467,19 @@ EAutomationState FAutomationReport::GetState(const int32 ClusterIndex, const int
 	return EAutomationState::NotRun;
 }
 
+void FAutomationReport::SetState(const EAutomationState State)
+{
+	if (IsEnabled())
+	{
+		for (int32 ResultsIndex = 0; ResultsIndex < Results.Num(); ++ResultsIndex)
+		{
+			for (int32 PassIndex = 0; PassIndex < Results[ResultsIndex].Num(); ++PassIndex)
+			{
+				Results[ResultsIndex][PassIndex].State = State;
+			}
+		}
+	}
+}
 
 const FAutomationTestResults& FAutomationReport::GetResults( const int32 ClusterIndex, const int32 PassIndex ) 
 {
diff --git a/Engine/Source/Developer/AutomationController/Private/AutomationReport.h b/Engine/Source/Developer/AutomationController/Private/AutomationReport.h
--- a/Engine/Source/Developer/AutomationController/Private/AutomationReport.h
+++ b/Engine/Source/Developer/AutomationController/Private/AutomationReport.h
@@ -51,6 +51,7 @@ public:
 	virtual void AddArtifact(const int32 ClusterIndex, const int32 PassIndex, const FAutomationArtifact& Artifact) override;
 	virtual void GetCompletionStatus(const int32 ClusterIndex, const int32 PassIndex, FAutomationCompleteState& OutCompletionState) override;
 	virtual EAutomationState GetState(const int32 ClusterIndex, const int32 PassIndex) const override;
+	virtual void SetState(const EAutomationState State) override;
 	virtual const FAutomationTestResults& GetResults( const int32 ClusterIndex, const int32 PassIndex ) override;
 	virtual const int32 GetNumResults( const int32 ClusterIndex ) override;
 	virtual const int32 GetCurrentPassIndex( const int32 ClusterIndex ) override;
diff --git a/Engine/Source/Developer/AutomationController/Public/IAutomationReport.h b/Engine/Source/Developer/AutomationController/Public/IAutomationReport.h
--- a/Engine/Source/Developer/AutomationController/Public/IAutomationReport.h
+++ b/Engine/Source/Developer/AutomationController/Public/IAutomationReport.h
@@ -390,6 +390,13 @@ public:
 	 */
 	virtual EAutomationState GetState(const int32 ClusterIndex, const int32 PassIndex) const = 0;
 
+	/**
+	 * Set the state of the test (not run, in process, success, failure).
+	 *
+	 * @param State EAutomationState state to set.
+	 */
+	virtual void SetState(const EAutomationState State) = 0;
+
 	/**
 	 * Gets a copy of errors and warnings that were found
 	 *
diff --git a/Engine/Source/Programs/AutomationTool/Gauntlet/Framework/Base/Gauntlet.BaseTestNode.cs b/Engine/Source/Programs/AutomationTool/Gauntlet/Framework/Base/Gauntlet.BaseTestNode.cs
--- a/Engine/Source/Programs/AutomationTool/Gauntlet/Framework/Base/Gauntlet.BaseTestNode.cs
+++ b/Engine/Source/Programs/AutomationTool/Gauntlet/Framework/Base/Gauntlet.BaseTestNode.cs
@@ -113,6 +113,7 @@ namespace Gauntlet
 		protected void MarkTestStarted()
 		{
 			InnerStatus = TestStatus.InProgress;
+			SetTestResult(TestResult.Invalid);
 		}
 
 		/// <summary>
diff --git a/Engine/Source/Programs/AutomationTool/Gauntlet/Framework/Base/Gauntlet.TestReport.cs b/Engine/Source/Programs/AutomationTool/Gauntlet/Framework/Base/Gauntlet.TestReport.cs
--- a/Engine/Source/Programs/AutomationTool/Gauntlet/Framework/Base/Gauntlet.TestReport.cs
+++ b/Engine/Source/Programs/AutomationTool/Gauntlet/Framework/Base/Gauntlet.TestReport.cs
@@ -3,10 +3,40 @@
 using System;
 using System.Reflection;
 using System.Collections.Generic;
+using System.Text.Json;
+using System.Text.Json.Serialization;
 
 
 namespace Gauntlet
 {
+	/// <summary>
+	/// Event Type
+	/// </summary>
+	/// 
+	[JsonConverter(typeof(JsonTryParseEnumConverter))]
+	public enum EventType
+	{
+		Unknown,
+		Info,
+		Warning,
+		Error
+	}
+
+	/// <summary>
+	/// Test State type
+	/// </summary>
+	/// 
+	[JsonConverter(typeof(JsonTryParseEnumConverter))]
+	public enum TestStateType
+	{
+		Unknown,
+		NotRun,
+		InProcess,
+		Fail,
+		Success,
+		NotEnoughParticipants
+	}
+
 	/// <summary>
 	/// Interface for test report
 	/// </summary>
@@ -32,7 +62,7 @@ namespace Gauntlet
 		/// <param name="Message"></param>
 		/// <param name="Context"></param>
 		/// <returns></returns>
-		void AddEvent(string Type, string Message, object Context = null);
+		void AddEvent(EventType Type, string Message, object Context = null);
 
 		/// <summary>
 		/// Attach an Artifact to the ITestReport
@@ -122,19 +152,19 @@ namespace Gauntlet
 		/// <param name="Message"></param>
 		/// <param name="Context"></param>
 		/// <returns></returns>
-		public abstract void AddEvent(string Type, string Message, object Context = null);
+		public abstract void AddEvent(EventType Type, string Message, object Context = null);
 
 		public void AddError(string Message, object Context = null)
 		{
-			AddEvent("Error", Message, Context);
+			AddEvent(EventType.Error, Message, Context);
 		}
 		public void AddWarning(string Message, object Context = null)
 		{
-			AddEvent("Warning", Message, Context);
+			AddEvent(EventType.Warning, Message, Context);
 		}
 		public void AddInfo(string Message, object Context = null)
 		{
-			AddEvent("Info", Message, Context);
+			AddEvent(EventType.Info, Message, Context);
 		}
 
 		/// <summary>
@@ -174,4 +204,59 @@ namespace Gauntlet
 			return TelemetryDataList;
 		}
 	}
+
+	class JsonTryParseEnumConverter : JsonConverterFactory
+	{
+		public override bool CanConvert(Type TypeToConvert)
+		{
+			return TypeToConvert.IsEnum;
+		}
+
+		public override JsonConverter CreateConverter(Type EnumType, JsonSerializerOptions Options)
+		{
+			JsonConverter Converter = (JsonConverter)Activator.CreateInstance(typeof(DefaultToFirstEnumConverterInner<>).MakeGenericType(EnumType));
+
+			return Converter;
+		}
+
+		private class DefaultToFirstEnumConverterInner<TEnum> :
+			JsonConverter<TEnum> where TEnum : struct, Enum
+		{
+			public override TEnum Read(ref Utf8JsonReader Reader, Type TypeToConvert, JsonSerializerOptions Options)
+			{
+				if (Reader.TokenType == JsonTokenType.String)
+				{
+					string PropertyName = Reader.GetString();
+					// Try to parse it, assigned it to the default(0) if not known.
+					TEnum EnumValue;
+					Enum.TryParse(PropertyName, true, out EnumValue);
+
+					return EnumValue;
+				}
+				else if (Reader.TokenType == JsonTokenType.Number)
+				{
+					int Value = Reader.GetInt32();
+					TEnum EnumValue;
+					// Out of bound value get assigned anyway, so instead of try parse, we check if the value is defined else we use default()
+					if(Enum.IsDefined(TypeToConvert, Value))
+					{
+						EnumValue = (TEnum)Enum.ToObject(TypeToConvert, Value);
+					}
+					else
+					{
+						EnumValue = default(TEnum);
+					}
+					return EnumValue;
+				}
+
+				throw new JsonException();
+			}
+
+			public override void Write(Utf8JsonWriter Writer, TEnum EnumValue, JsonSerializerOptions Options)
+			{
+				var PropertyName = EnumValue.ToString();
+				Writer.WriteStringValue(Options.PropertyNamingPolicy?.ConvertName(PropertyName) ?? PropertyName);
+			}
+		}
+	}
 }
\ No newline at end of file
diff --git a/Engine/Source/Programs/AutomationTool/Gauntlet/Framework/Gauntlet.TestExecutor.cs b/Engine/Source/Programs/AutomationTool/Gauntlet/Framework/Gauntlet.TestExecutor.cs
--- a/Engine/Source/Programs/AutomationTool/Gauntlet/Framework/Gauntlet.TestExecutor.cs
+++ b/Engine/Source/Programs/AutomationTool/Gauntlet/Framework/Gauntlet.TestExecutor.cs
@@ -751,6 +751,12 @@ namespace Gauntlet
 							TestInfo.FinalResult = TestResult.Invalid;
 							bCanFinalizeTest = false;
 						}
+						else
+						{
+							TestInfo.CancellationReason = "Failed to restart during retry.";
+							Log.Error(TestInfo.CancellationReason);
+							TestInfo.FinalResult = TestResult.Failed;
+						}
 					}
 
 					if (bCanFinalizeTest)
diff --git a/Engine/Source/Programs/AutomationTool/Gauntlet/Framework/Utils/Gauntlet.HordeReport.cs b/Engine/Source/Programs/AutomationTool/Gauntlet/Framework/Utils/Gauntlet.HordeReport.cs
--- a/Engine/Source/Programs/AutomationTool/Gauntlet/Framework/Utils/Gauntlet.HordeReport.cs
+++ b/Engine/Source/Programs/AutomationTool/Gauntlet/Framework/Utils/Gauntlet.HordeReport.cs
@@ -4,6 +4,7 @@ using System;
 using System.Collections.Generic;
 using System.IO;
 using System.Text.Json;
+using System.Text.Json.Serialization;
 using AutomationTool;
 
 namespace Gauntlet
@@ -159,7 +160,7 @@ namespace Gauntlet
 		/// </summary>
 		public class Event
 		{
-			public string Type { get; set; }
+			public EventType Type { get; set; }
 			public string Message { get; set; }
 			public string Context { get; set; }
 			public string Artifact { get; set; }
@@ -192,7 +193,7 @@ namespace Gauntlet
 
 			public string TestDisplayName { get; set; }
 			public string FullTestPath { get; set; }
-			public string State { get; set; }
+			public TestStateType State { get; set; }
 			public int Warnings { get; set; }
 			public int Errors { get; set; }
 			public List<Artifact> Artifacts { get; set; }
@@ -240,7 +241,7 @@ namespace Gauntlet
 				get { return TestDetailed.FullTestPath; }
 				set { TestDetailed.FullTestPath = value; }
 			}
-			public string State
+			public TestStateType State
 			{
 				get { return TestDetailed.State; }
 				set { TestDetailed.State = value; }
@@ -314,7 +315,7 @@ namespace Gauntlet
 				return NewTestResult;
 			}
 
-			public override void AddEvent(string Type, string Message, object Context = null)
+			public override void AddEvent(EventType Type, string Message, object Context = null)
 			{
 				throw new System.NotImplementedException("AddEvent not implemented");
 			}
@@ -407,7 +408,7 @@ namespace Gauntlet
 					string TestResultFilePath = Path.Combine(OutputArtifactPath, OutputTestResult.ArtifactName);
 					try
 					{
-						File.WriteAllText(TestResultFilePath, JsonSerializer.Serialize(OutputTestResultDetailed, new JsonSerializerOptions { WriteIndented = true }));
+						File.WriteAllText(TestResultFilePath, JsonSerializer.Serialize(OutputTestResultDetailed, GetDefaultJsonOptions()));
 						ArtifactsCount++;
 					}
 					catch (Exception Ex)
@@ -445,15 +446,16 @@ namespace Gauntlet
 			public List<String> Errors { get; set; }
 			public List<String> Warnings { get; set; }
 
-			public override void AddEvent(string Type, string Message, object Context = null)
+			public override void AddEvent(EventType Type, string Message, object Context = null)
 			{
-				if (Type.ToLower() == "error")
-				{
-					Errors.Add(Message);
-				}
-				else if (Type.ToLower() == "warning")
+				switch (Type)
 				{
-					Warnings.Add(Message);
+					case EventType.Error:
+						Errors.Add(Message);
+						break;
+					case EventType.Warning:
+						Warnings.Add(Message);
+						break;
 				}
 				// The rest is ignored with this report type.
 			}
@@ -526,7 +528,7 @@ namespace Gauntlet
 				Log.Verbose("Writing Test Data Collection for Horde at {0}", OutputTestDataFilePath);
 				try
 				{
-					File.WriteAllText(OutputTestDataFilePath, JsonSerializer.Serialize(this, new JsonSerializerOptions { WriteIndented = true }));
+					File.WriteAllText(OutputTestDataFilePath, JsonSerializer.Serialize(this, GetDefaultJsonOptions()));
 				}
 				catch (Exception Ex)
 				{
@@ -534,5 +536,12 @@ namespace Gauntlet
 				}
 			}
 		}
+		private static JsonSerializerOptions GetDefaultJsonOptions()
+		{
+			return new JsonSerializerOptions
+			{
+				WriteIndented = true
+			};
+		}
 	}
 }
diff --git a/Engine/Source/Programs/AutomationTool/Gauntlet/Unreal/Automation/UE.Automation.cs b/Engine/Source/Programs/AutomationTool/Gauntlet/Unreal/Automation/UE.Automation.cs
--- a/Engine/Source/Programs/AutomationTool/Gauntlet/Unreal/Automation/UE.Automation.cs
+++ b/Engine/Source/Programs/AutomationTool/Gauntlet/Unreal/Automation/UE.Automation.cs
@@ -59,7 +59,7 @@ namespace UE
 		/// </summary>
 		[AutoParam]
 		public string RunTest = "";
-		
+
 		/// <summary>
 		/// Absolute or project relative path to write an automation report to.
 		/// </summary>
@@ -111,6 +111,7 @@ namespace UE
 		/// <param name="OtherRoles"></param>
 		public override void ApplyToConfig(UnrealAppConfig AppConfig, UnrealSessionRole ConfigRole, IEnumerable<UnrealSessionRole> OtherRoles)
 		{
+
 			base.ApplyToConfig(AppConfig, ConfigRole, OtherRoles);
 
 			// The "RunTest" argument is required since it is what is passed to the editor to specify which tests to execute
@@ -119,21 +120,39 @@ namespace UE
 				throw new AutomationException("No AutomationTest argument specified. Use -RunTest=\"Group:AI\", -RunTest=\"Project\", -RunTest=\"Navigation.Landscape Ramp Test\" etc.");
 			}
 
-			// Are we writing out info for Horde?
-			if (WriteTestResultsForHorde)
+			if (ConfigRole.RoleType.IsEditor())
 			{
-				if (string.IsNullOrEmpty(ReportExportPath))
+				if (ResumeOnCriticalFailure)
 				{
-					ReportExportPath = Path.Combine(Globals.TempDir, "TestReport");
+					AppConfig.CommandLine += " -ResumeRunTest";
 				}
-				if (string.IsNullOrEmpty(HordeTestDataPath))
+				// Are we writing out info for Horde?
+				if (WriteTestResultsForHorde)
 				{
-					HordeTestDataPath = HordeReport.DefaultTestDataDir;
+					if (string.IsNullOrEmpty(ReportExportPath))
+					{
+						ReportExportPath = Path.Combine(Globals.TempDir, "TestReport");
+					}
+					if (string.IsNullOrEmpty(HordeTestDataPath))
+					{
+						HordeTestDataPath = HordeReport.DefaultTestDataDir;
+					}
+					if (string.IsNullOrEmpty(HordeArtifactPath))
+					{
+						HordeArtifactPath = HordeReport.DefaultArtifactsDir;
+					}
 				}
-				if (string.IsNullOrEmpty(HordeArtifactPath))
+				// Arguments for writing out the report and providing a URL where it can be viewed
+				string ReportArgs = "";
+				if (!string.IsNullOrEmpty(ReportExportPath))
 				{
-					HordeArtifactPath = HordeReport.DefaultArtifactsDir;
+					ReportArgs = string.Format(" -ReportExportPath=\"{0}\"", ReportExportPath);
 				}
+				if (!string.IsNullOrEmpty(ReportURL))
+				{
+					ReportArgs += string.Format(" -ReportURL=\"{0}\"", ReportURL);
+				}
+				AppConfig.CommandLine += ReportArgs;
 			}
 
 			// Setup commandline for telemetry outputs
@@ -151,14 +170,6 @@ namespace UE
 				AppConfig.CommandLine += string.Format("-TelemetryDirectory=\"{0}\"", TelemetryDirectory);
 			}
 
-			// Arguments for writing out the report and providing a URL where it can be viewed
-			string ReportArgs = string.IsNullOrEmpty(ReportExportPath) ? "" : string.Format("-ReportExportPath=\"{0}\"", ReportExportPath);
-
-			if (!string.IsNullOrEmpty(ReportURL))
-			{
-				ReportArgs += string.Format("-ReportURL=\"{0}\"", ReportURL);
-			}
-
 			string AutomationTestArgument = string.Format("RunTest {0};", RunTest);
 
 			// if this is not attended then we'll quit the editor after the tests and disable any user interactions
@@ -171,7 +182,7 @@ namespace UE
 			// If there's only one role and it's the editor then tests are running under the editor with no target
 			if (ConfigRole.RoleType == UnrealTargetRole.Editor && OtherRoles.Any() == false)
 			{ 
-				AppConfig.CommandLine += string.Format(" {0} -ExecCmds=\"Automation {1}\"", ReportArgs, AutomationTestArgument);		
+				AppConfig.CommandLine += string.Format(" -ExecCmds=\"Automation {0}\"", AutomationTestArgument);		
 			}
 			else
 			{
@@ -185,7 +196,7 @@ namespace UE
 				}
 				else if (ConfigRole.RoleType.IsEditor())
 				{
-					AppConfig.CommandLine += string.Format(" -ExecCmds=\"Automation StartRemoteSession {0}; {1}\" -TcpMessagingListen={2}:6666 -multihome={3} {4}", SessionID, AutomationTestArgument, HostIP, HostIP, ReportArgs);
+					AppConfig.CommandLine += string.Format(" -ExecCmds=\"Automation StartRemoteSession {0}; {1}\" -TcpMessagingListen={2}:6666 -multihome={3}", SessionID, AutomationTestArgument, HostIP, HostIP);
 				}
 			}
 
@@ -312,6 +323,8 @@ namespace UE
 		// used to track stdout from the processes 
 		private int LastAutomationEntryCount = 0;
 
+		private UnrealAutomatedTestPassResults TestPassResults = null;
+
 		private DateTime LastAutomationEntryTime = DateTime.MinValue;
 
 		public AutomationNodeBase(Gauntlet.UnrealTestContext InContext)
@@ -353,10 +366,34 @@ namespace UE
 		{
 			LastAutomationEntryTime = DateTime.MinValue;
 			LastAutomationEntryCount = 0;
+			TestPassResults = null;
+
+			if (GetConfiguration() is AutomationTestConfig Config)
+			{
+				string ReportExportPath = Config.ReportExportPath;
+				if (!string.IsNullOrEmpty(ReportExportPath))
+				{
+					// clean up previous run if any on StartTest (not RestartTest)
+					if (Directory.Exists(ReportExportPath))
+					{
+						DirectoryInfo ReportDirInfo = new DirectoryInfo(ReportExportPath);
+						ReportDirInfo.Delete(true);
+					}
+				}
+			}
 
 			return base.StartTest(Pass, InNumPasses);
 		}
 
+		public override bool RestartTest()
+		{
+			LastAutomationEntryTime = DateTime.MinValue;
+			LastAutomationEntryCount = 0;
+			TestPassResults = null;
+
+			return base.RestartTest();
+		}
+
 		/// <summary>
 		/// Override TickTest to log interesting things and make sure nothing has stalled
 		/// </summary>
@@ -402,6 +439,58 @@ namespace UE
 			base.TickTest();
 		}
 
+		private UnrealAutomatedTestPassResults GetTestPassResults()
+		{
+			if (TestPassResults == null)
+			{
+				string AutomationReportPath = string.Empty;
+				if (GetConfiguration() is AutomationTestConfig Config)
+				{
+					AutomationReportPath = Config.ReportExportPath;
+					if (!string.IsNullOrEmpty(AutomationReportPath))
+					{
+						TestPassResults = UnrealAutomatedTestPassResults.LoadFromJson(Path.Combine(AutomationReportPath, "index.json"));
+					}
+				}
+				
+				if(TestPassResults == null)
+				{
+					// Parse automaton info from the log then
+					TestPassResults = new UnrealAutomatedTestPassResults();
+					var EditorRole = RoleResults.Where(R => R.Artifacts.SessionRole.RoleType == UnrealTargetRole.Editor).FirstOrDefault();
+					if (EditorRole != null)
+					{
+						AutomationLogParser LogParser = new AutomationLogParser(EditorRole.LogSummary.FullLogContent);
+						IEnumerable<UnrealAutomatedTestResult> LogTestResults = LogParser.GetResults();
+						if (LogTestResults.Any())
+						{
+							foreach (UnrealAutomatedTestResult LogTestResult in LogTestResults)
+							{
+								TestPassResults.AddTest(LogTestResult);
+							}
+						}
+						else
+						{
+							foreach (UnrealLog.LogEntry Entry in LogParser.AutomationWarningsAndErrors)
+							{
+								switch (Entry.Level)
+								{
+									case UnrealLog.LogLevel.Error:
+										Events.Add(new UnrealAutomationEvent(EventType.Error, Entry.Message));
+										break;
+									case UnrealLog.LogLevel.Warning:
+										Events.Add(new UnrealAutomationEvent(EventType.Warning, Entry.Message));
+										break;
+								}
+							}
+						}
+					}
+				}
+			}
+
+			return TestPassResults;
+		}
+
 		/// <summary>
 		/// Override GetExitCodeAndReason to provide additional checking of success / failure based on what occurred
 		/// </summary>
@@ -419,21 +508,18 @@ namespace UE
 				// if no fatal errors, check test results
 				if (InLog.FatalError == null)
 				{
-					AutomationLogParser Parser = new AutomationLogParser(InLog.FullLogContent);
-
-					IEnumerable<AutomationTestResult> TotalTests = Parser.GetResults();
-					IEnumerable<AutomationTestResult> FailedTests = TotalTests.Where(R => !R.Passed);
+					var TestResults = GetTestPassResults();
 
 					// Tests failed so list that as our primary cause of failure
-					if (FailedTests.Any())
+					if (TestResults != null && TestResults.Failed > 0)
 					{
-						ExitReason = string.Format("{0} of {1} test(s) failed", FailedTests.Count(), TotalTests.Count());
+						ExitReason = string.Format("{0} of {1} test(s) failed", TestResults.Failed, TestResults.Tests.Count());
 						ExitCode = -1;
 						return UnrealProcessResult.TestFailure;
 					}
 
 					// If no tests were run then that's a failure (possibly a bad RunTest argument?)
-					if (!TotalTests.Any())
+					if (TestResults == null || TestResults.Tests.Count() == 0)
 					{
 						ExitReason = "No tests were executed!";
 						ExitCode = -1;
@@ -506,15 +592,24 @@ namespace UE
 
 			if (EditorRole != null)
 			{
-				// Parse automaton info from the log (TODO - use the json version)
-				AutomationLogParser Parser = new AutomationLogParser(EditorRole.LogSummary.FullLogContent);
+				UnrealAutomatedTestPassResults JsonTestPassResults = GetTestPassResults();
 
 				// Filter our tests into categories
-				IEnumerable<AutomationTestResult> AllTests = Parser.GetResults();
-				IEnumerable<AutomationTestResult> IncompleteTests = AllTests.Where(R => !R.Completed);
-				IEnumerable<AutomationTestResult> FailedTests = AllTests.Where(R => R.Completed && !R.Passed);
-				IEnumerable<AutomationTestResult> TestsWithWarnings = AllTests.Where(R => R.Completed && R.Passed && R.WarningEvents.Any());
-
+				IEnumerable<UnrealAutomatedTestResult> AllTests = JsonTestPassResults.Tests;
+				IEnumerable<UnrealAutomatedTestResult> IncompleteTests = AllTests.Where(T => !T.IsComplete);
+				IEnumerable<UnrealAutomatedTestResult> FailedTests = AllTests.Where(T => T.IsComplete && !T.HasSucceeded);
+				IEnumerable<UnrealAutomatedTestResult> TestsWithWarnings = AllTests.Where(T => T.HasSucceeded && T.HasWarnings);
+
+				Func<string, string> ErrorLineForHorde = (L) => string.Format("Err: {0}", L);
+				Func<string, string> WarningLineForHorde = (L) => string.Format("Warn: {0}", L);
+				Func<IEnumerable<string>, IEnumerable<string>> CapErrorOrWarningList = (E) =>
+				{
+					if (E.Count() > kMaxErrorsOrWarningsToDisplay)
+					{
+						E = E.Skip(E.Count() - kMaxErrorsOrWarningsToDisplay);
+					}
+					return E;
+				};
 				// If there were abnormal exits then look only at the incomplete tests to avoid confusing things.
 				if (GetRolesThatExitedAbnormally().Any())
 				{
@@ -524,12 +619,14 @@ namespace UE
 					}
 					else if (IncompleteTests.Count() > 0)
 					{
-						MB.H3("Error: The following test(s) were incomplete:");
+						MB.H3("The following test(s) were incomplete:");
 
-						foreach (AutomationTestResult Result in IncompleteTests)
+						foreach (UnrealAutomatedTestResult Result in IncompleteTests)
 						{
-							MB.H4(string.Format("{0}", Result.FullName));
-							MB.UnorderedList(Result.WarningAndErrorEvents.Distinct());
+							MB.H4(string.Format("Error: Test '{0}' did not run or complete", Result.TestDisplayName));
+							MB.Paragraph("FullName: " + Result.FullTestPath);
+							MB.UnorderedList(CapErrorOrWarningList(Result.ErrorEvents.Select(E => ErrorLineForHorde(E.Message)).Distinct()));
+							MB.UnorderedList(CapErrorOrWarningList(Result.WarningEvents.Select(E => WarningLineForHorde(E.Message)).Distinct()));
 						}
 					}
 				}
@@ -539,11 +636,12 @@ namespace UE
 					{
 						MB.H3("Error: No tests were executed.");
 
-						IEnumerable<UnrealLog.LogEntry> WarningsAndErrors = Parser.AutomationWarningsAndErrors;
+						IEnumerable<string> WarningsAndErrors = Events.Where(E => E.IsError || E.IsWarning).Select(E => E.IsError? ErrorLineForHorde(E.Message) : WarningLineForHorde(E.Message)).Distinct();
 
 						if (WarningsAndErrors.Any())
 						{
-							MB.UnorderedList(WarningsAndErrors.Select(E => E.ToString()));
+							WarningsAndErrors = CapErrorOrWarningList(WarningsAndErrors);
+							MB.UnorderedList(WarningsAndErrors);
 						}
 						else
 						{
@@ -552,34 +650,21 @@ namespace UE
 					}
 					else
 					{
-
-						Func<string, string> SantizeLineForHorde = (L) =>
-						{
-							L = L.Replace(": Error: ", ": Err: ");
-							L = L.Replace(": Warning: ", ": Warn: ");
-							L = L.Replace("LogAutomationController: ", "");
-
-							return L;
-						};
-
 						// Now list the tests that failed
 						if (FailedTests.Count() > 0)
 						{
 							MB.H3("The following test(s) failed:");
 
-							foreach (AutomationTestResult Result in FailedTests)
+							foreach (UnrealAutomatedTestResult Result in FailedTests)
 							{
 								// only show the last N items
-								IEnumerable<string> Events = Result.Events.Distinct();
+								IEnumerable<string> Events = Result.ErrorEvents.Select(E => ErrorLineForHorde(E.Message)).Distinct();
 
-								if (Events.Count() > kMaxErrorsOrWarningsToDisplay)
-								{
-									Events = Events.Skip(Events.Count() - kMaxErrorsOrWarningsToDisplay);
-								}
+								Events = CapErrorOrWarningList(Events);
 
-								MB.H4(string.Format("Error: Test '{0}' failed", Result.DisplayName));
-								MB.Paragraph("FullName: " + Result.TestName);
-								MB.UnorderedList(Events.Distinct().Select(E => SantizeLineForHorde(E)));
+								MB.H4(string.Format("Error: Test '{0}' failed", Result.TestDisplayName));
+								MB.Paragraph("FullName: " + Result.FullTestPath);
+								MB.UnorderedList(Events);
 							}
 						}
 
@@ -587,19 +672,16 @@ namespace UE
 						{
 							MB.H3("The following test(s) completed with warnings:");
 
-							foreach (AutomationTestResult Result in TestsWithWarnings)
+							foreach (UnrealAutomatedTestResult Result in TestsWithWarnings)
 							{
-							// only show the last N items
-							IEnumerable<string> WarningEvents = Result.WarningEvents.Distinct();
+								// only show the last N items
+								IEnumerable<string> WarningEvents = Result.WarningEvents.Select(E => WarningLineForHorde(E.Message)).Distinct();
 
-							if (WarningEvents.Count() > kMaxErrorsOrWarningsToDisplay)
-							{
-								WarningEvents = WarningEvents.Skip(WarningEvents.Count() - kMaxErrorsOrWarningsToDisplay);
-							}
+								WarningEvents = CapErrorOrWarningList(WarningEvents);
 
-							MB.H4(string.Format("Warning: Test '{0}' completed with warnings", Result.DisplayName));
-							MB.Paragraph("FullName: " + Result.TestName);
-							MB.UnorderedList(WarningEvents.Distinct().Select(E => SantizeLineForHorde(E)));
+								MB.H4(string.Format("Warning: Test '{0}' completed with warnings", Result.TestDisplayName));
+								MB.Paragraph("FullName: " + Result.FullTestPath);
+								MB.UnorderedList(WarningEvents);
 							}
 						}
 
@@ -607,19 +689,16 @@ namespace UE
 						{
 							MB.H3("The following test(s) timed out or did not run:");
 
-							foreach (AutomationTestResult Result in IncompleteTests)
+							foreach (UnrealAutomatedTestResult Result in IncompleteTests)
 							{
 								// only show the last N items
-								IEnumerable<string> WarningAndErrorEvents = Result.WarningAndErrorEvents.Distinct();
+								IEnumerable<string> ErrorAndWarningEvents = Result.WarningAndErrorEvents.Select(E => E.IsError? ErrorLineForHorde(E.Message) : WarningLineForHorde(E.Message)).Distinct();
 
-								if (WarningAndErrorEvents.Count() > kMaxErrorsOrWarningsToDisplay)
-								{
-									WarningAndErrorEvents = WarningAndErrorEvents.Skip(WarningAndErrorEvents.Count() - kMaxErrorsOrWarningsToDisplay);
-								}
+								ErrorAndWarningEvents = CapErrorOrWarningList(ErrorAndWarningEvents);
 
-								MB.H4(string.Format("Error: Test '{0}' did not run or complete", Result.DisplayName));
-								MB.Paragraph("FullName: " + Result.TestName);
-								MB.UnorderedList(WarningAndErrorEvents.Distinct().Select(E => SantizeLineForHorde(E)));
+								MB.H4(string.Format("Error: Test '{0}' did not run or complete", Result.TestDisplayName));
+								MB.Paragraph("FullName: " + Result.FullTestPath);
+								MB.UnorderedList(ErrorAndWarningEvents);
 							}
 						}
 
@@ -659,20 +738,12 @@ namespace UE
 
 				if (EditorRole.LogSummary.EngineInitialized)
 				{
-					// Use the paths from the report. If we passed these in they should be the same, and if not
-					// they'll be valid defaults
-					string AutomationReportPath = Parser.AutomationReportPath;
-					string AutomationReportURL = Parser.AutomationReportURL;
+					string AutomationReportPath = string.Empty;
+					string AutomationReportURL = string.Empty;
 					if (GetConfiguration() is AutomationTestConfig Config)
 					{
-						if (string.IsNullOrEmpty(AutomationReportPath))
-						{
-							AutomationReportPath = Config.ReportExportPath;
-						}
-						if (string.IsNullOrEmpty(AutomationReportURL))
-						{
-							AutomationReportURL = Config.ReportURL;
-						}
+						AutomationReportPath = Config.ReportExportPath;
+						AutomationReportURL = Config.ReportURL;
 					}
 					if (!string.IsNullOrEmpty(AutomationReportPath) || !string.IsNullOrEmpty(AutomationReportURL))
 					{
@@ -708,10 +779,10 @@ namespace UE
 				{
 					AutomationLogParser Parser = new AutomationLogParser(Role.LogSummary.FullLogContent);
 					AllErrors.AddRange(
-						Parser.GetResults().Where(R => !R.Passed)
-							.SelectMany(R => R.Events
-								.Where(E => E.ToLower().Contains("error"))
-								.Distinct().Select(E => string.Format("[test={0}] {1}", R.DisplayName, E))
+						Parser.GetResults().Where(R => !R.HasSucceeded)
+							.SelectMany(R => R.Entries
+								.Where(E => E.Event.Type == EventType.Error)
+								.Distinct().Select(E => string.Format("[test={0}] {1}", R.TestDisplayName, E))
 							)
 						);
 				}
@@ -740,9 +811,9 @@ namespace UE
 					AutomationLogParser Parser = new AutomationLogParser(Role.LogSummary.FullLogContent);
 					AllWarnings.AddRange(
 						Parser.GetResults()
-							.SelectMany(R => R.Events
-								.Where(E => E.ToLower().Contains("warning"))
-								.Distinct().Select(E => string.Format("[test={0}] {1}", R.DisplayName, E))
+							.SelectMany(R => R.Entries
+								.Where(E => E.Event.Type == EventType.Warning)
+								.Distinct().Select(E => string.Format("[test={0}] {1}", R.TestDisplayName, E))
 							)
 						);
 				}
diff --git a/Engine/Source/Programs/AutomationTool/Gauntlet/Unreal/Base/Gauntlet.UnrealTestConfiguration.cs b/Engine/Source/Programs/AutomationTool/Gauntlet/Unreal/Base/Gauntlet.UnrealTestConfiguration.cs
--- a/Engine/Source/Programs/AutomationTool/Gauntlet/Unreal/Base/Gauntlet.UnrealTestConfiguration.cs
+++ b/Engine/Source/Programs/AutomationTool/Gauntlet/Unreal/Base/Gauntlet.UnrealTestConfiguration.cs
@@ -544,6 +544,18 @@ namespace Gauntlet
 		[AutoParam(600.0f)]
 		public float MaxDuration { get; set; }
 
+		/// <summary>
+		/// Resume test run on critical failure through pass retry
+		/// </summary>
+		[AutoParam]
+		public bool ResumeOnCriticalFailure = false;
+
+		/// <summary>
+		/// Max number of retries in case of critical failure
+		/// </summary>
+		[AutoParam(3)]
+		public int MaxRetries { get; set; }
+
 		/// <summary>
 		/// Produce test artifacts for Horde build system
 		/// </summary>
@@ -859,8 +871,8 @@ namespace Gauntlet
 					AppConfig.CommandLineParams.GameMap = MapChoice;
 				}
 			}
-			
-						// we write results to Horde test data if we run under Horde agent
+
+			// we write results to Horde test data if we run under Horde agent
 			if (HordeReport.IsUnderHordeAgent)
 			{
 				WriteTestResultsForHorde = true;
diff --git a/Engine/Source/Programs/AutomationTool/Gauntlet/Unreal/Base/Gauntlet.UnrealTestNode.cs b/Engine/Source/Programs/AutomationTool/Gauntlet/Unreal/Base/Gauntlet.UnrealTestNode.cs
--- a/Engine/Source/Programs/AutomationTool/Gauntlet/Unreal/Base/Gauntlet.UnrealTestNode.cs
+++ b/Engine/Source/Programs/AutomationTool/Gauntlet/Unreal/Base/Gauntlet.UnrealTestNode.cs
@@ -89,10 +89,10 @@
 		{
 			if (RoleResults == null)
 			{
-				return Warnings;
+				return Events.Where(E => E.IsWarning).Select(E => E.Message);
 			}
 
-			Warnings = RoleResults.SelectMany(R =>
+			IEnumerable<string> Warnings = RoleResults.SelectMany(R =>
 			{
 				return R.LogSummary.Ensures.Select(E => E.Message);
 			}).ToList(); 
@@ -113,12 +113,12 @@
 		{
 			if (RoleResults == null)
 			{
-				return Errors;
+				return Events.Where(E => E.IsError).Select(E => E.Message);
 			}
 
 			var FailedRoles = GetRolesThatFailed();
 
-			Errors = FailedRoles.Where(R => R.LogSummary.FatalError != null).Select(R => R.LogSummary.FatalError.Message).ToList();
+			IEnumerable<string> Errors = FailedRoles.Where(R => R.LogSummary.FatalError != null).Select(R => R.LogSummary.FatalError.Message).ToList();
 
 			// add all errors from from roles if desired
 			if (Flags.HasFlag(BehaviorFlags.PromoteErrors))
@@ -177,7 +177,7 @@
 		public virtual void ReportError(string Message, params object[] Args)
 		{
 			Message = string.Format(Message, Args);
-			Errors.Add(Message);
+			Events.Add(new UnrealAutomationEvent(EventType.Error, Message));
 			if (!LogWarningsAndErrorsAfterSummary) { Log.Error(Message); }
 			if (GetTestStatus() == TestStatus.Complete && GetTestResult() == TestResult.Passed)
 			{
@@ -192,7 +192,7 @@
 		public virtual void ReportWarning(string Message, params object[] Args)
 		{
 			Message = string.Format(Message, Args);
-			if (!LogWarningsAndErrorsAfterSummary) { Warnings.Add(Message); }
+			if (!LogWarningsAndErrorsAfterSummary) { Events.Add(new UnrealAutomationEvent(EventType.Warning, Message)); }
 			Log.Warning(Message);
 		}
 
@@ -263,14 +263,9 @@
 		public IEnumerable<UnrealRoleArtifacts> SessionArtifacts { get; private set; }
 
 		/// <summary>
-		/// Error collection.
+		/// Error and warning collection.
 		/// </summary>
-		protected List<string> Errors { get; private set; } = new List<string>();
+		protected List<UnrealAutomationEvent> Events { get; private set; } = new List<UnrealAutomationEvent>();
 
 		/// <summary>
-		/// Warning collection.
-		/// </summary>
-		protected List<string> Warnings { get; private set; } = new List<string>();
-
-		/// <summary>
 		/// Whether we submit to the dashboard
@@ -293,6 +288,9 @@
 
 		static protected DateTime SessionStartTime = DateTime.MinValue;
 
+		private int Retries = 0;
+		private int MaxRetries = 3;
+
 		/// <summary>
 		/// Standard semantic versioning for tests. Should be overwritten within individual tests, and individual test maintainers
 		/// are responsible for updating their own versions. See https://semver.org/ for more info on maintaining semantic versions.
@@ -312,11 +310,6 @@
 
 		protected TConfigClass CachedConfig = null;
 
-		/// <summary>
-		/// If our test should exit suddenly, this is the process that caused it
-		/// </summary>
-		protected List<IAppInstance> MissingProcesses;
-
 		protected DateTime TimeOfFirstMissingProcess;
 
 		protected int TimeToWaitForProcesses { get; set; }
@@ -354,7 +347,6 @@
 			Context = InContext;
 
 			UnrealTestResult = TestResult.Invalid;
-			MissingProcesses = new List<IAppInstance>();
 			TimeToWaitForProcesses = 5;
 			LastLogCount = 0;
 			CurrentPass = 0;
@@ -646,31 +638,11 @@
 		}
 
 		/// <summary>
-		/// Called by the test executor to start our test running. After this
-		/// Test.Status should return InProgress or greater
+		/// Generate an unique path for the test artifacts and reserve it
 		/// </summary>
 		/// <returns></returns>
-		public override bool StartTest(int Pass, int InNumPasses)
+		protected String ReserveArtifactPath()
 		{
-			if (UnrealApp == null)
-			{
-				throw new AutomationException("Node already has a null UnrealApp, was PrepareUnrealSession or IsReadyToStart called?");
-			}
-
-			// ensure we reset things
-			SessionArtifacts = Enumerable.Empty<UnrealRoleArtifacts>();
-			RoleResults = Enumerable.Empty<UnrealRoleResult>();
-			UnrealTestResult = TestResult.Invalid;
-			MissingProcesses = new List<IAppInstance>();
-			CurrentPass = Pass;
-			NumPasses = InNumPasses;
-			LastLogCount = 0;
-			LastHeartbeatTime = DateTime.MinValue;
-			LastActiveHeartbeatTime = DateTime.MinValue;
-			
-
-			TConfigClass Config = GetCachedConfiguration();
-
 			// Either use the ArtifactName param or name of this test
 			string TestFolder = string.IsNullOrEmpty(Context.Options.ArtifactName) ? this.ToString() : Context.Options.ArtifactName;
 
@@ -685,13 +657,18 @@
 			TestFolder = TestFolder.Replace(",", "");
 
 			ArtifactPath = Path.Combine(Context.Options.LogDir, TestFolder);
-		
+
 			// if doing multiple passes, put each in a subdir
 			if (NumPasses > 1)
 			{
 				ArtifactPath = Path.Combine(ArtifactPath, string.Format("Pass_{0}_of_{1}", CurrentPass + 1, NumPasses));
 			}
 
+			if (Retries > 0)
+			{
+				ArtifactPath = Path.Combine(ArtifactPath, string.Format("Retry_{0}", Retries));
+			}
+
 			// When running with -parallel we could have several identical tests (same test, configurations) in flight so
 			// we need unique artifact paths. We also don't overwrite dest directories from the build machine for the same
 			// reason of multiple tests for a build. Really though these should use ArtifactPrefix to save to
@@ -724,6 +701,35 @@
 
 			ReservedArtifcactPaths.Add(ArtifactPath);
 
+			return ArtifactPath;
+		}
+
+		/// <summary>
+		/// Called by the test executor to start our test running. After this
+		/// Test.Status should return InProgress or greater
+		/// </summary>
+		/// <returns></returns>
+		public override bool StartTest(int Pass, int InNumPasses)
+		{
+			if (UnrealApp == null)
+			{
+				throw new AutomationException("Node already has a null UnrealApp, was PrepareUnrealSession or IsReadyToStart called?");
+			}
+
+			// ensure we reset things
+			SessionArtifacts = Enumerable.Empty<UnrealRoleArtifacts>();
+			RoleResults = Enumerable.Empty<UnrealRoleResult>();
+			UnrealTestResult = TestResult.Invalid;
+			CurrentPass = Pass;
+			NumPasses = InNumPasses;
+			LastLogCount = 0;
+			LastHeartbeatTime = DateTime.MinValue;
+			LastActiveHeartbeatTime = DateTime.MinValue;			
+
+			TConfigClass Config = GetCachedConfiguration();
+
+			ReserveArtifactPath();
+
 			// We need to create this directory at the start of the test rather than the end of the test - we are running into instances where multiple A/B tests
 			// on the same build are seeing the directory as non-existent and thinking it is safe to write to.
 			Log.Info("UnrealTestNode.StartTest Calling CreateDirectory for artifacts at {0}", ArtifactPath);
@@ -754,7 +760,7 @@
 				// Update these for the executor
 				MaxDuration = Config.MaxDuration;
 				MaxDurationReachedResult = Config.MaxDurationReachedResult;
-				UnrealTestResult = TestResult.Invalid;
+				MaxRetries = Config.MaxRetries;
 				MarkTestStarted();
 			}
 			
@@ -834,6 +840,19 @@
 		/// <returns></returns>
 		public override bool RestartTest()
 		{
+			//Reset/Increment artifact output
+			SessionArtifacts = Enumerable.Empty<UnrealRoleArtifacts>();
+			RoleResults = Enumerable.Empty<UnrealRoleResult>();
+
+			LastLogCount = 0;
+			LastHeartbeatTime = DateTime.MinValue;
+			LastActiveHeartbeatTime = DateTime.MinValue;
+
+			ReserveArtifactPath();
+			Log.Info("UnrealTestNode.ReStartTest Calling CreateDirectory for artifacts at {0}", ArtifactPath);
+			Directory.CreateDirectory(ArtifactPath);
+
+			// Relaunch the test
 			TestInstance = UnrealApp.RestartSession();
 
 			bool bWasRestarted = (TestInstance != null);
@@ -1151,7 +1170,7 @@
 			{
 				foreach (UnrealRoleResult RoleResult in RoleResults)
 				{
-					string LogName = Path.GetFullPath(RoleResult.Artifacts.LogPath).Replace(Path.GetFullPath(Path.Combine(ArtifactPath, "..")), "").TrimStart(Path.DirectorySeparatorChar);
+					string LogName = Path.GetFullPath(RoleResult.Artifacts.LogPath).Replace(Path.GetFullPath(Context.Options.LogDir), "").TrimStart(Path.DirectorySeparatorChar);
 					HordeTestReport.AttachArtifact(RoleResult.Artifacts.LogPath, LogName);
 
 					UnrealLog LogSummary = RoleResult.LogSummary;
@@ -1180,10 +1199,63 @@
 			string JsonReportPath = Path.Combine(UnrealAutomatedTestReportPath, "index.json");
 			if (File.Exists(JsonReportPath))
 			{
+				string HordeArtifactPath = GetConfiguration().HordeArtifactPath;
 				Log.Verbose("Reading json Unreal Automated test report from {0}", JsonReportPath);
 				UnrealAutomatedTestPassResults JsonTestPassResults = UnrealAutomatedTestPassResults.LoadFromJson(JsonReportPath);
+				if (JsonTestPassResults.InProcess > 0)
+				{
+					// The test pass did not run completely
+					Log.Verbose("Found Uncompleted tests: {0}", JsonTestPassResults.InProcess);
+					// Get any critical error and push it to json report and resave it.
+					if (RoleResults != null)
+					{
+						UnrealLog.CallstackMessage FatalError = null;
+						foreach (UnrealRoleResult Result in RoleResults)
+						{
+							if (Result.LogSummary.FatalError != null)
+							{
+								FatalError = Result.LogSummary.FatalError;
+								break;
+							}
+						}
+						if (FatalError != null)
+						{
+							var Test = JsonTestPassResults.Tests.FirstOrDefault((T => T.State == TestStateType.InProcess));
+							if (!String.IsNullOrEmpty(Test.TestDisplayName))
+							{
+								Test.AddError("Engine encountered a critical failure. \n" + FatalError.FormatForLog());
+								JsonTestPassResults.WriteToJson(JsonReportPath);
+							}
+						}
+					}
+					if (GetConfiguration().ResumeOnCriticalFailure)
+					{
+						if (Retries < MaxRetries)
+						{
+							Retries++;
+							// Reschedule test to resume from last 'in-process' test.
+							SetUnrealTestResult(TestResult.WantRetry);
+							// Attach current artifacts to Horde output
+							if (SessionArtifacts != null)
+							{
+								HordeReport.SimpleTestReport TempReport = new HordeReport.SimpleTestReport();
+								TempReport.SetOutputArtifactPath(HordeArtifactPath);
+								foreach (UnrealRoleArtifacts Artifact in SessionArtifacts)
+								{
+									string LogName = Path.GetFullPath(Artifact.LogPath).Replace(Path.GetFullPath(Context.Options.LogDir), "").TrimStart(Path.DirectorySeparatorChar);
+									TempReport.AttachArtifact(Artifact.LogPath, LogName);
+								}
+							}
+							// Discard the report as we are going to do another pass.
+							return null;
+						}
+						else
+						{
+							Log.Error("Reach maximum of retries({0}) to resume on critical failure!", Retries);
+						}
+					}
+				}
 				// Convert test results for Horde
-				string HordeArtifactPath = GetConfiguration().HordeArtifactPath;
 				HordeReport.UnrealEngineTestPassResults HordeTestPassResults = HordeReport.UnrealEngineTestPassResults.FromUnrealAutomatedTests(JsonTestPassResults, UnrealAutomatedTestReportPath, ReportURL);
 				HordeTestPassResults.CopyTestResultsArtifacts(HordeArtifactPath);
 				// Attached test Artifacts
@@ -1191,7 +1263,7 @@
 				{
 					foreach (UnrealRoleArtifacts Artifact in SessionArtifacts)
 					{
-						string LogName = Path.GetFullPath(Artifact.LogPath).Replace(Path.GetFullPath(Path.Combine(ArtifactPath, "..")), "").TrimStart(Path.DirectorySeparatorChar);
+						string LogName = Path.GetFullPath(Artifact.LogPath).Replace(Path.GetFullPath(Context.Options.LogDir), "").TrimStart(Path.DirectorySeparatorChar);
 						HordeTestPassResults.AttachArtifact(Artifact.LogPath, LogName);
 					}
 				}
diff --git a/Engine/Source/Programs/AutomationTool/Gauntlet/Unreal/Engine/UnrealTests.cs b/Engine/Source/Programs/AutomationTool/Gauntlet/Unreal/Engine/UnrealTests.cs
--- a/Engine/Source/Programs/AutomationTool/Gauntlet/Unreal/Engine/UnrealTests.cs
+++ b/Engine/Source/Programs/AutomationTool/Gauntlet/Unreal/Engine/UnrealTests.cs
@@ -76,9 +76,9 @@ namespace Gauntlet
 					{
 						AutomationLogParser Parser = new AutomationLogParser(Role.LogSummary.FullLogContent);
 						AllErrors.AddRange(
-							Parser.GetResults().Where(R => !R.Passed)
-							.SelectMany(R => R.Events
-								.Where(E => E.ToLower().Contains("error")).Select(E => string.Format("[test={0}] {1}", R.DisplayName, E))
+							Parser.GetResults().Where(R => !R.HasSucceeded)
+							.SelectMany(R => R.Entries
+								.Where(E => E.Event.Type == EventType.Error).Select(E => string.Format("[test={0}] {1}", R.TestDisplayName, E))
 								)
 							);
 					}
@@ -219,8 +219,8 @@ namespace Gauntlet
 					{
 						AutomationLogParser Parser = new AutomationLogParser(InLog.FullLogContent);
 
-						IEnumerable<AutomationTestResult> TotalTests = Parser.GetResults();
-						IEnumerable<AutomationTestResult> FailedTests = TotalTests.Where(R => !R.Passed);
+						IEnumerable<UnrealAutomatedTestResult> TotalTests = Parser.GetResults();
+						IEnumerable<UnrealAutomatedTestResult> FailedTests = TotalTests.Where(R => !R.HasSucceeded);
 
 						// Tests failed so list that as our primary cause of failure
 						if (FailedTests.Any())
@@ -257,9 +257,9 @@ namespace Gauntlet
 				{
 					AutomationLogParser Parser = new AutomationLogParser(EditorRole.LogSummary.FullLogContent);
 
-					IEnumerable<AutomationTestResult> AllTests = Parser.GetResults();
-					IEnumerable<AutomationTestResult> FailedTests = AllTests.Where(R => R.Completed && !R.Passed);
-					IEnumerable<AutomationTestResult> IncompleteTests = AllTests.Where(R => !R.Completed);
+					IEnumerable<UnrealAutomatedTestResult> AllTests = Parser.GetResults();
+					IEnumerable<UnrealAutomatedTestResult> FailedTests = AllTests.Where(R => R.IsComplete && !R.HasSucceeded);
+					IEnumerable<UnrealAutomatedTestResult> IncompleteTests = AllTests.Where(R => !R.IsComplete);
 
 					if (AllTests.Count() == 0)
 					{
@@ -274,10 +274,10 @@ namespace Gauntlet
 						{
 							MB.H3("The following tests failed:");
 
-							foreach (AutomationTestResult Result in FailedTests)
+							foreach (UnrealAutomatedTestResult Result in FailedTests)
 							{
-								MB.H4(string.Format("{0}: {1}", Result.DisplayName, Result.TestName));
-								MB.UnorderedList(Result.Events);
+								MB.H4(string.Format("{0}: {1}", Result.TestDisplayName, Result.FullTestPath));
+								MB.UnorderedList(Result.ErrorEvents.Select(E => string.Format("Error: {0}",E.Message)));
 							}
 						}
 
@@ -285,10 +285,10 @@ namespace Gauntlet
 						{
 							MB.H3("The following tests timed out:");
 
-							foreach (AutomationTestResult Result in IncompleteTests)
+							foreach (UnrealAutomatedTestResult Result in IncompleteTests)
 							{
-								MB.H4(string.Format("{0}: {1}", Result.DisplayName, Result.TestName));
-								MB.UnorderedList(Result.Events);
+								MB.H4(string.Format("{0}: {1}", Result.TestDisplayName, Result.FullTestPath));
+								MB.UnorderedList(Result.ErrorEvents.Select(E => string.Format("Error: {0}", E.Message)));
 							}
 						}
 
@@ -321,7 +321,7 @@ namespace Gauntlet
 		/// <summary>
 		/// Define a Config class for this test that contains the available options
 		/// </summary>
-		public class EngineTestConfig : UnrealTestConfiguration
+		public class EngineTestConfig : UE.AutomationTestConfig
 		{
 			/// <summary>
 			/// Set to true if the editor executes all these tests in is own process and PIE
@@ -329,36 +329,6 @@ namespace Gauntlet
 			[AutoParam]
 			public virtual bool UseEditor { get; set; } = false;
 
-			/// <summary>
-			/// Run with d3d12 RHI
-			/// </summary>
-			[AutoParam]
-			public virtual bool D3D12 { get; set; } = false;
-
-			/// <summary>
-			/// Run with Nvidia cards for raytracing support
-			/// </summary>
-			[AutoParam]
-			public virtual bool PreferNvidia { get; set; } = false;
-
-			/// <summary>
-			/// Run forcing raytracing 
-			/// </summary>
-			[AutoParam]
-			public virtual bool RayTracing { get; set; } = false;
-
-			/// <summary>
-			/// Use the StompMalloc memory allocator.
-			/// </summary>
-			[AutoParam]
-			public virtual bool StompMalloc { get; set; } = false;
-
-			/// <summary>
-			/// Enable the D3D debug layer.
-			/// </summary>
-			[AutoParam]
-			public virtual bool D3DDebug { get; set; } = false;
-
 			/// <summary>
 			/// Enable the RHI validation layer.
 			/// </summary>
@@ -366,51 +336,10 @@ namespace Gauntlet
 			public virtual bool RHIValidation { get; set; } = false;
 
 			/// <summary>
-			/// Disable capturing frame trace for image based tests
-			/// </summary>
-			[AutoParam]
-			public virtual bool DisableFrameTraceCapture { get; set; } = true;
-
-			/// <summary>
-			/// Filter or groups of tests to apply
-			/// </summary>
-			[AutoParam]
-			public virtual string TestFilter { get; set; } = "Filter:Project+Filter:System";
-
-			/// <summary>
-			/// Path to write the automation report to
-			/// </summary>
-			[AutoParam]
-			public virtual string ReportExportPath { get; set; } = "";
-
-			/// <summary>
-			/// Path the report can be found at
+			/// Filter or groups of tests to apply (alias for RunTest) 
 			/// </summary>
 			[AutoParam]
-			public virtual string ReportURL { get; set; } = "";
-
-			/// <summary>
-			/// Use Simple Horde Report instead of Unreal Automated Tests
-			/// </summary>
-			[AutoParam]
-			public virtual bool SimpleHordeReport { get; set; } = false;
-
-			/// <summary>
-			/// Validate DDC during tests
-			/// </summary>
-			[AutoParam]
-			public virtual bool VerifyDDC { get; set; } = false;
-
-			/// <summary>
-			/// Validate DDC during tests
-			/// </summary>
-			[AutoParam]
-			public virtual string DDC { get; set; } = "";
-
-			/// <summary>
-			/// Used for having the editor and any client communicate
-			/// </summary>
-			public string SessionID = Guid.NewGuid().ToString();
+			public virtual string TestFilter { get; set; } = "";
 
 			/// <summary>
 			/// Apply this config to the role that is passed in
@@ -420,100 +349,12 @@ namespace Gauntlet
 			/// <param name="OtherRoles"></param>
 			public override void ApplyToConfig(UnrealAppConfig AppConfig, UnrealSessionRole ConfigRole, IEnumerable<UnrealSessionRole> OtherRoles)
 			{
-				base.ApplyToConfig(AppConfig, ConfigRole, OtherRoles);
-
-				if (WriteTestResultsForHorde)
+				if (string.IsNullOrEmpty(RunTest))
 				{
-					if (string.IsNullOrEmpty(ReportExportPath))
-					{
-						ReportExportPath = Path.Combine(Globals.TempDir, "TestReport");
-					}
-					if (string.IsNullOrEmpty(HordeTestDataPath))
-					{
-						HordeTestDataPath = HordeReport.DefaultTestDataDir;
-					}
-					if (string.IsNullOrEmpty(HordeArtifactPath))
-					{
-						HordeArtifactPath = HordeReport.DefaultArtifactsDir;
-					}
-				}
-
-				string ReportOutputArg = "";
-				if (!string.IsNullOrEmpty(ReportExportPath))
-				{
-					if (Directory.Exists(ReportExportPath))
-					{
-						// clean up previous run if any
-						DirectoryInfo ReportDirInfo = new DirectoryInfo(ReportExportPath);
-						ReportDirInfo.Delete(true);
-					}
-					ReportOutputArg = string.Format("-ReportExportPath=\"{0}\"", ReportExportPath);
-				}
-
-				if (UseEditor)
-				{
-					if (ConfigRole.RoleType == UnrealTargetRole.Editor)
-					{
-						AppConfig.CommandLine += string.Format(" -NoWatchdog -stdout -FORCELOGFLUSH -CrashForUAT -unattended -buildmachine {0} -ExecCmds=\"Automation RunTests {1}; Quit;\"", ReportOutputArg, TestFilter);
-					}
-
-					if (VerifyDDC)
-					{
-						AppConfig.CommandLine += " -VerifyDDC";
-					}
-
-					if (!string.IsNullOrEmpty(DDC))
-					{
-						AppConfig.CommandLine += string.Format(" -ddc={0}", DDC);
-					}
-				}
-				else
-				{
-					// If the test isnt using the editor for both roles then pass the IP of the editor (us) to the client
-					string HostIP = UnrealHelpers.GetHostIpAddress();
-
-					if (ConfigRole.RoleType.IsClient())
-					{
-						AppConfig.CommandLine += string.Format(" -sessionid={0} -messaging -log -TcpMessagingConnect={1}:6666", SessionID, HostIP);
-					}
-					else if (ConfigRole.RoleType.IsEditor())
-					{
-						AppConfig.CommandLine += string.Format(" -nullrhi -ExecCmds=\"Automation list;StartRemoteSession {0};RunTests {1}; Quit\" -TcpMessagingListen={2}:6666 -multihome={3} -log {4}", SessionID, TestFilter, HostIP, HostIP, ReportOutputArg);
-					}
+					RunTest = TestFilter;
 				}
 
-				if (DisableFrameTraceCapture || RayTracing)
-				{
-					AppConfig.CommandLine += " -DisableFrameTraceCapture";
-				}
-
-				if (RayTracing)
-				{
-					AppConfig.CommandLine += " -dpcvars=r.RayTracing=1,r.SkinCache.CompileShaders=1,AutomationAllowFrameTraceCapture=0 -noxgeshadercompile";
-				}
-
-				if (ConfigRole.Platform == UnrealTargetPlatform.Win64)
-				{
-					if (PreferNvidia)
-					{
-						AppConfig.CommandLine += " -preferNvidia";
-					}
-
-					if (D3D12)
-					{
-						AppConfig.CommandLine += " -d3d12";
-					}
-
-					if (D3DDebug)
-					{
-						AppConfig.CommandLine += " -d3ddebug";
-					}
-
-					if (StompMalloc)
-					{
-						AppConfig.CommandLine += " -stompmalloc";
-					}
-				}
+				base.ApplyToConfig(AppConfig, ConfigRole, OtherRoles);
 
 				if (RHIValidation)
 				{
diff --git a/Engine/Source/Programs/AutomationTool/Gauntlet/Unreal/Utils/Gauntlet.AutomationLogParser.cs b/Engine/Source/Programs/AutomationTool/Gauntlet/Unreal/Utils/Gauntlet.AutomationLogParser.cs
--- a/Engine/Source/Programs/AutomationTool/Gauntlet/Unreal/Utils/Gauntlet.AutomationLogParser.cs
+++ b/Engine/Source/Programs/AutomationTool/Gauntlet/Unreal/Utils/Gauntlet.AutomationLogParser.cs
@@ -10,105 +10,6 @@ using System.Threading.Tasks;
 
 namespace Gauntlet
 {
-	/// <summary>
-	/// Result for a single automation result
-	/// </summary>
-	public class AutomationTestResult
-	{
-		/// <summary>
-		/// Short friendly name of this test
-		/// </summary>
-		public string		DisplayName;
-
-		/// <summary>
-		/// Full test name
-		/// </summary>
-		public string		TestName;
-
-		/// <summary>
-		/// True if the test completed (pass or fail)
-		/// </summary>
-		public bool			Completed;
-
-		/// <summary>
-		/// True if the test passed
-		/// </summary>
-		public bool			Passed;
-
-		/// <summary>
-		/// Events logged during the test. Should contain errors if the test failed
-		/// </summary>
-		public IEnumerable<string> Events;
-
-		/// <summary>
-		/// Returns all events that are a warning
-		/// </summary>
-		public IEnumerable<string> WarningEvents 
-		{  
-			get
-			{
-				return Events.Where(E => E.ToLower().Contains(": warning:"));
-			}
-		}
-
-		/// <summary>
-		/// Returns all events that are a error
-		/// </summary>
-		public IEnumerable<string> ErrorEvents
-		{
-			get
-			{
-				return Events.Where(E => E.ToLower().Contains(": error:"));
-			}
-		}
-
-		/// <summary>
-		/// Returns all events that are a error
-		/// </summary>
-		public IEnumerable<string> WarningAndErrorEvents
-		{
-			get
-			{
-				return Events.Where(E => E.ToLower().Contains(": error:") || E.ToLower().Contains(": warning:"));
-			}
-		}
-
-		/// <summary>
-		/// Returns the name of the test. If DisplayName and TestName differ they are joined
-		/// </summary>
-		public string FullName
-		{
-			get
-			{
-				string FullName = DisplayName;
-
-				if (TestName != DisplayName)
-				{
-					FullName += string.Format(": {0}", TestName);
-				}
-
-				return FullName;
-			}
-		}
-
-		/// <summary>
-		/// Result contstructor
-		/// </summary>
-		/// <param name="InDisplayName"></param>
-		/// <param name="InTestName"></param>
-		/// <param name="bInPassed"></param>
-		public AutomationTestResult(string InDisplayName, string InTestName, bool InCompleted, bool bInPassed, IEnumerable<string> InEvents)
-		{
-			DisplayName = InDisplayName;
-			TestName = InTestName;
-			Completed = InCompleted;
-			Passed = bInPassed;
-			Events = InEvents.ToArray();
-		}
-
-
-	};
-
 	/// <summary>
 	/// Helper class for parsing AutomationTest results from either an UnrealLogParser or log contents
 	/// </summary>
@@ -177,7 +78,7 @@ namespace Gauntlet
 		/// Returns all results found in our construction content.
 		/// </summary>
 		/// <returns></returns>
-		public IEnumerable<AutomationTestResult> GetResults()
+		public IEnumerable<UnrealAutomatedTestResult> GetResults()
 		{
 			IEnumerable<Match> TestStarts = Parser.GetAllMatches(@"LogAutomationController.+Test Started. Name={(.+?)}");
 
@@ -191,14 +92,23 @@ namespace Gauntlet
 
 			string[] AutomationChannel = Parser.GetLogChannel("AutomationController").ToArray();
 
+			Func<string, string> SantizeLine = (L) =>
+			{
+				L = L.Replace(": Error: ", "");
+				L = L.Replace(": Warning: ", "");
+				L = L.Replace("LogAutomationController: ", "");
+
+				return L;
+			};
+
 			// Convert these lines into results by parsing out the details and then adding the events
-			IEnumerable<AutomationTestResult> Results = TestStarts.Select(TestMatch =>
+			IEnumerable<UnrealAutomatedTestResult> Results = TestStarts.Select(TestMatch =>
 			{
+				UnrealAutomatedTestResult Result = new UnrealAutomatedTestResult();
 				string DisplayName = TestMatch.Groups[1].ToString();
 				string LongName = "";
 				bool Completed = false;
 				bool Passed = false;
-				List<string> Events = new List<string>();
 				
 				Match ResultMatch = TestResults.Where(M => M.Groups[2].ToString() == DisplayName).FirstOrDefault();
 
@@ -222,16 +132,21 @@ namespace Gauntlet
 								break;
 							}
 
-							Events.Add(Event);
+							EventType EventType = Event.Contains(": Error: ") ? EventType.Error : Event.Contains(": Warning: ") ? EventType.Warning : EventType.Info;
+
+							Result.AddEvent(EventType, SantizeLine(Event));
 						}
 					}
 				}
 				else
 				{
-					Events.Add(string.Format("Gauntlet.AutomationLogParser: Error: Test {0} incomplete.", DisplayName));
+					Result.AddEvent(EventType.Error, string.Format("Test {0} incomplete.", DisplayName));
 				}
 
-				AutomationTestResult Result = new AutomationTestResult(DisplayName, LongName, Completed, Passed, Events);
+				Result.TestDisplayName = DisplayName;
+				Result.FullTestPath = LongName;
+				Result.State = Passed ? TestStateType.Success : Completed ? TestStateType.Fail : TestStateType.InProcess;
+
 				return Result;
 			});
 
diff --git a/Engine/Source/Programs/AutomationTool/Gauntlet/Unreal/Utils/Gauntlet.UnrealLogParser.cs b/Engine/Source/Programs/AutomationTool/Gauntlet/Unreal/Utils/Gauntlet.UnrealLogParser.cs
--- a/Engine/Source/Programs/AutomationTool/Gauntlet/Unreal/Utils/Gauntlet.UnrealLogParser.cs
+++ b/Engine/Source/Programs/AutomationTool/Gauntlet/Unreal/Utils/Gauntlet.UnrealLogParser.cs
@@ -103,13 +103,11 @@ namespace Gauntlet
 			/// <returns>Formatted log string version of callstack.</returns>
 			public string FormatForLog()
 			{
-				string FormattedString = "";
-				FormattedString += string.Format("{0}", Message);
+				string FormattedString = string.Format("{0}\n", Message);
 				foreach (string StackLine in Callstack)
 				{
 					FormattedString += string.Format("\t{0}\n", StackLine);
 				}
-				FormattedString += "\n";
 				return FormattedString;
 			}
 		};
diff --git a/Engine/Source/Programs/AutomationTool/Gauntlet/Unreal/Utils/Gauntlet.UnrealTestResults.cs b/Engine/Source/Programs/AutomationTool/Gauntlet/Unreal/Utils/Gauntlet.UnrealTestResults.cs
--- a/Engine/Source/Programs/AutomationTool/Gauntlet/Unreal/Utils/Gauntlet.UnrealTestResults.cs
+++ b/Engine/Source/Programs/AutomationTool/Gauntlet/Unreal/Utils/Gauntlet.UnrealTestResults.cs
@@ -1,8 +1,11 @@
 // Copyright Epic Games, Inc. All Rights Reserved.
 
+using System;
 using System.IO;
 using System.Collections.Generic;
 using System.Text.Json;
+using System.Text.Json.Serialization;
+using System.Linq;
 
 namespace Gauntlet
 {
@@ -21,10 +24,38 @@ namespace Gauntlet
 	}
 	public class UnrealAutomationEvent
 	{
-		public string Type { get; set; }
+		public EventType Type { get; set; }
 		public string Message { get; set; }
 		public string Context { get; set; }
 		public string Artifact { get; set; }
+
+		public UnrealAutomationEvent()
+		{ }
+		public UnrealAutomationEvent(EventType InType, string InMessage)
+		{
+			Type = InType;
+			Message = InMessage;
+		}
+
+		[JsonIgnore]
+		public bool IsError
+		{
+			get
+			{
+				return Type == EventType.Error;
+			}
+
+		}
+
+		[JsonIgnore]
+		public bool IsWarning
+		{
+			get
+			{
+				return Type == EventType.Warning;
+			}
+
+		}
 	}
 	public class UnrealAutomationEntry
 	{
@@ -38,11 +69,97 @@ namespace Gauntlet
 		public string TestDisplayName { get; set; }
 		public string FullTestPath { get; set; }
 		public string ArtifactName { get; set; }
-		public string State { get; set; }
+		public TestStateType State { get; set; }
 		public int Warnings { get; set; }
 		public int Errors { get; set; }
 		public List<UnrealAutomationArtifact> Artifacts { get; set; }
 		public List<UnrealAutomationEntry> Entries { get; set; }
+
+		public UnrealAutomatedTestResult()
+		{
+			Artifacts = new List<UnrealAutomationArtifact>();
+			Entries = new List<UnrealAutomationEntry>();
+		}
+
+		public void AddEvent(EventType EventType, string Message)
+		{
+			var Event = new UnrealAutomationEvent(EventType, Message);
+			var Entry = new UnrealAutomationEntry();
+			Entry.Event = Event;
+			Entry.Timestamp = DateTime.Now.ToString("yyyy.MM.dd-HH.mm.ss");
+			Entries.Add(Entry);
+
+			switch(EventType)
+			{
+				case EventType.Error:
+					Errors++;
+					break;
+				case EventType.Warning:
+					Warnings++;
+					break;
+			}
+		}
+		public void AddError(string Message)
+		{
+			AddEvent(EventType.Error, Message);
+		}
+		public void AddWarning(string Message)
+		{
+			AddEvent(EventType.Warning, Message);
+		}
+		public void AddInfo(string Message)
+		{
+			AddEvent(EventType.Info, Message);
+		}
+
+		[JsonIgnore]
+		public bool IsComplete
+		{
+			get
+			{
+				return State != TestStateType.InProcess;
+			}
+		}
+		[JsonIgnore]
+		public bool HasSucceeded
+		{
+			get
+			{
+				return State == TestStateType.Success;
+			}
+		}
+		[JsonIgnore]
+		public bool HasWarnings
+		{
+			get
+			{
+				return Warnings > 0;
+			}
+		}
+		[JsonIgnore]
+		public IEnumerable<UnrealAutomationEvent> ErrorEvents
+		{
+			get
+			{
+				return Entries.Where(E => E.Event.IsError).Select(E => E.Event);
+			}
+		}
+		[JsonIgnore]
+		public IEnumerable<UnrealAutomationEvent> WarningEvents
+		{
+			get
+			{
+				return Entries.Where(E => E.Event.IsWarning).Select(E => E.Event);
+			}
+		}
+		[JsonIgnore]
+		public IEnumerable<UnrealAutomationEvent> WarningAndErrorEvents
+		{
+			get
+			{
+				return Entries.Where(E => E.Event.IsError || E.Event.IsWarning).Select(E => E.Event);
+			}
+		}
 	}
 	public class UnrealAutomatedTestPassResults
 	{
@@ -58,6 +175,51 @@ namespace Gauntlet
 		public string ComparisonExportDirectory { get; set; }
 		public List<UnrealAutomatedTestResult> Tests { get; set; }
 
+		public UnrealAutomatedTestPassResults()
+		{
+			Tests = new List<UnrealAutomatedTestResult>();
+		}
+
+		public UnrealAutomatedTestResult AddTest(string DisplayName, string FullName, TestStateType State)
+		{
+			var Test = new UnrealAutomatedTestResult();
+			Test.TestDisplayName = DisplayName;
+			Test.FullTestPath = FullName;
+			Test.State = State;
+
+			return AddTest(Test);
+		}
+
+		public UnrealAutomatedTestResult AddTest(UnrealAutomatedTestResult Test)
+		{
+			Tests.Add(Test);
+
+			switch (Test.State)
+			{
+				case TestStateType.Success:
+					if (Test.HasWarnings)
+					{
+						SucceededWithWarnings++;
+					}
+					else
+					{
+						Succeeded++;
+					}
+					break;
+				case TestStateType.Fail:
+					Failed++;
+					break;
+				case TestStateType.NotRun:
+					NotRun++;
+					break;
+				case TestStateType.InProcess:
+					InProcess++;
+					break;
+			}
+
+			return Test;
+		}
+
 		/// <summary>
 		/// Load Unreal Automated Test Results from json report
 		/// </summary>
@@ -72,5 +234,32 @@ namespace Gauntlet
 			UnrealAutomatedTestPassResults JsonTestPassResults = JsonSerializer.Deserialize<UnrealAutomatedTestPassResults>(JsonString, Options);
 			return JsonTestPassResults;
 		}
+
+		/// <summary>
+		/// Write json data into a file
+		/// </summary>
+		/// <param name="FilePath"></param>
+		public void WriteToJson(string FilePath)
+		{
+			string OutputTestDataDir = Path.GetDirectoryName(FilePath);
+			if (!Directory.Exists(OutputTestDataDir))
+			{
+				Directory.CreateDirectory(OutputTestDataDir);
+			}
+
+			Log.Verbose("Writing Json to file {0}", FilePath);
+			try
+			{
+				JsonSerializerOptions Options = new JsonSerializerOptions
+				{
+					WriteIndented = true
+				};
+				File.WriteAllText(FilePath, JsonSerializer.Serialize(this, Options));
+			}
+			catch (Exception Ex)
+			{
+				Log.Error("Failed to save json file. {0}", Ex);
+			}
+		}
 	}
 }
