Gauntlet - fix a warning and failure to produce a report when ResumeOnCriticalFaiure is on and the test pass timeout.

#rnx

diff --git a/Engine/Source/Programs/AutomationTool/Gauntlet/Unreal/Base/Gauntlet.UnrealTestNode.cs b/Engine/Source/Programs/AutomationTool/Gauntlet/Unreal/Base/Gauntlet.UnrealTestNode.cs
--- a/Engine/Source/Programs/AutomationTool/Gauntlet/Unreal/Base/Gauntlet.UnrealTestNode.cs
+++ b/Engine/Source/Programs/AutomationTool/Gauntlet/Unreal/Base/Gauntlet.UnrealTestNode.cs
@@ -1206,9 +1206,11 @@ namespace Gauntlet
 				{
 					// The test pass did not run completely
 					Log.Verbose("Found Uncompleted tests: {0}", JsonTestPassResults.InProcess);
+					bool HasTimeout = false;
 					// Get any critical error and push it to json report and resave it.
 					if (RoleResults != null)
 					{
+						HasTimeout = RoleResults.Where(R => R.ProcessResult == UnrealProcessResult.TimeOut).Any();
 						UnrealLog.CallstackMessage FatalError = null;
 						foreach (UnrealRoleResult Result in RoleResults)
 						{
@@ -1228,7 +1230,7 @@ namespace Gauntlet
 							}
 						}
 					}
-					if (GetConfiguration().ResumeOnCriticalFailure)
+					if (GetConfiguration().ResumeOnCriticalFailure && !HasTimeout)
 					{
 						if (Retries < MaxRetries)
 						{
