Gauntlet - Fix log not being parsed when no json report were generated during the test pass

#rnx

diff --git a/Engine/Source/Programs/AutomationTool/Gauntlet/Unreal/Automation/UE.Automation.cs b/Engine/Source/Programs/AutomationTool/Gauntlet/Unreal/Automation/UE.Automation.cs
--- a/Engine/Source/Programs/AutomationTool/Gauntlet/Unreal/Automation/UE.Automation.cs
+++ b/Engine/Source/Programs/AutomationTool/Gauntlet/Unreal/Automation/UE.Automation.cs
@@ -439,7 +439,7 @@ namespace UE
 			base.TickTest();
 		}
 
-		private UnrealAutomatedTestPassResults GetTestPassResults()
+		private UnrealAutomatedTestPassResults GetTestPassResults(UnrealLog InLog)
 		{
 			if (TestPassResults == null)
 			{
@@ -457,31 +457,27 @@ namespace UE
 				{
 					// Parse automaton info from the log then
 					TestPassResults = new UnrealAutomatedTestPassResults();
-					var EditorRole = RoleResults.Where(R => R.Artifacts.SessionRole.RoleType == UnrealTargetRole.Editor).FirstOrDefault();
-					if (EditorRole != null)
+					AutomationLogParser LogParser = new AutomationLogParser(InLog.FullLogContent);
+					IEnumerable<UnrealAutomatedTestResult> LogTestResults = LogParser.GetResults();
+					if (LogTestResults.Any())
 					{
-						AutomationLogParser LogParser = new AutomationLogParser(EditorRole.LogSummary.FullLogContent);
-						IEnumerable<UnrealAutomatedTestResult> LogTestResults = LogParser.GetResults();
-						if (LogTestResults.Any())
+						foreach (UnrealAutomatedTestResult LogTestResult in LogTestResults)
 						{
-							foreach (UnrealAutomatedTestResult LogTestResult in LogTestResults)
-							{
-								TestPassResults.AddTest(LogTestResult);
-							}
+							TestPassResults.AddTest(LogTestResult);
 						}
-						else
+					}
+					else
+					{
+						foreach (UnrealLog.LogEntry Entry in LogParser.AutomationWarningsAndErrors)
 						{
-							foreach (UnrealLog.LogEntry Entry in LogParser.AutomationWarningsAndErrors)
+							switch (Entry.Level)
 							{
-								switch (Entry.Level)
-								{
-									case UnrealLog.LogLevel.Error:
-										Events.Add(new UnrealAutomationEvent(EventType.Error, Entry.Message));
-										break;
-									case UnrealLog.LogLevel.Warning:
-										Events.Add(new UnrealAutomationEvent(EventType.Warning, Entry.Message));
-										break;
-								}
+								case UnrealLog.LogLevel.Error:
+									Events.Add(new UnrealAutomationEvent(EventType.Error, Entry.Message));
+									break;
+								case UnrealLog.LogLevel.Warning:
+									Events.Add(new UnrealAutomationEvent(EventType.Warning, Entry.Message));
+									break;
 							}
 						}
 					}
@@ -508,7 +504,7 @@ namespace UE
 				// if no fatal errors, check test results
 				if (InLog.FatalError == null)
 				{
-					var TestResults = GetTestPassResults();
+					var TestResults = GetTestPassResults(InLog);
 
 					// Tests failed so list that as our primary cause of failure
 					if (TestResults != null && TestResults.Failed > 0)
@@ -592,7 +588,7 @@ namespace UE
 
 			if (EditorRole != null)
 			{
-				UnrealAutomatedTestPassResults JsonTestPassResults = GetTestPassResults();
+				UnrealAutomatedTestPassResults JsonTestPassResults = GetTestPassResults(EditorRole.LogSummary);
 
 				// Filter our tests into categories
 				IEnumerable<UnrealAutomatedTestResult> AllTests = JsonTestPassResults.Tests;
