AutomationTool - Logging out where the log file is when telling people to check the log file.

diff --git a/Engine/Source/Programs/AutomationTool/CompileScriptModules.cs b/Engine/Source/Programs/AutomationTool/CompileScriptModules.cs
--- a/Engine/Source/Programs/AutomationTool/CompileScriptModules.cs
+++ b/Engine/Source/Programs/AutomationTool/CompileScriptModules.cs
@@ -871,7 +871,7 @@ namespace AutomationToolDriver
 
 			bool Result = true;
 
-			Log.TraceInformation($"Building {AutomationProjectGraph.EntryPointNodes.Count} projects (see Log for more details)");
+			Log.TraceInformation($"Building {AutomationProjectGraph.EntryPointNodes.Count} projects (see Log 'Engine/Programs/AutomationTool/Saved/Logs/Log.txt' for more details)");
 			foreach (string TargetToBuild in TargetsToBuild)
 			{
 				var GraphRequest = new GraphBuildRequestData(AutomationProjectGraph, new string[] { TargetToBuild });
