Gauntlet - avoid failing to produce a report if UE automation test json report is missing

#rnx

diff --git a/Engine/Source/Programs/AutomationTool/Gauntlet/Unreal/Automation/UE.Automation.cs b/Engine/Source/Programs/AutomationTool/Gauntlet/Unreal/Automation/UE.Automation.cs
--- a/Engine/Source/Programs/AutomationTool/Gauntlet/Unreal/Automation/UE.Automation.cs
+++ b/Engine/Source/Programs/AutomationTool/Gauntlet/Unreal/Automation/UE.Automation.cs
@@ -449,7 +449,18 @@ namespace UE
 					AutomationReportPath = Config.ReportExportPath;
 					if (!string.IsNullOrEmpty(AutomationReportPath))
 					{
-						TestPassResults = UnrealAutomatedTestPassResults.LoadFromJson(Path.Combine(AutomationReportPath, "index.json"));
+						string JsonReportFile = Path.Combine(AutomationReportPath, "index.json");
+						if (File.Exists(JsonReportFile))
+						{
+							try
+							{
+								TestPassResults = UnrealAutomatedTestPassResults.LoadFromJson(JsonReportFile);
+							}
+							catch (Exception Ex)
+							{
+								Log.Warning("Failed to load Json report. {0}", Ex);
+							}
+						}
 					}
 				}
 				
