Zen: Parse package store manifest when staging from Zen

#rnx

diff --git a/Engine/Source/Programs/AutomationTool/Scripts/CopyBuildToStagingDirectory.Automation.cs b/Engine/Source/Programs/AutomationTool/Scripts/CopyBuildToStagingDirectory.Automation.cs
--- a/Engine/Source/Programs/AutomationTool/Scripts/CopyBuildToStagingDirectory.Automation.cs
+++ b/Engine/Source/Programs/AutomationTool/Scripts/CopyBuildToStagingDirectory.Automation.cs
@@ -725,7 +725,7 @@ namespace AutomationScripts
 				List<FileReference> CookedFiles = DirectoryReference.EnumerateFiles(SC.PlatformCookDir, "*", SearchOption.AllDirectories).ToList();
 
 				// When cooking directly to I/O store contaier files look for the package store manifest
-				if (Params.IoStore)
+				if (Params.IoStore || Params.ZenStore)
 				{
 					CookedFiles.AddRange(PackageStoreManifest.GetCookedFiles(FileReference.Combine(SC.MetadataDir, "packagestore.manifest")));
 				}
@@ -1056,7 +1056,7 @@ namespace AutomationScripts
 						List<FileReference> CookedFiles = DirectoryReference.EnumerateFiles(CookOutputDir, "*", SearchOption.AllDirectories).ToList();
 
 						// When cooking directly to I/O store contaier files look for the package store manifest
-						if (Params.IoStore)
+						if (Params.IoStore || Params.ZenStore)
 						{
 							CookedFiles.AddRange(PackageStoreManifest.GetCookedFiles(FileReference.Combine(SC.MetadataDir, "packagestore.manifest")));
 						}
