Gauntlet - Use Not-Run test state instead of In-Process to evaluate if a test pass need to be resumed

#rnx

diff --git a/Engine/Source/Developer/AutomationController/Private/AutomationControllerManger.cpp b/Engine/Source/Developer/AutomationController/Private/AutomationControllerManger.cpp
--- a/Engine/Source/Developer/AutomationController/Private/AutomationControllerManger.cpp
+++ b/Engine/Source/Developer/AutomationController/Private/AutomationControllerManger.cpp
@@ -126,10 +126,14 @@ void FAutomatedTestPassResults::UpdateTestResultStatus(const IAutomationReportPt
 		NotRun--;
 		InProcess++;
 		break;
-	default:
+	case EAutomationState::NotRun:
 		NotRun++;
 		InProcess--;
 		break;
+	case EAutomationState::NotEnoughParticipants:
+		// Those never get to be set as InProcess
+		NotRun--;
+		break;
 	}
 }
 
@@ -762,6 +766,7 @@ void FAutomationControllerManager::ExecuteNextTask( int32 ClusterIndex, OUT bool
 				TestResults.AddEvent(FAutomationEvent(EAutomationEventType::Warning, FString::Printf(TEXT("Needed %d devices to participate, Only had %d available."), CurrentTest->GetNumParticipantsRequired(), DeviceClusterManager.GetNumDevicesInCluster(ClusterIndex))));
 
 				CurrentTest->SetResults(ClusterIndex, CurrentTestPass, TestResults);
+				CollectTestResults(CurrentTest, TestResults);
 				DeviceClusterManager.ResetAllDevicesRunningTest(ClusterIndex, CurrentTest);
 			}
 		}
diff --git a/Engine/Source/Programs/AutomationTool/Gauntlet/Unreal/Base/Gauntlet.UnrealTestNode.cs b/Engine/Source/Programs/AutomationTool/Gauntlet/Unreal/Base/Gauntlet.UnrealTestNode.cs
--- a/Engine/Source/Programs/AutomationTool/Gauntlet/Unreal/Base/Gauntlet.UnrealTestNode.cs
+++ b/Engine/Source/Programs/AutomationTool/Gauntlet/Unreal/Base/Gauntlet.UnrealTestNode.cs
@@ -1205,12 +1205,10 @@ namespace Gauntlet
 				if (JsonTestPassResults.InProcess > 0)
 				{
 					// The test pass did not run completely
-					Log.Verbose("Found Uncompleted tests: {0}", JsonTestPassResults.InProcess);
-					bool HasTimeout = false;
+					Log.Verbose("Found in-process tests: {0}", JsonTestPassResults.InProcess);
 					// Get any critical error and push it to json report and resave it.
 					if (RoleResults != null)
 					{
-						HasTimeout = RoleResults.Where(R => R.ProcessResult == UnrealProcessResult.TimeOut).Any();
 						UnrealLog.CallstackMessage FatalError = null;
 						foreach (UnrealRoleResult Result in RoleResults)
 						{
@@ -1230,6 +1228,12 @@ namespace Gauntlet
 							}
 						}
 					}
+				}
+				if (JsonTestPassResults.NotRun > 0)
+				{
+					// The test pass did not run at all
+					Log.Verbose("Found not-run tests: {0}", JsonTestPassResults.NotRun);
+					bool HasTimeout = RoleResults != null && RoleResults.Where(R => R.ProcessResult == UnrealProcessResult.TimeOut).Any();
 					if (GetConfiguration().ResumeOnCriticalFailure && !HasTimeout)
 					{
 						if (Retries < MaxRetries)
