Adding resilency to SafeCreateDirectory

All the other Safe<FileOperation>() methods perform a retry. Gave similar retry logic to SafeCreateDirectory

diff --git a/Engine/Source/Programs/AutomationTool/AutomationUtils/Utils.cs b/Engine/Source/Programs/AutomationTool/AutomationUtils/Utils.cs
--- a/Engine/Source/Programs/AutomationTool/AutomationUtils/Utils.cs
+++ b/Engine/Source/Programs/AutomationTool/AutomationUtils/Utils.cs
@@ -47,21 +47,46 @@ namespace AutomationTool
 		/// <returns>True if the directory was created, false otherwise.</returns>
 		public static bool SafeCreateDirectory(string Path, bool bQuiet = false)
 		{
-			if( !bQuiet)
+			if(!bQuiet)
 			{
 				Log.TraceLog("SafeCreateDirectory {0}", Path);
 			}
 
+			const int MaxAttempts = 10;
+			int Attempts = 0;
 			bool Result = true;
-			try
+			Exception LastException = null;
+			do
 			{
-				Result = Directory.CreateDirectory(Path).Exists;
-			}
-			catch (Exception)
+				Result = Directory.Exists(Path);
+				if (!Result)
+				{
+					try
+					{
+						Result = Directory.CreateDirectory(Path).Exists;
+					}
+					catch (Exception Ex)
+					{
+						if (!Directory.Exists(Path))
+						{
+							Thread.Sleep(3000);
+						}
+						Result = Directory.Exists(Path);
+						LastException = Ex;
+					}
+				}
+			} while (Result == false && ++Attempts < MaxAttempts);
+
+			if (Result == false && LastException != null)
 			{
-				if (Directory.Exists(Path) == false)
+				if (bQuiet)
 				{
-					Result = false;
+					Log.TraceLog("Failed to create directory {0} in {1} attempts.", Path, MaxAttempts);
+				}
+				else
+				{
+					Log.TraceWarning("Failed to create directory {0} in {1} attempts.", Path, MaxAttempts);
+					Log.TraceWarning(LogUtils.FormatException(LastException));
 				}
 			}
 			return Result;
