Added "RunUAT CreatePlatformExtension" command to auto-generate plugin and module platform extension files.

Examples:
    RunUAT CreatePlatformExtension -P4 -platform=MyPlatform -source=Engine/Plugins/Path/To/MyPlugin.uplugin
    RunUAT CreatePlatformExtension -P4 -platform=MyPlatform -source=Engine/Source/Path/To/MyModule.build.cs
    RunUAT CreatePlatformExtension -P4 -platform=MyPlatform -source=Engine/Path/To/Whatever/

    RunUAT CreatePlatformExtension -P4 -platform=MyPlatform -source=MyProject/Source/MyProject.target.cs  -project=MyProject/MyProject.uproject

diff --git a/Engine/Source/Programs/AutomationTool/Scripts/CreatePlatformExtension.cs b/Engine/Source/Programs/AutomationTool/Scripts/CreatePlatformExtension.cs
new file mode 100644
--- /dev/null
+++ b/Engine/Source/Programs/AutomationTool/Scripts/CreatePlatformExtension.cs
@@ -0,0 +1,559 @@
+// Copyright Epic Games, Inc. All Rights Reserved.
+using System;
+using System.Collections.Generic;
+using System.Text;
+using System.IO;
+using AutomationTool;
+using UnrealBuildTool;
+using EpicGames.Core;
+using UnrealBuildBase;
+using System.Linq;
+
+[Help("Create stub code for platform extension")]
+[Help("Source", "Path to source .uplugin, .build.cs or .target.cs, or a source folder to search")]
+[Help("Platform", "Platform(s) or Platform Groups to generate for")]
+[Help("Project", "Optional path to project (only required if not creating code for Engine modules/plugins")]
+[Help("SkipPluginModules", "Do not generate platform extension module files when generating a platform extension plugin")]
+[Help("AllowUnknownPlatforms", "Allow platform & platform groups that are not known, for example when generating code for extensions we do not have access to")]
+[Help("P4", "Create a changelist for the new files")]
+class CreatePlatformExtension : BuildCommand
+{
+	readonly List<ModuleHostType> ModuleTypeBlacklist = new List<ModuleHostType>
+	{
+		ModuleHostType.Developer,
+		ModuleHostType.Editor, 
+		ModuleHostType.EditorNoCommandlet,
+		ModuleHostType.EditorAndProgram,
+		ModuleHostType.Program,
+	};
+
+	ConfigHierarchy GameIni;
+	DirectoryReference ProjectDir;
+	List<string> NewFiles = new List<string>();
+	bool bSkipPluginModules;
+
+	public override void ExecuteBuild()
+	{
+		// Parse the parameters
+		string[] Platforms = ParseParamValue("Platform", "").Split('+', StringSplitOptions.RemoveEmptyEntries);
+		string Source = ParseParamValue("Source", "");
+		string Project = ParseParamValue("Project", "");
+		bSkipPluginModules = ParseParamBool("SkipPluginModules");
+
+		// make sure we have somewhere to look
+		if (string.IsNullOrEmpty(Source))
+		{
+			Log.TraceError("No -Source= directory/file specified");
+			return;
+		}
+
+		// Sanity check platforms list
+		Platforms = VerifyPlatforms(Platforms);
+		if (Platforms.Length == 0)
+		{
+			Log.TraceError("Please specify at least one platform or platform group");
+			return;
+		}
+
+		// Prepare values
+		ProjectDir = string.IsNullOrEmpty(Project) ? null : new FileReference(Project).Directory;
+		GameIni = ConfigCache.ReadHierarchy(ConfigHierarchyType.Game, ProjectDir, BuildHostPlatform.Current.Platform );
+		int CL = -1;
+
+		// Generate the code
+		try
+		{
+			if (Directory.Exists(Source))
+			{
+				// check the directory for plugins first, because the plugins will automatically generate the modules too
+				List<string> Plugins = Directory.EnumerateFiles(Source, "*.uplugin", SearchOption.AllDirectories ).ToList();
+				if (Plugins.Count > 0)
+				{
+					foreach (string Plugin in Plugins)
+					{
+						GeneratePluginPlatformExtension( new FileReference(Plugin), Platforms );
+					}
+				}
+				else
+				{
+					// there were no plugins found, so search for module & target rules instead
+					List<string> ModuleRules = Directory.EnumerateFiles(Source,"*.build.cs", SearchOption.AllDirectories).ToList();
+					ModuleRules.AddRange(Directory.EnumerateFiles(Source,"*.target.cs", SearchOption.AllDirectories));
+					if (ModuleRules.Count > 0)
+					{
+						foreach (string ModuleRule in ModuleRules)
+						{
+							GenerateModulePlatformExtension( new FileReference(ModuleRule), Platforms );
+						}
+					}
+					else
+					{
+						Log.TraceError($"Cannot find any supported files in {Source}");
+					}
+				}
+			}
+			else if (File.Exists(Source))
+			{
+				GeneratePlatformExtensionFromFile(new FileReference(Source), Platforms);
+			}
+			else
+			{
+				Log.TraceError($"Invalid path or file name {Source}");
+			}
+
+			// check the generated files
+			if (NewFiles.Count > 0)
+			{
+				bool bIsTest = ParseParam("Test");
+
+				// add the files to perforce if that is available
+				if (CommandUtils.P4Enabled && !bIsTest)
+				{
+					DirectoryReference SourceDir = new DirectoryReference(Source);
+					string Description = $"[AUTO-GENERATED] {string.Join('+', Platforms)} platform extension files from {SourceDir.MakeRelativeTo(Unreal.RootDirectory)}\n\n#nocheckin verify the code has been generated successfully before checking in!";
+
+					CL = P4.CreateChange(P4Env.Client, Description );
+					foreach (string NewFile in NewFiles)
+					{
+						P4.Add(CL, CommandUtils.MakePathSafeToUseWithCommandLine(NewFile) );
+					}
+				}
+
+				// display final report
+				Log.TraceInformation(System.Environment.NewLine);
+				Log.TraceInformation(System.Environment.NewLine);
+				Log.TraceInformation(System.Environment.NewLine);
+				Log.TraceInformation("The following files have been created" + ((CL > 0) ? $" and added to changelist {CL}:" : ":") );
+				foreach (string NewFile in NewFiles)
+				{
+					Log.TraceInformation($"\t{NewFile}");
+				}
+				Log.TraceInformation(System.Environment.NewLine);
+				Log.TraceInformation(System.Environment.NewLine);
+				Log.TraceInformation(System.Environment.NewLine);
+
+				// remove everything if requested (for debugging etc)
+				if (!CommandUtils.P4Enabled && bIsTest)
+				{
+					Log.TraceInformation("Deleting all the files because this is just a test...");
+					foreach (string NewFile in NewFiles)
+					{
+						File.Delete(NewFile);
+					}
+				}
+			}
+		}
+		catch(Exception)
+		{
+			// something went wrong - clean up anything we've created so far
+			foreach (string NewFile in NewFiles)
+			{
+				Log.TraceInformation($"Removing partial file ${NewFile} due to error");
+				File.Delete(NewFile);
+			}
+
+			// try to safely clean up the perforce changelist too
+			try
+			{
+				if (CL > 0 && CommandUtils.P4Enabled)
+				{
+					Log.TraceInformation($"Removing partial changelist ${CL} due to error");
+					P4.DeleteChange(CL,true);
+				}
+			}
+			catch(Exception e)
+			{
+				Log.TraceError(e.Message);
+			}
+
+
+			throw;
+		}
+	}
+
+
+
+
+	/// <summary>
+	/// Create the platform extension plugin files of the given plugin, for the given platforms
+	/// </summary>
+	private void GeneratePluginPlatformExtension(FileReference PluginPath, string[] Platforms)
+	{
+		// sanity check plugin path
+		if (!File.Exists(PluginPath.FullName))
+		{
+			Log.TraceError($"File not found: {PluginPath}");
+			return;
+		}
+
+		DirectoryReference PluginDir = PluginPath.Directory;
+		if (ProjectDir == null && !PluginDir.IsUnderDirectory(Unreal.EngineDirectory))
+		{
+			Log.TraceError($"{PluginPath} is not under the Engine directory, and no -project= has been specified");
+			return;
+		}
+
+		DirectoryReference RootDir = ProjectDir ?? Unreal.EngineDirectory;
+		if (!PluginDir.IsUnderDirectory(RootDir))
+		{
+			Log.TraceError($"{PluginPath} is not under {RootDir}");
+			return;
+		}
+
+		// load the plugin & find suitable modules, if required
+		PluginDescriptor ParentPlugin = PluginDescriptor.FromFile(PluginPath); //NOTE: if the PluginPath is itself a child plugin, not all whitelist, blacklist & supported platform information will be available.
+		List<ModuleDescriptor> ParentModuleDescs = new List<ModuleDescriptor>();
+		Dictionary<ModuleDescriptor, FileReference> ParentModuleRules = new Dictionary<ModuleDescriptor, FileReference>();
+		if (!bSkipPluginModules && ParentPlugin.Modules != null)
+		{
+			ParentModuleDescs = ParentPlugin.Modules.Where(ModuleDesc => CanCreatePlatformExtensionForPluginModule(ModuleDesc)).ToList();
+
+			// find all module rules that are listed in the plugin
+			DirectoryReference ModuleRulesPath = DirectoryReference.Combine( PluginDir, "Source" );
+			var ModuleRules = DirectoryReference.EnumerateFiles(ModuleRulesPath, "*.build.cs", SearchOption.AllDirectories);
+			foreach (FileReference ModuleRule in ModuleRules)
+			{
+				string ModuleRuleName = GetPlatformExtensionBaseNameFromPath(ModuleRule.FullName);
+				ModuleDescriptor ModuleDesc = ParentModuleDescs.Find(ParentModuleDesc => ParentModuleDesc.Name.Equals(ModuleRuleName, StringComparison.InvariantCultureIgnoreCase));
+				if (ModuleDesc != null)
+				{
+					ParentModuleRules.Add(ModuleDesc, ModuleRule);
+				}
+			}
+		}
+
+		// generate the platform extension files
+		string BasePluginDir = GetRelativeBaseDirectory(PluginDir, RootDir);
+		string BasePluginName = GetPlatformExtensionBaseNameFromPath(PluginPath.FullName);
+		foreach (string PlatformName in Platforms)
+		{
+			// verify final file name
+			string FinalFileName = Path.Combine(RootDir.FullName, "Platforms", PlatformName, BasePluginDir, BasePluginName + "_" + PlatformName + ".uplugin");
+			if (File.Exists(FinalFileName))
+			{
+				Log.TraceWarning($"Skipping {FinalFileName} as it already exists");
+				continue;
+			}
+
+			// create the child plugin
+			Directory.CreateDirectory(Path.GetDirectoryName(FinalFileName));
+			using (JsonWriter ChildPlugin = new JsonWriter(FinalFileName))
+			{
+				UnrealTargetPlatform Platform;
+				bool bHasPlatform = UnrealTargetPlatform.TryParse(PlatformName, out Platform);
+
+				// a platform reference is needed if there are already platforms listed in the parent, or the parent requires an explicit platform list
+				bool NeedsPlatformReference( List<UnrealTargetPlatform> ParentPlatforms, bool bHasExplicitPlatforms )
+				{
+					return (bHasPlatform && ((ParentPlatforms != null && ParentPlatforms.Count > 0) || bHasExplicitPlatforms));
+				}
+
+				// create the plugin definition
+				ChildPlugin.WriteObjectStart();
+				ChildPlugin.WriteValue("FileVersion", (int)PluginDescriptorVersion.ProjectPluginUnification ); // this is the version that this code has been tested against
+				ChildPlugin.WriteValue("bIsPluginExtension", true );
+				if (NeedsPlatformReference(ParentPlugin.SupportedTargetPlatforms, ParentPlugin.bHasExplicitPlatforms))
+				{
+					ChildPlugin.WriteStringArrayField("SupportedTargetPlatforms", new string[]{ Platform.ToString() } );
+				}
+
+				// select all modules that are not blacklisted
+				IEnumerable<ModuleDescriptor> ModuleDescs = ParentModuleDescs.Where( ModuleDesc => !(bHasPlatform && ModuleDesc.BlacklistPlatforms != null && ModuleDesc.BlacklistPlatforms.Contains(Platform)) );
+				if (ModuleDescs.Any() )
+				{
+					ChildPlugin.WriteArrayStart("Modules");
+					foreach (ModuleDescriptor ParentModuleDesc in ModuleDescs)
+					{
+						// create the child module reference
+						ChildPlugin.WriteObjectStart();
+						ChildPlugin.WriteValue("Name", ParentModuleDesc.Name);
+						ChildPlugin.WriteValue("Type", ParentModuleDesc.Type.ToString());
+						ChildPlugin.WriteValue("LoadingPhase", ParentModuleDesc.LoadingPhase.ToString());
+						if (NeedsPlatformReference(ParentModuleDesc.WhitelistPlatforms, ParentModuleDesc.bHasExplicitPlatforms))
+						{
+							ChildPlugin.WriteStringArrayField("WhitelistPlatforms", new string[] { Platform.ToString() } );
+						}
+						ChildPlugin.WriteObjectEnd();
+
+						// see if there is a module rule file too & generate the rules file for this platform
+						FileReference ParentModuleRule;
+						if (ParentModuleRules.TryGetValue(ParentModuleDesc, out ParentModuleRule))
+						{
+							GenerateModulePlatformExtension(ParentModuleRule, new string[] { PlatformName });
+						}
+					}
+					ChildPlugin.WriteArrayEnd();
+				}
+				ChildPlugin.WriteObjectEnd();
+			}
+			NewFiles.Add(FinalFileName);
+		}
+	}
+
+
+	/// <summary>
+	/// Creates the platform extension child class files of the given module, for the given platforms
+	/// </summary>
+	private void GenerateModulePlatformExtension(FileReference ModulePath, string[] Platforms)
+	{
+		// sanity check module path
+		if (!File.Exists(ModulePath.FullName))
+		{
+			Log.TraceError($"File not found: {ModulePath}");
+			return;
+		}
+
+		DirectoryReference ModuleDir = ModulePath.Directory;
+		if (ProjectDir == null && !ModuleDir.IsUnderDirectory(Unreal.EngineDirectory))
+		{
+			Log.TraceError($"{ModulePath} is not under the Engine directory, and no -project= has been specified");
+			return;
+		}
+
+		DirectoryReference RootDir = ProjectDir ?? Unreal.EngineDirectory;
+		if (!ModuleDir.IsUnderDirectory(RootDir))
+		{
+			Log.TraceError($"{ModulePath} is not under {RootDir}");
+			return;
+		}
+
+		// sanity check module file name
+		string ModuleFilename = ModulePath.GetFileName();
+		string ModuleExtension = ModuleFilename.Substring( ModuleFilename.IndexOf('.') );
+		ModuleFilename = ModuleFilename.Substring( 0, ModuleFilename.Length - ModuleExtension.Length );
+		if (!ModuleExtension.Equals(".build.cs", System.StringComparison.InvariantCultureIgnoreCase ) && !ModuleExtension.Equals(".target.cs", System.StringComparison.InvariantCultureIgnoreCase))
+		{
+			Log.TraceError($"{ModulePath} is a module/rules file. Expecting .build.cs or .target.cs");
+			return;
+		}
+
+		// load module file & find module class name, and optional class namespace
+		const string ClassDeclaration = "public class ";
+		const string NamespaceDeclaration = "namespace ";
+		string[] ModuleContents = File.ReadAllLines(ModulePath.FullName);
+		string ModuleClassDeclaration = ModuleContents.FirstOrDefault( L => L.Trim().StartsWith(ClassDeclaration) );
+		string ModuleNamespaceDeclaration = ModuleContents.FirstOrDefault( L => L.Trim().StartsWith(NamespaceDeclaration) );
+		if (string.IsNullOrEmpty(ModuleClassDeclaration))
+		{
+			Log.TraceError($"Cannot find class declaration in ${ModulePath}");
+			return;
+		}
+		string ParentModuleName = ModuleClassDeclaration.Trim().Remove(0, ClassDeclaration.Length ).Split(' ', StringSplitOptions.None ).First();
+		if (string.IsNullOrEmpty(ParentModuleName))
+		{
+			Log.TraceError($"Cannot parse class declaration in ${ModulePath}");
+			return;
+		}
+		string ParentNamespace = string.IsNullOrEmpty(ModuleNamespaceDeclaration) ? "" : (ModuleNamespaceDeclaration.Trim().Remove(0, NamespaceDeclaration.Length ).Split(' ', StringSplitOptions.None ).First() + ".");
+		string BaseModuleName = ParentModuleName;
+		int Index = BaseModuleName.IndexOf('_'); //trim off _[platform] suffix
+		if (Index != -1)
+		{
+			BaseModuleName = BaseModuleName.Substring(0, Index);
+		}
+
+		// load template and generate the platform extension files
+		string BaseModuleDir = GetRelativeBaseDirectory( ModuleDir, RootDir );
+		string BaseModuleFileName = GetPlatformExtensionBaseNameFromPath( ModulePath.FullName );
+		string CopyrightLine = MakeCopyrightLine();
+		string Template = LoadTemplate($"PlatformExtension{ModuleExtension}.template");
+		foreach (string PlatformName in Platforms)
+		{
+			// verify the final file name
+			string FinalFileName = Path.Combine(RootDir.FullName, "Platforms", PlatformName, BaseModuleDir, BaseModuleFileName + "_" + PlatformName + ModuleExtension );
+			if (File.Exists(FinalFileName))
+			{
+				Log.TraceWarning($"Skipping {FinalFileName} as it already exists");
+				continue;
+			}
+
+			// generate final code from the template
+			string FinalOutput = Template;
+			FinalOutput = FinalOutput.Replace("%COPYRIGHT_LINE%",     CopyrightLine,                     StringComparison.InvariantCultureIgnoreCase );
+			FinalOutput = FinalOutput.Replace("%PARENT_MODULE_NAME%", ParentNamespace+ParentModuleName,  StringComparison.InvariantCultureIgnoreCase );
+			FinalOutput = FinalOutput.Replace("%BASE_MODULE_NAME%",   BaseModuleName,                    StringComparison.InvariantCultureIgnoreCase );
+			FinalOutput = FinalOutput.Replace("%PLATFORM_NAME%",      PlatformName,                      StringComparison.InvariantCultureIgnoreCase );
+
+			// save the child .cs file
+			Directory.CreateDirectory(Path.GetDirectoryName(FinalFileName));
+			File.WriteAllText(FinalFileName, FinalOutput);
+			NewFiles.Add(FinalFileName);
+		}
+	}
+
+
+
+	/// <summary>
+	/// Generates platform extension files based on the given source file name
+	/// </summary>
+	private void GeneratePlatformExtensionFromFile(FileReference Source, string[] Platforms)
+	{
+		if (Source.FullName.ToLower().EndsWith(".uplugin") )
+		{
+			GeneratePluginPlatformExtension(Source, Platforms);
+		}
+		else if (Source.FullName.ToLower().EndsWith(".build.cs") || Source.FullName.ToLower().EndsWith(".target.cs"))
+		{
+			GenerateModulePlatformExtension(Source, Platforms);
+		}
+		else
+		{
+			Log.TraceError($"unsupported file type {Source}");
+		}
+	}
+
+
+
+	#region boilerplate & helpers
+
+
+	/// <summary>
+	/// Determines whether we should attempt to add this plugin module to the child plugin module references
+	/// </summary>
+	/// <param name="ModuleDesc"></param>
+	/// <returns></returns>
+	private bool CanCreatePlatformExtensionForPluginModule( ModuleDescriptor ModuleDesc )
+	{
+		// make sure it's a type that is usually associated with platform extensions
+		if (ModuleTypeBlacklist.Contains(ModuleDesc.Type))
+		{
+			return false;
+		}
+
+		// this module must have supported platforms explicitly listed so we must create a child reference
+		if (ModuleDesc.bHasExplicitPlatforms)
+		{
+			return true;
+		}
+
+		// the module has a non-empty whitelist platform list so we must create a child reference
+		if (ModuleDesc.WhitelistPlatforms != null && ModuleDesc.WhitelistPlatforms.Count >= 0)
+		{
+			return true;
+		}
+
+		// the module has an empty whitelist platform list so no explicit platform reference is needed
+		return false;
+	}
+
+
+
+	/// <summary>
+	/// Gets the relative path from the given root. If the root is also under a Platforms/[name] then that is also removed
+	/// </summary>
+	private string GetRelativeBaseDirectory(DirectoryReference ChildDir, DirectoryReference RootDir)
+	{
+		DirectoryReference PlatformsDir = DirectoryReference.Combine(RootDir, "Platforms");
+		string BaseDir;
+		if (ChildDir.IsUnderDirectory(PlatformsDir))
+		{
+			BaseDir = ChildDir.MakeRelativeTo(PlatformsDir);
+			BaseDir = BaseDir.Substring(BaseDir.IndexOf(Path.DirectorySeparatorChar) + 1);
+		}
+		else
+		{
+			BaseDir = ChildDir.MakeRelativeTo(RootDir);
+		}
+
+		return BaseDir;
+	}
+
+
+	/// <summary>
+	/// Given a full path to a plugin or module file, returns the raw file name - trimming off any _[platform] suffix too
+	/// </summary>
+	private string GetPlatformExtensionBaseNameFromPath( string FileName )
+	{
+		// trim off path
+		string BaseName = Path.GetFileName(FileName);
+
+		// trim off any extensions
+		string Extensions = BaseName.Substring( BaseName.IndexOf('.') );
+		BaseName = BaseName.Substring( 0, BaseName.Length - Extensions.Length );
+
+		// trim off any platform suffix
+		int Idx = BaseName.IndexOf('_');
+		if (Idx != -1)
+		{
+			BaseName = BaseName.Substring(0,Idx);
+		}
+
+		return BaseName;
+	}
+
+
+	
+	/// <summary>
+	/// Load the given file from the engine templates folder
+	/// </summary>
+	private string LoadTemplate( string FileName )
+	{
+		string TemplatePath = Path.Combine( Unreal.EngineDirectory.FullName, "Content", "Editor", "Templates", FileName );
+		return File.ReadAllText(TemplatePath);
+	}
+
+
+
+	/// <summary>
+	/// Look up the project/engine specific copyright string
+	/// </summary>
+	private string MakeCopyrightLine()
+	{
+		string CopyrightNotice = "";
+		GameIni.GetString("/Script/EngineSettings.GeneralProjectSettings", "CopyrightNotice", out CopyrightNotice);
+
+		if (!string.IsNullOrEmpty(CopyrightNotice))
+		{
+			return "// " + CopyrightNotice;
+		}
+		else
+		{
+			return "";
+		}
+	}
+
+
+
+	/// <summary>
+	/// Returns a list of validated and case-corrected platform and platform groups
+	/// </summary>
+	private string[] VerifyPlatforms(string[] Platforms)
+	{
+		bool bAllowUnknownPlatforms = ParseParamBool("AllowUnknownPlatforms");
+
+		List<string> Result = new List<string>();
+		foreach (string PlatformName in Platforms)
+		{
+			// see if this is a platform
+			UnrealTargetPlatform Platform;
+			if (UnrealTargetPlatform.TryParse(PlatformName, out Platform))
+			{
+				Result.Add(Platform.ToString());
+				continue;
+			}
+
+			// see if this is a platform group
+			UnrealPlatformGroup PlatformGroup;
+			if (UnrealPlatformGroup.TryParse(PlatformName, out PlatformGroup))
+			{
+				Result.Add(PlatformGroup.ToString());
+				continue;
+			}
+
+			// this is an unknown item - see if we will accept it anyway...
+			if (bAllowUnknownPlatforms)
+			{
+				Log.TraceWarning($"{PlatformName} is not a known Platform or Platform Group. The code will still be generated but you may not be able to test it locally");
+				Result.Add(PlatformName);
+			}
+			else
+			{
+				Log.TraceWarning($"{PlatformName} is not a known Platform or Platform Group and so it will be ignored. Specify -AllowUnknownPlatforms to allow it anyway");
+			}
+		}
+
+		return Result.ToArray();
+		
+	}
+	#endregion
+}
diff --git a/Engine/Source/Programs/UnrealBuildTool/Configuration/UEBuildTarget.cs b/Engine/Source/Programs/UnrealBuildTool/Configuration/UEBuildTarget.cs
--- a/Engine/Source/Programs/UnrealBuildTool/Configuration/UEBuildTarget.cs
+++ b/Engine/Source/Programs/UnrealBuildTool/Configuration/UEBuildTarget.cs
@@ -445,6 +445,25 @@ namespace UnrealBuildTool
 		{
 			return GetUniqueStringRegistry().GetStringForId(Id);
 		}
+
+		/// <summary>
+		/// 
+		/// </summary>
+		/// <param name="Name"></param>
+		/// <param name="Group"></param>
+		/// <returns></returns>
+		public static bool TryParse(string Name, out UnrealPlatformGroup Group)
+		{
+			if (GetUniqueStringRegistry().HasString(Name))
+			{
+				Group.Id = GetUniqueStringRegistry().FindOrAddByName(Name);
+				return true;
+			}
+
+			Group.Id = -1;
+			return false;
+		}
+
 		/// <summary>
 		/// 
 		/// </summary>
