BuildGraph: Add task for deploying AWS ECS services

diff --git a/Engine/Source/Programs/AutomationTool/BuildGraph/Tasks/AwsEcsDeployTask.cs b/Engine/Source/Programs/AutomationTool/BuildGraph/Tasks/AwsEcsDeployTask.cs
new file mode 100644
--- /dev/null
+++ b/Engine/Source/Programs/AutomationTool/BuildGraph/Tasks/AwsEcsDeployTask.cs
@@ -0,0 +1,136 @@
+// Copyright Epic Games, Inc. All Rights Reserved.
+
+using AutomationTool;
+using EpicGames.Core;
+using System;
+using System.Collections.Generic;
+using System.IO;
+using System.Text;
+using System.Text.Json;
+using System.Xml;
+using UnrealBuildBase;
+
+namespace BuildGraph.Tasks
+{
+	/// <summary>
+	/// Parameters for a AWS ECS deploy task
+	/// </summary>
+	public class AwsEcsDeployTaskParameters
+	{
+		/// <summary>
+		/// Task definition file to use
+		/// </summary>
+		[TaskParameter(Optional = false)]
+		public string TaskDefinitionFile;
+
+		/// <summary>
+		/// Docker image to set in new task definition (will replace DOCKER_PATTERN with this value)
+		/// </summary>
+		[TaskParameter(Optional = false)]
+		public string DockerImage;
+
+		/// <summary>
+		/// Cluster ARN representing AWS ECS cluster to operate on
+		/// </summary>
+		[TaskParameter(Optional = false)]
+		public string Cluster;
+
+		/// <summary>
+		/// Service name to update and deploy to
+		/// </summary>
+		[TaskParameter(Optional = false)]
+		public string Service;
+
+		/// <summary>
+		/// Environment variables
+		/// </summary>
+		[TaskParameter(Optional = true)]
+		public string Environment;
+
+		/// <summary>
+		/// File to read environment from
+		/// </summary>
+		[TaskParameter(Optional = true)]
+		public string EnvironmentFile;
+
+		/// <summary>
+		/// Write output to the log
+		/// </summary>
+		[TaskParameter(Optional = true)]
+		public bool LogOutput = false;
+	}
+
+	/// <summary>
+	/// Creates a new AWS ECS task definition and updates the ECS service to use this new revision of the task def
+	/// </summary>
+	[TaskElement("Aws-EcsDeploy", typeof(AwsEcsDeployTaskParameters))]
+	public class AwsEcsDeployTask : SpawnTaskBase
+	{
+		/// <summary>
+		/// Parameters for this task
+		/// </summary>
+		AwsEcsDeployTaskParameters Parameters;
+
+		/// <summary>
+		/// Construct an AWS ECS deploy task
+		/// </summary>
+		/// <param name="InParameters">Parameters for the task</param>
+		public AwsEcsDeployTask(AwsEcsDeployTaskParameters InParameters)
+		{
+			Parameters = InParameters;
+		}
+
+		/// <summary>
+		/// Execute the task.
+		/// </summary>
+		/// <param name="Job">Information about the current job</param>
+		/// <param name="BuildProducts">Set of build products produced by this node.</param>
+		/// <param name="TagNameToFileSet">Mapping from tag names to the set of files they include</param>
+		public override void Execute(JobContext Job, HashSet<FileReference> BuildProducts, Dictionary<string, HashSet<FileReference>> TagNameToFileSet)
+		{
+			string TaskDefContents = File.ReadAllText(ResolveFile(Parameters.TaskDefinitionFile).FullName);
+			string TaskDefWithDockerImageContents = TaskDefContents.Replace("DOCKER_IMAGE", Parameters.DockerImage);
+
+			FileReference TempTaskDefFile = FileReference.Combine(Unreal.RootDirectory, "Engine", "Intermediate", "Build", "AwsEcsDeployTaskTemp.json");
+			DirectoryReference.CreateDirectory(TempTaskDefFile.Directory);
+			File.WriteAllText(TempTaskDefFile.FullName, TaskDefWithDockerImageContents, new UTF8Encoding(encoderShouldEmitUTF8Identifier: false));
+			
+			IProcessResult CreateTaskDefResult = SpawnTaskBase.Execute("aws", $"ecs register-task-definition --cli-input-json \"file://{TempTaskDefFile.FullName}\"", EnvVars: ParseEnvVars(Parameters.Environment, Parameters.EnvironmentFile), LogOutput: Parameters.LogOutput);
+
+			JsonDocument TaskDefJson = JsonDocument.Parse(CreateTaskDefResult.Output);
+			string TaskDefFamily = TaskDefJson.RootElement.GetProperty("taskDefinition").GetProperty("family").GetString();
+			string TaskDefRevision = TaskDefJson.RootElement.GetProperty("taskDefinition").GetProperty("revision").ToString();
+
+			string Params = $"ecs update-service --cluster {Parameters.Cluster} --service {Parameters.Service} --task-definition {TaskDefFamily}:{TaskDefRevision}";
+			SpawnTaskBase.Execute("aws", Params, EnvVars: ParseEnvVars(Parameters.Environment, Parameters.EnvironmentFile), LogOutput: Parameters.LogOutput);
+
+			Log.TraceInformation($"Service {Parameters.Service} updated to use new task def {TaskDefFamily}:{TaskDefRevision}");
+		}
+
+		/// <summary>
+		/// Output this task out to an XML writer.
+		/// </summary>
+		public override void Write(XmlWriter Writer)
+		{
+			Write(Writer, Parameters);
+		}
+
+		/// <summary>
+		/// Find all the tags which are used as inputs to this task
+		/// </summary>
+		/// <returns>The tag names which are read by this task</returns>
+		public override IEnumerable<string> FindConsumedTagNames()
+		{
+			yield break;
+		}
+
+		/// <summary>
+		/// Find all the tags which are modified by this task
+		/// </summary>
+		/// <returns>The tag names which are modified by this task</returns>
+		public override IEnumerable<string> FindProducedTagNames()
+		{
+			yield break;
+		}
+	}
+}
