EngineTest - enable Rendering tests on PS5 and XSX
+ Apply rendering and debug flags to the main role only instead of all process in UE.Automation
+ Remove MachineLearning test suite

#rnx

diff --git a/Engine/Source/Programs/AutomationTool/Gauntlet/Unreal/Automation/UE.Automation.cs b/Engine/Source/Programs/AutomationTool/Gauntlet/Unreal/Automation/UE.Automation.cs
--- a/Engine/Source/Programs/AutomationTool/Gauntlet/Unreal/Automation/UE.Automation.cs
+++ b/Engine/Source/Programs/AutomationTool/Gauntlet/Unreal/Automation/UE.Automation.cs
@@ -179,8 +179,9 @@ namespace UE
 				AppConfig.CommandLine += " -unattended";
 			}
 
+			bool HasNoOtherRole = !OtherRoles.Any();
 			// If there's only one role and it's the editor then tests are running under the editor with no target
-			if (ConfigRole.RoleType == UnrealTargetRole.Editor && OtherRoles.Any() == false)
+			if (ConfigRole.RoleType == UnrealTargetRole.Editor && HasNoOtherRole)
 			{ 
 				AppConfig.CommandLine += string.Format(" -ExecCmds=\"Automation {0}\"", AutomationTestArgument);		
 			}
@@ -200,37 +201,41 @@ namespace UE
 				}
 			}
 
-			if (DisableFrameTraceCapture || RayTracing)
+			if ((ConfigRole.RoleType == UnrealTargetRole.Editor && HasNoOtherRole) || ConfigRole.RoleType.IsClient())
 			{
-				AppConfig.CommandLine += " -DisableFrameTraceCapture";
-			}
-
-			if (RayTracing)
-			{
-				AppConfig.CommandLine += " -dpcvars=r.RayTracing=1,r.SkinCache.CompileShaders=1,AutomationAllowFrameTraceCapture=0";
-			}
-
-			// Options specific to windows
-			if (ConfigRole.Platform == UnrealTargetPlatform.Win64)
-			{
-				if (PreferNvidia)
+				// These are flags that are required only on the main role that is going to execute the tests. ie raytracing is required only on the client or if it is an editor test.
+				if (DisableFrameTraceCapture || RayTracing)
 				{
-					AppConfig.CommandLine += " -preferNvidia";
+					AppConfig.CommandLine += " -DisableFrameTraceCapture";
 				}
 
-				if (D3D12)
+				if (RayTracing)
 				{
-					AppConfig.CommandLine += " -d3d12";
+					AppConfig.CommandLine += " -dpcvars=r.RayTracing=1,r.SkinCache.CompileShaders=1,AutomationAllowFrameTraceCapture=0";
 				}
 
-				if (D3DDebug)
+				// Options specific to windows
+				if (ConfigRole.Platform == UnrealTargetPlatform.Win64)
 				{
-					AppConfig.CommandLine += " -d3ddebug";
-				}
+					if (PreferNvidia)
+					{
+						AppConfig.CommandLine += " -preferNvidia";
+					}
 
-				if (StompMalloc)
-				{
-					AppConfig.CommandLine += " -stompmalloc";
+					if (D3D12)
+					{
+						AppConfig.CommandLine += " -d3d12";
+					}
+
+					if (D3DDebug)
+					{
+						AppConfig.CommandLine += " -d3ddebug";
+					}
+
+					if (StompMalloc)
+					{
+						AppConfig.CommandLine += " -stompmalloc";
+					}
 				}
 			}
 
