When iterating through script assemblies, skip AutomationTool

diff --git a/Engine/Source/Programs/AutomationTool/Localization/LocalizationProvider.cs b/Engine/Source/Programs/AutomationTool/Localization/LocalizationProvider.cs
--- a/Engine/Source/Programs/AutomationTool/Localization/LocalizationProvider.cs
+++ b/Engine/Source/Programs/AutomationTool/Localization/LocalizationProvider.cs
@@ -163,6 +163,13 @@ namespace EpicGames.Localization
 				var LoadedAssemblies = AppDomain.CurrentDomain.GetAssemblies();
 				foreach (var Dll in LoadedAssemblies)
 				{
+					if (String.Equals("AutomationTool", Dll.GetName().Name))
+					{
+						// Exclude the AutomationTool driver assembly - it contains no relevant types, and trying to load it here
+						// can result in an exception if Microsoft.Build.Framework is not able to be loaded
+						continue;
+					}
+					
 					var AllTypes = Dll.GetTypes();
 					foreach (var PotentialLocalizationNodeType in AllTypes)
 					{
