AutomationUtils:

PS4 automation script is not able to be built on Linux - prevent the attempt.

diff --git a/Engine/Source/Programs/AutomationTool/CompileScriptModules.cs b/Engine/Source/Programs/AutomationTool/CompileScriptModules.cs
--- a/Engine/Source/Programs/AutomationTool/CompileScriptModules.cs
+++ b/Engine/Source/Programs/AutomationTool/CompileScriptModules.cs
@@ -409,6 +409,10 @@ namespace AutomationToolDriver
 				{
 					return false;
 				}
+				if (ModuleName.StartsWith("PS4", StringComparison.OrdinalIgnoreCase))
+				{
+					return false;
+				}
             }
 
 			return true;
