AutomationTool:

When looking for types that come from script modules, look at just the script module assemblies.

diff --git a/Engine/Source/Programs/AutomationTool/AutomationUtils/ScriptManager.cs b/Engine/Source/Programs/AutomationTool/AutomationUtils/ScriptManager.cs
--- a/Engine/Source/Programs/AutomationTool/AutomationUtils/ScriptManager.cs
+++ b/Engine/Source/Programs/AutomationTool/AutomationUtils/ScriptManager.cs
@@ -18,7 +18,7 @@ namespace AutomationTool
 	public static class ScriptManager
 	{
 		private static Dictionary<string, Type> ScriptCommands;
-		private static HashSet<Assembly> AllCompiledAssemblies;
+		public static readonly HashSet<Assembly> AllScriptAssemblies = new HashSet<Assembly>();
 		public static HashSet<FileReference> BuildProducts { get; private set; } = new HashSet<FileReference>();
 
 		/// <summary>
@@ -27,7 +27,7 @@ namespace AutomationTool
 		static void EnumerateScriptCommands()
 		{
 			ScriptCommands = new Dictionary<string, Type>(StringComparer.InvariantCultureIgnoreCase);
-			foreach (Assembly CompiledScripts in AllCompiledAssemblies)
+			foreach (Assembly CompiledScripts in AllScriptAssemblies)
 			{
 				try
 				{
@@ -77,7 +77,7 @@ namespace AutomationTool
 
 			HashSet<DirectoryReference> OutputDirs = new HashSet<DirectoryReference>();
 
-			foreach (Assembly CompiledAssembly in AllCompiledAssemblies)
+			foreach (Assembly CompiledAssembly in AllScriptAssemblies)
 			{
 				DirectoryReference AssemblyDirectory = FileReference.FromString(CompiledAssembly.Location).Directory;
 				AssemblyDirectory = DirectoryReference.FindCorrectCase(AssemblyDirectory);
@@ -106,7 +106,6 @@ namespace AutomationTool
 		/// <returns>List of compiled assemblies</returns>
 		public static void LoadScriptAssemblies(IEnumerable<FileReference> AssemblyPaths)
 		{
-			HashSet<Assembly> Assemblies = new HashSet<Assembly>();
 			foreach (FileReference AssemblyLocation in AssemblyPaths)
 			{
 				// Load the assembly into our app domain
@@ -117,7 +116,7 @@ namespace AutomationTool
 					// Add a resolver for the Assembly directory, so that its dependencies may be found alongside it
 					AssemblyUtils.InstallRecursiveAssemblyResolver(AssemblyLocation.Directory.FullName);
 					Assembly Assembly = AppDomain.CurrentDomain.Load(AssemblyName.GetAssemblyName(AssemblyLocation.FullName));
-					Assemblies.Add(Assembly);
+					AllScriptAssemblies.Add(Assembly);
 				}
 				catch (Exception Ex)
 				{
@@ -125,27 +124,13 @@ namespace AutomationTool
 				}
 			}
 
-			AllCompiledAssemblies = Assemblies;
-
-			Platform.InitializePlatforms(AllCompiledAssemblies);
+			Platform.InitializePlatforms(AllScriptAssemblies);
 
 			EnumerateScriptCommands();
 
 			EnumerateBuildProducts();
 		}
 
-		public static HashSet<Assembly> GetCompiledAssemblies()
-		{
-			if (AllCompiledAssemblies == null)
-			{
-				return new HashSet<Assembly>();
-			}
-			else
-			{
-				return AllCompiledAssemblies;
-			}
-		}
-
 		public static Dictionary<string, Type> Commands
 		{
 			get { return ScriptCommands; }
diff --git a/Engine/Source/Programs/AutomationTool/BuildGraph/BuildGraph.cs b/Engine/Source/Programs/AutomationTool/BuildGraph/BuildGraph.cs
--- a/Engine/Source/Programs/AutomationTool/BuildGraph/BuildGraph.cs
+++ b/Engine/Source/Programs/AutomationTool/BuildGraph/BuildGraph.cs
@@ -530,16 +530,13 @@ namespace AutomationTool
 		/// <param name="bPublicTasksOnly">Whether to include just public tasks, or all the tasks in any loaded assemblies</param>
 		static bool FindAvailableTasks(Dictionary<string, ScriptTask> NameToTask, bool bPublicTasksOnly)
 		{
-			IEnumerable<Assembly> LoadedAssemblies = AppDomain.CurrentDomain.GetAssemblies()
-				// Exclude the AutomationTool driver assembly - it contains no tasks, and trying to load it in the context of
-				// BuildGraph can result in an exception if Microsoft.Build.Framework is not able to be loaded
-				.Where(A => !String.Equals("AutomationTool", A.GetName().Name));
+			IEnumerable<Assembly> LoadedScriptAssemblies = ScriptManager.AllScriptAssemblies;
 
 			if(bPublicTasksOnly)
 			{
-				LoadedAssemblies = LoadedAssemblies.Where(x => IsPublicAssembly(new FileReference(x.Location)));
+				LoadedScriptAssemblies = LoadedScriptAssemblies.Where(x => IsPublicAssembly(new FileReference(x.Location)));
 			}
-			foreach (Assembly LoadedAssembly in LoadedAssemblies)
+			foreach (Assembly LoadedAssembly in LoadedScriptAssemblies)
 			{
 				Type[] Types;
 				try
diff --git a/Engine/Source/Programs/AutomationTool/Gauntlet/Framework/Gauntlet.Utils.cs b/Engine/Source/Programs/AutomationTool/Gauntlet/Framework/Gauntlet.Utils.cs
--- a/Engine/Source/Programs/AutomationTool/Gauntlet/Framework/Gauntlet.Utils.cs
+++ b/Engine/Source/Programs/AutomationTool/Gauntlet/Framework/Gauntlet.Utils.cs
@@ -512,15 +512,8 @@ namespace Gauntlet
 				// find all types from loaded assemblies that implement testnode
 					List < Type> CandidateTypes = new List<Type>();
 
-				foreach (var Assembly in AppDomain.CurrentDomain.GetAssemblies())
+				foreach (Assembly Assembly in ScriptManager.AllScriptAssemblies) 
 				{
-					if (String.Equals("AutomationTool", Assembly.GetName().Name))
-					{
-						// Skip the AutomationTool driver assembly - it contains no relevant types, and trying to process it here
-						// can result in an exception if Microsoft.Build.Framework is not able to be loaded.
-						continue;
-					}
-
 					foreach (var Type in Assembly.GetTypes())
 					{
 						if (typeof(ITestNode).IsAssignableFrom(Type))
@@ -733,9 +726,9 @@ namespace Gauntlet
 
 				if (bIncludeCompiledScripts)
 				{
-					foreach (Assembly assembly in ScriptManager.GetCompiledAssemblies())
+					foreach (Assembly Assembly in ScriptManager.AllScriptAssemblies)
 					{
-						List<Type> AssemblyTypes = assembly.GetTypes().Where(T => typeof(InterfaceType).IsAssignableFrom(T)).ToList();
+						List<Type> AssemblyTypes = Assembly.GetTypes().Where(T => typeof(InterfaceType).IsAssignableFrom(T)).ToList();
 						AllTypes.UnionWith(AssemblyTypes);
 					}
 				}
diff --git a/Engine/Source/Programs/AutomationTool/Localization/LocalizationProvider.cs b/Engine/Source/Programs/AutomationTool/Localization/LocalizationProvider.cs
--- a/Engine/Source/Programs/AutomationTool/Localization/LocalizationProvider.cs
+++ b/Engine/Source/Programs/AutomationTool/Localization/LocalizationProvider.cs
@@ -6,6 +6,7 @@ using AutomationTool;
 using UnrealBuildTool;
 using EpicGames.Core;
 using System.Threading.Tasks;
+using System.Reflection;
 
 namespace EpicGames.Localization
 {
@@ -160,16 +161,8 @@ namespace EpicGames.Localization
 			{
 				// Find all types that derive from LocalizationProvider in any of our DLLs
 				CachedLocalizationProviderTypes = new Dictionary<string, Type>();
-				var LoadedAssemblies = AppDomain.CurrentDomain.GetAssemblies();
-				foreach (var Dll in LoadedAssemblies)
+				foreach (Assembly Dll in ScriptManager.AllScriptAssemblies)
 				{
-					if (String.Equals("AutomationTool", Dll.GetName().Name))
-					{
-						// Exclude the AutomationTool driver assembly - it contains no relevant types, and trying to load it here
-						// can result in an exception if Microsoft.Build.Framework is not able to be loaded
-						continue;
-					}
-					
 					var AllTypes = Dll.GetTypes();
 					foreach (var PotentialLocalizationNodeType in AllTypes)
 					{
