Fix platform name for per-platform inis.

The problem was reported on UDN so there is no JIRA.

diff --git a/Engine/Source/Programs/AutomationTool/Scripts/CopyBuildToStagingDirectory.Automation.cs b/Engine/Source/Programs/AutomationTool/Scripts/CopyBuildToStagingDirectory.Automation.cs
--- a/Engine/Source/Programs/AutomationTool/Scripts/CopyBuildToStagingDirectory.Automation.cs
+++ b/Engine/Source/Programs/AutomationTool/Scripts/CopyBuildToStagingDirectory.Automation.cs
@@ -870,7 +870,7 @@ namespace AutomationScripts
 					// Initialize platform extensions to stage.
 					var PlatformExtensionsToStage = new List<string>();
 					{
-						string PlatformExtensionName = SC.StageTargetPlatform.PlatformType.ToString();
+						string PlatformExtensionName = ConfigHierarchy.GetIniPlatformName(SC.StageTargetPlatform.PlatformType);
 						PlatformExtensionsToStage.Add(PlatformExtensionName);
 
 						DataDrivenPlatformInfo.ConfigDataDrivenPlatformInfo Info = DataDrivenPlatformInfo.GetDataDrivenInfoForPlatform(SC.StageTargetPlatform.PlatformType);
