BuildPhysX.Automation.cs:

Iterate just script assemblies, not everything loaded - to avoid an assertion when processing AutomationTool.dll due to Microsoft.Build.Framework

diff --git a/Engine/Source/Programs/AutomationTool/Scripts/BuildPhysX.Automation.cs b/Engine/Source/Programs/AutomationTool/Scripts/BuildPhysX.Automation.cs
--- a/Engine/Source/Programs/AutomationTool/Scripts/BuildPhysX.Automation.cs
+++ b/Engine/Source/Programs/AutomationTool/Scripts/BuildPhysX.Automation.cs
@@ -474,7 +474,7 @@ public sealed class BuildPhysX : BuildCommand
 		var TargetPlatforms = new Dictionary<string, TargetPlatform>();
 
 		// Grab all the non-abstract subclasses of TargetPlatform from the executing assembly.
-		var AvailablePlatformTypes = from Assembly in AppDomain.CurrentDomain.GetAssemblies()
+		var AvailablePlatformTypes = from Assembly in ScriptManager.AllScriptAssemblies
 									 from Type in Assembly.GetTypes()
 									 where !Type.IsAbstract && Type.IsSubclassOf(typeof(TargetPlatform)) && !Type.IsAbstract
 									 select Type;
