Making project file generation wrt to target files consistent for temporary targets ( used for BP-only projects that are either nativized or enable a non-default plugin ) and project file generation.  Temp targets generated in this way didn't previously set DefaultBuildSettings to v2.

diff --git a/Engine/Source/Programs/AutomationTool/AutomationUtils/ProjectUtils.cs b/Engine/Source/Programs/AutomationTool/AutomationUtils/ProjectUtils.cs
--- a/Engine/Source/Programs/AutomationTool/AutomationUtils/ProjectUtils.cs
+++ b/Engine/Source/Programs/AutomationTool/AutomationUtils/ProjectUtils.cs
@@ -448,6 +448,7 @@ namespace AutomationTool
 				Writer.WriteLine("{");
 				Writer.WriteLine("\tpublic {0}Target(TargetInfo Target) : base(Target)", ProjectName);
 				Writer.WriteLine("\t{");
+				Writer.WriteLine("\t\tDefaultBuildSettings = BuildSettingsVersion.V2;");
 				Writer.WriteLine("\t\tType = TargetType.Game;");
 				Writer.WriteLine("\t\tExtraModuleNames.Add(\"{0}\");", ProjectName);
 				Writer.WriteLine("\t}");
