Horde: Support version pattern replacement in AWS ECS deploy task

diff --git a/Engine/Source/Programs/AutomationTool/BuildGraph/Tasks/AwsEcsDeployTask.cs b/Engine/Source/Programs/AutomationTool/BuildGraph/Tasks/AwsEcsDeployTask.cs
--- a/Engine/Source/Programs/AutomationTool/BuildGraph/Tasks/AwsEcsDeployTask.cs
+++ b/Engine/Source/Programs/AutomationTool/BuildGraph/Tasks/AwsEcsDeployTask.cs
@@ -24,11 +24,17 @@ namespace BuildGraph.Tasks
 		public string TaskDefinitionFile;
 
 		/// <summary>
-		/// Docker image to set in new task definition (will replace DOCKER_PATTERN with this value)
+		/// Docker image to set in new task definition (will replace %%DOCKER_PATTERN%% with this value)
 		/// </summary>
 		[TaskParameter(Optional = false)]
 		public string DockerImage;
 
+		/// <summary>
+		/// App version to set in new task definition (will replace %%VERSION%% with this value)
+		/// </summary>
+		[TaskParameter(Optional = true)]
+		public string Version;
+
 		/// <summary>
 		/// Cluster ARN representing AWS ECS cluster to operate on
 		/// </summary>
@@ -88,12 +94,16 @@ namespace BuildGraph.Tasks
 		/// <param name="TagNameToFileSet">Mapping from tag names to the set of files they include</param>
 		public override void Execute(JobContext Job, HashSet<FileReference> BuildProducts, Dictionary<string, HashSet<FileReference>> TagNameToFileSet)
 		{
-			string TaskDefContents = File.ReadAllText(ResolveFile(Parameters.TaskDefinitionFile).FullName);
-			string TaskDefWithDockerImageContents = TaskDefContents.Replace("DOCKER_IMAGE", Parameters.DockerImage);
+			string TaskDefTemplate = File.ReadAllText(ResolveFile(Parameters.TaskDefinitionFile).FullName);
+			string TaskDefRendered = TaskDefTemplate.Replace("%%DOCKER_IMAGE%%", Parameters.DockerImage);
+			if (Parameters.Version != null)
+			{
+				TaskDefRendered = TaskDefRendered.Replace("%%VERSION%%", Parameters.Version);
+			}
 
 			FileReference TempTaskDefFile = FileReference.Combine(Unreal.RootDirectory, "Engine", "Intermediate", "Build", "AwsEcsDeployTaskTemp.json");
 			DirectoryReference.CreateDirectory(TempTaskDefFile.Directory);
-			File.WriteAllText(TempTaskDefFile.FullName, TaskDefWithDockerImageContents, new UTF8Encoding(encoderShouldEmitUTF8Identifier: false));
+			File.WriteAllText(TempTaskDefFile.FullName, TaskDefRendered, new UTF8Encoding(encoderShouldEmitUTF8Identifier: false));
 			
 			IProcessResult CreateTaskDefResult = SpawnTaskBase.Execute("aws", $"ecs register-task-definition --cli-input-json \"file://{TempTaskDefFile.FullName}\"", EnvVars: ParseEnvVars(Parameters.Environment, Parameters.EnvironmentFile), LogOutput: Parameters.LogOutput);
 
