AutomationTool:

Re-add -NoCompile to prevent compilation of AutomationTool or script modules

diff --git a/Engine/Build/BatchFiles/RunUAT.bat b/Engine/Build/BatchFiles/RunUAT.bat
--- a/Engine/Build/BatchFiles/RunUAT.bat
+++ b/Engine/Build/BatchFiles/RunUAT.bat
@@ -29,7 +29,7 @@ if not defined ARGUMENT goto ParseArguments_Done
 set ARGUMENT=%ARGUMENT:"=%
 if /I "%ARGUMENT%" == "-msbuild-verbose" set MSBUILD_LOGLEVEL=normal
 if /I "%ARGUMENT%" == "-compile" set FORCECOMPILE_UAT=FORCE
-rem if /I "%ARGUMENT%" == "-nocompile" set NOCOMPILE_UAT=1
+if /I "%ARGUMENT%" == "-nocompile" set NOCOMPILE_UAT=1
 if /I "%ARGUMENT%" == "-noturnkeyvariables" set SET_TURNKEY_VARIABLES=0
 shift
 goto ParseArguments
diff --git a/Engine/Source/Programs/AutomationTool/CompileScriptModules.cs b/Engine/Source/Programs/AutomationTool/CompileScriptModules.cs
--- a/Engine/Source/Programs/AutomationTool/CompileScriptModules.cs
+++ b/Engine/Source/Programs/AutomationTool/CompileScriptModules.cs
@@ -46,7 +46,7 @@ namespace AutomationToolDriver
 		/// <param name="bUseBuildRecords"></param>
 		/// <param name="bBuildSuccess"></param>
 		/// <returns></returns>
-		public static HashSet<FileReference> InitializeScriptModules(string ScriptsForProjectFileName, List<string> AdditionalScriptsFolders, bool bForceCompile, bool bUseBuildRecords, out bool bBuildSuccess)
+		public static HashSet<FileReference> InitializeScriptModules(string ScriptsForProjectFileName, List<string> AdditionalScriptsFolders, bool bForceCompile, bool bNoCompile, bool bUseBuildRecords, out bool bBuildSuccess)
 		{
 			WriteTimeCache WriteTimeCache = new WriteTimeCache();
 			
@@ -74,7 +74,7 @@ namespace AutomationToolDriver
 
 			// Fall back to the slower approach: use msbuild to load csproj files & build as necessary
 			RegisterMsBuildPath();
-			return BuildAllScriptPlugins(FoundAutomationProjects, bForceCompile, out bBuildSuccess, ref WriteTimeCache);
+			return BuildAllScriptPlugins(FoundAutomationProjects, bForceCompile, bNoCompile, out bBuildSuccess, ref WriteTimeCache);
 		}
 
 		// Acceleration structure:
@@ -569,7 +569,7 @@ namespace AutomationToolDriver
 #endif
             };
 
-		private static HashSet<FileReference> BuildAllScriptPlugins(HashSet<FileReference> FoundAutomationProjects, bool bForceCompile, out bool bBuildSuccess, ref WriteTimeCache Cache)
+		private static HashSet<FileReference> BuildAllScriptPlugins(HashSet<FileReference> FoundAutomationProjects, bool bForceCompile, bool bNoCompile, out bool bBuildSuccess, ref WriteTimeCache Cache)
 		{
 			// The -IgnoreBuildRecords prevents the loading & parsing of .uatbuildrecord files - but UATBuildRecord objects will be used in this function regardless
 			Dictionary<FileReference, UATBuildRecord> BuildRecords = new Dictionary<FileReference, UATBuildRecord>();
@@ -830,6 +830,13 @@ namespace AutomationToolDriver
 					}
 				}
 
+				if (bNoCompile)
+				{
+					bBuildSuccess = true;
+					// return a list of all script modules that are up to date (without touching any uatbuildrecords)
+					return new HashSet<FileReference>(InputProjectGraph.EntryPointNodes.Where(P => !OutOfDateProjects.Contains(P)).Select(P => FileReference.FromString(P.ProjectInstance.GetPropertyValue("TargetPath"))));
+				}
+
 				if (OutOfDateProjects.Count > 0)
 				{
 					BuildProjectGraph = new ProjectGraph(OutOfDateProjects.Select(P => P.ProjectInstance.FullPath), GlobalProperties);
diff --git a/Engine/Source/Programs/AutomationTool/Program.cs b/Engine/Source/Programs/AutomationTool/Program.cs
--- a/Engine/Source/Programs/AutomationTool/Program.cs
+++ b/Engine/Source/Programs/AutomationTool/Program.cs
@@ -245,6 +245,7 @@ namespace AutomationToolDriver
 					{"-AllowStdOutLogVerbosity", ""},
 					{"-NoAutoSDK", ""},
 					{"-Compile", "Force all script modules to be compiled"},
+					{"-NoCompile", "Do not attempt to compile any script modules - attempts to run with whatever is up to date" },
 					{"-IgnoreBuildRecords", "Ignore build record (.uatbuildrecord) files when determining if script modules are up to date" },
 					{"-UseLocalBuildStorage", @"Allows you to use local storage for your root build storage dir {default of P:\Builds {on PC} is changed to Engine\Saved\LocalBuilds}. Used for local testing."},
 					{"-WaitForDebugger", "Waits for a debugger to be attached, and breaks once debugger successfully attached."},
@@ -407,13 +408,14 @@ namespace AutomationToolDriver
 			string ScriptsForProject = (string)AutomationToolCommandLine.GetValueUnchecked("-ScriptsForProject");
 			List<string> AdditionalScriptDirs = (List<string>) AutomationToolCommandLine.GetValueUnchecked("-ScriptDir");
 			bool bForceCompile = AutomationToolCommandLine.IsSetGlobal("-Compile");
+			bool bNoCompile = AutomationToolCommandLine.IsSetGlobal("-NoCompile");
 			bool bUseBuildRecords = !AutomationToolCommandLine.IsSetGlobal("-IgnoreBuildRecords");
 			List<CommandInfo> Commands = AutomationToolCommandLine.IsSetGlobal("-List")
 				? null
 				: AutomationToolCommandLine.CommandsToExecute;
 			bool bBuildSuccess;
 			HashSet<FileReference> ScriptModuleAssemblyPaths = InitializeScriptModules(
-					ScriptsForProject, AdditionalScriptDirs, bForceCompile, bUseBuildRecords, out bBuildSuccess);
+					ScriptsForProject, AdditionalScriptDirs, bForceCompile, bNoCompile,bUseBuildRecords, out bBuildSuccess);
 
 			if (!bBuildSuccess)
             {
