AutomationTool:

Rename -NoCompile as -NoCompileUAT, to work around problems with HordeAgent's use of the former.

diff --git a/Engine/Build/BatchFiles/RunUAT.bat b/Engine/Build/BatchFiles/RunUAT.bat
--- a/Engine/Build/BatchFiles/RunUAT.bat
+++ b/Engine/Build/BatchFiles/RunUAT.bat
@@ -29,7 +29,7 @@ if not defined ARGUMENT goto ParseArguments_Done
 set ARGUMENT=%ARGUMENT:"=%
 if /I "%ARGUMENT%" == "-msbuild-verbose" set MSBUILD_LOGLEVEL=normal
 if /I "%ARGUMENT%" == "-compile" set FORCECOMPILE_UAT=FORCE
-if /I "%ARGUMENT%" == "-nocompile" set NOCOMPILE_UAT=1
+if /I "%ARGUMENT%" == "-nocompileuat" set NOCOMPILE_UAT=1
 if /I "%ARGUMENT%" == "-noturnkeyvariables" set SET_TURNKEY_VARIABLES=0
 shift
 goto ParseArguments
diff --git a/Engine/Source/Programs/AutomationTool/Program.cs b/Engine/Source/Programs/AutomationTool/Program.cs
--- a/Engine/Source/Programs/AutomationTool/Program.cs
+++ b/Engine/Source/Programs/AutomationTool/Program.cs
@@ -245,7 +245,7 @@ namespace AutomationToolDriver
 					{"-AllowStdOutLogVerbosity", ""},
 					{"-NoAutoSDK", ""},
 					{"-Compile", "Force all script modules to be compiled"},
-					{"-NoCompile", "Do not attempt to compile any script modules - attempts to run with whatever is up to date" },
+					{"-NoCompileUAT", "Do not attempt to compile any script modules - attempts to run with whatever is up to date" },
 					{"-IgnoreBuildRecords", "Ignore build record (.uatbuildrecord) files when determining if script modules are up to date" },
 					{"-UseLocalBuildStorage", @"Allows you to use local storage for your root build storage dir {default of P:\Builds {on PC} is changed to Engine\Saved\LocalBuilds}. Used for local testing."},
 					{"-WaitForDebugger", "Waits for a debugger to be attached, and breaks once debugger successfully attached."},
@@ -408,7 +408,7 @@ namespace AutomationToolDriver
 			string ScriptsForProject = (string)AutomationToolCommandLine.GetValueUnchecked("-ScriptsForProject");
 			List<string> AdditionalScriptDirs = (List<string>) AutomationToolCommandLine.GetValueUnchecked("-ScriptDir");
 			bool bForceCompile = AutomationToolCommandLine.IsSetGlobal("-Compile");
-			bool bNoCompile = AutomationToolCommandLine.IsSetGlobal("-NoCompile");
+			bool bNoCompile = AutomationToolCommandLine.IsSetGlobal("-NoCompileUAT");
 			bool bUseBuildRecords = !AutomationToolCommandLine.IsSetGlobal("-IgnoreBuildRecords");
 			List<CommandInfo> Commands = AutomationToolCommandLine.IsSetGlobal("-List")
 				? null
