AutomationTool UpdateVersionFiles:

Do not modify MetaData.cs

diff --git a/Engine/Build/BatchFiles/BuildUAT.bat b/Engine/Build/BatchFiles/BuildUAT.bat
--- a/Engine/Build/BatchFiles/BuildUAT.bat
+++ b/Engine/Build/BatchFiles/BuildUAT.bat
@@ -31,6 +31,7 @@ dir /s^
 
 rem note: no /s
 dir ^
+ Programs\Shared\MetaData.cs^
  Programs\AutomationTool\*.cs^
  Programs\AutomationTool\*.csproj^
  | find ".cs" >>..\Intermediate\Build\AutomationToolFiles.txt
diff --git a/Engine/Build/BatchFiles/BuildUAT.sh b/Engine/Build/BatchFiles/BuildUAT.sh
--- a/Engine/Build/BatchFiles/BuildUAT.sh
+++ b/Engine/Build/BatchFiles/BuildUAT.sh
@@ -47,6 +47,7 @@ elif [ -f ../Intermediate/Build/AutomationToolLastBuildTime ]; then
   fi
 
   UPDATED_AUTOMATIONTOOL_FILES="$(find \
+    Programs/Shared \
     Programs/AutomationTool \
     -maxdepth 1 \
     -type f \
diff --git a/Engine/Source/Programs/AutomationTool/AutomationUtils/UE4Build.cs b/Engine/Source/Programs/AutomationTool/AutomationUtils/UE4Build.cs
--- a/Engine/Source/Programs/AutomationTool/AutomationUtils/UE4Build.cs
+++ b/Engine/Source/Programs/AutomationTool/AutomationUtils/UE4Build.cs
@@ -340,7 +340,7 @@ namespace AutomationTool
 				Version = new BuildVersion();
 			}
 
-			List<FileReference> Result = new List<FileReference>();
+			List<FileReference> Result = new List<FileReference>(1);
 			{
 				if (bDoUpdateVersionFiles)
 				{
@@ -380,34 +380,6 @@ namespace AutomationTool
 				Result.Add(BuildVersionFile);
 			}
 
-            {
-                // Use Version.h data to update MetaData.cs so the assemblies match the engine version.
-				FileReference MetaDataFile = FileReference.Combine(Unreal.EngineDirectory, "Source", "Programs", "Shared", "MetaData.cs");
-
-				if (bDoUpdateVersionFiles)
-                {
-					// Get the revision to sync files to before 
-					if(CommandUtils.P4Enabled && ChangelistNumber > 0)
-					{
-						CommandUtils.P4.Sync(String.Format("-f \"{0}@{1}\"", MetaDataFile, ChangelistNumber), false, false);
-					}
-
-                    // Get the MAJOR/MINOR/PATCH from the Engine Version file, as it is authoritative. The rest we get from the P4Env.
-                    string NewInformationalVersion = String.Format("{0}.{1}.{2}-{3}+{4}", Version.MajorVersion, Version.MinorVersion, Version.PatchVersion, Version.Changelist.ToString(), Version.BranchName);
-
-                    CommandUtils.LogLog("Updating {0} with AssemblyInformationalVersion: {1}", MetaDataFile, NewInformationalVersion);
-
-                    VersionFileUpdater VersionH = new VersionFileUpdater(MetaDataFile);
-                    VersionH.SetAssemblyInformationalVersion(NewInformationalVersion);
-                    VersionH.Commit();
-                }
-                else
-                {
-                    CommandUtils.LogVerbose("{0} will not be updated because P4 is not enabled.", MetaDataFile);
-                }
-                Result.Add(MetaDataFile);
-            }
-
 			return Result;
 		}
 
diff --git a/Engine/Source/Programs/AutomationTool/AutomationUtils/Utils.cs b/Engine/Source/Programs/AutomationTool/AutomationUtils/Utils.cs
--- a/Engine/Source/Programs/AutomationTool/AutomationUtils/Utils.cs
+++ b/Engine/Source/Programs/AutomationTool/AutomationUtils/Utils.cs
@@ -1169,31 +1169,6 @@ namespace AutomationTool
 			return Line.Substring(StartOffset, Offset - StartOffset);
 		}
 
-		/// <summary>
-		/// Doc
-		/// </summary>
-        public void SetAssemblyInformationalVersion(string NewInformationalVersion)
-		{
-            // This searches for the AssemblyInformationalVersion string. Most the mess is to allow whitespace in places that are possible.
-            // Captures the string into a group called "Ver" for replacement.
-            var regex = new Regex(@"\[assembly:\s+AssemblyInformationalVersion\s*\(\s*""(?<Ver>.*)""\s*\)\s*]", RegexOptions.IgnoreCase | RegexOptions.ExplicitCapture);
-            foreach (var Index in Enumerable.Range(0, Lines.Count))
-            {
-                var line = Lines[Index];
-                var match = regex.Match(line);
-                if (match.Success)
-                {
-                    var verGroup = match.Groups["Ver"];
-                    var sb = new StringBuilder(line);
-                    sb.Remove(verGroup.Index, verGroup.Length);
-                    sb.Insert(verGroup.Index, NewInformationalVersion);
-                    Lines[Index] = sb.ToString();
-                    return;
-                }
-            }
-            throw new AutomationException("Failed to find the AssemblyInformationalVersion attribute in {1}", MyFile.FullName);
-		}
-
 		/// <summary>
 		/// Doc
 		/// </summary>
diff --git a/Engine/Source/Programs/Shared/MetaData.cs b/Engine/Source/Programs/Shared/MetaData.cs
--- a/Engine/Source/Programs/Shared/MetaData.cs
+++ b/Engine/Source/Programs/Shared/MetaData.cs
@@ -18,7 +18,7 @@ using System.Resources;
 
 #if !SPECIFIC_VERSION
 // Automatically generate a version number based on the time of compilation
-[assembly: AssemblyVersion( "4.0.0.0" )]
-[assembly: AssemblyInformationalVersion("4.0.0")]
+[assembly: AssemblyVersion( "5.0.0.0" )]
+[assembly: AssemblyInformationalVersion("5.0.0")]
 #endif
 
