OodleCompressDLL config option needs to be read in CopyBuildToStagingDirectory and passed on the command line
because UnrealPak as a program can't read project config hierarchy

diff --git a/Engine/Source/Programs/AutomationTool/Scripts/CopyBuildToStagingDirectory.Automation.cs b/Engine/Source/Programs/AutomationTool/Scripts/CopyBuildToStagingDirectory.Automation.cs
--- a/Engine/Source/Programs/AutomationTool/Scripts/CopyBuildToStagingDirectory.Automation.cs
+++ b/Engine/Source/Programs/AutomationTool/Scripts/CopyBuildToStagingDirectory.Automation.cs
@@ -62,7 +62,7 @@ namespace AutomationScripts
 				}
 
 				if(OrderType == OrderFileType.Cooker && AppendOrder == -1)
-	            {
+				{
 					// make sure to sort CookerOpenOrder.log at the end since its produce by the cooking stage
 					AppendOrder = Int32.MaxValue;
 				}
@@ -208,7 +208,7 @@ namespace AutomationScripts
 			{
 				CmdLine.AppendFormat(" -secondaryOrder={0}", CommandUtils.MakePathSafeToUseWithCommandLine(string.Join(",", SecondaryPakOrderFileLocations.Select(u => u.File.FullName).ToArray())));
 			}
-			if (!String.IsNullOrEmpty(PatchSourceContentPath)) 
+			if (!String.IsNullOrEmpty(PatchSourceContentPath))
 			{
 				CmdLine.AppendFormat(" -generatepatch={0} -tempfiles={1}", CommandUtils.MakePathSafeToUseWithCommandLine(PatchSourceContentPath), CommandUtils.MakePathSafeToUseWithCommandLine(CommandUtils.CombinePaths(CommandUtils.CmdEnv.LocalRoot, "TempFiles" + Path.GetFileNameWithoutExtension(OutputLocation.FullName))));
 			}
@@ -227,7 +227,7 @@ namespace AutomationScripts
 					CmdLine.AppendFormat(" -sign");
 				}
 			}
-			
+
 			CmdLine.Append(PlatformOptions);
 
 			return CmdLine.ToString();
@@ -916,7 +916,7 @@ namespace AutomationScripts
 
 						SC.RestrictedFolderNames = ToRestoreRestrictedFolder;
 					}
-					
+
 					// Stage platform extension config files
 					foreach (string PlatformExtensionToStage in PlatformExtensionsToStage)
 					{
@@ -1054,7 +1054,7 @@ namespace AutomationScripts
 					if (DirectoryReference.Exists(CookOutputDir))
 					{
 						List<FileReference> CookedFiles = DirectoryReference.EnumerateFiles(CookOutputDir, "*", SearchOption.AllDirectories).ToList();
-						
+
 						// When cooking directly to I/O store contaier files look for the package store manifest
 						if (Params.IoStore || Params.ZenStore)
 						{
@@ -1840,7 +1840,7 @@ namespace AutomationScripts
 				ConfigFiles.Add(new ConfigFile(ProjectConfigFileReference));
 				bFoundConfig = true;
 			}
-		
+
 			if (!bFoundConfig)
 			{
 				return null;
@@ -1848,7 +1848,7 @@ namespace AutomationScripts
 			ConfigHierarchy PakRulesConfig = new ConfigHierarchy(ConfigFiles);*/
 
 
-			
+
 			ConfigHierarchy PakRulesConfig = ConfigCache.ReadHierarchy(ConfigHierarchyType.PakFileRules, DirectoryReference.FromFile(Params.RawProjectPath), SC.StageTargetPlatform.IniPlatformType, SC.CustomConfig);
 
 			bool bChunkedBuild = SC.PlatformUsesChunkManifests && DoesChunkPakManifestExist(Params, SC);
@@ -2010,7 +2010,7 @@ namespace AutomationScripts
 			// Search in order, return on first match
 			foreach (var PakRules in RulesList)
 			{
-				bool bMatched = !PakRules.bDisabled && 
+				bool bMatched = !PakRules.bDisabled &&
 					((!PakRules.bUseExactFilePathMatch && PakRules.Filter != null && PakRules.Filter.Matches(StagingFile.Key)) ||
 					(PakRules.bUseExactFilePathMatch && PakRules.ExactFileMatches != null && PakRules.ExactFileMatches.Contains(StagingFile.Value)));
 
@@ -2432,6 +2432,21 @@ namespace AutomationScripts
 					//LogInformation(" CompressionLevel = {0}", CompressionLevel);
 					CompressionFormats += " -compresslevel=" + CompressionLevel;
 				}
+
+
+				// unlike the other config, the DLL should be fetched from the build platform not the target platform
+				// it's the DLL to use during packaging, not during runtime, so it needs to be for the platform the packager runs on, not the platform you're packaging for.
+				// Most of the compression options come from the platform you're packaging FOR
+
+				ConfigHierarchy BuildPlatformEngineConfig = ConfigCache.ReadHierarchy(ConfigHierarchyType.Engine, DirectoryReference.FromFile(Params.RawProjectPath),
+					BuildHostPlatform.Current.Platform);
+
+				string OodleCompressDLL;
+				BuildPlatformEngineConfig.GetString("OodleDataCompressionFormat", "OodleCompressDLL", out OodleCompressDLL);
+				if (!string.IsNullOrWhiteSpace(OodleCompressDLL))
+				{
+					CompressionFormats += " -OodleCompressDLL=" + OodleCompressDLL;
+				}
 			}
 
 			Func<string, string> GetPostFix = PakFilename => GetPakFilePostFix(bShouldGeneratePatch, Params.HasIterateSharedCookedBuild, PlatformGameConfig, PakFilename);
@@ -2569,10 +2584,10 @@ namespace AutomationScripts
 			{
 				// Default filespecs
 				OrderFileSpecs.AddRange(new OrderFileSpec[] {
-					new OrderFileSpec{ FileNamePattern = "GameOpenOrder*.log", OrderType = OrderFile.OrderFileType.Game, Required=false, Priority = 0 },
-					new OrderFileSpec{ FileNamePattern = "CookerOpenOrder*.log", OrderType = OrderFile.OrderFileType.Cooker, Required=false, Priority = 0 },
-					new OrderFileSpec{ FileNamePattern = "EditorOpenOrder.log", OrderType = OrderFile.OrderFileType.Editor, Required=false, Priority = 0 }
-				});
+				new OrderFileSpec{ FileNamePattern = "GameOpenOrder*.log", OrderType = OrderFile.OrderFileType.Game, Required=false, Priority = 0 },
+				new OrderFileSpec{ FileNamePattern = "CookerOpenOrder*.log", OrderType = OrderFile.OrderFileType.Cooker, Required=false, Priority = 0 },
+				new OrderFileSpec{ FileNamePattern = "EditorOpenOrder.log", OrderType = OrderFile.OrderFileType.Editor, Required=false, Priority = 0 }
+			});
 			}
 
 			// search CookPlaform (e.g. IOSClient and then regular platform (e.g. IOS).
@@ -2631,12 +2646,12 @@ namespace AutomationScripts
 				OrderFiles.RemoveAll(x => x.OrderType == OrderFile.OrderFileType.Editor);
 			}
 
-	        // Check second order 
-	        bool bUseSecondaryOrder = false;
-	        PlatformGameConfig.GetBool("/Script/UnrealEd.ProjectPackagingSettings", "bPakUsesSecondaryOrder", out bUseSecondaryOrder);
+			// Check second order 
+			bool bUseSecondaryOrder = false;
+			PlatformGameConfig.GetBool("/Script/UnrealEd.ProjectPackagingSettings", "bPakUsesSecondaryOrder", out bUseSecondaryOrder);
 
 			if(!bUseSecondaryOrder)
-	        {
+			{
 				OrderFiles.RemoveAll(x => (x.OrderType == OrderFile.OrderFileType.Cooker && x.AppendOrder == Int32.MaxValue));
 
 				// chekc if theres a merging cookopenorder file
@@ -2782,7 +2797,7 @@ namespace AutomationScripts
 						}
 
 						ConfigHierarchy PlatformEngineConfig;
-						if (Params.EngineConfigs.TryGetValue(SC.StageTargetPlatform.PlatformType, out PlatformEngineConfig))					
+						if (Params.EngineConfigs.TryGetValue(SC.StageTargetPlatform.PlatformType, out PlatformEngineConfig))
 						{
 							// if the runtime will want to reduce memory usage, we have to disable the pak index freezing
 							bool bUnloadPakEntries = false;
@@ -3174,7 +3189,7 @@ namespace AutomationScripts
 			{
 				CommandletParams += String.Format(" -Order={0}", MakePathSafeToUseWithCommandLine(string.Join(",", OrderFiles.Select(u => u.File.FullName).ToArray())));
 				CommandletParams += String.Format(" -OrderPriority={0}", MakePathSafeToUseWithCommandLine(string.Join(",", OrderFiles.Select(u => u.Priority).ToArray())));
-			} 
+			}
 			if (!string.IsNullOrWhiteSpace(AdditionalArgs))
 			{
 				CommandletParams += AdditionalArgs;
@@ -3236,7 +3251,7 @@ namespace AutomationScripts
 					NumFailed++;
 				}
 			}
-				
+
 			// Abort if any instance failed
 			if(NumFailed > 0)
 			{
@@ -3486,14 +3501,14 @@ namespace AutomationScripts
 					string NoExtension = CombinePaths(Path.GetDirectoryName(OriginalFilename), Path.GetFileNameWithoutExtension(OriginalFilename));
 					if(Path.GetExtension(NoExtension) == ".m")
 					{
-						// Hack around .m.ubulk files having a double extension
-						NoExtension = CombinePaths(Path.GetDirectoryName(OriginalFilename), Path.GetFileNameWithoutExtension(NoExtension));
+					// Hack around .m.ubulk files having a double extension
+					NoExtension = CombinePaths(Path.GetDirectoryName(OriginalFilename), Path.GetFileNameWithoutExtension(NoExtension));
 					}
 					string OriginalReplaceSlashes = OriginalFilename.Replace('/', '\\');
 					string NoExtensionReplaceSlashes = NoExtension.Replace('/', '\\');
 
-					// First read manifest
-					List<ChunkDefinition> Chunks;
+				// First read manifest
+				List<ChunkDefinition> Chunks;
 					if (FileNameToChunks.TryGetValue(OriginalFilename, out Chunks))
 					{
 						PakList.UnionWith(Chunks);
@@ -3511,24 +3526,24 @@ namespace AutomationScripts
 						PakList.UnionWith(Chunks);
 					}
 
-					// Now run through the pak rules which may override things
-					ApplyPakFileRules(PakRulesList, StagingFile, PakList, ChunkNameToDefinition, out bExcludeFromPaks);
+				// Now run through the pak rules which may override things
+				ApplyPakFileRules(PakRulesList, StagingFile, PakList, ChunkNameToDefinition, out bExcludeFromPaks);
 
 					if (bExcludeFromPaks)
 					{
 						return;
 					}
 
-					// Actually add to chunk
-					foreach (ChunkDefinition Chunk in PakList)
+				// Actually add to chunk
+				foreach (ChunkDefinition Chunk in PakList)
 					{
 						ChunkDefinition TargetChunk = Chunk;
 
 						string OrigExt = Path.GetExtension(OriginalFilename);
 						if (OrigExt.Equals(OptionalBulkDataFileExtension))
 						{
-							// any optional files encountered we want to put in a separate pak file
-							string OptionalChunkName = Chunk.ChunkName + "optional";
+						// any optional files encountered we want to put in a separate pak file
+						string OptionalChunkName = Chunk.ChunkName + "optional";
 							if (!OptionalChunks.ContainsKey(OptionalChunkName))
 							{
 								ChunkDefinition OptionalChunk = new ChunkDefinition(OptionalChunkName);
@@ -3544,15 +3559,15 @@ namespace AutomationScripts
 
 						if (bForceOneChunkPerFile)
 						{
-							// Files are only allowed to be in a single chunk
-							break;
+						// Files are only allowed to be in a single chunk
+						break;
 						}
 					}
 
 					if (!bAddedToChunk)
 					{
-						//LogInformation("No chunk assigned found for {0}. Using default chunk.", StagingFile.Key);
-						DefaultChunk.ResponseFile.TryAdd(StagingFile.Key, StagingFile.Value);
+					//LogInformation("No chunk assigned found for {0}. Using default chunk.", StagingFile.Key);
+					DefaultChunk.ResponseFile.TryAdd(StagingFile.Key, StagingFile.Value);
 					}
 				});
 
@@ -3635,7 +3650,7 @@ namespace AutomationScripts
 				if (Chunk.ResponseFile.Count > 0)
 				{
 					string EncryptionKeyToUse = Params.SkipEncryption ? "" : Chunk.EncryptionKeyGuid;
-					PakInputs.Add(new CreatePakParams(Chunk.ChunkName, 
+					PakInputs.Add(new CreatePakParams(Chunk.ChunkName,
 						Chunk.ResponseFile.ToDictionary(entry => entry.Key, entry => entry.Value),
 						Params.Compressed || Chunk.bCompressed, EncryptionKeyToUse));
 				}
diff --git a/Engine/Source/Runtime/Core/Private/Compression/OodleDataCompressionFormat.cpp b/Engine/Source/Runtime/Core/Private/Compression/OodleDataCompressionFormat.cpp
--- a/Engine/Source/Runtime/Core/Private/Compression/OodleDataCompressionFormat.cpp
+++ b/Engine/Source/Runtime/Core/Private/Compression/OodleDataCompressionFormat.cpp
@@ -448,43 +448,55 @@ ICompressionFormat * CreateOodleDataCompressionFormat()
 		bool bUseCompressionFormatOodle = 
 			( FCString::Strifind(FCommandLine::Get(), TEXT("-compressionformats=oodle")) != NULL ) ||
 			( FCString::Strifind(FCommandLine::Get(), TEXT("-compressionformat=oodle")) != NULL );
+
 		if ( bUseCompressionFormatOodle )			
 		{
 			UE_LOG(OodleDataCompression, Display, TEXT("Oodle v%s format for pak/iostore with method=%s, level=%d=%s"), TEXT(OodleVersion), **MethodMap.FindKey(UsedCompressor), (int)UsedLevel, **LevelMap.FindKey(UsedLevel) );
-		}
 		
-		// OodleCompressDLL to use an earlier *encoder* to maintain identical bit streams
-		// mainly for shipped games to avoid patches
-		// do not specify this if you want to use the latest Engine version of Oodle for encoding
-
-		// NOTE : we get OodleCompressDLL from Engine.ini not from command line
-		//	because this is a function of the platform we are running *on* not the platform we are packaging *for*
-		// the ProjectPackaging settings we get on the command line from Game.ini come from the *target* platform
-		FString OodleDLL = "";
-		if ( GConfig && GConfig->GetString(TEXT("OodleDataCompressionFormat"), TEXT("OodleCompressDLL"), OodleDLL, GEngineIni) && ! OodleDLL.IsEmpty() )
-		{
-			UE_LOG(OodleDataCompression, Display, TEXT("OodleCompressDLL=%s"), *OodleDLL);
-				
-			// just load from default locations :
-			//  the DLL you want here should be added in RuntimeDependencies.Add in a build.cs
-			//  to ensure it is in the binaries directory of the process
-			void * 	OodleDLLHandle = FPlatformProcess::GetDllHandle(*OodleDLL);
+			// OodleCompressDLL to use an earlier *encoder* to maintain identical bit streams
+			// mainly for shipped games to avoid patches
+			// do not specify this if you want to use the latest Engine version of Oodle for encoding
+
+			// NOTE : we get OodleCompressDLL from Engine.ini for the platform we are running *on* not the platform we are packaging *for*
+			// the ProjectPackaging settings we get on the command line from Game.ini come from the *target* platform
+			FString OodleDLL = "";
+
+			// check command line first : 
+			FParse::Value(FCommandLine::Get(), TEXT("OodleCompressDLL="), OodleDLL);
 
-			if ( OodleDLLHandle == nullptr )
+			if ( OodleDLL.IsEmpty() && GConfig )
 			{
-				UE_LOG(OodleDataCompression, Warning, TEXT("OodleCompressDLL %s requested but could not be loaded"), *OodleDLL);
+				// @todo Oodle : possibly remove this? unnecessary if it's always been put on command line
+				// UnrealPak and other "programs" do not read the project config hierarchy
+				// CopyBuildToStagingDirectory.Automation.cs reads this config value and passes it on the command line
+				GConfig->GetString(TEXT("OodleDataCompressionFormat"), TEXT("OodleCompressDLL"), OodleDLL, GEngineIni);
 			}
-			else
+		
+			if ( ! OodleDLL.IsEmpty() )
 			{
-				// GetDllExport is GetProcAddress or dlsym on Mac
-				//  LinuxPlatformProcess::GetDllExport seems to be missing ?
-				//OodleCompressFuncPtr = GetProcAddress( (HMODULE)OodleDLLHandle, "OodleLZ_Compress" );
-				OodleCompressFuncPtr = FPlatformProcess::GetDllExport( OodleDLLHandle, TEXT("OodleLZ_Compress") );
-				if ( OodleCompressFuncPtr == nullptr )
+				UE_LOG(OodleDataCompression, Display, TEXT("OodleCompressDLL=%s"), *OodleDLL);
+				
+				// just load from default locations :
+				//  the DLL you want here should be added in RuntimeDependencies.Add in a build.cs
+				//  to ensure it is in the binaries directory of the process
+				void * 	OodleDLLHandle = FPlatformProcess::GetDllHandle(*OodleDLL);
+
+				if ( OodleDLLHandle == nullptr )
+				{
+					UE_LOG(OodleDataCompression, Warning, TEXT("OodleCompressDLL %s requested but could not be loaded"), *OodleDLL);
+				}
+				else
 				{
-					UE_LOG(OodleDataCompression, Warning, TEXT("OodleCompressDLL %s loaded but no OodleLZ_Compress !?"), *OodleDLL);
+					// GetDllExport is GetProcAddress or dlsym on Mac
+					//  LinuxPlatformProcess::GetDllExport seems to be missing ?
+					//OodleCompressFuncPtr = GetProcAddress( (HMODULE)OodleDLLHandle, "OodleLZ_Compress" );
+					OodleCompressFuncPtr = FPlatformProcess::GetDllExport( OodleDLLHandle, TEXT("OodleLZ_Compress") );
+					if ( OodleCompressFuncPtr == nullptr )
+					{
+						UE_LOG(OodleDataCompression, Warning, TEXT("OodleCompressDLL %s loaded but no OodleLZ_Compress !?"), *OodleDLL);
+					}
+					OodleSetAllocatorsFuncPtr = FPlatformProcess::GetDllExport( OodleDLLHandle, TEXT("OodlePlugins_SetAllocators") );
 				}
-				OodleSetAllocatorsFuncPtr = FPlatformProcess::GetDllExport( OodleDLLHandle, TEXT("OodlePlugins_SetAllocators") );
 			}
 		}
 	}
