Added a faster option for platforms to quickly return a single device in cases where enumerating all devices can potentially be slow.

#rnx

diff --git a/Engine/Source/Programs/AutomationTool/AutomationUtils/Platform.cs b/Engine/Source/Programs/AutomationTool/AutomationUtils/Platform.cs
--- a/Engine/Source/Programs/AutomationTool/AutomationUtils/Platform.cs
+++ b/Engine/Source/Programs/AutomationTool/AutomationUtils/Platform.cs
@@ -199,6 +199,24 @@ namespace AutomationTool
 			return null;
 		}
 
+		public virtual DeviceInfo GetDeviceByName( string DeviceName )
+		{
+			DeviceInfo[] Devices = GetDevices();
+			if (Devices == null)
+			{
+				return null;
+			}
+			// look by Id first
+			DeviceInfo Device = Array.Find(Devices, x => string.Compare(x.Id, DeviceName, true) == 0);
+			// if that fails, use Name
+			if (Device == null)
+			{
+				Device = Array.Find(Devices, x => string.Compare(x.Name, DeviceName, true) == 0);
+			}
+			return Device;
+
+		}
+
 		public virtual bool InstallSDK(BuildCommand BuildCommand, ITurnkeyContext TurnkeyContext, DeviceInfo Device, bool bUnattended)
 		{
 			string Command, Params;
diff --git a/Engine/Source/Programs/AutomationTool/Turnkey/TurnkeyUtils.cs b/Engine/Source/Programs/AutomationTool/Turnkey/TurnkeyUtils.cs
--- a/Engine/Source/Programs/AutomationTool/Turnkey/TurnkeyUtils.cs
+++ b/Engine/Source/Programs/AutomationTool/Turnkey/TurnkeyUtils.cs
@@ -383,19 +383,7 @@ namespace Turnkey
 
 		private static DeviceInfo GetDeviceByPlatformAndName(UnrealTargetPlatform Platform, string DeviceName)
 		{
-			DeviceInfo[] Devices = AutomationTool.Platform.GetPlatform(Platform).GetDevices();
-			if (Devices == null)
-			{
-				return null;
-			}
-			// look by Id first
-			DeviceInfo Device = Array.Find(Devices, x => string.Compare(x.Id, DeviceName, true) == 0);
-			// if that fails, use Name
-			if (Device == null)
-			{
-				Device = Array.Find(Devices, x => string.Compare(x.Name, DeviceName, true) == 0);
-			}
-			return Device;
+			return AutomationTool.Platform.GetPlatform(Platform).GetDeviceByName(DeviceName);
 		}
 
 		public static List<DeviceInfo> GetDevicesFromCommandLineOrUser(string[] CommandOptions, UnrealTargetPlatform Platform)
