Iterate just script assemblies, not everything loaded - to avoid an assertion when processing AutomationTool.dll due to Microsoft.Build.Framework

#rnx

diff --git a/Engine/Source/Programs/AutomationTool/Scripts/BuildCMakeLib.Automation.cs b/Engine/Source/Programs/AutomationTool/Scripts/BuildCMakeLib.Automation.cs
--- a/Engine/Source/Programs/AutomationTool/Scripts/BuildCMakeLib.Automation.cs
+++ b/Engine/Source/Programs/AutomationTool/Scripts/BuildCMakeLib.Automation.cs
@@ -466,7 +466,7 @@ public sealed class BuildCMakeLib : BuildCommand
 	private TargetPlatform GetTargetPlatform()
 	{
 		// Grab all the non-abstract subclasses of TargetPlatform from the executing assembly.
-		var AvailablePlatformTypes = from Assembly in AppDomain.CurrentDomain.GetAssemblies()
+		var AvailablePlatformTypes = from Assembly in ScriptManager.AllScriptAssemblies
 									 from Type in Assembly.GetTypes()
 									 where !Type.IsAbstract && Type.IsSubclassOf(typeof(TargetPlatform))
 									 select Type;
