Add support for making Regex operations optional, leaving Properties unchanged if the match fails.

diff --git a/Engine/Source/Programs/AutomationTool/BuildGraph/Schema.cs b/Engine/Source/Programs/AutomationTool/BuildGraph/Schema.cs
--- a/Engine/Source/Programs/AutomationTool/BuildGraph/Schema.cs
+++ b/Engine/Source/Programs/AutomationTool/BuildGraph/Schema.cs
@@ -836,6 +836,7 @@ namespace AutomationTool
 			Extension.Attributes.Add(CreateSchemaAttribute("Input", StringTypeName, XmlSchemaUse.Required));
 			Extension.Attributes.Add(CreateSchemaAttribute("Pattern", StringTypeName, XmlSchemaUse.Required));
 			Extension.Attributes.Add(CreateSchemaAttribute("Capture", ScriptSchemaStandardType.NameList, XmlSchemaUse.Required));
+			Extension.Attributes.Add(CreateSchemaAttribute("Optional", ScriptSchemaStandardType.Boolean, XmlSchemaUse.Optional));
 			Extension.Attributes.Add(CreateSchemaAttribute("If", ScriptSchemaStandardType.BalancedString, XmlSchemaUse.Optional));
 
 			XmlSchemaSimpleContent ContentModel = new XmlSchemaSimpleContent();
diff --git a/Engine/Source/Programs/AutomationTool/BuildGraph/Script.cs b/Engine/Source/Programs/AutomationTool/BuildGraph/Script.cs
--- a/Engine/Source/Programs/AutomationTool/BuildGraph/Script.cs
+++ b/Engine/Source/Programs/AutomationTool/BuildGraph/Script.cs
@@ -781,9 +781,13 @@ namespace AutomationTool
 						string Input = ReadAttribute(Element, "Input");
 						Match Match = RegexValue.Match(Input);
 
+						bool Optional = Condition.Evaluate(ReadAttribute(Element, "Optional"));
 						if (!Match.Success)
 						{
-							LogError(Element, "Regex {0} did not find a match against input string {1}", RegexString, Input);
+							if (!Optional)
+							{
+								LogError(Element, "Regex {0} did not find a match against input string {1}", RegexString, Input);
+							}
 						}
 						else
 						{
@@ -1885,7 +1889,7 @@ namespace AutomationTool
 			{
 				if (Idx > 0 && Name[Idx] == ' ' && Name[Idx - 1] == ' ')
 				{
-					LogError(Element, "Consecutive spaces in object name");
+					LogError(Element, "Consecutive spaces in object name - '{0}'", Name);
 					return false;
 				}
 				if(Char.IsControl(Name[Idx]) || ScriptSchema.IllegalNameCharacters.IndexOf(Name[Idx]) != -1)
