Remove workaround for UE-117842 which is resolved by [REDACTED_HANDLE]

#rnx

diff --git a/Engine/Source/Programs/AutomationTool/Gauntlet/Unreal/Editor/EditorTests.cs b/Engine/Source/Programs/AutomationTool/Gauntlet/Unreal/Editor/EditorTests.cs
--- a/Engine/Source/Programs/AutomationTool/Gauntlet/Unreal/Editor/EditorTests.cs
+++ b/Engine/Source/Programs/AutomationTool/Gauntlet/Unreal/Editor/EditorTests.cs
@@ -134,9 +134,6 @@ namespace UnrealEditor
 				AppConfig.CommandLineParams.Add("BENCHMARK");
 				AppConfig.CommandLineParams.Add("Deterministic");
 			}
-
-			// Workaround for UE-117842
-			AppConfig.CommandLineParams.Add("dx11");
 		}
 	}
 
