Make Compile Task "Target" Optional

It's totally reasonable to pass -TargetType to UBT via Arguments.

#rnx

diff --git a/Engine/Source/Programs/AutomationTool/BuildGraph/Tasks/CompileTask.cs b/Engine/Source/Programs/AutomationTool/BuildGraph/Tasks/CompileTask.cs
--- a/Engine/Source/Programs/AutomationTool/BuildGraph/Tasks/CompileTask.cs
+++ b/Engine/Source/Programs/AutomationTool/BuildGraph/Tasks/CompileTask.cs
@@ -20,7 +20,7 @@ namespace AutomationTool
 		/// <summary>
 		/// The target to compile.
 		/// </summary>
-		[TaskParameter]
+		[TaskParameter(Optional=true)]
 		public string Target;
 
 		/// <summary>
