Fixed incorrect handling of -iostore argument when staging.

diff --git a/Engine/Source/Programs/AutomationTool/Scripts/CopyBuildToStagingDirectory.Automation.cs b/Engine/Source/Programs/AutomationTool/Scripts/CopyBuildToStagingDirectory.Automation.cs
--- a/Engine/Source/Programs/AutomationTool/Scripts/CopyBuildToStagingDirectory.Automation.cs
+++ b/Engine/Source/Programs/AutomationTool/Scripts/CopyBuildToStagingDirectory.Automation.cs
@@ -606,11 +606,7 @@ namespace AutomationScripts
 
 		private static void LoadPackageStoreManifest(ProjectParams Params, DeploymentContext SC)
 		{
-			if (!Params.Stage)
-			{
-				return;
-			}
-			if (!Params.ZenStore && !Params.IoStore)
+			if (!ShouldCreateIoStoreContainerFiles(Params, SC.StageTargetPlatform, SC.CustomConfig))
 			{
 				return;
 			}
@@ -625,6 +621,10 @@ namespace AutomationScripts
 					FileInfo.Path = FileReference.Combine(SC.PlatformCookDir, FileInfo.Path).FullName;
 				}
 			}
+			else
+			{
+				throw new AutomationException(String.Format("A package store manifest is required when staging to IoStore. Expected to find {0}. Ensure that cooking was successful.", PackageStoreManifestFile.FullName));
+			}
 		}
 
 		public static void CreateStagingManifest(ProjectParams Params, DeploymentContext SC)
@@ -3233,10 +3233,6 @@ namespace AutomationScripts
 				CommandletParams += String.Format("-CreateGlobalContainer={0}", MakePathSafeToUseWithCommandLine(GlobalContainerOutputLocation.FullName));
 			}
 
-			if (SC.PackageStoreManifest == null)
-			{
-				throw new AutomationException("A package store manifest is required when staging to IoStore");
-			}
 			CommandletParams += String.Format(" -CookedDirectory={0} -PackageStoreManifest={1} -Commands={2}",
 				MakePathSafeToUseWithCommandLine(SC.PlatformCookDir.ToString()),
 				MakePathSafeToUseWithCommandLine(SC.PackageStoreManifest.FullPath),
@@ -3246,7 +3242,7 @@ namespace AutomationScripts
 				FileReference ScriptObjectsFile = FileReference.Combine(SC.MetadataDir, "scriptobjects.bin");
 				if (!FileReference.Exists(ScriptObjectsFile))
 				{
-					throw new AutomationException("scriptobjects.bin is required when staging to IoStore");
+					throw new AutomationException(String.Format("A script objects descriptor file is required when staging to IoStore. Expected to find {0}. Ensure cooking was successful.", ScriptObjectsFile.FullName));
 				}
 				CommandletParams += String.Format(" -ScriptObjects={0}", MakePathSafeToUseWithCommandLine(ScriptObjectsFile.FullName));
 			}
@@ -3803,20 +3799,20 @@ namespace AutomationScripts
 				return false;
 			}
 
-			if (Params.IoStore)
+			if (!Params.Stage || Params.SkipStage)
 			{
-				return true;
+				return false;
 			}
 
-			if (Params.Stage && !Params.SkipStage)
+			if (Params.IoStore)
 			{
-				ConfigHierarchy PlatformGameConfig = ConfigCache.ReadHierarchy(ConfigHierarchyType.Game, DirectoryReference.FromFile(Params.RawProjectPath), StageTargetPlatform.IniPlatformType, CustomConfig);
-				bool bUseIoStore = false;
-				PlatformGameConfig.GetBool("/Script/UnrealEd.ProjectPackagingSettings", "bUseIoStore", out bUseIoStore);
-				return bUseIoStore;
+				return true;
 			}
 
-			return false;
+			ConfigHierarchy PlatformGameConfig = ConfigCache.ReadHierarchy(ConfigHierarchyType.Game, DirectoryReference.FromFile(Params.RawProjectPath), StageTargetPlatform.IniPlatformType, CustomConfig);
+			bool bUseIoStore = false;
+			PlatformGameConfig.GetBool("/Script/UnrealEd.ProjectPackagingSettings", "bUseIoStore", out bUseIoStore);
+			return bUseIoStore;
 		}
 
 		private static bool ShouldCreatePak(ProjectParams Params, DeploymentContext SC)
