AutomationTool: Support for Editor in Project BuildEnvironments

This adds support for running Editor in BuildGraph Command and Commandlet tasks where the Project at hand has it's own BuildEnvironment.

diff --git a/Engine/Source/Programs/AutomationTool/AutomationUtils/ProjectUtils.cs b/Engine/Source/Programs/AutomationTool/AutomationUtils/ProjectUtils.cs
--- a/Engine/Source/Programs/AutomationTool/AutomationUtils/ProjectUtils.cs
+++ b/Engine/Source/Programs/AutomationTool/AutomationUtils/ProjectUtils.cs
@@ -922,6 +922,37 @@ namespace AutomationTool
 			// either valid or we're out of ideas...
 			return ProjectPath;
 		}
+
+		/// <summary>
+		/// Full path to the Project executable for the current platform.
+		/// </summary>
+		/// <param name="ProjectFile">Path to Project file</param>
+		/// <param name="TargetType">Target type</param>
+		/// <param name="TargetPlatform">Target platform</param>
+		/// <param name="TargetConfiguration">Target build configuration</param>
+		/// <param name="Cmd">Do you want the console subsystem/commandlet executable?</param>
+		/// <returns></returns>
+		public static FileReference GetProjectTarget(FileReference ProjectFile, UnrealBuildTool.TargetType TargetType, UnrealBuildTool.UnrealTargetPlatform TargetPlatform, UnrealBuildTool.UnrealTargetConfiguration TargetConfiguration = UnrealBuildTool.UnrealTargetConfiguration.Development, bool Cmd = false)
+		{
+			ProjectProperties Properties = ProjectUtils.GetProjectProperties(ProjectFile);
+			List<SingleTargetProperties> Targets = Properties.Targets.Where(x => x.Rules.Type == TargetType).ToList();
+			string TargetName = null;
+			switch (Targets.Count)
+			{
+				case 0:
+					return null;
+				case 1:
+					TargetName = Targets.First().TargetName;
+					break;
+				default:
+					Properties.EngineConfigs[TargetPlatform].GetString("/Script/BuildSettings.BuildSettings", "DefaultEditorTarget", out TargetName);
+					break;
+			}
+
+			FileReference TargetReceiptFileName = UnrealBuildTool.TargetReceipt.GetDefaultPath(ProjectFile.Directory, TargetName, TargetPlatform, TargetConfiguration, "");
+			UnrealBuildTool.TargetReceipt TargetReceipt = UnrealBuildTool.TargetReceipt.Read(TargetReceiptFileName);
+			return Cmd ? TargetReceipt.LaunchCmd : TargetReceipt.Launch;
+		}
 	}
 
     public class BranchInfo
diff --git a/Engine/Source/Programs/AutomationTool/BuildGraph/Tasks/CommandletTask.cs b/Engine/Source/Programs/AutomationTool/BuildGraph/Tasks/CommandletTask.cs
--- a/Engine/Source/Programs/AutomationTool/BuildGraph/Tasks/CommandletTask.cs
+++ b/Engine/Source/Programs/AutomationTool/BuildGraph/Tasks/CommandletTask.cs
@@ -100,7 +100,11 @@ namespace AutomationTool.Tasks
 			FileReference EditorExe;
 			if(Parameters.EditorExe == null)
 			{
-                EditorExe = new FileReference(HostPlatform.Current.GetUE4ExePath("UnrealEditor-Cmd.exe"));
+				EditorExe = ProjectUtils.GetProjectTarget(ProjectFile, UnrealBuildTool.TargetType.Editor, BuildHostPlatform.Current.Platform, UnrealTargetConfiguration.Development, true);
+				if (EditorExe == null)
+				{
+					EditorExe = new FileReference(HostPlatform.Current.GetUE4ExePath("UnrealEditor-Cmd.exe"));
+				}
 			}
 			else
 			{
diff --git a/Engine/Source/Programs/AutomationTool/Gauntlet/Unreal/BuildSource/Gauntlet.UnrealBuildSource.cs b/Engine/Source/Programs/AutomationTool/Gauntlet/Unreal/BuildSource/Gauntlet.UnrealBuildSource.cs
--- a/Engine/Source/Programs/AutomationTool/Gauntlet/Unreal/BuildSource/Gauntlet.UnrealBuildSource.cs
+++ b/Engine/Source/Programs/AutomationTool/Gauntlet/Unreal/BuildSource/Gauntlet.UnrealBuildSource.cs
@@ -627,41 +627,48 @@ namespace Gauntlet
 
 
 		/// <summary>
-		/// Given a platform, a build config, and true/false for client, returns the path to the binary for that config. E.g.
-		/// Win64, Shipping, false = Binaries\Win64\FooServer-Win64-Shipping.exe
+		/// Given a role, platform, and config, returns the path to the binary for that config. E.g. Binaries\Win64\FooServer-Win64-Shipping.exe
 		/// </summary>
+		/// <param name="TargetRole"></param>
 		/// <param name="TargetPlatform"></param>
-		/// <param name="BuildConfig"></param>
-		/// <param name="IsClient"></param>
+		/// <param name="TargetConfiguration"></param>
 		/// <returns></returns>
-		virtual public string GetRelativeExecutablePath(UnrealTargetRole TargetType, UnrealTargetPlatform TargetPlatform, UnrealTargetConfiguration TargetConfiguration)
+		virtual public string GetRelativeExecutablePath(UnrealTargetRole TargetRole, UnrealTargetPlatform TargetPlatform, UnrealTargetConfiguration TargetConfiguration)
 		{
 			string ExePath;
 
-			if (TargetType.UsesEditor())
+			if (TargetRole.UsesEditor())
 			{
-				string ExeFileName = "UnrealEditor";
-				if (TargetConfiguration != UnrealTargetConfiguration.Development)
+				FileReference EditorExe = ProjectUtils.GetProjectTarget(ProjectPath, UnrealBuildTool.TargetType.Editor, TargetPlatform, TargetConfiguration);
+				if (EditorExe != null)
 				{
-					ExeFileName += string.Format("-{0}-{1}", TargetPlatform.ToString(), TargetConfiguration.ToString());
+					ExePath = EditorExe.FullName;
 				}
+				else
+				{
+					string ExeFileName = "UnrealEditor";
+					if (TargetConfiguration != UnrealTargetConfiguration.Development)
+					{
+						ExeFileName += string.Format("-{0}-{1}", TargetPlatform.ToString(), TargetConfiguration.ToString());
+					}
 
-				ExeFileName += Platform.GetExeExtension(TargetPlatform);
+					ExeFileName += Platform.GetExeExtension(TargetPlatform);
 
-				ExePath = string.Format("Engine/Binaries/{0}/{1}", BuildHostPlatform.Current.Platform, ExeFileName);
+					ExePath = string.Format("Engine/Binaries/{0}/{1}", BuildHostPlatform.Current.Platform, ExeFileName);
+				}
 			}
 			else
 			{
 				string BuildType = "";
 
-				if (TargetType == UnrealTargetRole.Client)
+				if (TargetRole == UnrealTargetRole.Client)
 				{
 					if (!UsesSharedBuildType)
 					{
 						BuildType = "Client";
 					}
 				}
-				else if (TargetType == UnrealTargetRole.Server)
+				else if (TargetRole == UnrealTargetRole.Server)
 				{
 					if (!UsesSharedBuildType)
 					{
@@ -685,16 +692,16 @@ namespace Gauntlet
 					{
 						Flags |= BuildFlags.CanReplaceExecutable;
 					}
-                    if (Globals.Params.ParseParam("bulk"))
-                    {
-                        Flags |= BuildFlags.Bulk;
-                    }
+					if (Globals.Params.ParseParam("bulk"))
+					{
+						Flags |= BuildFlags.Bulk;
+					}
 					else if(Globals.Params.ParseParam("notbulk"))
 					{
 						Flags |= BuildFlags.NotBulk;
 					}
 
-                    var Build = GetMatchingBuilds(TargetType, TargetPlatform, TargetConfiguration, Flags).FirstOrDefault();
+					var Build = GetMatchingBuilds(TargetRole, TargetPlatform, TargetConfiguration, Flags).FirstOrDefault();
 
 					if (Build != null)
 					{
@@ -719,7 +726,7 @@ namespace Gauntlet
 
 					ExeFileName += Platform.GetExeExtension(TargetPlatform);
 
-					string BasePath = GetPlatformPath(TargetType, TargetPlatform);
+					string BasePath = GetPlatformPath(TargetRole, TargetPlatform);
 					string ProjectBinary = string.Format("{0}\\Binaries\\{1}\\{2}", ProjectName, TargetPlatform.ToString(), ExeFileName);
 					string StubBinary = Path.Combine(BasePath, ExeFileName);
 					string DevBinary = Path.Combine(Environment.CurrentDirectory, ProjectBinary);
diff --git a/Engine/Source/Programs/UnrealBuildTool/Configuration/UEBuildTarget.cs b/Engine/Source/Programs/UnrealBuildTool/Configuration/UEBuildTarget.cs
--- a/Engine/Source/Programs/UnrealBuildTool/Configuration/UEBuildTarget.cs
+++ b/Engine/Source/Programs/UnrealBuildTool/Configuration/UEBuildTarget.cs
@@ -1445,8 +1445,19 @@ namespace UnrealBuildTool
 				{
 					if (Pair.Value == BuildProductType.Executable)
 					{
-						Receipt.Launch = Pair.Key;
-						break;
+						if (System.IO.Path.GetFileNameWithoutExtension(Pair.Key.FullName).EndsWith("-Cmd"))
+						{
+							Receipt.LaunchCmd = Pair.Key;
+						}
+						else
+						{
+							Receipt.Launch = Pair.Key;
+						}
+
+						if ( Receipt.Launch != null && Receipt.LaunchCmd != null )
+						{
+							break;
+						}
 					}
 				}
 			}
diff --git a/Engine/Source/Programs/UnrealBuildTool/System/TargetReceipt.cs b/Engine/Source/Programs/UnrealBuildTool/System/TargetReceipt.cs
--- a/Engine/Source/Programs/UnrealBuildTool/System/TargetReceipt.cs
+++ b/Engine/Source/Programs/UnrealBuildTool/System/TargetReceipt.cs
@@ -304,6 +304,11 @@ namespace UnrealBuildTool
 		/// </summary>
 		public FileReference? Launch;
 
+		/// <summary>
+		/// The console subsystem/commmandlet exectuable to launch for this target
+		/// </summary>
+		public FileReference? LaunchCmd;
+
 		/// <summary>
 		/// The build products which are part of this target
 		/// </summary>
@@ -560,6 +565,11 @@ namespace UnrealBuildTool
 			{
 				Receipt.Launch = ExpandPathVariables(Launch, EngineDir, ProjectDir);
 			}
+			string LaunchCmd;
+			if(RawObject.TryGetStringField("LaunchCmd", out LaunchCmd))
+			{
+				Receipt.LaunchCmd = ExpandPathVariables(LaunchCmd, EngineDir, ProjectDir);
+			}
 
 			// Read the build products
 			JsonObject[] BuildProductObjects;
@@ -721,6 +731,11 @@ namespace UnrealBuildTool
 					Writer.WriteValue("Launch", InsertPathVariables(Launch, EngineDir, ProjectDir));
 				}
 
+				if(LaunchCmd != null)
+				{
+					Writer.WriteValue("LaunchCmd", InsertPathVariables(LaunchCmd, EngineDir, ProjectDir));
+				}
+
 				Writer.WriteObjectStart("Version");
 				Version.WriteProperties(Writer);
 				Writer.WriteObjectEnd();
