Add lib/so rules for TempStorage in IsDuplicateBuildProduct for Linux

diff --git a/Engine/Source/Programs/AutomationTool/BuildGraph/TempStorage.cs b/Engine/Source/Programs/AutomationTool/BuildGraph/TempStorage.cs
--- a/Engine/Source/Programs/AutomationTool/BuildGraph/TempStorage.cs
+++ b/Engine/Source/Programs/AutomationTool/BuildGraph/TempStorage.cs
@@ -1166,6 +1166,10 @@ namespace AutomationTool
 			{
 				return true;
 			}
+			if (FileName.StartsWith("lib", StringComparison.OrdinalIgnoreCase) && FileName.EndsWith(".so", StringComparison.OrdinalIgnoreCase))
+			{
+				return true;
+			}
 			if (FileName.Equals("plugInfo.json", StringComparison.OrdinalIgnoreCase))
 			{
 				return true;
