Fix missing / in the root path for Linux/Mac

diff --git a/Engine/Source/Programs/AutomationTool/AutomationUtils/CommandUtils.cs b/Engine/Source/Programs/AutomationTool/AutomationUtils/CommandUtils.cs
--- a/Engine/Source/Programs/AutomationTool/AutomationUtils/CommandUtils.cs
+++ b/Engine/Source/Programs/AutomationTool/AutomationUtils/CommandUtils.cs
@@ -2052,11 +2052,11 @@ namespace AutomationTool
                 {
 					if (RuntimePlatform.IsMac)
 					{
-						CachedRootBuildStorageDirectory = CombinePaths("Volumes", "Builds");
+						CachedRootBuildStorageDirectory = CombinePaths("/Volumes", "Builds");
 					}
 					else if (RuntimePlatform.IsLinux)
 					{
-						CachedRootBuildStorageDirectory = CombinePaths("mnt", "Builds");
+						CachedRootBuildStorageDirectory = CombinePaths("/mnt", "Builds");
 					}
 					else
 					{
