AutomationTool:

When running as a child instance, log that fact to the console.

#trivial

diff --git a/Engine/Source/Programs/AutomationTool/AutomationUtils/CommandEnvironment.cs b/Engine/Source/Programs/AutomationTool/AutomationUtils/CommandEnvironment.cs
--- a/Engine/Source/Programs/AutomationTool/AutomationUtils/CommandEnvironment.cs
+++ b/Engine/Source/Programs/AutomationTool/AutomationUtils/CommandEnvironment.cs
@@ -132,9 +132,14 @@ namespace AutomationTool
 			CommandUtils.SetEnvVar(EnvVarNames.MacMallocNanoZone, MallocNanoZone);
 
 			int IsChildInstanceInt;
-			int.TryParse(CommandUtils.GetEnvVar("uebp_UATChildInstance", "0"), out IsChildInstanceInt);
+			int.TryParse(CommandUtils.GetEnvVar(EnvVarNames.IsChildInstance, "0"), out IsChildInstanceInt);
 			IsChildInstance = (IsChildInstanceInt != 0);
 
+			if (IsChildInstance)
+			{
+				Log.TraceInformation($"AutomationTool is running as a child instance ({EnvVarNames.IsChildInstance}={IsChildInstanceInt.ToString()})");
+			}
+
 			// Setup the timestamp string
 			DateTime LocalTime = DateTime.Now;
 
