AutomationTool

* Skip flushing ProjectUtils temp folders on startup of a child process. Leave files intact that were generated by the parent process.

This change fixes some issues where ProjectUtils is booted up in a parent AutomationTool, which causes some UATRules assemblies to be generated and loaded. Child processes try to clean up these files on startup, but they are locked by the parent process. Skip the cleanup in this case.

#rnx

diff --git a/Engine/Source/Programs/AutomationTool/AutomationUtils/Automation.cs b/Engine/Source/Programs/AutomationTool/AutomationUtils/Automation.cs
--- a/Engine/Source/Programs/AutomationTool/AutomationUtils/Automation.cs
+++ b/Engine/Source/Programs/AutomationTool/AutomationUtils/Automation.cs
@@ -113,7 +113,10 @@ namespace AutomationTool
 				}
 				
 				// Clean rules folders up
-				ProjectUtils.CleanupFolders();
+				if (!CommandUtils.CmdEnv.IsChildInstance)
+				{
+					ProjectUtils.CleanupFolders();
+				}
 
 				// Compile scripts.
 				using (TelemetryStopwatch ScriptLoadStopwatch = new TelemetryStopwatch("ScriptLoad"))
