Gauntlet: fix project path when launching on Mac

diff --git a/Engine/Source/Programs/AutomationTool/AutomationUtils/ProjectUtils.cs b/Engine/Source/Programs/AutomationTool/AutomationUtils/ProjectUtils.cs
--- a/Engine/Source/Programs/AutomationTool/AutomationUtils/ProjectUtils.cs
+++ b/Engine/Source/Programs/AutomationTool/AutomationUtils/ProjectUtils.cs
@@ -932,7 +932,7 @@ namespace AutomationTool
 		/// <param name="TargetConfiguration">Target build configuration</param>
 		/// <param name="Cmd">Do you want the console subsystem/commandlet executable?</param>
 		/// <returns></returns>
-		public static FileReference GetProjectTarget(FileReference ProjectFile, UnrealBuildTool.TargetType TargetType, UnrealBuildTool.UnrealTargetPlatform TargetPlatform, UnrealBuildTool.UnrealTargetConfiguration TargetConfiguration = UnrealBuildTool.UnrealTargetConfiguration.Development, bool Cmd = false)
+		public static FileSystemReference GetProjectTarget(FileReference ProjectFile, UnrealBuildTool.TargetType TargetType, UnrealBuildTool.UnrealTargetPlatform TargetPlatform, UnrealBuildTool.UnrealTargetConfiguration TargetConfiguration = UnrealBuildTool.UnrealTargetConfiguration.Development, bool Cmd = false)
 		{
 			ProjectProperties Properties = ProjectUtils.GetProjectProperties(ProjectFile);
 			List<SingleTargetProperties> Targets = Properties.Targets.Where(x => x.Rules.Type == TargetType).ToList();
@@ -951,7 +951,19 @@ namespace AutomationTool
 
 			FileReference TargetReceiptFileName = UnrealBuildTool.TargetReceipt.GetDefaultPath(ProjectFile.Directory, TargetName, TargetPlatform, TargetConfiguration, "");
 			UnrealBuildTool.TargetReceipt TargetReceipt = UnrealBuildTool.TargetReceipt.Read(TargetReceiptFileName);
-			return Cmd ? TargetReceipt.LaunchCmd : TargetReceipt.Launch;
+
+			if (Cmd)
+			{
+				return TargetReceipt.LaunchCmd;
+			}
+			
+			if (TargetPlatform == UnrealTargetPlatform.Mac)
+			{
+				// Remove trailing "/Contents/MacOS/UnrealEngine" to get back to .app directory
+				return TargetReceipt.Launch.Directory.ParentDirectory.ParentDirectory;
+			}
+			
+			return TargetReceipt.Launch;
 		}
 	}
 
diff --git a/Engine/Source/Programs/AutomationTool/BuildGraph/Tasks/CommandletTask.cs b/Engine/Source/Programs/AutomationTool/BuildGraph/Tasks/CommandletTask.cs
--- a/Engine/Source/Programs/AutomationTool/BuildGraph/Tasks/CommandletTask.cs
+++ b/Engine/Source/Programs/AutomationTool/BuildGraph/Tasks/CommandletTask.cs
@@ -97,7 +97,7 @@ namespace AutomationTool.Tasks
 			}
 
 			// Get the path to the editor, and check it exists
-			FileReference EditorExe;
+			FileSystemReference EditorExe;
 			if(Parameters.EditorExe == null)
 			{
 				EditorExe = ProjectUtils.GetProjectTarget(ProjectFile, UnrealBuildTool.TargetType.Editor, BuildHostPlatform.Current.Platform, UnrealTargetConfiguration.Development, true);
@@ -111,12 +111,6 @@ namespace AutomationTool.Tasks
 				EditorExe = Parameters.EditorExe;
 			}
 
-			// Make sure the editor exists
-			if(!FileReference.Exists(EditorExe))
-			{
-				throw new AutomationException("{0} does not exist", EditorExe.FullName);
-			}
-
 			// Run the commandlet
 			CommandUtils.RunCommandlet(ProjectFile, EditorExe.FullName, Parameters.Name, Parameters.Arguments, Parameters.ErrorLevel);
 		}
diff --git a/Engine/Source/Programs/AutomationTool/Gauntlet/Unreal/BuildSource/Gauntlet.UnrealBuildSource.cs b/Engine/Source/Programs/AutomationTool/Gauntlet/Unreal/BuildSource/Gauntlet.UnrealBuildSource.cs
--- a/Engine/Source/Programs/AutomationTool/Gauntlet/Unreal/BuildSource/Gauntlet.UnrealBuildSource.cs
+++ b/Engine/Source/Programs/AutomationTool/Gauntlet/Unreal/BuildSource/Gauntlet.UnrealBuildSource.cs
@@ -639,7 +639,7 @@ namespace Gauntlet
 
 			if (TargetRole.UsesEditor())
 			{
-				FileReference EditorExe = ProjectUtils.GetProjectTarget(ProjectPath, UnrealBuildTool.TargetType.Editor, TargetPlatform, TargetConfiguration);
+				FileSystemReference EditorExe = ProjectUtils.GetProjectTarget(ProjectPath, UnrealBuildTool.TargetType.Editor, TargetPlatform, TargetConfiguration);
 				if (EditorExe != null)
 				{
 					ExePath = EditorExe.FullName;
