Fixed issue that caused iostore to be disabled when staging and paking in separate steps.
#rnx

diff --git a/Engine/Source/Programs/AutomationTool/Scripts/CopyBuildToStagingDirectory.Automation.cs b/Engine/Source/Programs/AutomationTool/Scripts/CopyBuildToStagingDirectory.Automation.cs
--- a/Engine/Source/Programs/AutomationTool/Scripts/CopyBuildToStagingDirectory.Automation.cs
+++ b/Engine/Source/Programs/AutomationTool/Scripts/CopyBuildToStagingDirectory.Automation.cs
@@ -606,7 +606,7 @@ namespace AutomationScripts
 
 		private static void LoadPackageStoreManifest(ProjectParams Params, DeploymentContext SC)
 		{
-			if (!ShouldCreateIoStoreContainerFiles(Params, SC.StageTargetPlatform, SC.CustomConfig))
+			if (!ShouldCreateIoStoreContainerFiles(Params, SC))
 			{
 				return;
 			}
@@ -2821,7 +2821,7 @@ namespace AutomationScripts
 						AdditionalArgs += " -platform=" + ConfigHierarchy.GetIniPlatformName(SC.StageTargetPlatform.IniPlatformType);
 
 						Dictionary<string, string> UnrealPakResponseFile = PakParams.UnrealPakResponseFile;
-						if (ShouldCreateIoStoreContainerFiles(Params, SC.StageTargetPlatform, SC.CustomConfig))
+						if (ShouldCreateIoStoreContainerFiles(Params, SC))
 						{
 							bool bLegacyBulkDataOffsets = false;
 							if (PlatformEngineConfig.GetBool("Core.System", "LegacyBulkDataOffsets", out bLegacyBulkDataOffsets) && bLegacyBulkDataOffsets)
@@ -3042,7 +3042,7 @@ namespace AutomationScripts
 
 					InternalUtils.SafeCreateDirectory(Path.GetDirectoryName(ReleaseVersionPath));
 					InternalUtils.SafeCopyFile(OutputLocation.FullName, ReleaseVersionPath);
-					if (ShouldCreateIoStoreContainerFiles(Params, SC.StageTargetPlatform, SC.CustomConfig))
+					if (ShouldCreateIoStoreContainerFiles(Params, SC))
 					{
 						InternalUtils.SafeCopyFile(Path.ChangeExtension(OutputLocation.FullName, ".utoc"), Path.ChangeExtension(ReleaseVersionPath, ".utoc"));
 						InternalUtils.SafeCopyFile(Path.ChangeExtension(OutputLocation.FullName, ".ucas"), Path.ChangeExtension(ReleaseVersionPath, ".ucas"));
@@ -3152,7 +3152,7 @@ namespace AutomationScripts
 						{
 							HashSet<string> IncludedExtensions = new HashSet<string>();
 							IncludedExtensions.Add(OutputFilenameExtension);
-							if (ShouldCreateIoStoreContainerFiles(Params, SC.StageTargetPlatform, SC.CustomConfig))
+							if (ShouldCreateIoStoreContainerFiles(Params, SC))
 							{
 								IncludedExtensions.Add(".ucas");
 								IncludedExtensions.Add(".utoc");
@@ -3787,9 +3787,9 @@ namespace AutomationScripts
 			return TmpPackagingPath;
 		}
 
-		private static bool ShouldCreateIoStoreContainerFiles(ProjectParams Params, Platform StageTargetPlatform, String CustomConfig)
+		private static bool ShouldCreateIoStoreContainerFiles(ProjectParams Params, DeploymentContext SC)
 		{
-			if (Params.CookOnTheFly)
+			if (!ShouldCreatePak(Params, SC))
 			{
 				return false;
 			}
@@ -3799,17 +3799,12 @@ namespace AutomationScripts
 				return false;
 			}
 
-			if (!Params.Stage || Params.SkipStage)
-			{
-				return false;
-			}
-
 			if (Params.IoStore)
 			{
 				return true;
 			}
 
-			ConfigHierarchy PlatformGameConfig = ConfigCache.ReadHierarchy(ConfigHierarchyType.Game, DirectoryReference.FromFile(Params.RawProjectPath), StageTargetPlatform.IniPlatformType, CustomConfig);
+			ConfigHierarchy PlatformGameConfig = ConfigCache.ReadHierarchy(ConfigHierarchyType.Game, DirectoryReference.FromFile(Params.RawProjectPath), SC.StageTargetPlatform.IniPlatformType, SC.CustomConfig);
 			bool bUseIoStore = false;
 			PlatformGameConfig.GetBool("/Script/UnrealEd.ProjectPackagingSettings", "bUseIoStore", out bUseIoStore);
 			return bUseIoStore;
