AutomationTool
Add support for Project names in Compile tasks, similar to Commandlet

diff --git a/Engine/Source/Programs/AutomationTool/BuildGraph/Tasks/CompileTask.cs b/Engine/Source/Programs/AutomationTool/BuildGraph/Tasks/CompileTask.cs
--- a/Engine/Source/Programs/AutomationTool/BuildGraph/Tasks/CompileTask.cs
+++ b/Engine/Source/Programs/AutomationTool/BuildGraph/Tasks/CompileTask.cs
@@ -38,7 +38,7 @@ namespace AutomationTool
 		/// <summary>
 		/// The project to compile with.
 		/// </summary>
-		[TaskParameter(Optional = true)]
+		[TaskParameter(Optional = true, ValidationType = TaskParameterValidationType.FileSpec)]
 		public string Project;
 
 		/// <summary>
@@ -135,7 +135,7 @@ namespace AutomationTool
 			bAllowXGE &= Parameters.AllowXGE;
 			bAllowParallelExecutor &= Parameters.AllowParallelExecutor;
 
-			UE4Build.BuildTarget Target = new UE4Build.BuildTarget { TargetName = Parameters.Target, Platform = Parameters.Platform, Config = Parameters.Configuration, UprojectPath = String.IsNullOrEmpty(Parameters.Project)? null : new FileReference(Parameters.Project), UBTArgs = "-nobuilduht " + (Parameters.Arguments ?? ""), Clean = Parameters.Clean };
+			UE4Build.BuildTarget Target = new UE4Build.BuildTarget { TargetName = Parameters.Target, Platform = Parameters.Platform, Config = Parameters.Configuration, UprojectPath = CompileTask.ProjectFile, UBTArgs = "-nobuilduht " + (Parameters.Arguments ?? ""), Clean = Parameters.Clean };
 			if(!String.IsNullOrEmpty(Parameters.Tag))
 			{
 				TargetToTagName.Add(Target, Parameters.Tag);
@@ -199,6 +199,11 @@ namespace AutomationTool
 		/// </summary>
 		public CompileTaskParameters Parameters;
 
+		/// <summary>
+		/// Resolved path to Project file
+		/// </summary>
+		public FileReference ProjectFile = null;
+
 		/// <summary>
 		/// Construct a compile task
 		/// </summary>
@@ -206,6 +211,24 @@ namespace AutomationTool
 		public CompileTask(CompileTaskParameters Parameters)
 		{
 			this.Parameters = Parameters;
+
+			// Get the full path to the project file
+			if(!String.IsNullOrEmpty(Parameters.Project))
+			{
+				if(Parameters.Project.EndsWith(".uproject", StringComparison.OrdinalIgnoreCase))
+				{
+					ProjectFile = CustomTask.ResolveFile(Parameters.Project);
+				}
+				else
+				{
+					ProjectFile = NativeProjects.EnumerateProjectFiles().FirstOrDefault(x => x.GetFileNameWithoutExtension().Equals(Parameters.Project, StringComparison.OrdinalIgnoreCase));
+				}
+
+				if(ProjectFile == null || !FileReference.Exists(ProjectFile))
+				{
+					throw new BuildException("Unable to resolve project '{0}'", Parameters.Project);
+				}
+			}
 		}
 
 		/// <summary>
@@ -216,6 +239,9 @@ namespace AutomationTool
 		/// <param name="TagNameToFileSet">Mapping from tag names to the set of files they include</param>
 		public override void Execute(JobContext Job, HashSet<FileReference> BuildProducts, Dictionary<string, HashSet<FileReference>> TagNameToFileSet)
 		{
+			//
+			// Don't do any logic here. You have to do it in the ctor (at graph parse time) otherwise you break the ParallelExecutor pathway!!
+			//
 			GetExecutor().Execute(Job, BuildProducts, TagNameToFileSet);
 		}
 
