Tolerate referenced project files that are completely missing.

#rnx

diff --git a/Engine/Source/Programs/AutomationTool/BuildGraph/Tasks/CompileTask.cs b/Engine/Source/Programs/AutomationTool/BuildGraph/Tasks/CompileTask.cs
--- a/Engine/Source/Programs/AutomationTool/BuildGraph/Tasks/CompileTask.cs
+++ b/Engine/Source/Programs/AutomationTool/BuildGraph/Tasks/CompileTask.cs
@@ -223,11 +223,6 @@ namespace AutomationTool
 				{
 					ProjectFile = NativeProjects.EnumerateProjectFiles().FirstOrDefault(x => x.GetFileNameWithoutExtension().Equals(Parameters.Project, StringComparison.OrdinalIgnoreCase));
 				}
-
-				if(ProjectFile == null || !FileReference.Exists(ProjectFile))
-				{
-					throw new BuildException("Unable to resolve project '{0}'", Parameters.Project);
-				}
 			}
 		}
 
