AutomationTool

Defer resolution of the Project path until Execution time. This allows for build graphs to refer to projects that may be missing on the workspace, and still execute valid nodes in the graph.

#rnx

diff --git a/Engine/Source/Programs/AutomationTool/BuildGraph/Tasks/CompileTask.cs b/Engine/Source/Programs/AutomationTool/BuildGraph/Tasks/CompileTask.cs
--- a/Engine/Source/Programs/AutomationTool/BuildGraph/Tasks/CompileTask.cs
+++ b/Engine/Source/Programs/AutomationTool/BuildGraph/Tasks/CompileTask.cs
@@ -135,7 +135,7 @@ namespace AutomationTool
 			bAllowXGE &= Parameters.AllowXGE;
 			bAllowParallelExecutor &= Parameters.AllowParallelExecutor;
 
-			UE4Build.BuildTarget Target = new UE4Build.BuildTarget { TargetName = Parameters.Target, Platform = Parameters.Platform, Config = Parameters.Configuration, UprojectPath = CompileTask.ProjectFile, UBTArgs = "-nobuilduht " + (Parameters.Arguments ?? ""), Clean = Parameters.Clean };
+			UE4Build.BuildTarget Target = new UE4Build.BuildTarget { TargetName = Parameters.Target, Platform = Parameters.Platform, Config = Parameters.Configuration, UprojectPath = CompileTask.FindProjectFile(), UBTArgs = "-nobuilduht " + (Parameters.Arguments ?? ""), Clean = Parameters.Clean };
 			if(!String.IsNullOrEmpty(Parameters.Tag))
 			{
 				TargetToTagName.Add(Target, Parameters.Tag);
@@ -155,12 +155,12 @@ namespace AutomationTool
 		public void Execute(JobContext Job, HashSet<FileReference> BuildProducts, Dictionary<string, HashSet<FileReference>> TagNameToFileSet)
 		{
 			// Create the agenda
-            UE4Build.BuildAgenda Agenda = new UE4Build.BuildAgenda();
+			UE4Build.BuildAgenda Agenda = new UE4Build.BuildAgenda();
 			Agenda.Targets.AddRange(Targets);
 
 			// Build everything
 			Dictionary<UE4Build.BuildTarget, BuildManifest> TargetToManifest = new Dictionary<UE4Build.BuildTarget,BuildManifest>();
-            UE4Build Builder = new UE4Build(Job.OwnerCommand);
+			UE4Build Builder = new UE4Build(Job.OwnerCommand);
 
 			bool bCanUseParallelExecutor = (BuildHostPlatform.Current.Platform == UnrealTargetPlatform.Win64 && bAllowParallelExecutor);	// parallel executor is only available on Windows as of 2016-09-22
 			Builder.Build(Agenda, InDeleteBuildProducts: null, InUpdateVersionFiles: false, InForceNoXGE: !bAllowXGE, InUseParallelExecutor: bCanUseParallelExecutor, InTargetToManifest: TargetToManifest);
@@ -211,8 +211,16 @@ namespace AutomationTool
 		public CompileTask(CompileTaskParameters Parameters)
 		{
 			this.Parameters = Parameters;
+		}
+
+		/// <summary>
+		/// Resolve the path to the project file
+		/// </summary>
+		public FileReference FindProjectFile()
+		{
+			FileReference ProjectFile = null;
 
-			// Get the full path to the project file
+			// Resolve the full path to the project file
 			if(!String.IsNullOrEmpty(Parameters.Project))
 			{
 				if(Parameters.Project.EndsWith(".uproject", StringComparison.OrdinalIgnoreCase))
@@ -223,7 +231,14 @@ namespace AutomationTool
 				{
 					ProjectFile = NativeProjects.EnumerateProjectFiles().FirstOrDefault(x => x.GetFileNameWithoutExtension().Equals(Parameters.Project, StringComparison.OrdinalIgnoreCase));
 				}
+
+				if(ProjectFile == null || !FileReference.Exists(ProjectFile))
+				{
+					throw new BuildException("Unable to resolve project '{0}'", Parameters.Project);
+				}
 			}
+
+			return ProjectFile;
 		}
 
 		/// <summary>
@@ -235,7 +250,8 @@ namespace AutomationTool
 		public override void Execute(JobContext Job, HashSet<FileReference> BuildProducts, Dictionary<string, HashSet<FileReference>> TagNameToFileSet)
 		{
 			//
-			// Don't do any logic here. You have to do it in the ctor (at graph parse time) otherwise you break the ParallelExecutor pathway!!
+			// Don't do any logic here. You have to do it in the ctor or a getter
+			//  otherwise you break the ParallelExecutor pathway, which doesn't call this function!
 			//
 			GetExecutor().Execute(Job, BuildProducts, TagNameToFileSet);
 		}
