AutomationTool

Fix GetRelativeExecutablePath for usages when specific builds are specified

#rnx

diff --git a/Engine/Source/Programs/AutomationTool/Gauntlet/Unreal/BuildSource/Gauntlet.UnrealBuildSource.cs b/Engine/Source/Programs/AutomationTool/Gauntlet/Unreal/BuildSource/Gauntlet.UnrealBuildSource.cs
--- a/Engine/Source/Programs/AutomationTool/Gauntlet/Unreal/BuildSource/Gauntlet.UnrealBuildSource.cs
+++ b/Engine/Source/Programs/AutomationTool/Gauntlet/Unreal/BuildSource/Gauntlet.UnrealBuildSource.cs
@@ -639,7 +639,13 @@ namespace Gauntlet
 
 			if (TargetRole.UsesEditor())
 			{
-				FileSystemReference EditorExe = ProjectUtils.GetProjectTarget(ProjectPath, UnrealBuildTool.TargetType.Editor, TargetPlatform, TargetConfiguration);
+				FileSystemReference EditorExe = null;
+
+				if (BuildName == "Editor")
+				{
+					ProjectUtils.GetProjectTarget(ProjectPath, UnrealBuildTool.TargetType.Editor, TargetPlatform, TargetConfiguration);
+				}
+
 				if (EditorExe != null)
 				{
 					ExePath = EditorExe.FullName;
diff --git a/Engine/Source/Programs/AutomationTool/Gauntlet/Unreal/RunUnreal.cs b/Engine/Source/Programs/AutomationTool/Gauntlet/Unreal/RunUnreal.cs
--- a/Engine/Source/Programs/AutomationTool/Gauntlet/Unreal/RunUnreal.cs
+++ b/Engine/Source/Programs/AutomationTool/Gauntlet/Unreal/RunUnreal.cs
@@ -98,7 +98,7 @@ namespace Gauntlet
 
 			if (string.IsNullOrEmpty(ContextOptions.Build))
 			{
-				throw new AutomationException("No builds specified. Use -builds=p:\\path\\to\\build");
+				throw new AutomationException("No build specified. Use -build=p:\\path\\to\\build");
 			}
 
 			if (typeof(UnrealBuildSource).IsAssignableFrom(ContextOptions.BuildSourceType) == false)
