RunLowLevelTests compilation fix

diff --git a/Engine/Source/Programs/AutomationTool/LowLevelTests/RunLowLevelTests.cs b/Engine/Source/Programs/AutomationTool/LowLevelTests/RunLowLevelTests.cs
--- a/Engine/Source/Programs/AutomationTool/LowLevelTests/RunLowLevelTests.cs
+++ b/Engine/Source/Programs/AutomationTool/LowLevelTests/RunLowLevelTests.cs
@@ -417,13 +417,9 @@ namespace LowLevelTests
 
 		public static string ExeFileRegEx { get { return @"\w+Tests.exe$"; } }
 
-		public static string SelfFileRegEx { get { return @"\w+Tests.self$"; } }
-
 		// To be used when *both* Mac and Linux apps can be launched from Gauntlet
 		public static string PosixFileRegEx { get { return @"\w+Tests$"; } }
 
-		public static string SwitchFileRegEx { get { return @"\w+Tests.nspd"; } }
-
 		public string BuildName { get { return TestApp; } }
 
 		public class LooseStagedBuild : StagedBuild
@@ -449,39 +445,7 @@ namespace LowLevelTests
 
 			public bool CopyPlatformSpecificFiles()
 			{
-				try
-				{
-					if (Platform.IsInGroup(UnrealPlatformGroup.GDK))
-					{
-						const string MsGameConfigFile = "MicrosoftGame.Config";
-						const string GameOSXvdFile = "gameos.xvd";
-						const string ResourcesPriFile = "resources.pri";
-						string LayoutDirectory = Path.Combine(Globals.UE4RootDir, "Engine", "Saved", Platform.ToString(), "Layout", "Image", "Loose");
-						File.Copy(Path.Combine(LayoutDirectory, MsGameConfigFile), Path.Combine(BuildPath, MsGameConfigFile), true);
-						File.Copy(Path.Combine(LayoutDirectory, GameOSXvdFile), Path.Combine(BuildPath, GameOSXvdFile), true);
-						File.Copy(Path.Combine(LayoutDirectory, ResourcesPriFile), Path.Combine(BuildPath, ResourcesPriFile), true);
-					}
-					if (Platform.IsInGroup(UnrealPlatformGroup.Sony))
-					{
-						const string SceModuleDir = "sce_module";
-						string UEBuildDirectory = Path.Combine(Globals.UE4RootDir, "Engine", "Build", Platform.ToString());
-						string SourceDir = Path.Combine(UEBuildDirectory, SceModuleDir);
-						string TargetDir = Path.Combine(BuildPath, SceModuleDir);
-						if (!Directory.Exists(TargetDir))
-						{
-							Directory.CreateDirectory(TargetDir);
-						}
-						foreach(string SceFile in Directory.EnumerateFiles(SourceDir))
-						{
-							File.Copy(SceFile, SceFile.Replace(SourceDir, TargetDir), true);
-						}
-					}
-				}
-				catch (Exception copyEx)
-				{
-					Log.Error("Could not copy {0} specific files: {1}.", Platform.ToString(), copyEx);
-					return false;
-				}
+				// TODO
 				
 				return true;
 			}
@@ -512,18 +476,10 @@ namespace LowLevelTests
 			{
 				LowLevelTestsBuild DiscoveredBuild = null;
 				IEnumerable<string> Executables = new List<string>();
-				if (InPlatform.IsInGroup(UnrealPlatformGroup.GDK))
+				if (InPlatform.IsInGroup(UnrealPlatformGroup.Windows))
 				{
 					Executables = DirectoryUtils.FindFiles(InBuildPath, new Regex(ExeFileRegEx));
 				}
-				else if (InPlatform.IsInGroup(UnrealPlatformGroup.Sony))
-				{
-					Executables = DirectoryUtils.FindFiles(InBuildPath, new Regex(SelfFileRegEx));
-				}
-				else if (InPlatform == UnrealTargetPlatform.Switch)
-				{
-					Executables = DirectoryUtils.FindFiles(InBuildPath, new Regex(SwitchFileRegEx));
-				}
 				foreach (string Executable in Executables)
 				{
 					UnrealTargetConfiguration UnrealConfig = UnrealTargetConfiguration.Development;
