Cleanup useless code
#rnx

diff --git a/Engine/Source/Programs/AutomationTool/LowLevelTests/RunLowLevelTests.cs b/Engine/Source/Programs/AutomationTool/LowLevelTests/RunLowLevelTests.cs
--- a/Engine/Source/Programs/AutomationTool/LowLevelTests/RunLowLevelTests.cs
+++ b/Engine/Source/Programs/AutomationTool/LowLevelTests/RunLowLevelTests.cs
@@ -68,18 +68,7 @@ namespace LowLevelTests
 
 		void DoCleanup(UnrealTargetPlatform Platform)
 		{
-			if (!Globals.Params.ParseParam("removedevices"))
-			{
-				return;
-			}
-
-			if (Platform == UnrealTargetPlatform.PS4)
-			{
-				string DevKitUtilPath = Path.Combine(Environment.CurrentDirectory, "Engine/Platforms/PS4/Binaries/DotNET/PS4DevKitUtil.exe");
-				Log.Verbose("PS4DevkitUtil executing 'removeall'");
-				IProcessResult BootResult = CommandUtils.Run(DevKitUtilPath, "removeall");
-				Log.Verbose("PS4DevkitUtil 'removeall' completed with exit code {0}", BootResult.ExitCode);
-			}
+			// TODO: Platform specific cleanup
 		}
 
 		private bool ExecuteTest(LowLevelTestExecutorOptions Options, ITestNode LowLevelTestNode)
