Add support for Linux in the Legacy Unzip path

diff --git a/Engine/Source/Programs/AutomationTool/AutomationUtils/BuildUtils.cs b/Engine/Source/Programs/AutomationTool/AutomationUtils/BuildUtils.cs
--- a/Engine/Source/Programs/AutomationTool/AutomationUtils/BuildUtils.cs
+++ b/Engine/Source/Programs/AutomationTool/AutomationUtils/BuildUtils.cs
@@ -85,22 +85,37 @@ namespace AutomationTool
 			MsBuild(Env, ProjectFile, CmdLine, LogName);
 		}
 
-        /// <summary>
-        /// returns true if this is a mac executable using some awful conventions
-        /// <param name="Filename">Filename</param>
-        /// </summary>
-        public static bool IsProbablyAMacOrIOSExe(string Filename)
-        {
-            return
-                Path.GetExtension(Filename) == ".sh" ||
-                Path.GetExtension(Filename) == ".command" ||
-                ((
-                    CommandUtils.CombinePaths(Filename).ToLower().Contains(CommandUtils.CombinePaths("Binaries", "Mac").ToLower()) ||
-                    CommandUtils.CombinePaths(Filename).ToLower().Contains(CommandUtils.CombinePaths("Binaries", "IOS").ToLower()) ||
+		/// <summary>
+		/// returns true if this is a linux executable using some awful conventions
+		/// <param name="Filename">Filename</param>
+		/// </summary>
+		public static bool IsProbablyALinuxExe(string Filename)
+		{
+			return
+				Path.GetExtension(Filename) == ".sh" ||
+				((
+					CommandUtils.CombinePaths(Filename).ToLower().Contains(CommandUtils.CombinePaths("Binaries", "Linux").ToLower()) ||
+					CommandUtils.CombinePaths(Filename).ToLower().Contains(CommandUtils.CombinePaths("Binaries", "LinuxArm64").ToLower()) ||
+					CommandUtils.CombinePaths(Filename).ToLower().Contains(CommandUtils.CombinePaths("Binaries", "ThirdParty").ToLower())
+				) && (Path.GetExtension(Filename) == "" || Path.GetExtension(Filename) == "."));
+		}
+
+		/// <summary>
+		/// returns true if this is a mac executable using some awful conventions
+		/// <param name="Filename">Filename</param>
+		/// </summary>
+		public static bool IsProbablyAMacOrIOSExe(string Filename)
+		{
+			return
+				Path.GetExtension(Filename) == ".sh" ||
+				Path.GetExtension(Filename) == ".command" ||
+				((
+					CommandUtils.CombinePaths(Filename).ToLower().Contains(CommandUtils.CombinePaths("Binaries", "Mac").ToLower()) ||
+					CommandUtils.CombinePaths(Filename).ToLower().Contains(CommandUtils.CombinePaths("Binaries", "IOS").ToLower()) ||
 					CommandUtils.CombinePaths(Filename).ToLower().Contains(CommandUtils.CombinePaths("Binaries", "ThirdParty").ToLower()) ||
 					CommandUtils.CombinePaths(PathSeparator.Slash, Filename).ToLower().Contains(".app/Contents/MacOS".ToLower())
-                ) && (Path.GetExtension(Filename) == "" || Path.GetExtension(Filename) == "."));
-        }
+				) && (Path.GetExtension(Filename) == "" || Path.GetExtension(Filename) == "."));
+		}
 
 		/// <summary>
 		/// Sets an executable bit for Unix executables and adds read permission for all users plus write permission for owner.
@@ -108,7 +123,7 @@ namespace AutomationTool
 		/// </summary>
 		public static void FixUnixFilePermissions(string Filename)
 		{
-			string Permissions = IsProbablyAMacOrIOSExe(Filename) ? "0755" : "0644";
+			string Permissions = IsProbablyAMacOrIOSExe(Filename) || IsProbablyALinuxExe(Filename) ? "0755" : "0644";
 			string CommandString = string.Format("-c 'chmod {0} \"{1}\"'", Permissions, Filename.Replace("'", "'\"'\"'"));
 			var Result = CommandUtils.Run("sh", CommandString, Options:ERunOptions.SpewIsVerbose);
 			if (Result.ExitCode != 0)
diff --git a/Engine/Source/Programs/AutomationTool/AutomationUtils/CommandUtils.cs b/Engine/Source/Programs/AutomationTool/AutomationUtils/CommandUtils.cs
--- a/Engine/Source/Programs/AutomationTool/AutomationUtils/CommandUtils.cs
+++ b/Engine/Source/Programs/AutomationTool/AutomationUtils/CommandUtils.cs
@@ -2259,7 +2259,7 @@ namespace AutomationTool
 						CommandUtils.LogInformation(FilePath);
 						if (File.Exists(FilePath) && !OutputFileNames.Contains(FilePath) && FilePath != ZipFileName)
 						{
-							if (CommandUtils.IsProbablyAMacOrIOSExe(FilePath))
+							if (CommandUtils.IsProbablyAMacOrIOSExe(FilePath) || CommandUtils.IsProbablyALinuxExe(FilePath))
 							{
 								FixUnixFilePermissions(FilePath);
 							}
