Add support for RHI DX11, DX12 and Vulkan to EngineTest buildgraph
Introduce Test configs to EngineTest buildgraph in order to improve test schedule readibility and maintenance by platform.
Add support for scheduling EngineTest by test priority
Schedule Mac and Win64+Console separatlly

#rnx

diff --git a/Engine/Source/Programs/AutomationTool/Gauntlet/Unreal/Automation/UE.Automation.cs b/Engine/Source/Programs/AutomationTool/Gauntlet/Unreal/Automation/UE.Automation.cs
--- a/Engine/Source/Programs/AutomationTool/Gauntlet/Unreal/Automation/UE.Automation.cs
+++ b/Engine/Source/Programs/AutomationTool/Gauntlet/Unreal/Automation/UE.Automation.cs
@@ -19,10 +19,26 @@ namespace UE
 	public class AutomationTestConfig : UnrealTestConfiguration
 	{
 		/// <summary>
-		/// Run with d3d12 RHI
+		/// Run with specify RHI
 		/// </summary>
 		[AutoParam]
-		public bool D3D12 = false;
+		public string RHI = "";
+		/// <summary>
+		/// Valid RHI names based on Platform
+		/// </summary>
+		private static class ValidRHI
+		{
+			public enum Win64
+			{
+				d3d11,
+				d3d12,
+				vulkan
+			}
+			public enum Linux
+			{
+				vulkan
+			}
+		}
 
 		/// <summary>
 		/// Run with Nvidia cards for raytracing support
@@ -222,9 +238,16 @@ namespace UE
 						AppConfig.CommandLine += " -preferNvidia";
 					}
 
-					if (D3D12)
+					if (!string.IsNullOrEmpty(RHI))
 					{
-						AppConfig.CommandLine += " -d3d12";
+						if (Enum.IsDefined(typeof(ValidRHI.Win64), RHI.ToLower()))
+						{
+							AppConfig.CommandLine += string.Format(" -{0}", RHI);
+						}
+						else
+						{
+							throw new AutomationException(string.Format("Unknown RHI target '{0}' for Win64", RHI));
+						}
 					}
 
 					if (D3DDebug)
@@ -567,6 +590,11 @@ namespace UE
 						if (!string.IsNullOrEmpty(ReportPath))
 						{
 							Report = CreateUnrealEngineTestPassReport(ReportPath, Config.ReportURL);
+							if (Report != null)
+							{
+								var MainRolePlatform = Context.GetRoleContext(Config.GetMainRequiredRole().Type).Platform;
+								Report.SetMetadata("RHI", string.IsNullOrEmpty(Config.RHI) || MainRolePlatform != UnrealTargetPlatform.Win64 ? "default" : Config.RHI.ToLower());
+							}
 						}
 					}
 				}
diff --git a/Engine/Source/Programs/AutomationTool/Gauntlet/Unreal/Automation/UE.Networking.cs b/Engine/Source/Programs/AutomationTool/Gauntlet/Unreal/Automation/UE.Networking.cs
new file mode 100644
--- /dev/null
+++ b/Engine/Source/Programs/AutomationTool/Gauntlet/Unreal/Automation/UE.Networking.cs
@@ -0,0 +1,102 @@
+// Copyright Epic Games, Inc. All Rights Reserved.
+using System;
+using System.Collections.Generic;
+using System.Linq;
+using Gauntlet;
+using EpicGame;
+using AutomationTool;
+using UnrealBuildTool;
+
+namespace UE
+{
+	namespace NetworkAutomation
+	{
+		public class DedicatedServer : UnrealTestNode<EpicGameTestConfig>
+		{
+			public DedicatedServer(UnrealTestContext InContext) : base(InContext)
+			{
+			}
+
+			public override EpicGameTestConfig GetConfiguration()
+			{
+				EpicGameTestConfig Config = base.GetConfiguration();
+				Config.PreAssignAccount = false;
+				Config.NoMCP = true;
+
+				IEnumerable<UnrealTestRole> Clients = Config.RequireRoles(UnrealTargetRole.Client, 2);
+				UnrealTestRole Server = Config.RequireRole(UnrealTargetRole.Server);
+
+				Server.Controllers.Add("NetTestGauntletServerController");
+				Server.CommandLine += " -log";
+				Clients.ElementAt(0).CommandLine += " -log";
+				Clients.ElementAt(0).Controllers.Add("NetTestGauntletClientController");
+				Clients.ElementAt(1).CommandLine += " -log";
+				Clients.ElementAt(1).Controllers.Add("NetTestGauntletClientController");
+
+				return Config;
+			}
+		}
+
+		public class ListenServerConfig : UnrealGame.UE4TestConfig
+		{
+			string ListenServerIP = "";
+
+			public override void ApplyToConfig(UnrealAppConfig AppConfig, UnrealSessionRole ConfigRole, IEnumerable<UnrealSessionRole> OtherRoles)
+			{
+				base.ApplyToConfig(AppConfig, ConfigRole, OtherRoles);
+
+				if (ConfigRole.RoleType.IsClient())
+				{
+					if (ListenServerIP.Length == 0)
+					{
+						if (AppConfig.Platform == UnrealTargetPlatform.Win64)
+						{
+							var LocalAddress = UnrealHelpers.GetHostIpAddress();
+
+							if (LocalAddress == null)
+							{
+								throw new AutomationException("Could not find local IP address");
+							}
+
+							string RequestedClientIP = Globals.Params.ParseValue("clientip", "");
+							string ClientIP = string.IsNullOrEmpty(RequestedClientIP) ? LocalAddress.ToString() : RequestedClientIP;
+
+							// client too, but only desktop platforms
+							if (!string.IsNullOrEmpty(RequestedClientIP))
+							{
+								AppConfig.CommandLine += string.Format(" -multihome={0}", ClientIP);
+							}
+
+							ListenServerIP = ClientIP;
+						}
+					}
+					else
+					{
+						AppConfig.CommandLine += string.Format(" -ExecCmds=\"open {0}\" -log", ListenServerIP);
+					}
+				}
+			}
+		}
+
+		public class ListenServer : UnrealTestNode<ListenServerConfig>
+		{
+			public ListenServer(UnrealTestContext InContext) : base(InContext)
+			{
+			}
+
+			public override ListenServerConfig GetConfiguration()
+			{
+				ListenServerConfig Config = base.GetConfiguration();
+
+				IEnumerable<UnrealTestRole> Clients = Config.RequireRoles(UnrealTargetRole.Client, 3);
+
+				Clients.ElementAt(0).Controllers.Add("NetTestGauntletServerController");
+				Clients.ElementAt(0).CommandLine += string.Format(" -ExecCmds=\"open {0}?Listen\" -log", Config.Map);
+				Clients.ElementAt(1).Controllers.Add("NetTestGauntletClientController");
+				Clients.ElementAt(2).Controllers.Add("NetTestGauntletClientController");
+
+				return Config;
+			}
+		}
+	}
+}
\ No newline at end of file
