Remove SubmitToDashboard unecessary warning

#rnx

diff --git a/Engine/Source/Programs/AutomationTool/Gauntlet/Unreal/Base/Gauntlet.UnrealTestNode.cs b/Engine/Source/Programs/AutomationTool/Gauntlet/Unreal/Base/Gauntlet.UnrealTestNode.cs
--- a/Engine/Source/Programs/AutomationTool/Gauntlet/Unreal/Base/Gauntlet.UnrealTestNode.cs
+++ b/Engine/Source/Programs/AutomationTool/Gauntlet/Unreal/Base/Gauntlet.UnrealTestNode.cs
@@ -995,11 +995,10 @@ namespace Gauntlet
 
 			try
 			{
-				// Check if the deprecated signature is overriden, call it anyway if that the case, and trigger a warning.
+				// Check if the deprecated signature is overriden, call it anyway if that the case.
 				var DeprecatedSignature = new[] { typeof(TestResult), typeof(UnrealTestContext), typeof(UnrealBuildSource), typeof(IEnumerable<UnrealRoleArtifacts>), typeof(string) };
 				if (Utils.InterfaceHelpers.HasOverriddenMethod(this.GetType(), "SubmitToDashboard", DeprecatedSignature))
 				{
-					Log.Warning("This SubmitToDashboard signature is deprecated, please use the overloaded version.");
 					SubmitToDashboard(GetTestResult(), Context, Context.BuildInfo, SessionArtifacts, ArtifactPath);
 				}
 
