BuildGraph: Add Docker-Compose-Up and Docker-Compose-Down tasks.

Also disable ANSI escape sequences while running any docker commands, and allow passing an environment file into <DotNet> tasks.

diff --git a/Engine/Source/Programs/AutomationTool/BuildGraph/Tasks/DockerBuildTask.cs b/Engine/Source/Programs/AutomationTool/BuildGraph/Tasks/DockerBuildTask.cs
--- a/Engine/Source/Programs/AutomationTool/BuildGraph/Tasks/DockerBuildTask.cs
+++ b/Engine/Source/Programs/AutomationTool/BuildGraph/Tasks/DockerBuildTask.cs
@@ -114,7 +114,7 @@ namespace BuildGraph.Tasks
 					}
 				}
 
-				StringBuilder Arguments = new StringBuilder("build .");
+				StringBuilder Arguments = new StringBuilder("--no-ansi build .");
 				if (Parameters.Tag != null)
 				{
 					Arguments.Append($" -t {Parameters.Tag}");
diff --git a/Engine/Source/Programs/AutomationTool/BuildGraph/Tasks/DockerComposeDownTask.cs b/Engine/Source/Programs/AutomationTool/BuildGraph/Tasks/DockerComposeDownTask.cs
new file mode 100644
--- /dev/null
+++ b/Engine/Source/Programs/AutomationTool/BuildGraph/Tasks/DockerComposeDownTask.cs
@@ -0,0 +1,105 @@
+// Copyright Epic Games, Inc. All Rights Reserved.
+
+using EpicGames.Core;
+using System;
+using System.Collections.Generic;
+using System.IO;
+using System.Linq;
+using System.Text;
+using System.Xml;
+using AutomationTool;
+using UnrealBuildBase;
+
+namespace BuildGraph.Tasks
+{
+	/// <summary>
+	/// Parameters for a Docker-Compose task
+	/// </summary>
+	public class DockerComposeDownTaskParameters
+	{
+		/// <summary>
+		/// Path to the docker-compose file
+		/// </summary>
+		[TaskParameter]
+		public string File;
+
+		/// <summary>
+		/// Arguments for the command
+		/// </summary>
+		[TaskParameter(Optional = true)]
+		public string Arguments;
+	}
+
+	/// <summary>
+	/// Spawns Docker and waits for it to complete.
+	/// </summary>
+	[TaskElement("Docker-Compose-Down", typeof(DockerComposeDownTaskParameters))]
+	public class DockerComposeDownTask : SpawnTaskBase
+	{
+		/// <summary>
+		/// Parameters for this task
+		/// </summary>
+		DockerComposeDownTaskParameters Parameters;
+
+		/// <summary>
+		/// Construct a Docker-Compose task
+		/// </summary>
+		/// <param name="InParameters">Parameters for the task</param>
+		public DockerComposeDownTask(DockerComposeDownTaskParameters InParameters)
+		{
+			Parameters = InParameters;
+		}
+
+		/// <summary>
+		/// Execute the task.
+		/// </summary>
+		/// <param name="Job">Information about the current job</param>
+		/// <param name="BuildProducts">Set of build products produced by this node.</param>
+		/// <param name="TagNameToFileSet">Mapping from tag names to the set of files they include</param>
+		public override void Execute(JobContext Job, HashSet<FileReference> BuildProducts, Dictionary<string, HashSet<FileReference>> TagNameToFileSet)
+		{
+			StringBuilder Arguments = new StringBuilder("--no-ansi ");
+			if (!String.IsNullOrEmpty(Parameters.File))
+			{
+				Arguments.Append($"--file {Parameters.File.QuoteArgument()} ");
+			}
+			Arguments.Append("down");
+			if (!String.IsNullOrEmpty(Parameters.Arguments))
+			{
+				Arguments.Append($" {Parameters.Arguments}");
+			}
+
+			Log.TraceInformation("Running docker compose {0}", Arguments.ToString());
+			using (LogIndentScope Scope = new LogIndentScope("  "))
+			{
+				SpawnTaskBase.Execute("docker-compose", Arguments.ToString());
+			}
+		}
+
+		/// <summary>
+		/// Output this task out to an XML writer.
+		/// </summary>
+		public override void Write(XmlWriter Writer)
+		{
+			Write(Writer, Parameters);
+		}
+
+		/// <summary>
+		/// Find all the tags which are used as inputs to this task
+		/// </summary>
+		/// <returns>The tag names which are read by this task</returns>
+		public override IEnumerable<string> FindConsumedTagNames()
+		{
+			return Enumerable.Empty<string>();
+		}
+
+		/// <summary>
+		/// Find all the tags which are modified by this task
+		/// </summary>
+		/// <returns>The tag names which are modified by this task</returns>
+		public override IEnumerable<string> FindProducedTagNames()
+		{
+			return Enumerable.Empty<string>();
+		}
+	}
+}
diff --git a/Engine/Source/Programs/AutomationTool/BuildGraph/Tasks/DockerComposeUpTask.cs b/Engine/Source/Programs/AutomationTool/BuildGraph/Tasks/DockerComposeUpTask.cs
new file mode 100644
--- /dev/null
+++ b/Engine/Source/Programs/AutomationTool/BuildGraph/Tasks/DockerComposeUpTask.cs
@@ -0,0 +1,105 @@
+// Copyright Epic Games, Inc. All Rights Reserved.
+
+using EpicGames.Core;
+using System;
+using System.Collections.Generic;
+using System.IO;
+using System.Linq;
+using System.Text;
+using System.Xml;
+using AutomationTool;
+using UnrealBuildBase;
+
+namespace BuildGraph.Tasks
+{
+	/// <summary>
+	/// Parameters for a Docker-Compose task
+	/// </summary>
+	public class DockerComposeUpTaskParameters
+	{
+		/// <summary>
+		/// Path to the docker-compose file
+		/// </summary>
+		[TaskParameter]
+		public string File;
+
+		/// <summary>
+		/// Arguments for the command
+		/// </summary>
+		[TaskParameter(Optional = true)]
+		public string Arguments;
+	}
+
+	/// <summary>
+	/// Spawns Docker and waits for it to complete.
+	/// </summary>
+	[TaskElement("Docker-Compose-Up", typeof(DockerComposeUpTaskParameters))]
+	public class DockerComposeUpTask : SpawnTaskBase
+	{
+		/// <summary>
+		/// Parameters for this task
+		/// </summary>
+		DockerComposeUpTaskParameters Parameters;
+
+		/// <summary>
+		/// Construct a Docker-Compose task
+		/// </summary>
+		/// <param name="InParameters">Parameters for the task</param>
+		public DockerComposeUpTask(DockerComposeUpTaskParameters InParameters)
+		{
+			Parameters = InParameters;
+		}
+
+		/// <summary>
+		/// Execute the task.
+		/// </summary>
+		/// <param name="Job">Information about the current job</param>
+		/// <param name="BuildProducts">Set of build products produced by this node.</param>
+		/// <param name="TagNameToFileSet">Mapping from tag names to the set of files they include</param>
+		public override void Execute(JobContext Job, HashSet<FileReference> BuildProducts, Dictionary<string, HashSet<FileReference>> TagNameToFileSet)
+		{
+			StringBuilder Arguments = new StringBuilder("--no-ansi ");
+			if (!String.IsNullOrEmpty(Parameters.File))
+			{
+				Arguments.Append($"--file {Parameters.File.QuoteArgument()} ");
+			}
+			Arguments.Append("up --detach");
+			if (!String.IsNullOrEmpty(Parameters.Arguments))
+			{
+				Arguments.Append($" {Parameters.Arguments}");
+			}
+
+			Log.TraceInformation("Running docker compose {0}", Arguments.ToString());
+			using (LogIndentScope Scope = new LogIndentScope("  "))
+			{
+				SpawnTaskBase.Execute("docker-compose", Arguments.ToString());
+			}
+		}
+
+		/// <summary>
+		/// Output this task out to an XML writer.
+		/// </summary>
+		public override void Write(XmlWriter Writer)
+		{
+			Write(Writer, Parameters);
+		}
+
+		/// <summary>
+		/// Find all the tags which are used as inputs to this task
+		/// </summary>
+		/// <returns>The tag names which are read by this task</returns>
+		public override IEnumerable<string> FindConsumedTagNames()
+		{
+			return Enumerable.Empty<string>();
+		}
+
+		/// <summary>
+		/// Find all the tags which are modified by this task
+		/// </summary>
+		/// <returns>The tag names which are modified by this task</returns>
+		public override IEnumerable<string> FindProducedTagNames()
+		{
+			return Enumerable.Empty<string>();
+		}
+	}
+}
diff --git a/Engine/Source/Programs/AutomationTool/BuildGraph/Tasks/DockerPushTask.cs b/Engine/Source/Programs/AutomationTool/BuildGraph/Tasks/DockerPushTask.cs
--- a/Engine/Source/Programs/AutomationTool/BuildGraph/Tasks/DockerPushTask.cs
+++ b/Engine/Source/Programs/AutomationTool/BuildGraph/Tasks/DockerPushTask.cs
@@ -90,12 +90,12 @@ namespace BuildGraph.Tasks
 				if (Parameters.AwsEcr)
 				{
 					IProcessResult Result = SpawnTaskBase.Execute("aws", "ecr get-login-password", EnvVars: Environment, LogOutput: false);
-					Execute("docker", $"login {Parameters.Repository} --username AWS --password-stdin", Input: Result.Output);
+					Execute("docker", $"--no-ansi login {Parameters.Repository} --username AWS --password-stdin", Input: Result.Output);
 				}
 
 				string TargetImage = Parameters.TargetImage ?? Parameters.Image;
-				Execute("docker", $"tag {Parameters.Image} {Parameters.Repository}/{TargetImage}", EnvVars: Environment);
-				Execute("docker", $"push {Parameters.Repository}/{TargetImage}", EnvVars: Environment);
+				Execute("docker", $"--no-ansi tag {Parameters.Image} {Parameters.Repository}/{TargetImage}", EnvVars: Environment);
+				Execute("docker", $"--no-ansi push {Parameters.Repository}/{TargetImage}", EnvVars: Environment);
 			}
 		}
 
diff --git a/Engine/Source/Programs/AutomationTool/BuildGraph/Tasks/DotNetTask.cs b/Engine/Source/Programs/AutomationTool/BuildGraph/Tasks/DotNetTask.cs
--- a/Engine/Source/Programs/AutomationTool/BuildGraph/Tasks/DotNetTask.cs
+++ b/Engine/Source/Programs/AutomationTool/BuildGraph/Tasks/DotNetTask.cs
@@ -1,5 +1,6 @@
 // Copyright Epic Games, Inc. All Rights Reserved.
 
+using AutomationTool;
 using EpicGames.Core;
 using System;
 using System.Collections.Generic;
@@ -7,7 +8,7 @@ using System.IO;
 using System.Text;
 using System.Xml;
 
-namespace AutomationTool.Tasks
+namespace BuildGraph.Tasks
 {
 	/// <summary>
 	/// Parameters for a DotNet task
@@ -26,6 +27,18 @@ namespace AutomationTool.Tasks
 		[TaskParameter(Optional = true)]
 		public string BaseDir;
 
+		/// <summary>
+		/// Environment variables to set
+		/// </summary>
+		[TaskParameter(Optional = true)]
+		public string Environment;
+
+		/// <summary>
+		/// File to read environment variables from
+		/// </summary>
+		[TaskParameter(Optional = true)]
+		public string EnvironmentFile;
+
 		/// <summary>
 		/// The minimum exit code, which is treated as an error.
 		/// </summary>
@@ -37,7 +50,7 @@ namespace AutomationTool.Tasks
 	/// Spawns Docker and waits for it to complete.
 	/// </summary>
 	[TaskElement("DotNet", typeof(DotNetTaskParameters))]
-	public class DotNetTask : CustomTask
+	public class DotNetTask : SpawnTaskBase
 	{
 		/// <summary>
 		/// Parameters for this task
@@ -67,7 +80,7 @@ namespace AutomationTool.Tasks
 				throw new AutomationException("DotNet is missing from {0}", DotNetFile);
 			}
 
-			IProcessResult Result = CommandUtils.Run(DotNetFile.FullName, Parameters.Arguments, WorkingDir: Parameters.BaseDir);
+			IProcessResult Result = Execute(DotNetFile.FullName, Parameters.Arguments, WorkingDir: Parameters.BaseDir, EnvVars: ParseEnvVars(Parameters.Environment, Parameters.EnvironmentFile));
 			if (Result.ExitCode < 0 || Result.ExitCode >= Parameters.ErrorLevel)
 			{
 				throw new AutomationException("Docker terminated with an exit code indicating an error ({0})", Result.ExitCode);
