Fix all Unreal Gauntlet failing tests reporting: Failed due to undiagnosed reasons.
But it is diagnosed...

#rnx

diff --git a/Engine/Source/Programs/AutomationTool/Gauntlet/Unreal/Base/Gauntlet.UnrealTestNode.cs b/Engine/Source/Programs/AutomationTool/Gauntlet/Unreal/Base/Gauntlet.UnrealTestNode.cs
--- a/Engine/Source/Programs/AutomationTool/Gauntlet/Unreal/Base/Gauntlet.UnrealTestNode.cs
+++ b/Engine/Source/Programs/AutomationTool/Gauntlet/Unreal/Base/Gauntlet.UnrealTestNode.cs
@@ -1462,7 +1462,7 @@ namespace Gauntlet
 			}*/
 			
 			// The process is gone but we don't know why. This is likely bad and signifies an unhandled or undiagnosed error
-						ExitReason = "app exited with code 0";
+			ExitReason = "app exited with code 0";
 			ExitCode = -1;
 			return UnrealProcessResult.Unknown;
 		}
@@ -1991,7 +1991,7 @@ namespace Gauntlet
 				string.Format("Result: {0}", GetTestResult())
 			});
 
-			if (TestFailed && GetRolesThatExitedAbnormally().Any() == false)
+			if (TestFailed && GetRolesThatFailed().Where(R => R.ProcessResult == UnrealProcessResult.Unknown).Any())
 			{
 				MB.Paragraph(string.Format("Error: {0} failed due to undiagnosed reasons", this.Name));
 				MB.Paragraph("See 'Role Report' below for more details on each role");
