Remove a some platform references in Gauntlet
#rnx

diff --git a/Engine/Source/Programs/AutomationTool/Gauntlet/Unreal/Base/Gauntlet.UnrealTestContext.cs b/Engine/Source/Programs/AutomationTool/Gauntlet/Unreal/Base/Gauntlet.UnrealTestContext.cs
--- a/Engine/Source/Programs/AutomationTool/Gauntlet/Unreal/Base/Gauntlet.UnrealTestContext.cs
+++ b/Engine/Source/Programs/AutomationTool/Gauntlet/Unreal/Base/Gauntlet.UnrealTestContext.cs
@@ -309,7 +309,6 @@ namespace Gauntlet
 			// do we have any tests? Need to check the global test list
 			bool HaveTests = TestList.Count > 0 || PlatformList.Where(Plat => Plat.ParseValues("test").Count() > 0).Count() > 0;
 
-			// turn -device=BobsKit,BobsKit(PS4) into a device list
 			List<string> DeviceArgStrings = Params.ParseValues("device=");
 
 			if (DeviceArgStrings.Count == 0)
diff --git a/Engine/Source/Programs/AutomationTool/Gauntlet/Unreal/RunUnreal.cs b/Engine/Source/Programs/AutomationTool/Gauntlet/Unreal/RunUnreal.cs
--- a/Engine/Source/Programs/AutomationTool/Gauntlet/Unreal/RunUnreal.cs
+++ b/Engine/Source/Programs/AutomationTool/Gauntlet/Unreal/RunUnreal.cs
@@ -168,7 +168,7 @@ namespace Gauntlet
 					Role.Platform = DefaultPlatform;
 					Role.Configuration = DefaultConfiguration;
 
-					// globally, what was requested (e.g -platform=PS4 -configuration=Shipping)
+					// globally, what was requested (e.g -platform=Win64 -configuration=Shipping)
 					UnrealTargetPlatform RequestedPlatform = PlatformType;
 					UnrealTargetConfiguration RequestedConfiguration = ContextOptions.Configuration;
 
diff --git a/Engine/Source/Programs/AutomationTool/Gauntlet/Unreal/Tasks/Gauntlet.Task.CleanDevices.cs b/Engine/Source/Programs/AutomationTool/Gauntlet/Unreal/Tasks/Gauntlet.Task.CleanDevices.cs
--- a/Engine/Source/Programs/AutomationTool/Gauntlet/Unreal/Tasks/Gauntlet.Task.CleanDevices.cs
+++ b/Engine/Source/Programs/AutomationTool/Gauntlet/Unreal/Tasks/Gauntlet.Task.CleanDevices.cs
@@ -94,11 +94,6 @@ namespace Gauntlet
 				return;
 			}
 
-			/*if (Device is TargetDevicePS4)
-			{
-				CleanPS4(Device as TargetDevicePS4);
-			}*/
-
 			// disconnect and power down
 			Gauntlet.Log.Info("Powering down and disconnecting from {0}", Device.Name);
 			Device.Disconnect();
diff --git a/Engine/Source/Programs/AutomationTool/Gauntlet/Unreal/Utils/Gauntlet.UnrealLogParser.cs b/Engine/Source/Programs/AutomationTool/Gauntlet/Unreal/Utils/Gauntlet.UnrealLogParser.cs
--- a/Engine/Source/Programs/AutomationTool/Gauntlet/Unreal/Utils/Gauntlet.UnrealLogParser.cs
+++ b/Engine/Source/Programs/AutomationTool/Gauntlet/Unreal/Utils/Gauntlet.UnrealLogParser.cs
@@ -792,7 +792,7 @@ namespace Gauntlet
 
 			// Now, because platforms sometimes dump asserts to the log and low-level logging, and we might have a post-mortem stack, we
 			// need to prune out redundancies. Basic approach - find errors with the same assert message and keep the one with the
-			// longest callstack. If we have a post-mortem error, overwrite the previous trace with its info (on PS4 the post-mortem
+			// longest callstack. If we have a post-mortem error, overwrite the previous trace with its info (on some devices the post-mortem
 			// info is way more informative).
 
 			List<UnrealLog.CallstackMessage> FilteredTraces = new List<UnrealLog.CallstackMessage>();
