Horde: Remove --no-ansi flag for docker buld task
#rnx

diff --git a/Engine/Source/Programs/AutomationTool/BuildGraph/Tasks/DockerBuildTask.cs b/Engine/Source/Programs/AutomationTool/BuildGraph/Tasks/DockerBuildTask.cs
--- a/Engine/Source/Programs/AutomationTool/BuildGraph/Tasks/DockerBuildTask.cs
+++ b/Engine/Source/Programs/AutomationTool/BuildGraph/Tasks/DockerBuildTask.cs
@@ -114,7 +114,7 @@ namespace BuildGraph.Tasks
 					}
 				}
 
-				StringBuilder Arguments = new StringBuilder("--no-ansi build .");
+				StringBuilder Arguments = new StringBuilder("build .");
 				if (Parameters.Tag != null)
 				{
 					Arguments.Append($" -t {Parameters.Tag}");
