Horde: Use --ansi never instead of since removed --no-ansi
#rnx

diff --git a/Engine/Source/Programs/AutomationTool/BuildGraph/Tasks/DockerBuildTask.cs b/Engine/Source/Programs/AutomationTool/BuildGraph/Tasks/DockerBuildTask.cs
--- a/Engine/Source/Programs/AutomationTool/BuildGraph/Tasks/DockerBuildTask.cs
+++ b/Engine/Source/Programs/AutomationTool/BuildGraph/Tasks/DockerBuildTask.cs
@@ -114,7 +114,7 @@ namespace BuildGraph.Tasks
 					}
 				}
 
-				StringBuilder Arguments = new StringBuilder("build .");
+				StringBuilder Arguments = new StringBuilder("--ansi never build .");
 				if (Parameters.Tag != null)
 				{
 					Arguments.Append($" -t {Parameters.Tag}");
