Horde: Change to --ansi never in compose up and down as well
#rnx

diff --git a/Engine/Source/Programs/AutomationTool/BuildGraph/Tasks/DockerComposeDownTask.cs b/Engine/Source/Programs/AutomationTool/BuildGraph/Tasks/DockerComposeDownTask.cs
--- a/Engine/Source/Programs/AutomationTool/BuildGraph/Tasks/DockerComposeDownTask.cs
+++ b/Engine/Source/Programs/AutomationTool/BuildGraph/Tasks/DockerComposeDownTask.cs
@@ -58,7 +58,7 @@ namespace BuildGraph.Tasks
 		/// <param name="TagNameToFileSet">Mapping from tag names to the set of files they include</param>
 		public override void Execute(JobContext Job, HashSet<FileReference> BuildProducts, Dictionary<string, HashSet<FileReference>> TagNameToFileSet)
 		{
-			StringBuilder Arguments = new StringBuilder("--no-ansi ");
+			StringBuilder Arguments = new StringBuilder("--ansi never ");
 			if (!String.IsNullOrEmpty(Parameters.File))
 			{
 				Arguments.Append($"--file {Parameters.File.QuoteArgument()} ");
diff --git a/Engine/Source/Programs/AutomationTool/BuildGraph/Tasks/DockerComposeUpTask.cs b/Engine/Source/Programs/AutomationTool/BuildGraph/Tasks/DockerComposeUpTask.cs
--- a/Engine/Source/Programs/AutomationTool/BuildGraph/Tasks/DockerComposeUpTask.cs
+++ b/Engine/Source/Programs/AutomationTool/BuildGraph/Tasks/DockerComposeUpTask.cs
@@ -58,7 +58,7 @@ namespace BuildGraph.Tasks
 		/// <param name="TagNameToFileSet">Mapping from tag names to the set of files they include</param>
 		public override void Execute(JobContext Job, HashSet<FileReference> BuildProducts, Dictionary<string, HashSet<FileReference>> TagNameToFileSet)
 		{
-			StringBuilder Arguments = new StringBuilder("--no-ansi ");
+			StringBuilder Arguments = new StringBuilder("--ansi never ");
 			if (!String.IsNullOrEmpty(Parameters.File))
 			{
 				Arguments.Append($"--file {Parameters.File.QuoteArgument()} ");
