Horde: Remove --no-ansi from docker pushes
#rnx

diff --git a/Engine/Source/Programs/AutomationTool/BuildGraph/Tasks/DockerPushTask.cs b/Engine/Source/Programs/AutomationTool/BuildGraph/Tasks/DockerPushTask.cs
--- a/Engine/Source/Programs/AutomationTool/BuildGraph/Tasks/DockerPushTask.cs
+++ b/Engine/Source/Programs/AutomationTool/BuildGraph/Tasks/DockerPushTask.cs
@@ -90,12 +90,12 @@ namespace BuildGraph.Tasks
 				if (Parameters.AwsEcr)
 				{
 					IProcessResult Result = SpawnTaskBase.Execute("aws", "ecr get-login-password", EnvVars: Environment, LogOutput: false);
-					Execute("docker", $"--no-ansi login {Parameters.Repository} --username AWS --password-stdin", Input: Result.Output);
+					Execute("docker", $"login {Parameters.Repository} --username AWS --password-stdin", Input: Result.Output);
 				}
 
 				string TargetImage = Parameters.TargetImage ?? Parameters.Image;
-				Execute("docker", $"--no-ansi tag {Parameters.Image} {Parameters.Repository}/{TargetImage}", EnvVars: Environment);
-				Execute("docker", $"--no-ansi push {Parameters.Repository}/{TargetImage}", EnvVars: Environment);
+				Execute("docker", $"tag {Parameters.Image} {Parameters.Repository}/{TargetImage}", EnvVars: Environment);
+				Execute("docker", $"push {Parameters.Repository}/{TargetImage}", EnvVars: Environment);
 			}
 		}
 
