Remove platform reference in build graph stage task, originating CL 3187153
#rnx

diff --git a/Engine/Source/Programs/AutomationTool/BuildGraph/Tasks/StageTask.cs b/Engine/Source/Programs/AutomationTool/BuildGraph/Tasks/StageTask.cs
--- a/Engine/Source/Programs/AutomationTool/BuildGraph/Tasks/StageTask.cs
+++ b/Engine/Source/Programs/AutomationTool/BuildGraph/Tasks/StageTask.cs
@@ -147,12 +147,6 @@ namespace BuildGraph.Tasks
 					TargetFile = FileReference.Combine(TargetProjectDir, SourceFile.MakeRelativeTo(SourceProjectDir));
 				}
 
-				// Fixup the case of the output file. Would expect Platform.DeployLowerCaseFilenames() to return true here, but seems not to be the case.
-				if(Parameters.Platform == UnrealTargetPlatform.PS4)
-				{
-					TargetFile = FileReference.Combine(TargetDir, TargetFile.MakeRelativeTo(TargetDir).ToLowerInvariant());
-				}
-
 				// Only copy the output file if it doesn't already exist. We can stage multiple targets to the same output directory.
 				if(Parameters.Overwrite || !FileReference.Exists(TargetFile))
 				{
