Remove platform reference from Gauntlet device reservations
#rnx

diff --git a/Engine/Source/Programs/AutomationTool/Gauntlet/Framework/Devices/Gauntlet.DevicePool.cs b/Engine/Source/Programs/AutomationTool/Gauntlet/Framework/Devices/Gauntlet.DevicePool.cs
--- a/Engine/Source/Programs/AutomationTool/Gauntlet/Framework/Devices/Gauntlet.DevicePool.cs
+++ b/Engine/Source/Programs/AutomationTool/Gauntlet/Framework/Devices/Gauntlet.DevicePool.cs
@@ -569,14 +569,7 @@ namespace Gauntlet
 
 			foreach (string Platform in UnrealTargetPlatform.GetValidPlatformNames())
 			{
-				if (Platform == "PS4" || Platform == "XboxOne")
-				{
-					DeviceMap.Add(UnrealTargetPlatform.Parse(Platform), string.Format("{0}-DevKit", Platform));
-				}
-				else
-				{
-					DeviceMap.Add(UnrealTargetPlatform.Parse(Platform), Platform);
-				}
+				DeviceMap.Add(UnrealTargetPlatform.Parse(Platform), Platform);
 			}
 
 			List<string> Devices = new List<string>();
@@ -649,7 +642,7 @@ namespace Gauntlet
 					DeviceDefinition Def = new DeviceDefinition();
 					Def.Address = Device.IPOrHostName;
 					Def.Name = Device.Name;
-					Def.Platform = DeviceMap.FirstOrDefault(Entry => Entry.Value == Device.Type).Key;
+					Def.Platform = DeviceMap.FirstOrDefault(Entry => Entry.Value == Device.Type.Replace("-DevKit", "", StringComparison.OrdinalIgnoreCase)).Key;
 					Def.DeviceData = Device.DeviceData;
 					Def.Model = string.Empty;
 
@@ -1037,16 +1030,11 @@ namespace Gauntlet
 			var TooFewTotalDevices = RequiredDevices.Where(KP => TotalDeviceTypes[KP.Key] < RequiredDevices[KP.Key]).Select(KP => KP.Key);
 			var TooFewCurrentDevices = RequiredDevices.Where(KP => AvailableDeviceTypes[KP.Key] < RequiredDevices[KP.Key]).Select(KP => KP.Key);
 
-			List<UnrealTargetPlatform> ServicePlatforms = new List<UnrealTargetPlatform>()
-			{
-				UnrealTargetPlatform.PS4, UnrealTargetPlatform.Switch
-			};
 
-			var Devices = TooFewTotalDevices.Concat(TooFewCurrentDevices);
-			var UnsupportedPlatforms = Devices.Where(D => !ServicePlatforms.Contains(D.Platform.Value) && (!D.Platform.Value.ToString().StartsWith("XboxOne")) && (D.Platform.Value.ToString() != "XSX") && (D.Platform.Value.ToString() != "PS5"));
+			var Devices = TooFewTotalDevices.Concat(TooFewCurrentDevices);		
 
 			// Request devices from the service if we need them
-			if (UseServiceDevices && !String.IsNullOrEmpty(DeviceURL) && UnsupportedPlatforms.Count() == 0 && (TooFewTotalDevices.Count() > 0 || TooFewCurrentDevices.Count() > 0))
+			if (UseServiceDevices && !String.IsNullOrEmpty(DeviceURL)  && (TooFewTotalDevices.Count() > 0 || TooFewCurrentDevices.Count() > 0))
 			{				
 
 				Dictionary<UnrealDeviceTargetConstraint, int> DeviceCounts = new Dictionary<UnrealDeviceTargetConstraint, int>();
