Added support for copy_failed_binary_exists error to BPT error handler.

diff --git a/Engine/Source/Programs/AutomationTool/AutomationUtils/MCPPublic.cs b/Engine/Source/Programs/AutomationTool/AutomationUtils/MCPPublic.cs
--- a/Engine/Source/Programs/AutomationTool/AutomationUtils/MCPPublic.cs
+++ b/Engine/Source/Programs/AutomationTool/AutomationUtils/MCPPublic.cs
@@ -1015,6 +1015,7 @@ namespace EpicGames.MCP.Automation
 				Success,
 				Unknown,
 				Copy_Failed,
+				Copy_Failed_Binary_Exists,
 				Num
 			};
 			public BPTErrorCode ResolvedErrorCode;
@@ -1024,6 +1025,7 @@ namespace EpicGames.MCP.Automation
 				ResolvedErrorCode = BPTErrorCode.Success;
 				LogFileLines = null;
 				// try to classify the type of error
+				CommandUtils.LogInformation("Parsing log for build patch tool exception");
 				if (File.Exists(StdOutLogFileName))
 				{
 					LogFileLines = CommandUtils.ReadAllLines(StdOutLogFileName);
@@ -1034,12 +1036,14 @@ namespace EpicGames.MCP.Automation
 						System.Text.RegularExpressions.Match Match = ErrorCodeRegex.Match(Line);
 						if (Match.Success)
 						{
+							
 							// found an error code
 							if (ResolvedErrorCode == BPTErrorCode.Success)
 							{
 								ResolvedErrorCode = BPTErrorCode.Unknown;
 							}
 							string ErrorCodeString = Match.Groups[@"errorCode"].Value;
+							CommandUtils.LogInformation("Found match for build patch tool exception \"{0}\"", ErrorCodeString);
 							if (string.IsNullOrEmpty(ErrorCodeString) == false)
 							{
 								if (Enum.TryParse<BPTErrorCode>(ErrorCodeString, true, out ResolvedErrorCode))
