PR #7683: putting all information to the editor from device (Contributed by [REDACTED_PERSON])

diff --git a/Engine/Source/Programs/AutomationTool/HoloLens/HoloLensPlatform.Automation.cs b/Engine/Source/Programs/AutomationTool/HoloLens/HoloLensPlatform.Automation.cs
--- a/Engine/Source/Programs/AutomationTool/HoloLens/HoloLensPlatform.Automation.cs
+++ b/Engine/Source/Programs/AutomationTool/HoloLens/HoloLensPlatform.Automation.cs
@@ -207,15 +207,12 @@ namespace HoloLens.Automation
 
 		bool ProcessHasExited;
 
-		public HoloLensDevicePortalCreatedProcess(Microsoft.Tools.WindowsDevicePortal.DevicePortal InPortal, string InPackageName, string InFriendlyName, uint InProcId)
+		public HoloLensDevicePortalCreatedProcess(Microsoft.Tools.WindowsDevicePortal.DevicePortal InPortal, string InPackageName, string InFriendlyName)
 		{
 			Portal = InPortal;
-			ProcId = InProcId;
 			PackageName = InPackageName;
 			FriendlyName = InFriendlyName;
 			StateLock = new object();
-			Portal.RunningProcessesMessageReceived += Portal_RunningProcessesMessageReceived;
-			Portal.StartListeningForRunningProcessesAsync().Wait();
 			//MonitorThread = new System.Threading.Thread(()=>(MonitorProcess))
 
 			// No ETW on Xbox One, so can't collect trace that way
@@ -228,6 +225,13 @@ namespace HoloLens.Automation
 			}
 		}
 
+		public void SetProcId(uint InProcId)
+		{
+			ProcId = InProcId;
+			Portal.RunningProcessesMessageReceived += Portal_RunningProcessesMessageReceived;
+			Portal.StartListeningForRunningProcessesAsync().Wait();
+		}
+
 		private void Portal_RealtimeEventsMessageReceived(Microsoft.Tools.WindowsDevicePortal.DevicePortal sender, Microsoft.Tools.WindowsDevicePortal.WebSocketMessageReceivedEventArgs<Microsoft.Tools.WindowsDevicePortal.DevicePortal.EtwEvents> args)
 		{
 			foreach (var EtwEvent in args.Message.Events)
@@ -1615,13 +1619,14 @@ namespace HoloLens.Automation
 					}
 				}
 
+				var Result = new HoloLensDevicePortalCreatedProcess(portal, FullName, Params.ShortProjectName);
 				var LaunchTask = portal.LaunchApplicationAsync(Aumid, FullName);
 
 				// Message back to the UE4 Editor to correctly set the app id for each device
 				Console.WriteLine("Running Package@Device:{0}@{1}", FullName, DeviceAddress);
 
 				LaunchTask.Wait();
-				IProcessResult Result = new HoloLensDevicePortalCreatedProcess(portal, FullName, Params.ShortProjectName, LaunchTask.Result);
+				Result.SetProcId(LaunchTask.Result);
 
 				ProcessManager.AddProcess(Result);
 				if (!ClientRunFlags.HasFlag(ERunOptions.NoWaitForExit))
diff --git a/Engine/Source/Runtime/Core/Private/HoloLens/HoloLensOutputDevicesPrivate.h b/Engine/Source/Runtime/Core/Private/HoloLens/HoloLensOutputDevicesPrivate.h
--- a/Engine/Source/Runtime/Core/Private/HoloLens/HoloLensOutputDevicesPrivate.h
+++ b/Engine/Source/Runtime/Core/Private/HoloLens/HoloLensOutputDevicesPrivate.h
@@ -44,7 +44,9 @@ public:
 		if (EtwLogChannel != nullptr)
 		{
 			Windows::Foundation::Diagnostics::LoggingLevel LogWithLevel = GetWindowsLoggingLevelFromUEVerbosity(Verbosity);
-			if (LogWithLevel >= Windows::Foundation::Diagnostics::LoggingLevel::Warning)
+			// we need to use Information level as minimal for getting all information to the editor
+			// and giving user ability to read the whole meaningful device's log
+			if (LogWithLevel >= Windows::Foundation::Diagnostics::LoggingLevel::Information)
 			{
 				EtwLogChannel->LogMessage(ref new Platform::String(Buffer), LogWithLevel);
 			}
