Follow up to merge CL 16348436 - integrate InProcess property

#rnx

diff --git a/Engine/Source/Programs/AutomationTool/Gauntlet/Framework/Utils/Gauntlet.HordeReport.cs b/Engine/Source/Programs/AutomationTool/Gauntlet/Framework/Utils/Gauntlet.HordeReport.cs
--- a/Engine/Source/Programs/AutomationTool/Gauntlet/Framework/Utils/Gauntlet.HordeReport.cs
+++ b/Engine/Source/Programs/AutomationTool/Gauntlet/Framework/Utils/Gauntlet.HordeReport.cs
@@ -261,6 +261,7 @@ namespace Gauntlet
 			public int SucceededWithWarningsCount { get; set; }
 			public int FailedCount { get; set; }
 			public int NotRunCount { get; set; }
+			public int InProcessCount { get; set; }
 			public float TotalDurationSeconds { get; set; }
 			public List<TestResult> Tests { get; set; }
 
@@ -296,6 +297,7 @@ namespace Gauntlet
 				OutTestPassResults.SucceededWithWarningsCount = InTestPassResults.SucceededWithWarnings;
 				OutTestPassResults.FailedCount = InTestPassResults.Failed;
 				OutTestPassResults.NotRunCount = InTestPassResults.NotRun;
+				OutTestPassResults.InProcessCount = InTestPassResults.InProcess;
 				OutTestPassResults.TotalDurationSeconds = InTestPassResults.TotalDuration;
 				if (InTestPassResults.Tests != null)
 				{
diff --git a/Engine/Source/Programs/AutomationTool/Gauntlet/Unreal/Utils/Gauntlet.UnrealTestResults.cs b/Engine/Source/Programs/AutomationTool/Gauntlet/Unreal/Utils/Gauntlet.UnrealTestResults.cs
--- a/Engine/Source/Programs/AutomationTool/Gauntlet/Unreal/Utils/Gauntlet.UnrealTestResults.cs
+++ b/Engine/Source/Programs/AutomationTool/Gauntlet/Unreal/Utils/Gauntlet.UnrealTestResults.cs
@@ -52,6 +52,7 @@ namespace Gauntlet
 		public int SucceededWithWarnings { get; set; }
 		public int Failed { get; set; }
 		public int NotRun { get; set; }
+		public int InProcess { get; set; }
 		public float TotalDuration { get; set; }
 		public bool ComparisonExported { get; set; }
 		public string ComparisonExportDirectory { get; set; }
