Gauntlet - Improve error message when no gauntlet controller test exitcode

#rnx

diff --git a/Engine/Source/Programs/AutomationTool/Gauntlet/Unreal/Base/Gauntlet.UnrealTestNode.cs b/Engine/Source/Programs/AutomationTool/Gauntlet/Unreal/Base/Gauntlet.UnrealTestNode.cs
--- a/Engine/Source/Programs/AutomationTool/Gauntlet/Unreal/Base/Gauntlet.UnrealTestNode.cs
+++ b/Engine/Source/Programs/AutomationTool/Gauntlet/Unreal/Base/Gauntlet.UnrealTestNode.cs
@@ -1266,7 +1266,7 @@ namespace Gauntlet
 					{
 						Log.Verbose("Role {0} had 0 exit code but used Gauntlet and no TestExitCode was found. Assuming failure", InArtifacts.SessionRole.RoleType);
 						ExitCode = -1;
-						ExitReason = "No test result from Gauntlet controller";
+						ExitReason = "Process has terminated prematurely! No exit code from Gauntlet controller.";
 					}
 				}
 				else
@@ -1275,7 +1275,7 @@ namespace Gauntlet
 					ExitCode = InArtifacts.AppInstance.ExitCode;
 					if (ExitCode == 0)
 					{
-						ExitReason = "app exited with code 0";
+						ExitReason = "Process exited with code 0";
 					}
 				}
 			}
