Gauntlet - add overall test pass execution failure in case of abnormal exit

#rnx

diff --git a/Engine/Source/Programs/AutomationTool/Gauntlet/Unreal/Automation/UE.Automation.cs b/Engine/Source/Programs/AutomationTool/Gauntlet/Unreal/Automation/UE.Automation.cs
--- a/Engine/Source/Programs/AutomationTool/Gauntlet/Unreal/Automation/UE.Automation.cs
+++ b/Engine/Source/Programs/AutomationTool/Gauntlet/Unreal/Automation/UE.Automation.cs
@@ -483,12 +483,6 @@ namespace UE
 		{
 			MarkdownBuilder MB = new MarkdownBuilder(base.GetTestSummaryHeader());
 
-			// If there were abnormal exits then don't add any custom summary. It just confuses things.
-			if (GetArtifactsThatExitedAbnormally().Any())
-			{
-				return MB.ToString();
-			}
-
 			// Everything we need is in the editor artifacts
 			var EditorArtifacts = SessionArtifacts.Where(A => A.SessionRole.RoleType == UnrealTargetRole.Editor).FirstOrDefault();
 
@@ -503,6 +497,27 @@ namespace UE
 				IEnumerable<AutomationTestResult> FailedTests = AllTests.Where(R => R.Completed && !R.Passed);
 				IEnumerable<AutomationTestResult> TestsWithWarnings = AllTests.Where(R => R.Completed && R.Passed && R.WarningEvents.Any());
 
+				// If there were abnormal exits then look only at the incomplete tests to avoid confusing things.
+				if (GetArtifactsThatExitedAbnormally().Any())
+				{
+					if (AllTests.Count() == 0)
+					{
+						MB.H3("No tests were executed.");
+					}
+					else if (IncompleteTests.Count() > 0)
+					{
+						MB.H3("The following test(s) were incomplete:");
+
+						foreach (AutomationTestResult Result in IncompleteTests)
+						{
+							MB.H4(string.Format("{0}", Result.FullName));
+							MB.UnorderedList(Result.WarningAndErrorEvents.Distinct());
+						}
+					}
+
+					return MB.ToString();
+				}
+
 				if (AllTests.Count() == 0)
 				{
 					MB.H3("No tests were executed.");
