Gauntlet - report html report links even if test pass did not fully complete
+ treat timeout messaging as an error.

#rnx

diff --git a/Engine/Source/Programs/AutomationTool/Gauntlet/Framework/Gauntlet.TestExecutor.cs b/Engine/Source/Programs/AutomationTool/Gauntlet/Framework/Gauntlet.TestExecutor.cs
--- a/Engine/Source/Programs/AutomationTool/Gauntlet/Framework/Gauntlet.TestExecutor.cs
+++ b/Engine/Source/Programs/AutomationTool/Gauntlet/Framework/Gauntlet.TestExecutor.cs
@@ -528,14 +528,11 @@ namespace Gauntlet
 			// display the final result
 			if (TestInfo.FinalResult != TestResult.Passed)
 			{
-				string Cause = TestInfo.FinalResult.ToString();
-
+				Log.Info("{0} result={1}", TestInfo, TestInfo.FinalResult);
 				if (string.IsNullOrEmpty(TestInfo.CancellationReason) == false)
 				{
-					Cause = TestInfo.CancellationReason;
-				}
-
-				Log.Info("{0} result={1}", TestInfo, Cause);
+					Log.Info("\tReason: {0}", TestInfo.CancellationReason);
+				}				
 			}
 			else
 			{
@@ -650,7 +647,7 @@ namespace Gauntlet
 				{
 					TestInfo.CancellationReason = string.Format("Terminating Test {0} due to maximum duration of {1} seconds. ", TestInfo.TestNode, TestInfo.TestNode.MaxDuration);
 					TestInfo.FinalResult = TestResult.TimedOut;
-					Log.Info("{0}", TestInfo.CancellationReason);
+					Log.Error("{0}", TestInfo.CancellationReason);
 				}
 				else if (TestInfo.TestNode.MaxDurationReachedResult == EMaxDurationReachedResult.Success)
 				{
diff --git a/Engine/Source/Programs/AutomationTool/Gauntlet/Unreal/Automation/UE.Automation.cs b/Engine/Source/Programs/AutomationTool/Gauntlet/Unreal/Automation/UE.Automation.cs
--- a/Engine/Source/Programs/AutomationTool/Gauntlet/Unreal/Automation/UE.Automation.cs
+++ b/Engine/Source/Programs/AutomationTool/Gauntlet/Unreal/Automation/UE.Automation.cs
@@ -514,106 +514,111 @@ namespace UE
 							MB.UnorderedList(Result.WarningAndErrorEvents.Distinct());
 						}
 					}
-
-					return MB.ToString();
-				}
-
-				if (AllTests.Count() == 0)
-				{
-					MB.H3("No tests were executed.");
-
-					IEnumerable<UnrealLogParser.LogEntry> WarningsAndErrors = Parser.AutomationWarningsAndErrors;
-
-					if (WarningsAndErrors.Any())
-					{
-						MB.UnorderedList(WarningsAndErrors.Select(E => E.ToString()));
-					}
-					else
-					{
-						MB.Paragraph("Unknown failure.");
-					}
 				}
 				else
 				{
-					// Now list the tests that failed
-					if (FailedTests.Count() > 0)
+					if (AllTests.Count() == 0)
 					{
-						MB.H3("The following test(s) failed:");
+						MB.H3("No tests were executed.");
 
-						foreach (AutomationTestResult Result in FailedTests)
+						IEnumerable<UnrealLogParser.LogEntry> WarningsAndErrors = Parser.AutomationWarningsAndErrors;
+
+						if (WarningsAndErrors.Any())
+						{
+							MB.UnorderedList(WarningsAndErrors.Select(E => E.ToString()));
+						}
+						else
 						{
-							MB.H4(Result.FullName);
-							MB.UnorderedList(Result.Events.Distinct());
+							MB.Paragraph("Unknown failure.");
 						}
 					}
-
-					if (TestsWithWarnings.Count() > 0)
+					else
 					{
-						MB.H3("The following test(s) completed with warnings:");
-
-						foreach (AutomationTestResult Result in TestsWithWarnings)
+						// Now list the tests that failed
+						if (FailedTests.Count() > 0)
 						{
-							MB.Paragraph(Result.FullName);
-							MB.UnorderedList(Result.WarningEvents.Distinct());
+							MB.H3("The following test(s) failed:");
+
+							foreach (AutomationTestResult Result in FailedTests)
+							{
+								MB.H4(Result.FullName);
+								MB.UnorderedList(Result.Events.Distinct());
+							}
 						}
-					}
 
-					if (IncompleteTests.Count() > 0)
-					{
-						MB.H3("The following test(s) timed out or did not run:");
+						if (TestsWithWarnings.Count() > 0)
+						{
+							MB.H3("The following test(s) completed with warnings:");
 
-						foreach (AutomationTestResult Result in IncompleteTests)
+							foreach (AutomationTestResult Result in TestsWithWarnings)
+							{
+								MB.Paragraph(Result.FullName);
+								MB.UnorderedList(Result.WarningEvents.Distinct());
+							}
+						}
+
+						if (IncompleteTests.Count() > 0)
 						{
-							MB.H4(string.Format("{0}", Result.FullName));
-							MB.UnorderedList(Result.WarningAndErrorEvents.Distinct());
+							MB.H3("The following test(s) timed out or did not run:");
+
+							foreach (AutomationTestResult Result in IncompleteTests)
+							{
+								MB.H4(string.Format("{0}", Result.FullName));
+								MB.UnorderedList(Result.WarningAndErrorEvents.Distinct());
+							}
 						}
-					}
 
-					// show a brief summary at the end where it's most visible
-					List<string> TestSummary = new List<string>();
+						// show a brief summary at the end where it's most visible
+						List<string> TestSummary = new List<string>();
 
-					int PassedTests = AllTests.Count() - (FailedTests.Count() + IncompleteTests.Count());
-					int TestsPassedWithoutWarnings = PassedTests - TestsWithWarnings.Count();
+						int PassedTests = AllTests.Count() - (FailedTests.Count() + IncompleteTests.Count());
+						int TestsPassedWithoutWarnings = PassedTests - TestsWithWarnings.Count();
 
-					TestSummary.Add(string.Format("{0} Test(s) Requested", AllTests.Count()));
+						TestSummary.Add(string.Format("{0} Test(s) Requested", AllTests.Count()));
 
-					// Print out a summary of each category of result
-					if (TestsPassedWithoutWarnings > 0)
-					{
-						TestSummary.Add(string.Format("{0} Test(s) Passed", TestsPassedWithoutWarnings));
-					}
+						// Print out a summary of each category of result
+						if (TestsPassedWithoutWarnings > 0)
+						{
+							TestSummary.Add(string.Format("{0} Test(s) Passed", TestsPassedWithoutWarnings));
+						}
 
-					if (TestsWithWarnings.Count() > 0)
-					{
-						TestSummary.Add(string.Format("{0} Test(s) Passed with warnings", TestsWithWarnings.Count()));
-					}
+						if (TestsWithWarnings.Count() > 0)
+						{
+							TestSummary.Add(string.Format("{0} Test(s) Passed with warnings", TestsWithWarnings.Count()));
+						}
 
-					if (FailedTests.Count() > 0)
-					{
-						TestSummary.Add(string.Format("{0} Test(s) Failed", FailedTests.Count()));
-					}
+						if (FailedTests.Count() > 0)
+						{
+							TestSummary.Add(string.Format("{0} Test(s) Failed", FailedTests.Count()));
+						}
 
-					if (IncompleteTests.Count() > 0)
-					{
-						TestSummary.Add(string.Format("{0} Test(s) didn't complete", IncompleteTests.Count()));
-					}
+						if (IncompleteTests.Count() > 0)
+						{
+							TestSummary.Add(string.Format("{0} Test(s) didn't complete", IncompleteTests.Count()));
+						}
 
-					MB.UnorderedList(TestSummary);
+						MB.UnorderedList(TestSummary);
+					}
+				}
 
+				if (EditorArtifacts.LogParser.GetSummary().EngineInitialized)
+				{
 					// Use the paths from the report. If we passed these in they should be the same, and if not
 					// they'll be valid defaults
-					if (!string.IsNullOrEmpty(Parser.AutomationReportPath) || !string.IsNullOrEmpty(Parser.AutomationReportURL))
+					string AutomationReportPath = string.IsNullOrEmpty(Parser.AutomationReportPath) ? GetConfiguration().ReportExportPath : Parser.AutomationReportPath;
+					string AutomationReportURL = string.IsNullOrEmpty(Parser.AutomationReportURL) ? GetConfiguration().ReportURL : Parser.AutomationReportURL;
+					if (!string.IsNullOrEmpty(AutomationReportPath) || !string.IsNullOrEmpty(AutomationReportURL))
 					{
 						MB.H3("Links");
 
-						if (string.IsNullOrEmpty(Parser.AutomationReportURL) == false)
+						if (string.IsNullOrEmpty(AutomationReportURL) == false)
 						{
-							MB.Paragraph(string.Format("View results here: {0}", Parser.AutomationReportURL));
+							MB.Paragraph(string.Format("View results here: {0}", AutomationReportURL));
 						}
 
-						if (string.IsNullOrEmpty(Parser.AutomationReportPath) == false)
+						if (string.IsNullOrEmpty(AutomationReportPath) == false)
 						{
-							MB.Paragraph(string.Format("Open results in UnrealEd from {0}", Parser.AutomationReportPath));
+							MB.Paragraph(string.Format("Open results in UnrealEd from {0}", AutomationReportPath));
 						}
 					}
 				}
