Fix build break
Dynamicly cast the test configuration before accessing the properties

#rnx
#horde-issue 73203

diff --git a/Engine/Source/Programs/AutomationTool/Gauntlet/Unreal/Automation/UE.Automation.cs b/Engine/Source/Programs/AutomationTool/Gauntlet/Unreal/Automation/UE.Automation.cs
--- a/Engine/Source/Programs/AutomationTool/Gauntlet/Unreal/Automation/UE.Automation.cs
+++ b/Engine/Source/Programs/AutomationTool/Gauntlet/Unreal/Automation/UE.Automation.cs
@@ -605,8 +605,19 @@ namespace UE
 				{
 					// Use the paths from the report. If we passed these in they should be the same, and if not
 					// they'll be valid defaults
-					string AutomationReportPath = string.IsNullOrEmpty(Parser.AutomationReportPath) ? GetConfiguration().ReportExportPath : Parser.AutomationReportPath;
-					string AutomationReportURL = string.IsNullOrEmpty(Parser.AutomationReportURL) ? GetConfiguration().ReportURL : Parser.AutomationReportURL;
+					string AutomationReportPath = Parser.AutomationReportPath;
+					string AutomationReportURL = Parser.AutomationReportURL;
+					if (GetConfiguration() is AutomationTestConfig Config)
+					{
+						if (string.IsNullOrEmpty(AutomationReportPath))
+						{
+							AutomationReportPath = Config.ReportExportPath;
+						}
+						if (string.IsNullOrEmpty(AutomationReportURL))
+						{
+							AutomationReportURL = Config.ReportURL;
+						}
+					}
 					if (!string.IsNullOrEmpty(AutomationReportPath) || !string.IsNullOrEmpty(AutomationReportURL))
 					{
 						MB.H3("Links");
