Tests that do not have ray tracing enabled now disable ray tracing explicitly to avoid to enable it by accident when D3D12 is active.

diff --git a/Engine/Source/Programs/AutomationTool/Gauntlet/Unreal/Automation/UE.Automation.cs b/Engine/Source/Programs/AutomationTool/Gauntlet/Unreal/Automation/UE.Automation.cs
--- a/Engine/Source/Programs/AutomationTool/Gauntlet/Unreal/Automation/UE.Automation.cs
+++ b/Engine/Source/Programs/AutomationTool/Gauntlet/Unreal/Automation/UE.Automation.cs
@@ -229,6 +229,10 @@ namespace UE
 				{
 					AppConfig.CommandLine += " -dpcvars=r.RayTracing=1,r.SkinCache.CompileShaders=1,AutomationAllowFrameTraceCapture=0";
 				}
+				else
+				{
+					AppConfig.CommandLine += " -dpcvars=r.RayTracing=0";
+				}
 
 				// Options specific to windows
 				if (ConfigRole.Platform == UnrealTargetPlatform.Win64)
