Telemetry - modify data model to use a build table and a record table separately
+ support for staging db target

diff --git a/Engine/Build/EditorPerf.xml b/Engine/Build/EditorPerf.xml
--- a/Engine/Build/EditorPerf.xml
+++ b/Engine/Build/EditorPerf.xml
@@ -85,6 +85,11 @@
 	<Property Name="SummarizeArgs" Value=""/>
 	<Property Name="SummarizeArgs" Value="-skipbaseline" If="$(RunAndBaseline)"/>
 
+	<!-- Publish telemetry options -->
+	<Option Name="TelemetryConfig" DefaultValue="" Description="The configuration to use when publishing telemetry data to a database"/>
+	<Property Name="TelemetryConfig" Value="UETelemetryStaging" If="$(IsBuildMachine) And $(IsPreflight) And '$(TelemetryConfig)' == ''"/>
+	<Property Name="TelemetryConfig" Value="UETelemetry" If="$(IsBuildMachine) And '$(TelemetryConfig)' == ''"/>
+
 	<!-- The main test running macro logic that is expanded with platform and dependency node names -->
 	<Macro Name="RunTest" Arguments="OutputName;Platform;Project;NodeName;Phase;Args">
 		<Property Name="TestName" Value="EditorPerf$(Phase)"/>
@@ -113,11 +118,9 @@
 				<!-- Save all files named for this run in the trace directory -->
 				<Copy From="$(TraceDir)/$(IterationName)*.*" To="$(AutomationOutputDir)"/>
 
-				<Do If="$(IsBuildMachine)">
-					<Do If="!$(IsPreflight)">
-						<!-- Publish telemetry to the db -->
-						<Command Name="PublishUnrealAutomationTelemetry" Arguments="-CSVFile=&quot;$(TelemetryFilePath)&quot; -project=$(Project) -changelist=$(Change) -branch=$(Branch)"/>
-					</Do>
+				<Do If="'$(TelemetryConfig)' != ''">
+					<!-- Publish telemetry to the db -->
+					<Command Name="PublishUnrealAutomationTelemetry" Arguments="-CSVFile=&quot;$(TelemetryFilePath)&quot; -project=$(Project) -changelist=$(Change) -branch=$(Branch) -TelemetryConfig=$(TelemetryConfig)"/>
 				</Do>
 			</Do>
 
diff --git a/Engine/Source/Programs/AutomationTool/Gauntlet/Framework/Base/Gauntlet.Database.cs b/Engine/Source/Programs/AutomationTool/Gauntlet/Framework/Base/Gauntlet.Database.cs
--- a/Engine/Source/Programs/AutomationTool/Gauntlet/Framework/Base/Gauntlet.Database.cs
+++ b/Engine/Source/Programs/AutomationTool/Gauntlet/Framework/Base/Gauntlet.Database.cs
@@ -3,6 +3,7 @@
 using System;
 using System.IO;
 using System.Linq;
+using System.Data;
 using System.Collections.Generic;
 using MySql.Data.MySqlClient;
 using AutomationTool;
@@ -30,6 +31,12 @@ namespace Gauntlet
 		/// <param name="Context"></param>
 		/// <returns></returns>
 		bool SubmitDataItems(IEnumerable<Data> DataItems, ITelemetryContext Context);
+		/// <summary>
+		/// Execute a Query through the database driver
+		/// </summary>
+		/// <param name="Query"></param>
+		/// <returns></returns>
+		DataSet ExecuteQuery(string Query);
 	}
 
 	/// <summary>
@@ -56,8 +63,11 @@ namespace Gauntlet
 		/// </summary>
 		IEnumerable<string> GetTableColumns();
 
+		/// <summary>
 		/// Format the data for target table based on data type
 		/// </summary>
+		/// <param name="InData">Data to format</param>
+		/// <param name="InContext">ITelemetryContext of the telemetry data</param>
 		/// <returns></returns>
 		IEnumerable<string> FormatDataForTable(Data InData, ITelemetryContext InContext);
 	}
@@ -84,25 +94,55 @@ namespace Gauntlet
 			return Config.GetConfigValue("Server");
 		}
 
-		public object Insert(string Table, IEnumerable<string> Columns, IEnumerable<IEnumerable<string>> Rows)
+		public DataSet ExecuteQuery(string SqlQuery)
+		{
+			return MySqlHelper.ExecuteDataset(Config.ConfigString, SqlQuery);
+		}
+
+		public bool Insert(string Table, IEnumerable<string> Columns, IEnumerable<IEnumerable<string>> Rows)
+		{
+			foreach (var Chunk in ChunkIt(Rows))
+			{
+				string SqlQuery = string.Format(
+					"INSERT INTO `{0}`.{1} ({2}) VALUES {3}",
+					Config.DatabaseName, Table, string.Join(", ", Columns), string.Join(", ", Chunk.Select(R => string.Format("({0})", string.Join(", ", R.Select(V => string.Format("'{0}'", V))))))
+				);
+				if(MySqlHelper.ExecuteScalar(Config.ConfigString, SqlQuery) == null)
+				{
+					return false;
+				}
+			}
+			return true;
+		}
+
+		private IEnumerable<IEnumerable<T>> ChunkIt<T>(IEnumerable<T> ToChunk, int ChunkSize = 1000)
 		{
-			string SqlQuery = string.Format(
-				"INSERT INTO `{0}`.{1} ({2}) VALUES {3}",
-				Config.DatabaseName, Table, string.Join(", ", Columns), string.Join(", ", Rows.Select(R => string.Format("({0})", string.Join(", ", R.Select(V => string.Format("'{0}'", V))))))
-			);
-			return MySqlHelper.ExecuteScalar(Config.ConfigString, SqlQuery);
+			return ToChunk.Select((v, i) => new { Value = v, Index = i }).GroupBy(x => x.Index / ChunkSize).Select(g => g.Select(x => x.Value));
 		}
 
-		public bool SubmitDataItems(IEnumerable<Data> DataRows, ITelemetryContext TestContext)
+		public bool SubmitDataItems(IEnumerable<Data> DataRows, ITelemetryContext Context)
 		{
 			if (Config is IMySQLConfig<Data> DataConfig)
 			{
-				object Value = Insert(
+				if(!PreSubmitQuery(DataRows, Context))
+				{
+					Log.Error("Fail MySQL '{0}' pre-submit query.", Config.GetType().FullName);
+					return false;
+				}
+
+				bool Success = Insert(
 					DataConfig.GetTableName(),
 					DataConfig.GetTableColumns(),
-					DataRows.Select(D => DataConfig.FormatDataForTable(D, TestContext))
+					DataRows.Select(D => DataConfig.FormatDataForTable(D, Context))
 				);
-				return Value != null;
+
+				if(Success && !PostSubmitQuery(DataRows, Context))
+				{
+					Log.Error("Fail MySQL '{0}' post-submit query.", Config.GetType().FullName);
+					return false;
+				}
+
+				return Success;
 			}
 			else
 			{
@@ -111,6 +151,16 @@ namespace Gauntlet
 
 			return false;
 		}
+
+		public virtual bool PreSubmitQuery(IEnumerable<Data> DataRows, ITelemetryContext Context)
+		{
+			return Config.PreSubmitQuery(this, DataRows, Context);
+		}
+		public virtual bool PostSubmitQuery(IEnumerable<Data> DataRows, ITelemetryContext Context)
+		{
+			return Config.PostSubmitQuery(this, DataRows, Context);
+		}
+
 	}
 
 	public abstract class MySQLConfig<Data> : IDatabaseConfig<Data>, IMySQLConfig<Data> where Data : class
@@ -185,6 +235,28 @@ namespace Gauntlet
 			return new MySQLDriver<Data>(this);
 		}
 
+		/// <summary>
+		/// Override to add pre-submit query.
+		/// </summary>
+		/// <param name="Driver"></param>
+		/// <param name="DataRows"></param>
+		/// <param name="Context"></param>
+		/// <returns></returns>
+		public virtual bool PreSubmitQuery(IDatabaseDriver<Data> Driver, IEnumerable<Data> DataRows, ITelemetryContext Context)
+		{
+			return true;
+		}
+		/// <summary>
+		/// Override to add post-submit query.
+		/// </summary>
+		/// <param name="Driver"></param>
+		/// <param name="DataRows"></param>
+		/// <param name="Context"></param>
+		/// <returns></returns>
+		public virtual bool PostSubmitQuery(IDatabaseDriver<Data> Driver, IEnumerable<Data> DataRows, ITelemetryContext Context)
+		{
+			return true;
+		}
 		/// <summary>
 		/// Get Target Table name based on data type
 		/// </summary>
@@ -206,7 +278,7 @@ namespace Gauntlet
 		protected static IEnumerable<IDatabaseConfig<Data>> Configs;
 		static DatabaseConfigManager()
 		{
-			Configs = Gauntlet.Utils.InterfaceHelpers.FindImplementations<IDatabaseConfig<Data>>();
+			Configs = Gauntlet.Utils.InterfaceHelpers.FindImplementations<IDatabaseConfig<Data>>(true);
 		}
 		public static IDatabaseConfig<Data> GetConfigByName(string Name)
 		{
diff --git a/Engine/Source/Programs/AutomationTool/Gauntlet/Unreal/Config/Gauntlet.EpicConfig.cs b/Engine/Source/Programs/AutomationTool/Gauntlet/Unreal/Config/Gauntlet.EpicConfig.cs
--- a/Engine/Source/Programs/AutomationTool/Gauntlet/Unreal/Config/Gauntlet.EpicConfig.cs
+++ b/Engine/Source/Programs/AutomationTool/Gauntlet/Unreal/Config/Gauntlet.EpicConfig.cs
@@ -3,13 +3,20 @@
 using AutomationTool;
 using Gauntlet;
 using System;
+using System.Data;
 using System.Collections.Generic;
 using UnrealBuildTool;
 
 namespace EpicConfig
 {
+	/// <summary>
+	/// Production UE Telemetry config
+	/// </summary>
 	public class UETelemetry : MySQLConfig<Gauntlet.TelemetryData>
 	{
+		protected virtual string TableName { get { return "test_records_prod"; } }
+		protected virtual string BuildTableName { get { return "test_builds_prod"; } }
+		private int BuildID = -1;
 		public override void LoadConfig(string ConfigFilePath)
 		{
 			if (string.IsNullOrEmpty(ConfigFilePath))
@@ -30,7 +37,7 @@ namespace EpicConfig
 		/// </summary>
 		public override string GetTableName()
 		{
-			return "test_records";
+			return TableName;
 		}
 		/// <summary>
 		/// Get Target Table columns based on data type
@@ -38,11 +45,9 @@ namespace EpicConfig
 		public override IEnumerable<string> GetTableColumns()
 		{
 			return new List<string>() {
-				"changelist",
+				"buildId",
 				"system",
-				"stream",
 				"dateTime",
-				"project",
 				"platform",
 				"config",
 				"testName",
@@ -61,11 +66,9 @@ namespace EpicConfig
 		public override IEnumerable<string> FormatDataForTable(Gauntlet.TelemetryData Data, ITelemetryContext Context)
 		{
 			return new List<string>() {
-				Context.GetProperty("Changelist").ToString(),
+				BuildID.ToString(),
 				Environment.MachineName,
-				Context.GetProperty("Branch").ToString(),
 				DateTime.Now.ToString("yyyy-MM-dd HH:mm:ss"),
-				Context.GetProperty("ProjectName").ToString(),
 				Context.GetProperty("Platform").ToString(),
 				Context.GetProperty("Configuration").ToString(),
 				Data.TestName,
@@ -77,5 +80,49 @@ namespace EpicConfig
 				Context.GetProperty("JobLink").ToString()
 			};
 		}
+		/// <summary>
+		/// Get Build Table name
+		/// </summary>
+		public virtual string GetBuildTableName()
+		{
+			return BuildTableName;
+		}
+		/// <summary>
+		/// Query DB to get build id or add it and store the build id.
+		/// </summary>
+		/// <param name="Driver"></param>
+		/// <param name="DataRows"></param>
+		/// <param name="Context"></param>
+		public override bool PreSubmitQuery(IDatabaseDriver<Gauntlet.TelemetryData> Driver, IEnumerable<Gauntlet.TelemetryData> DataRows, ITelemetryContext Context)
+		{
+			string Changelist = Context.GetProperty("Changelist").ToString();
+			string ChangelistDateTime = ((DateTime)Context.GetProperty("ChangelistDateTime")).ToString("yyyy-MM-dd HH:mm:ss");
+			string Stream = Context.GetProperty("Branch").ToString();
+			string Project = Context.GetProperty("ProjectName").ToString();
+
+			string SqlQuery = @$"INSERT INTO `{DatabaseName}`.{BuildTableName} (changelist, changelistDateTime, stream, project)
+									SELECT '{Changelist}', '{ChangelistDateTime}', '{Stream}', '{Project}' FROM DUAL
+										WHERE NOT EXISTS (SELECT 1 FROM `{DatabaseName}`.{BuildTableName} WHERE changelist = '{Changelist}' AND stream = '{Stream}' AND project = '{Project}');
+									SELECT id FROM `{DatabaseName}`.{BuildTableName} WHERE changelist = '{Changelist}' AND stream = '{Stream}' AND project = '{Project}';";
+
+			DataSet Set = Driver.ExecuteQuery(SqlQuery);
+			// Get value from first row from first column
+			if (Set.Tables.Count == 0 || Set.Tables[0].Rows.Count == 0)
+			{
+				Log.Error("No data return from {0} insert Query.", BuildTableName);
+				return false;
+			}
+			BuildID = (int)Set.Tables[0].Rows[0][0];
+
+			return true;
+		}
+	}
+	/// <summary>
+	/// Staging UE Telemetry Config
+	/// </summary>
+	public class UETelemetryStaging : UETelemetry
+	{
+		protected override string TableName { get { return "test_records_staging"; } }
+		protected override string BuildTableName { get { return "test_builds_staging"; } }
 	}
 }
\ No newline at end of file
diff --git a/Engine/Source/Programs/AutomationTool/Gauntlet/Unreal/Utils/Gauntlet.UnrealTelemetry.cs b/Engine/Source/Programs/AutomationTool/Gauntlet/Unreal/Utils/Gauntlet.UnrealTelemetry.cs
--- a/Engine/Source/Programs/AutomationTool/Gauntlet/Unreal/Utils/Gauntlet.UnrealTelemetry.cs
+++ b/Engine/Source/Programs/AutomationTool/Gauntlet/Unreal/Utils/Gauntlet.UnrealTelemetry.cs
@@ -3,6 +3,7 @@
 using System;
 using System.IO;
 using System.Collections.Generic;
+using System.Text.RegularExpressions;
 using AutomationTool;
 using UnrealBuildTool;
 using UnrealBuildBase;
@@ -150,7 +151,7 @@ namespace Gauntlet
 		[Help("CSVFile", "Path to the csv file to parse.")]
 		[Help("CSVDirectory", "Path to a folder containing csv files to parse.")]
 		[Help("CSVMapping", "Optional CSV column mapping. Format is: <target key>:<source key>,...")]
-		[Help("TelemetryConfig", "Telemetry configuration to use to publish to Database. Default: UETelemtry.")]
+		[Help("TelemetryConfig", "Telemetry configuration to use to publish to Database. Default: UETelemetryStaging.")]
 		[Help("DatabaseConfigPath", "Path to alternate Database config. Default is TelemetryConfig default.")]
 		[Help("Project", "Target Project name.")]
 		[Help("Platform", "Target platform name. Default: current environment platform.")]
@@ -164,7 +165,7 @@ namespace Gauntlet
 		{
 			string CSVFile = ParseParamValue("CSVFile=", "");
 			string CSVDirectory = ParseParamValue("CSVDirectory=", "");
-			string Config = ParseParamValue("TelemetryConfig=", "UETelemetry");
+			string Config = ParseParamValue("TelemetryConfig=", "UETelemetryStaging");
 			string DatabaseConfigPath = ParseParamValue("DatabaseConfigPath=", "");
 			List<string> CSVMappingStrings = Globals.Params.ParseValues("CSVMapping", true);
 			string ProjectString = ParseParamValue("Project=", "");
@@ -245,6 +246,7 @@ namespace Gauntlet
 
 			Context.SetProperty("Branch", BranchString);
 			Context.SetProperty("Changelist", ChangelistString);
+			Context.SetProperty("ChangelistDateTime", GetChangelistDateTime(int.Parse(ChangelistString)));
 
 			if (string.IsNullOrEmpty(JobLink))
 			{
@@ -294,5 +296,28 @@ namespace Gauntlet
 
 			return 0;
 		}
+
+		/// <summary>
+		/// Get the DateTime of the given changelist
+		/// </summary>
+		/// <param name="Changelist"></param>
+		/// <returns></returns>
+		private DateTime GetChangelistDateTime(int Changelist)
+		{
+			try
+			{
+				P4Connection.DescribeRecord Record = new P4Connection.DescribeRecord();
+				P4.DescribeChangelist(Changelist, out Record, false);
+
+				Regex Regx = new Regex(@" on ([0-9/]+ [0-9:]+)", RegexOptions.Compiled | RegexOptions.IgnoreCase);
+				Match DateMatch = Regx.Match(Record.Header);
+				return DateTime.Parse(DateMatch.Groups[1].Value);
+			}
+			catch
+			{
+				return DateTime.Now;
+			}
+		}
+
 	}
 }
