Build fix.

diff --git a/Engine/Source/Programs/AutomationTool/BuildGraph/BuildGraph.cs b/Engine/Source/Programs/AutomationTool/BuildGraph/BuildGraph.cs
--- a/Engine/Source/Programs/AutomationTool/BuildGraph/BuildGraph.cs
+++ b/Engine/Source/Programs/AutomationTool/BuildGraph/BuildGraph.cs
@@ -676,7 +676,7 @@ namespace AutomationTool
 		bool CreateTaskInstances(BgScript Graph, BgNode Node, Dictionary<string, ScriptTaskBinding> NameToTask, Dictionary<BgTask, CustomTask> TaskInfoToTask)
 		{
 			bool bResult = true;
-			foreach (BgTask TaskInfo in Node.TaskInfos)
+			foreach (BgTask TaskInfo in Node.Tasks)
 			{
 				CustomTask Task = BindTask(Node, TaskInfo, NameToTask, Graph.TagNameToNodeOutput);
 				if (Task == null)
@@ -1229,7 +1229,7 @@ namespace AutomationTool
 		/// <returns>Whether the task succeeded or not. Exiting with an exception will be caught and treated as a failure.</returns>
 		bool ExecuteTasks(BgNode Node, JobContext Job, Dictionary<BgTask, CustomTask> TaskInfoToTask, Dictionary<string, HashSet<FileReference>> TagNameToFileSet)
 		{
-			List<CustomTask> Tasks = Node.TaskInfos.ConvertAll(x => TaskInfoToTask[x]);
+			List<CustomTask> Tasks = Node.Tasks.ConvertAll(x => TaskInfoToTask[x]);
 
 			// Run each of the tasks in order
 			HashSet<FileReference> BuildProducts = TagNameToFileSet[Node.DefaultOutput.TagName];
