[UE-133899] McpConfigHelper::Find() logs messages from any exceptions encountered when reflecting config types

diff --git a/Engine/Source/Programs/AutomationTool/AutomationUtils/MCPPublic.cs b/Engine/Source/Programs/AutomationTool/AutomationUtils/MCPPublic.cs
--- a/Engine/Source/Programs/AutomationTool/AutomationUtils/MCPPublic.cs
+++ b/Engine/Source/Programs/AutomationTool/AutomationUtils/MCPPublic.cs
@@ -2558,9 +2558,10 @@ namespace EpicGames.MCP.Config
                                     Configs.Add(Config.Name, Config);
                                 }
                             }
-                            catch
+                            catch (Exception Ex)
                             {
-                                BuildCommand.LogWarning("Unable to create McpConfig: {0}", PotentialConfigType.Name);
+                                BuildCommand.LogWarning("Unable to create McpConfig [{0}] with error: {1}",
+                                    PotentialConfigType.Name, Ex.Message);
                             }
                         }
                     }
