Add all files under UBT output directory, as with .net core we now need all of them.
Fixes issue running build setup on linux builders.

diff --git a/Engine/Source/Programs/AutomationTool/AutomationUtils/UnrealBuild.cs b/Engine/Source/Programs/AutomationTool/AutomationUtils/UnrealBuild.cs
--- a/Engine/Source/Programs/AutomationTool/AutomationUtils/UnrealBuild.cs
+++ b/Engine/Source/Programs/AutomationTool/AutomationUtils/UnrealBuild.cs
@@ -1223,20 +1223,13 @@ namespace AutomationTool
 		/// </summary>
 		public void AddUBTFilesToBuildProducts()
 		{
-			var UBTLocation = Path.GetDirectoryName(GetUBTExecutable());
-			var UBTFiles = new List<string>(new string[] 
-					{
-						"UnrealBuildTool.exe",
-					});
+			string UBTLocation = Path.GetDirectoryName(GetUBTExecutable());
+			// copy all the files from the UBT output directory
+			string[] UBTFiles = CommandUtils.FindFiles_NoExceptions("*.*", true, UBTLocation);
 
-			foreach (var UBTFile in UBTFiles)
+			foreach (string UBTFile in UBTFiles)
 			{
-				var UBTProduct = CommandUtils.CombinePaths(UBTLocation, UBTFile);
-				if (!CommandUtils.FileExists_NoExceptions(UBTProduct))
-				{
-					throw new UnrealBuildException("Cannot add UBT to the build products because {0} does not exist.", UBTProduct);
-				}
-				AddBuildProduct(UBTProduct);
+				AddBuildProduct(UBTFile);
 			}
 		}
 
