Fixes for HelmTask in BuildGraph

* No longer modify the global kubecontext but rather pass the context to use into helm
* Added option to specify which kubeconfig to use
* Fixed issue when values not being handled correctly

diff --git a/Engine/Source/Programs/AutomationTool/BuildGraph/Tasks/HelmTask.cs b/Engine/Source/Programs/AutomationTool/BuildGraph/Tasks/HelmTask.cs
--- a/Engine/Source/Programs/AutomationTool/BuildGraph/Tasks/HelmTask.cs
+++ b/Engine/Source/Programs/AutomationTool/BuildGraph/Tasks/HelmTask.cs
@@ -40,11 +40,17 @@ namespace BuildGraph.Tasks
 		[TaskParameter(Optional = true)]
 		public string KubeContext;
 
+		/// <summary>
+		/// The kubectl config file to use
+		/// </summary>
+		[TaskParameter(Optional = true)]
+		public string KubeConfig;
+
 		/// <summary>
 		/// Values to set for running the chart
 		/// </summary>
 		[TaskParameter(Optional = true)]
-		public List<string> Values;
+		public string Values;
 
 		/// <summary>
 		/// Values to set for running the chart
@@ -105,11 +111,6 @@ namespace BuildGraph.Tasks
 		/// <param name="TagNameToFileSet">Mapping from tag names to the set of files they include</param>
 		public override void Execute(JobContext Job, HashSet<FileReference> BuildProducts, Dictionary<string, HashSet<FileReference>> TagNameToFileSet)
 		{
-			// Switch Kubernetes config
-			if (Parameters.KubeContext != null)
-			{
-				Execute("kubectl", $"config use-context {Parameters.KubeContext}", WorkingDir: Parameters.WorkingDir);
-			}
 
 			// Build the argument list
 			List<string> Arguments = new List<string>();
@@ -123,7 +124,17 @@ namespace BuildGraph.Tasks
 				Arguments.Add("--namespace");
 				Arguments.Add(Parameters.Namespace);
 			}
-			foreach (string Value in Parameters.Values)
+			if (Parameters.KubeContext != null)
+			{
+				Arguments.Add("--kube-context");
+				Arguments.Add(Parameters.KubeContext);
+			}
+			if (Parameters.KubeConfig != null)
+			{
+				Arguments.Add("--kubeconfig");
+				Arguments.Add(Parameters.KubeConfig);
+			}
+			foreach (string Value in SplitDelimitedList(Parameters.Values))
 			{
 				Arguments.Add("--set");
 				Arguments.Add(Value);
