Add an explicit error message when we run out of retries and fail to unzip temp storage.

#rnx

diff --git a/Engine/Source/Programs/AutomationTool/BuildGraph/TempStorage.cs b/Engine/Source/Programs/AutomationTool/BuildGraph/TempStorage.cs
--- a/Engine/Source/Programs/AutomationTool/BuildGraph/TempStorage.cs
+++ b/Engine/Source/Programs/AutomationTool/BuildGraph/TempStorage.cs
@@ -1111,6 +1111,7 @@ namespace AutomationTool
 						{
 							if (UnzipFileAttempts == 0)
 							{
+								Log.TraceError("All retries exhausted attempting to unzip entries from '{0}'. Terminating.");
 								throw;
 							}
 
