Change UnrealPak Parallelization to use Parallel.ForEach instead of ThreadPools

#rnx

diff --git a/Engine/Source/Programs/AutomationTool/Scripts/CopyBuildToStagingDirectory.Automation.cs b/Engine/Source/Programs/AutomationTool/Scripts/CopyBuildToStagingDirectory.Automation.cs
--- a/Engine/Source/Programs/AutomationTool/Scripts/CopyBuildToStagingDirectory.Automation.cs
+++ b/Engine/Source/Programs/AutomationTool/Scripts/CopyBuildToStagingDirectory.Automation.cs
@@ -3315,25 +3315,15 @@ namespace AutomationScripts
 
 			// Spawn tasks for each command
 			IProcessResult[] Results = new IProcessResult[Commands.Count];
-			using(ThreadPoolWorkQueue Queue = new ThreadPoolWorkQueue())
+			string UnrealPakExe = GetUnrealPakLocation().FullName;
+			Parallel.ForEach(Commands, new ParallelOptions() { MaxDegreeOfParallelism = Environment.ProcessorCount }, (Command, LoopState, Idx) =>
 			{
-				string UnrealPakExe = GetUnrealPakLocation().FullName;
-				for(int Idx = 0; Idx < Commands.Count; Idx++)
-				{
-					string LogFile = LogUtils.GetUniqueLogName(CombinePaths(CmdEnv.LogFolder, String.Format("UnrealPak-{0}", LogNames[Idx])));
-					string Arguments = Commands[Idx] = String.Format("{0} -multiprocess -abslog={1} {2}", Commands[Idx], MakePathSafeToUseWithCommandLine(LogFile), AdditionalCompressionOptionsOnCommandLine);
-
-					int LocalIdx = Idx;
-					Queue.Enqueue(() => Results[LocalIdx] = Run(UnrealPakExe, Arguments, Options: ERunOptions.AppMustExist | ERunOptions.NoLoggingOfRunCommand));
-				}
+				int LocalIdx = (int)Idx;
+				string LogFile = LogUtils.GetUniqueLogName(CombinePaths(CmdEnv.LogFolder, String.Format("UnrealPak-{0}", LogNames[(int)Idx])));
+				string Arguments = Commands[LocalIdx] = String.Format("{0} -multiprocess -abslog={1} {2}", Commands[LocalIdx], MakePathSafeToUseWithCommandLine(LogFile), AdditionalCompressionOptionsOnCommandLine);
 
-				bool bComplete = false;
-				while(!bComplete)
-				{
-					bComplete = Queue.Wait(10 * 1000);
-					LogInformation("Waiting for child processes to complete ({0}/{1})", Results.Count(x => x != null), Results.Length);
-				}
-			}
+				Results[LocalIdx] = Run(UnrealPakExe, Arguments, Options: ERunOptions.AppMustExist | ERunOptions.NoLoggingOfRunCommand);
+			});
 
 			// Output all the results
 			int NumFailed = 0;
