BuildGraph - Fixed issue with helm task not handling when there were no values files

diff --git a/Engine/Source/Programs/AutomationTool/BuildGraph/Tasks/HelmTask.cs b/Engine/Source/Programs/AutomationTool/BuildGraph/Tasks/HelmTask.cs
--- a/Engine/Source/Programs/AutomationTool/BuildGraph/Tasks/HelmTask.cs
+++ b/Engine/Source/Programs/AutomationTool/BuildGraph/Tasks/HelmTask.cs
@@ -134,10 +134,13 @@ namespace BuildGraph.Tasks
 				Arguments.Add("--kubeconfig");
 				Arguments.Add(Parameters.KubeConfig);
 			}
-			foreach (string Value in SplitDelimitedList(Parameters.Values))
+			if (!string.IsNullOrEmpty(Parameters.Values))
 			{
-				Arguments.Add("--set");
-				Arguments.Add(Value);
+				foreach (string Value in SplitDelimitedList(Parameters.Values))
+				{
+					Arguments.Add("--set");
+					Arguments.Add(Value);
+				}
 			}
 			if (!String.IsNullOrEmpty(Parameters.ValuesFile))
 			{
