During staging, guard against the target platforms cooked data folder not existing. This can happen when making builds that re-use previously generated pak files.

This change is a re-implementation of a previous fix for the same issue that was lost during the move from UE4->UE5

diff --git a/Engine/Source/Programs/AutomationTool/Scripts/CopyBuildToStagingDirectory.Automation.cs b/Engine/Source/Programs/AutomationTool/Scripts/CopyBuildToStagingDirectory.Automation.cs
--- a/Engine/Source/Programs/AutomationTool/Scripts/CopyBuildToStagingDirectory.Automation.cs
+++ b/Engine/Source/Programs/AutomationTool/Scripts/CopyBuildToStagingDirectory.Automation.cs
@@ -727,7 +727,12 @@ namespace AutomationScripts
 				}
 
 				// Stage all the cooked data, this is the same rule as normal stage except we may skip Engine
-				List<FileReference> CookedFiles = DirectoryReference.EnumerateFiles(SC.PlatformCookDir, "*", SearchOption.AllDirectories).ToList();
+				List<FileReference> CookedFiles = new List<FileReference>();
+
+				if (DirectoryReference.Exists(SC.PlatformCookDir))
+				{
+					CookedFiles.AddRange(DirectoryReference.EnumerateFiles(SC.PlatformCookDir, "*", SearchOption.AllDirectories).ToList());
+				}
 
 				// When cooking to Zen get the list of cooked package files from the manifest
 				LoadPackageStoreManifest(Params, SC);
