fixed gauntlet duplicating additional args after CommandLine += string ops

diff --git a/Engine/Source/Programs/AutomationTool/Gauntlet/Unreal/Base/Gauntlet.UnrealTestConfiguration.cs b/Engine/Source/Programs/AutomationTool/Gauntlet/Unreal/Base/Gauntlet.UnrealTestConfiguration.cs
--- a/Engine/Source/Programs/AutomationTool/Gauntlet/Unreal/Base/Gauntlet.UnrealTestConfiguration.cs
+++ b/Engine/Source/Programs/AutomationTool/Gauntlet/Unreal/Base/Gauntlet.UnrealTestConfiguration.cs
@@ -216,6 +216,7 @@ namespace Gauntlet
 		public void ClearCommandLine()
 		{
 			Params.Clear();
+			AdditionalExplicitCommandLineArgs = string.Empty;
 		}
 
 		/// <summary>
