Log WebException on device reservation when there aren't any retries at all.

diff --git a/Engine/Source/Programs/AutomationTool/AutomationUtils/DeviceReservation.cs b/Engine/Source/Programs/AutomationTool/AutomationUtils/DeviceReservation.cs
--- a/Engine/Source/Programs/AutomationTool/AutomationUtils/DeviceReservation.cs
+++ b/Engine/Source/Programs/AutomationTool/AutomationUtils/DeviceReservation.cs
@@ -29,7 +29,7 @@ namespace AutomationTool.DeviceReservation
 		// Max times to attempt reservation renewal, in case device starvation, reservation service being restarted, etc
 		private static readonly int RenewRetryMax = 14;
 		private static readonly TimeSpan RenewRetryTime = TimeSpan.FromMinutes(1);
-		
+
 		private Thread RenewThread;
 		private AutoResetEvent WaitEvent = new AutoResetEvent(false);
 
@@ -98,7 +98,7 @@ namespace AutomationTool.DeviceReservation
 					// try again
 					RetryCurrent++;
 					RenewTimeCurrent = RenewRetryTime;
-				}				
+				}
 			}
 
 			// Delete reservation on server, if the web request fails backend has logic to cleanup reservations
@@ -224,7 +224,7 @@ namespace AutomationTool.DeviceReservation
 			{
 				Request.ContentLength = 0;
 			}
-			
+
 			using (var Response = (HttpWebResponse)Request.GetResponse())
 			using (var ResponseStream = Response.GetResponseStream())
 			{
@@ -262,7 +262,7 @@ namespace AutomationTool.DeviceReservation
 			while (true)
 			{
 				if (!bFirst)
-				{					
+				{
 					Thread.Sleep(RetryTime);
 				}
 
@@ -286,34 +286,32 @@ namespace AutomationTool.DeviceReservation
 					});
 				}
 				catch (WebException WebEx)
-				{					
+				{
 
 					if (RetryCount == RetryMax)
 					{
 						Console.WriteLine("Device reservation unsuccessful");
-						throw new AutomationException(WebEx, "Device reservation unsuccessful, devices unavailable");
+						string FinalMessage = "Device reservation unsuccessful";
+						string MaxRetryWebExMessage = GetReservationErrorMessage(WebEx);
+						if (!string.IsNullOrEmpty(MaxRetryWebExMessage))
+						{
+							FinalMessage = string.Format("{0} - Caused by: {1}", FinalMessage, MaxRetryWebExMessage);
+						}
+						else
+						{
+							FinalMessage = string.Format("{0} - Request returned status {1}", FinalMessage, WebEx.Status.ToString());
+						}
+						throw new AutomationException(WebEx, FinalMessage);
 					}
 
 					string RetryMessage = String.Format("retry {0} of {1} in {2} minutes", RetryCount + 1, RetryMax, RetryTime.Minutes);
 					string Message = String.Format("Unknown device server error, {0}", RetryMessage);
 
-					if (WebEx.Response == null)
-					{
-						Message = String.Format("Devices service currently not available, {0}", RetryMessage);
-					}
-					else if ((WebEx.Response as HttpWebResponse).StatusCode == HttpStatusCode.Conflict)
+					string WebExMessage = GetReservationErrorMessage(WebEx, RetryMessage);
+
+					if (!string.IsNullOrEmpty(WebExMessage))
 					{
-						Message = String.Format("No devices currently available, {0}", RetryMessage);
-					}
-					else
-					{
-						using (HttpWebResponse Response = (HttpWebResponse)WebEx.Response)
-						{
-							using (StreamReader Reader = new StreamReader(Response.GetResponseStream()))
-							{
-								Message = String.Format("WebException on reservation request: {0} : {1} : {2}", WebEx.Message, WebEx.Status, Reader.ReadToEnd());
-							}
-						}						
+						Message = WebExMessage;
 					}
 
 					Console.WriteLine(Message);
@@ -328,18 +326,44 @@ namespace AutomationTool.DeviceReservation
 			}
 		}
 
+		private static string GetReservationErrorMessage(WebException WebEx, string RetryMessage = "")
+		{
+			string Message = string.Empty;
+			string RetryMessageEx = string.IsNullOrEmpty(RetryMessage) ? string.Empty : (", " + RetryMessage);
+
+			if (WebEx.Response == null)
+			{
+				Message = String.Format("Devices service currently not available{0}", RetryMessageEx);
+			}
+			else if ((WebEx.Response as HttpWebResponse).StatusCode == HttpStatusCode.Conflict)
+			{
+				Message = String.Format("No devices currently available{0}", RetryMessageEx);
+			}
+			else
+			{
+				using (HttpWebResponse Response = (HttpWebResponse)WebEx.Response)
+				{
+					using (StreamReader Reader = new StreamReader(Response.GetResponseStream()))
+					{
+						Message = String.Format("WebException on reservation request: {0} : {1} : {2}", WebEx.Message, WebEx.Status, Reader.ReadToEnd());
+					}
+				}
+			}
+			return Message;
+		}
+
 		public void Renew(Uri BaseUri, TimeSpan NewDuration)
 		{
 			try
 			{
-				Utils.InvokeAPI(BaseUri.AppendPath("api/v1/reservations/" + Guid.ToString()), "PUT", NewDuration);				
+				Utils.InvokeAPI(BaseUri.AppendPath("api/v1/reservations/" + Guid.ToString()), "PUT", NewDuration);
 			}
 			catch (Exception ex)
 			{
 				throw new AutomationException(ex, "Failed to renew device reservation.");
 			}
 		}
-		
+
 		public void Delete(Uri BaseUri)
 		{
 			try
