[UE-133899] McpConfigHelper::Find() logs messages from any exceptions (and inner exceptions!) encountered when reflecting config types
+ grepping more logs/info to debug a specific warning

diff --git a/Engine/Source/Programs/AutomationTool/AutomationUtils/MCPPublic.cs b/Engine/Source/Programs/AutomationTool/AutomationUtils/MCPPublic.cs
--- a/Engine/Source/Programs/AutomationTool/AutomationUtils/MCPPublic.cs
+++ b/Engine/Source/Programs/AutomationTool/AutomationUtils/MCPPublic.cs
@@ -2560,8 +2560,14 @@ namespace EpicGames.MCP.Config
                             }
                             catch (Exception Ex)
                             {
+                                var ErrorBuilder = new System.Text.StringBuilder(Ex.Message);
+                                if (Ex.InnerException != null)
+                                {
+                                    ErrorBuilder.AppendLine(Ex.InnerException.Message);
+                                }
+                                ErrorBuilder.AppendLine(Ex.StackTrace);
                                 BuildCommand.LogWarning("Unable to create McpConfig [{0}] with error: {1}",
-                                    PotentialConfigType.Name, Ex.Message);
+                                    PotentialConfigType.Name, ErrorBuilder.ToString());
                             }
                         }
                     }
