[MemreportHelper] - remove debugger launch

diff --git a/Engine/Source/Programs/AutomationTool/Scripts/MemreportHelper.cs b/Engine/Source/Programs/AutomationTool/Scripts/MemreportHelper.cs
--- a/Engine/Source/Programs/AutomationTool/Scripts/MemreportHelper.cs
+++ b/Engine/Source/Programs/AutomationTool/Scripts/MemreportHelper.cs
@@ -93,7 +93,6 @@ public class MemreportHelper : BuildCommand
 
 	public override ExitCode Execute()
 	{
-		System.Diagnostics.Debugger.Launch();
 		LogInformation("************************ STARTING MemreportToGSheet ************************");
 
 		if (Params.Length < 2)
