[MemreportHelper] - copyright notice added

diff --git a/Engine/Source/Programs/AutomationTool/Scripts/MemreportHelper.cs b/Engine/Source/Programs/AutomationTool/Scripts/MemreportHelper.cs
--- a/Engine/Source/Programs/AutomationTool/Scripts/MemreportHelper.cs
+++ b/Engine/Source/Programs/AutomationTool/Scripts/MemreportHelper.cs
@@ -1,4 +1,4 @@
-// Copyright 2006-2016 Donya Labs AB. All Rights Reserved.
+// Copyright Epic Games, Inc. All Rights Reserved.
 
 using System;
 using System.IO;
