[MemreportHelper] - small tweaks to support batch processing

diff --git a/Engine/Source/Programs/AutomationTool/Scripts/MemreportHelper.cs b/Engine/Source/Programs/AutomationTool/Scripts/MemreportHelper.cs
--- a/Engine/Source/Programs/AutomationTool/Scripts/MemreportHelper.cs
+++ b/Engine/Source/Programs/AutomationTool/Scripts/MemreportHelper.cs
@@ -93,33 +93,28 @@ public class MemreportHelper : BuildCommand
 
 	public override ExitCode Execute()
 	{
-		LogInformation("************************ STARTING MemreportToGSheet ************************");
+		LogInformation("************************ STARTING MemreportHelper ************************");
 
-		if (Params.Length < 2)
+		if (Params.Length < 1)
 		{
 			LogError("Invalid number of arguments {0}", Params.Length);
 			return ExitCode.Error_Arguments;
 		}
 
-		string MemreportFile = ParseParamValue("report", "");
-		string OutputDir = ParseParamValue("outcsvdir", "");
+		string MemreportFileorPath = ParseParamValue("report", "");
+		string InOutputDir = ParseParamValue("outcsvdir", "");
 		string MaxOutputString = ParseParamValue("maxoutputcount", "");
 		string MinKBLim = ParseParamValue("minkblim", "");
 		int MaxOutputCount = 75;
 		float MinOutputKBLimit = 1 * 1024;
 		
-
-
-		if (string.IsNullOrWhiteSpace(MemreportFile))
+		if (string.IsNullOrWhiteSpace(MemreportFileorPath))
 		{
-			LogError("No memreport file specified");
+			LogError("No memreport file or directory specified");
 			return ExitCode.Error_Arguments;
 		}
 
-		if (string.IsNullOrEmpty(OutputDir))
-		{
-			OutputDir = Path.GetDirectoryName(MemreportFile);
-		}
+		
 
 		if (string.IsNullOrEmpty(MaxOutputString) == false)
 		{
@@ -134,63 +129,91 @@ public class MemreportHelper : BuildCommand
 		float MeshMin = MinOutputKBLimit / 2;
 		float SkelMeshMin = MinOutputKBLimit / 2;
 
-		if (File.Exists(MemreportFile)== false)
+		
+
+		List<string> FilesToProcess = new List<string>();
+
+		if (Directory.Exists(MemreportFileorPath))
 		{
-			LogError("Memreport file not found: " + MemreportFile);
-			return ExitCode.Error_Arguments;
-		}
+			var FilesReturned = Directory.EnumerateFiles(MemreportFileorPath, "*.memreport");
 
-		// Pull the file into a buffer of lines (this will be stomped by the process)
-		InputFileAllLines = File.ReadAllLines(MemreportFile);
+			foreach(string file in FilesReturned)
+			{
+				FilesToProcess.Add(file);
+			}
+		}
+		else
+		{
+			if (File.Exists(MemreportFileorPath) == false)
+			{
+				LogError("Memreport file not found: " + MemreportFileorPath);
+				return ExitCode.Error_Arguments;
+			}
 
-		bool CsvFound = ReadCsvListAndAdd("Obj List:", 3, "ResExcKB");
+			FilesToProcess.Add(MemreportFileorPath);
+		}
 
-		if (CsvFound)
+		foreach (string MRFile in FilesToProcess)
 		{
-			// look for other 'obj list' entries only if we found the main Obj List
+			string OutputDir = InOutputDir;
 
-			while (CsvFound)
+			if (string.IsNullOrEmpty(OutputDir))
 			{
-				CsvFound = ReadCsvListAndAdd("Obj List:", 3, "ResExcKB");
+				OutputDir = Path.Combine(Path.GetDirectoryName(MRFile), Path.GetFileName(MRFile).Replace('.', '_'));
 			}
 
-			// search for 'Listing NONVT textures.' and its ilk, these are ad-hoc and not explicitly set to -csv 
-			// order here is important for latest badness filtering to ensure badness from "all" is unique from NONVT/UNCOMP list
-			ReadCsvListAndAdd("Listing NONVT textures", 1, "Current Size (KB)", true);
-			ReadCsvListAndAdd("Listing uncompressed textures.", 1, "Current Size (KB)", true);
-			ReadCsvListAndAdd("Listing all textures.", 1, "Current Size (KB)", true);
+			// Pull the file into a buffer of lines (this will be stomped by the process)
+			InputFileAllLines = File.ReadAllLines(MRFile);
+			CsvBlocks.Clear();
+
+			bool CsvFound = ReadCsvListAndAdd("Obj List:", 3, "ResExcKB");
 
-			foreach (var csvBlock in CsvBlocks)
+			if (CsvFound)
 			{
+				// look for other 'obj list' entries only if we found the main Obj List
 
-				csvBlock.sb.AppendLine(csvBlock.Title);
+				while (CsvFound)
+				{
+					CsvFound = ReadCsvListAndAdd("Obj List:", 3, "ResExcKB");
+				}
 
-				// write out headers
-				if (csvBlock.Title.Contains("-csv"))
+				// search for 'Listing NONVT textures.' and its ilk, these are ad-hoc and not explicitly set to -csv 
+				// order here is important for latest badness filtering to ensure badness from "all" is unique from NONVT/UNCOMP list
+				ReadCsvListAndAdd("Listing NONVT textures", 1, "Current Size (KB)", true);
+				ReadCsvListAndAdd("Listing uncompressed textures.", 1, "Current Size (KB)", true);
+				ReadCsvListAndAdd("Listing all textures.", 1, "Current Size (KB)", true);
+
+				foreach (var csvBlock in CsvBlocks)
 				{
-					foreach (var key in csvBlock.SortedRows[0].Keys)
-					{
-						csvBlock.sb.Append(key + ",");
-					}
-					csvBlock.sb.Append(Environment.NewLine);
-					foreach (var row in csvBlock.SortedRows)
+
+					csvBlock.sb.AppendLine(csvBlock.Title);
+
+					// write out headers
+					if (csvBlock.Title.Contains("-csv"))
 					{
-						foreach (var ele in row.Values)
+						foreach (var key in csvBlock.SortedRows[0].Keys)
 						{
-							csvBlock.sb.Append(ele + ",");
+							csvBlock.sb.Append(key + ",");
+						}
+						csvBlock.sb.Append(Environment.NewLine);
+						foreach (var row in csvBlock.SortedRows)
+						{
+							foreach (var ele in row.Values)
+							{
+								csvBlock.sb.Append(ele + ",");
+							}
+							csvBlock.sb.Append(Environment.NewLine);
 						}
 						csvBlock.sb.Append(Environment.NewLine);
+						csvBlock.sb.Append(Environment.NewLine);
 					}
-					csvBlock.sb.Append(Environment.NewLine);
-					csvBlock.sb.Append(Environment.NewLine);
-				}
-				else
-				{
-					// Max Height  Max Size(KB)   Bias Authored   Current Width   Current Height  Current Size(KB)   Format LODGroup    Name Streaming   UnknownRef VT  Usage Count NumMips Uncompressed
-					// These are the explicit textures lists, slightly different formating
-
-					var headers = new[]
+					else
 					{
+						// Max Height  Max Size(KB)   Bias Authored   Current Width   Current Height  Current Size(KB)   Format LODGroup    Name Streaming   UnknownRef VT  Usage Count NumMips Uncompressed
+						// These are the explicit textures lists, slightly different formating
+
+						var headers = new[]
+						{
 						"Name",
 						"Format",
 						"Current Size (KB)",
@@ -208,49 +231,49 @@ public class MemreportHelper : BuildCommand
 						"NumMips",
 					};
 
-					csvBlock.sb.AppendLine(string.Join(",", headers));
+						csvBlock.sb.AppendLine(string.Join(",", headers));
 
-					foreach (var row in csvBlock.SortedRows)
-					{
-						csvBlock.sb.AppendLine(string.Join(",", headers.Select(h => row[h])));
+						foreach (var row in csvBlock.SortedRows)
+						{
+							csvBlock.sb.AppendLine(string.Join(",", headers.Select(h => row[h])));
+						}
+
+						csvBlock.sb.AppendLine();
+						csvBlock.sb.AppendLine();
 					}
 
-					csvBlock.sb.AppendLine();
-					csvBlock.sb.AppendLine();
 				}
 
-			}
-
-			if (!DirectoryExists(OutputDir))
-			{
-				CreateDirectory(OutputDir);
-			}
+				if (!DirectoryExists(OutputDir))
+				{
+					CreateDirectory(OutputDir);
+				}
 
-			string OutputRootFileName = Path.Combine(OutputDir, Path.GetFileName(MemreportFile));
+				string OutputRootFileName = Path.Combine(OutputDir, Path.GetFileName(MRFile));
 
-			File.WriteAllText(OutputRootFileName + ".csv", "");
+				File.WriteAllText(OutputRootFileName + ".csv", "");
 
-			// output the WHOLE set to 1 file + each set to individual files
-			foreach (var csvBlock in CsvBlocks)
-			{
-				File.AppendAllText(OutputRootFileName + ".csv", csvBlock.sb.ToString());
+				// output the WHOLE set to 1 file + each set to individual files
+				foreach (var csvBlock in CsvBlocks)
+				{
+					File.AppendAllText(OutputRootFileName + ".csv", csvBlock.sb.ToString());
 
-				string SanitizedFileName = OutputRootFileName + "." + csvBlock.Title.Replace("-csv", "").Replace(' ', '_').Replace('-', '_').Replace(':', '_');
+					string SanitizedFileName = OutputRootFileName + "." + csvBlock.Title.Replace("-csv", "").Replace(' ', '_').Replace('-', '_').Replace(':', '_');
 
-				File.WriteAllText(SanitizedFileName + ".csv", csvBlock.sb.ToString());
-			}
+					File.WriteAllText(SanitizedFileName + ".csv", csvBlock.sb.ToString());
+				}
 
-			{
-				// Now output a 'badness' file, this is heavily tailored to workflows in frosty right now (9/20/21 andrew.firth)
+				{
+					// Now output a 'badness' file, this is heavily tailored to workflows in frosty right now (9/20/21 andrew.firth)
 
-				List<string> PreOutput = new List<string>();
+					List<string> PreOutput = new List<string>();
 
-				StringBuilder sb = new StringBuilder();
+					StringBuilder sb = new StringBuilder();
 
-				// Output "bad" textures
-				{
-					var TextureHeaders = new[]
+					// Output "bad" textures
 					{
+						var TextureHeaders = new[]
+						{
 						"Name",
 						"Format",
 						"Current Size (KB)",
@@ -261,135 +284,136 @@ public class MemreportHelper : BuildCommand
 						"Streaming",
 					};
 
-					sb.AppendLine("Type," + string.Join(",", TextureHeaders));
+						sb.AppendLine("Type," + string.Join(",", TextureHeaders));
 
-					foreach (var csvBlock in CsvBlocks)
-					{
-						// limit this to just "all textures" 
-						//if (csvBlock.Title.Contains("Listing NONVT textures") || csvBlock.Title.Contains("Listing uncompressed textures") || csvBlock.Title.Contains("Listing all textures."))
-						if (csvBlock.Title.Contains("Listing all textures."))
+						foreach (var csvBlock in CsvBlocks)
 						{
-							sb.AppendLine(csvBlock.Title);
-							int OutputCount = 0;
-							// output top 10 enties, limited info
-							foreach (var row in csvBlock.SortedRows)
+							// limit this to just "all textures" 
+							//if (csvBlock.Title.Contains("Listing NONVT textures") || csvBlock.Title.Contains("Listing uncompressed textures") || csvBlock.Title.Contains("Listing all textures."))
+							if (csvBlock.Title.Contains("Listing all textures."))
 							{
-								if (PreOutput.Contains(row["Name"]) == false)
+								sb.AppendLine(csvBlock.Title);
+								int OutputCount = 0;
+								// output top 10 enties, limited info
+								foreach (var row in csvBlock.SortedRows)
 								{
-									sb.AppendLine("," + string.Join(",", TextureHeaders.Select(h => row[h])));
+									if (PreOutput.Contains(row["Name"]) == false)
+									{
+										sb.AppendLine("," + string.Join(",", TextureHeaders.Select(h => row[h])));
 
-									PreOutput.Add(row["Name"]);
+										PreOutput.Add(row["Name"]);
 
-									// handle limiting the badness output based on user options
-									OutputCount++;
-									if (OutputCount > MaxOutputCount)
-									{
-										break;
-									}
+										// handle limiting the badness output based on user options
+										OutputCount++;
+										if (OutputCount > MaxOutputCount)
+										{
+											break;
+										}
 
-									if (float.Parse(row["Current Size (KB)"]) < MinOutputKBLimit)
-									{
-										break;
+										if (float.Parse(row["Current Size (KB)"]) < MinOutputKBLimit)
+										{
+											break;
+										}
 									}
 								}
+								sb.AppendLine("");
 							}
-							sb.AppendLine("");
 						}
 					}
-				}
 
-				// Output "bad" meshes
-				{
-					var MeshHeaders = new[]
+					// Output "bad" meshes
 					{
+						var MeshHeaders = new[]
+						{
 						"Object",
 						"ResExcKB"
 					};
 
-					sb.AppendLine("Type," + string.Join(",", MeshHeaders));
+						sb.AppendLine("Type," + string.Join(",", MeshHeaders));
 
-					foreach (var csvBlock in CsvBlocks)
-					{
-						if (csvBlock.Title.Contains("Obj List: class=StaticMesh -alphasort"))
+						foreach (var csvBlock in CsvBlocks)
 						{
-							sb.AppendLine(csvBlock.Title);
-							int OutputCount = 0;
-							// output top 10 enties, limited info
-							foreach (var row in csvBlock.SortedRows)
+							if (csvBlock.Title.Contains("Obj List: class=StaticMesh -alphasort"))
 							{
-								if (PreOutput.Contains(row["Object"]) == false)
+								sb.AppendLine(csvBlock.Title);
+								int OutputCount = 0;
+								// output top 10 enties, limited info
+								foreach (var row in csvBlock.SortedRows)
 								{
-									sb.AppendLine("," + string.Join(",", MeshHeaders.Select(h => row[h])));
+									if (PreOutput.Contains(row["Object"]) == false)
+									{
+										sb.AppendLine("," + string.Join(",", MeshHeaders.Select(h => row[h])));
 
-									PreOutput.Add(row["Object"]);
+										PreOutput.Add(row["Object"]);
 
-									// handle limiting the badness output based on user options
-									OutputCount++;
-									if (OutputCount > MaxOutputCount)
-									{
-										break;
-									}
+										// handle limiting the badness output based on user options
+										OutputCount++;
+										if (OutputCount > MaxOutputCount)
+										{
+											break;
+										}
 
-									if (float.Parse(row["ResExcKB"]) < MeshMin)
-									{
-										break;
+										if (float.Parse(row["ResExcKB"]) < MeshMin)
+										{
+											break;
+										}
 									}
 								}
+								sb.AppendLine("");
 							}
-							sb.AppendLine("");
 						}
 					}
-				}
 
-				// Output "bad" skelmesh
-				{
-					var MeshHeaders = new[]
+					// Output "bad" skelmesh
 					{
+						var MeshHeaders = new[]
+						{
 						"Object",
 						"ResExcKB"
 					};
 
-					sb.AppendLine("Type," + string.Join(",", MeshHeaders));
+						sb.AppendLine("Type," + string.Join(",", MeshHeaders));
 
-					foreach (var csvBlock in CsvBlocks)
-					{
-						if (csvBlock.Title.Contains("Obj List: class=SkeletalMesh"))
+						foreach (var csvBlock in CsvBlocks)
 						{
-							sb.AppendLine(csvBlock.Title);
-							int OutputCount = 0;
-							// output top 10 enties, limited info
-							foreach (var row in csvBlock.SortedRows)
+							if (csvBlock.Title.Contains("Obj List: class=SkeletalMesh"))
 							{
-								if (PreOutput.Contains(row["Object"]) == false)
+								sb.AppendLine(csvBlock.Title);
+								int OutputCount = 0;
+								// output top 10 enties, limited info
+								foreach (var row in csvBlock.SortedRows)
 								{
-									sb.AppendLine("," + string.Join(",", MeshHeaders.Select(h => row[h])));
+									if (PreOutput.Contains(row["Object"]) == false)
+									{
+										sb.AppendLine("," + string.Join(",", MeshHeaders.Select(h => row[h])));
 
-									PreOutput.Add(row["Object"]);
+										PreOutput.Add(row["Object"]);
 
-									// handle limiting the badness output based on user options
-									OutputCount++;
-									if (OutputCount > MaxOutputCount)
-									{
-										break;
-									}
+										// handle limiting the badness output based on user options
+										OutputCount++;
+										if (OutputCount > MaxOutputCount)
+										{
+											break;
+										}
 
-									if (float.Parse(row["ResExcKB"]) < SkelMeshMin)
-									{
-										break;
+										if (float.Parse(row["ResExcKB"]) < SkelMeshMin)
+										{
+											break;
+										}
 									}
 								}
+								sb.AppendLine("");
 							}
-							sb.AppendLine("");
 						}
 					}
-				}
 
 
-				File.WriteAllText(OutputRootFileName + ".badness.csv", sb.ToString());
+					File.WriteAllText(OutputRootFileName + ".badness.csv", sb.ToString());
+				}
 			}
 		}
 
-		LogInformation("************************ MemreportToGSheet WORK COMPLETED ************************");
+		LogInformation("************************ MemreportHelper WORK COMPLETED ************************");
 
 		return ExitCode.Success;
 	}
