[MemreportHelper] - minor fix for the case when a CSV block does not have any entries

diff --git a/Engine/Source/Programs/AutomationTool/Scripts/MemreportHelper.cs b/Engine/Source/Programs/AutomationTool/Scripts/MemreportHelper.cs
--- a/Engine/Source/Programs/AutomationTool/Scripts/MemreportHelper.cs
+++ b/Engine/Source/Programs/AutomationTool/Scripts/MemreportHelper.cs
@@ -39,46 +39,51 @@ public class MemreportHelper : BuildCommand
 			if (InputFileAllLines[StartIndex].Contains("-csv") || SkipCsvRequired)
 			{
 				// hackily remove duped list (these come through as alphasort + resourcesizesort)
-				if (InputFileAllLines[StartIndex].Contains("-resourcesizesort") == false)
+				if (InputFileAllLines[StartIndex].Contains("-alphasort") == false)
 				{
-					List<Dictionary<string, string>> Rows = new List<Dictionary<string, string>>();
-					MRCsvBlock block = new MRCsvBlock();
-
-					block.LineIndex = StartIndex;
-					block.Title = InputFileAllLines[StartIndex];
-					block.SizeKey = SizeKey;
-
 					int HeaderIndex = StartIndex + LineOffset;
 					int DataIndex = HeaderIndex + 1;
 					string[] HeaderValues = InputFileAllLines[HeaderIndex].Split(',');
 
-					// Parse the rest of the data
-					while (DataIndex < InputFileAllLines.Length)
+					if (HeaderValues.Length > 1)
 					{
-						Dictionary<string, string> OneRow = new Dictionary<string, string>();
-						string Line = InputFileAllLines[DataIndex];
-						DataIndex++;
+						List<Dictionary<string, string>> Rows = new List<Dictionary<string, string>>();
+						MRCsvBlock block = new MRCsvBlock();
+
+						block.LineIndex = StartIndex;
+						block.Title = InputFileAllLines[StartIndex];
+						block.SizeKey = SizeKey;
 
-						string[] Values = Line.Split(',');
-						if (Values.Length >= HeaderValues.Length)
+						// Parse the rest of the data
+						while (DataIndex < InputFileAllLines.Length)
 						{
-							// used to remove the starting blank column on obj lists
-							int startindex = SkipCsvRequired ? 0 : 1;
-							for (int i = startindex; i < HeaderValues.Length; i++)
+							Dictionary<string, string> OneRow = new Dictionary<string, string>();
+							string Line = InputFileAllLines[DataIndex];
+							DataIndex++;
+
+							string[] Values = Line.Split(',');
+							if (Values.Length >= HeaderValues.Length)
 							{
-								OneRow[HeaderValues[i].Trim()] = Values[i].Trim();
+								// used to remove the starting blank column on obj lists
+								int startindex = SkipCsvRequired ? 0 : 1;
+								for (int i = startindex; i < HeaderValues.Length; i++)
+								{
+									OneRow[HeaderValues[i].Trim()] = Values[i].Trim();
+								}
+								Rows.Add(OneRow);
+							}
+							else
+							{
+								break;
 							}
-							Rows.Add(OneRow);
 						}
-						else
+						if (Rows.Count > 0)
 						{
-							break;
+							block.SortedRows = Rows.OrderByDescending(x => float.Parse(x[SizeKey])).ToList();
 						}
-					}
 
-					block.SortedRows = Rows.OrderByDescending(x => float.Parse(x[SizeKey])).ToList();
-
-					CsvBlocks.Add(block);
+						CsvBlocks.Add(block);
+					}
 				}
 
 				// replace the line so subsequent calls do not find the same array
