Add PackageEncryptionKeyFile to UAT project params

- added optional UAT parameter PackageEncryptionKeyFile to specify the encryption key file to use during packaging for supporting platforms
  - example usage: -PackageEncryptionKeyFile=C:\MyProject\MyEncryptionKey.key

diff --git a/Engine/Source/Programs/AutomationTool/AutomationUtils/ProjectParams.cs b/Engine/Source/Programs/AutomationTool/AutomationUtils/ProjectParams.cs
--- a/Engine/Source/Programs/AutomationTool/AutomationUtils/ProjectParams.cs
+++ b/Engine/Source/Programs/AutomationTool/AutomationUtils/ProjectParams.cs
@@ -395,6 +395,7 @@ namespace AutomationTool
 			this.ArchiveMetaData = InParams.ArchiveMetaData;
 			this.CreateAppBundle = InParams.CreateAppBundle;
 			this.Distribution = InParams.Distribution;
+			this.PackageEncryptionKeyFile = InParams.PackageEncryptionKeyFile;
 			this.Prereqs = InParams.Prereqs;
 			this.AppLocalDirectory = InParams.AppLocalDirectory;
 			this.NoBootstrapExe = InParams.NoBootstrapExe;
@@ -542,7 +543,8 @@ namespace AutomationTool
 			string SpecifiedServerTarget = null,
 			ParamList<string> ProgramTargets = null,
 			bool? Distribution = null,
-            bool? Prebuilt = null,
+			string PackageEncryptionKeyFile = null,
+			bool? Prebuilt = null,
             int? RunTimeoutSeconds = null,
 			string SpecifiedArchitecture = null,
 			string UbtArgs = null,
@@ -806,6 +808,7 @@ namespace AutomationTool
 			this.ArchiveMetaData = GetParamValueIfNotSpecified(Command, ArchiveMetaData, this.ArchiveMetaData, "archivemetadata");
 			this.CreateAppBundle = GetParamValueIfNotSpecified(Command, CreateAppBundle, true, "createappbundle");
 			this.Distribution = GetParamValueIfNotSpecified(Command, Distribution, this.Distribution, "distribution");
+			this.PackageEncryptionKeyFile = ParseParamValueIfNotSpecified(Command, PackageEncryptionKeyFile, "packageencryptionkeyfile", null);
 			this.Prereqs = GetParamValueIfNotSpecified(Command, Prereqs, this.Prereqs, "prereqs");
 			this.AppLocalDirectory = ParseParamValueIfNotSpecified(Command, AppLocalDirectory, "applocaldirectory", String.Empty, true);
 			this.NoBootstrapExe = GetParamValueIfNotSpecified(Command, NoBootstrapExe, this.NoBootstrapExe, "nobootstrapexe");
@@ -2086,6 +2089,9 @@ namespace AutomationTool
 		[Help("distribution", "package for distribution the project")]
 		public bool Distribution { get; set; }
 
+		[Help("PackageEncryptionKeyFile", "Path to file containing encryption key to use in packaging")]
+		public string PackageEncryptionKeyFile { get; set; }
+
 		[Help("prereqs", "stage prerequisites along with the project")]
 		public bool Prereqs { get; set; }
 
@@ -3023,7 +3029,8 @@ namespace AutomationTool
 				CommandUtils.LogLog("ProjectPlatformBinariesPaths={0}", string.Join(",", ProjectPlatformBinariesPaths));
 				CommandUtils.LogLog("ProjectExePaths={0}", string.Join(",", ProjectExePaths));
 				CommandUtils.LogLog("Distribution={0}", Distribution);
-                CommandUtils.LogLog("Prebuilt={0}", Prebuilt);
+				CommandUtils.LogLog("PackageEncryptionKeyFile={0}", PackageEncryptionKeyFile);
+				CommandUtils.LogLog("Prebuilt={0}", Prebuilt);
 				CommandUtils.LogLog("Prereqs={0}", Prereqs);
 				CommandUtils.LogLog("AppLocalDirectory={0}", AppLocalDirectory);
 				CommandUtils.LogLog("NoBootstrapExe={0}", NoBootstrapExe);
