Add support for on-demand virtual devices in Gauntlet

diff --git a/Engine/Source/Programs/AutomationTool/Gauntlet/Framework/Base/Gauntlet.TargetDevice.cs b/Engine/Source/Programs/AutomationTool/Gauntlet/Framework/Base/Gauntlet.TargetDevice.cs
--- a/Engine/Source/Programs/AutomationTool/Gauntlet/Framework/Base/Gauntlet.TargetDevice.cs
+++ b/Engine/Source/Programs/AutomationTool/Gauntlet/Framework/Base/Gauntlet.TargetDevice.cs
@@ -145,6 +145,15 @@ namespace Gauntlet
 		void CleanupDevices();
 	}
 
+	/// <summary>
+	/// Represents a class that can provides virtual local devices 
+	/// </summary>
+	public interface IVirtualLocalDevice : IDeviceSource
+	{
+		bool CanRunVirtualFromPlatform(UnrealTargetPlatform? Platfrom);
+		UnrealTargetPlatform? GetPlatform();
+	}
+
 	/// <summary>
 	/// Represents a class that tell what build the device support
 	/// </summary>
diff --git a/Engine/Source/Programs/AutomationTool/Gauntlet/Framework/Devices/Gauntlet.DevicePool.cs b/Engine/Source/Programs/AutomationTool/Gauntlet/Framework/Devices/Gauntlet.DevicePool.cs
--- a/Engine/Source/Programs/AutomationTool/Gauntlet/Framework/Devices/Gauntlet.DevicePool.cs
+++ b/Engine/Source/Programs/AutomationTool/Gauntlet/Framework/Devices/Gauntlet.DevicePool.cs
@@ -557,24 +557,45 @@ namespace Gauntlet
 			DeviceURL = InDeviceURL;
 		}
 
-		public void AddLocalDevices(int Count)
+		public void AddLocalDevices(int MaxCount)
 		{
-			// todo - default per platform?
-
-			UnrealTargetPlatform LocalPlat = UnrealTargetPlatform.Win64;
+			UnrealTargetPlatform LocalPlat = BuildHostPlatform.Current.Platform;
 
 			int NumDevices = GetAvailableDeviceCount(new UnrealDeviceTargetConstraint(LocalPlat));
 
-			// add local PCs (todo - some max number?)
-			for (int i = NumDevices; i < Count + NumDevices; i++)
+			for (int i = NumDevices; i < MaxCount; i++)
 			{
 				DeviceDefinition Def = new DeviceDefinition();
 				Def.Name = string.Format("LocalDevice{0}", i);
-				Def.Platform = BuildHostPlatform.Current.Platform;
+				Def.Platform = LocalPlat;
 				UnprovisionedDevices.Add(Def);
 			}
 		}
 
+		public void AddVirtualDevices(int MaxCount)
+		{
+			UnrealTargetPlatform LocalPlat = BuildHostPlatform.Current.Platform;
+
+			IEnumerable<IVirtualLocalDevice> VirtualDevices = Gauntlet.Utils.InterfaceHelpers.FindImplementations<IVirtualLocalDevice>()
+					.Where(F => F.CanRunVirtualFromPlatform(LocalPlat));
+
+			foreach (IVirtualLocalDevice Device in VirtualDevices)
+			{
+				UnrealTargetPlatform? DevicePlatform = Device.GetPlatform();
+				if (DevicePlatform != null)
+				{
+					int NumDevices = GetAvailableDeviceCount(new UnrealDeviceTargetConstraint(DevicePlatform));
+					for (int i = NumDevices; i < MaxCount; i++)
+					{
+						DeviceDefinition Def = new DeviceDefinition();
+						Def.Name = string.Format("Virtual{0}{1}", DevicePlatform.ToString(), i);
+						Def.Platform = DevicePlatform;
+						UnprovisionedDevices.Add(Def);
+					}
+				}
+			}
+		}
+
 		/// <summary>
 		/// Reserve devices from service
 		/// </summary>
diff --git a/Engine/Source/Programs/AutomationTool/Gauntlet/Unreal/Base/Gauntlet.UnrealTestContext.cs b/Engine/Source/Programs/AutomationTool/Gauntlet/Unreal/Base/Gauntlet.UnrealTestContext.cs
--- a/Engine/Source/Programs/AutomationTool/Gauntlet/Unreal/Base/Gauntlet.UnrealTestContext.cs
+++ b/Engine/Source/Programs/AutomationTool/Gauntlet/Unreal/Base/Gauntlet.UnrealTestContext.cs
@@ -99,6 +99,18 @@ namespace Gauntlet
 		[AutoParam("")]
 		public string DeviceURL;
 
+		/// <summary>
+		/// Maximum number of available local(host platform) devices
+		/// </summary>
+		[AutoParam(10)]
+		public int MaxLocalDevices;
+
+		/// <summary>
+		/// Maximum number of available virtual devices of each type (supported on the host platform)
+		/// </summary>
+		[AutoParam(2)]
+		public int MaxVirtualDevices;
+
 		/// <summary>
 		/// Details for current job (example: link to CIS, etc)
 		/// </summary>
diff --git a/Engine/Source/Programs/AutomationTool/Gauntlet/Unreal/RunUnreal.cs b/Engine/Source/Programs/AutomationTool/Gauntlet/Unreal/RunUnreal.cs
--- a/Engine/Source/Programs/AutomationTool/Gauntlet/Unreal/RunUnreal.cs
+++ b/Engine/Source/Programs/AutomationTool/Gauntlet/Unreal/RunUnreal.cs
@@ -417,7 +417,8 @@ namespace Gauntlet
 			Reservation.ReservationDetails = Options.JobDetails;
 
 			DevicePool.Instance.SetLocalOptions(Options.TempDir, Options.Parallel > 1, Options.DeviceURL);
-			DevicePool.Instance.AddLocalDevices(10);
+			DevicePool.Instance.AddLocalDevices(Options.MaxLocalDevices);
+			DevicePool.Instance.AddVirtualDevices(Options.MaxVirtualDevices);
 
 			foreach (var DeviceWithParams in Options.DeviceList)
 			{
