Catch IO issues when creating or copying ZIP files

diff --git a/Engine/Source/Programs/AutomationTool/BuildGraph/TempStorage.cs b/Engine/Source/Programs/AutomationTool/BuildGraph/TempStorage.cs
--- a/Engine/Source/Programs/AutomationTool/BuildGraph/TempStorage.cs
+++ b/Engine/Source/Programs/AutomationTool/BuildGraph/TempStorage.cs
@@ -994,27 +994,44 @@ namespace AutomationTool
 					FileReference File;
 					if (FilesToZip.TryDequeue(out File))
 					{
-						// Create one zip per thread using the given basename
-						using (var ZipArchive = System.IO.Compression.ZipFile.Open(ZipFileName.FullName, System.IO.Compression.ZipArchiveMode.Create))
+						try
 						{
-							// pull from the queue until we are out of files.
-							do
+							// Create one zip per thread using the given basename
+							using (var ZipArchive = System.IO.Compression.ZipFile.Open(ZipFileName.FullName, System.IO.Compression.ZipArchiveMode.Create))
 							{
-								// use fastest compression. In our best case we are CPU bound, so this is a good tradeoff,
-								// cutting overall time by 2/3 while only modestly increasing the compression ratio (22.7% -> 23.8% for RootEditor PDBs).
-								// This is in cases of a super hot cache, so the operation was largely CPU bound.
-								ZipArchiveExtensions.CreateEntryFromFile_CrossPlatform(ZipArchive, File.FullName, CommandUtils.ConvertSeparators(PathSeparator.Slash, File.MakeRelativeTo(RootDir)), System.IO.Compression.CompressionLevel.Fastest);
-							} while (FilesToZip.TryDequeue(out File));
+								// pull from the queue until we are out of files.
+								do
+								{
+									// use fastest compression. In our best case we are CPU bound, so this is a good tradeoff,
+									// cutting overall time by 2/3 while only modestly increasing the compression ratio (22.7% -> 23.8% for RootEditor PDBs).
+									// This is in cases of a super hot cache, so the operation was largely CPU bound.
+									ZipArchiveExtensions.CreateEntryFromFile_CrossPlatform(ZipArchive, File.FullName, CommandUtils.ConvertSeparators(PathSeparator.Slash, File.MakeRelativeTo(RootDir)), System.IO.Compression.CompressionLevel.Fastest);
+								} while (FilesToZip.TryDequeue(out File));
+							}
+						}
+
+						catch (IOException)
+						{
+							throw new AutomationException("Unable to open file {0}", ZipFileName.FullName);
 						}
-						// if we are using a staging dir, copy to the final location and delete the staged copy.
-						FileInfo ZipFile = new FileInfo(ZipFileName.FullName);
-						if (StagingDir != null)
+
+						try
+						{
+							// if we are using a staging dir, copy to the final location and delete the staged copy.
+							FileInfo ZipFile = new FileInfo(ZipFileName.FullName);
+							if (StagingDir != null)
+							{
+								FileInfo NewZipFile = ZipFile.CopyTo(CommandUtils.MakeRerootedFilePath(ZipFile.FullName, StagingDir.FullName, OutputDir.FullName));
+								ZipFile.Delete();
+								ZipFile = NewZipFile;
+							}
+							ZipFiles.Add(ZipFile);					
+						}
+
+						catch (IOException)
 						{
-							FileInfo NewZipFile = ZipFile.CopyTo(CommandUtils.MakeRerootedFilePath(ZipFile.FullName, StagingDir.FullName, OutputDir.FullName));
-							ZipFile.Delete();
-							ZipFile = NewZipFile;
+							throw new AutomationException("Unable to copy file {0}", ZipFileName.FullName);
 						}
-						ZipFiles.Add(ZipFile);
 					}
 				})).ToList();
 
