Add error logging for temp storage zip operations
#rnx

diff --git a/Engine/Source/Programs/AutomationTool/BuildGraph/TempStorage.cs b/Engine/Source/Programs/AutomationTool/BuildGraph/TempStorage.cs
--- a/Engine/Source/Programs/AutomationTool/BuildGraph/TempStorage.cs
+++ b/Engine/Source/Programs/AutomationTool/BuildGraph/TempStorage.cs
@@ -1012,6 +1012,7 @@ namespace AutomationTool
 
 						catch (IOException)
 						{
+							CommandUtils.LogError("Unable to open file for TempStorage zip: \"{0}\"", ZipFileName.FullName);
 							throw new AutomationException("Unable to open file {0}", ZipFileName.FullName);
 						}
 
@@ -1030,6 +1031,7 @@ namespace AutomationTool
 
 						catch (IOException)
 						{
+							CommandUtils.LogError("Unable to copy file for TempStorage zip: \"{0}\"", ZipFileName.FullName);
 							throw new AutomationException("Unable to copy file {0}", ZipFileName.FullName);
 						}
 					}
