Gauntlet - Fix test passing even with fatal error

#rnx

diff --git a/Engine/Source/Programs/AutomationTool/Gauntlet/Unreal/Base/Gauntlet.UnrealTestNode.cs b/Engine/Source/Programs/AutomationTool/Gauntlet/Unreal/Base/Gauntlet.UnrealTestNode.cs
--- a/Engine/Source/Programs/AutomationTool/Gauntlet/Unreal/Base/Gauntlet.UnrealTestNode.cs
+++ b/Engine/Source/Programs/AutomationTool/Gauntlet/Unreal/Base/Gauntlet.UnrealTestNode.cs
@@ -80,8 +80,8 @@ namespace Gauntlet
 		}
 
 		/// <summary>
-		/// Returns Errors found during tests. By default only fatal errors are considered
-		/// returning lists of events
+		/// Returns Errors found during tests. By default fatal and error severities are considered.
+		/// returning lists of event summaries
 		/// </summary>
 		public override IEnumerable<string> GetErrors()
 		{
@@ -89,14 +89,12 @@ namespace Gauntlet
 			
 			if (RoleResults != null)
 			{
-				ErrorList = ErrorList.Union(RoleResults.SelectMany(R => R.Events.Where(E => E.Severity == EventSeverity.Error)).Select(E => E.Summary));
+				ErrorList = ErrorList.Union(RoleResults.SelectMany(R => R.Events.Where(E => E.Severity == EventSeverity.Error || E.Severity == EventSeverity.Fatal)).Select(E => E.Summary));
 			}
 
-			return ErrorList.ToArray();
-			
+			return ErrorList.ToArray();			
 		}
 
-
 		/// <summary>
 		/// Returns Errors found during tests. Including Abnormal Exit reasons
 		/// </summary>
@@ -108,24 +106,8 @@ namespace Gauntlet
 				return Errors;
 			}
 
-			Dictionary<UnrealRoleResult, Tuple<int, string>> ErrorCodesAndReasons = new Dictionary<UnrealRoleResult, Tuple<int, string>>();
-
-			var FailedRoles = RoleResults.Where(
-				R => {
-					if (R.Artifacts.AppInstance.WasKilled)
-					{
-						return false;
-					}
-					
-					ErrorCodesAndReasons.Add(R, new Tuple<int, string>(R.ExitCode, R.Summary));
-					return R.ExitCode != 0;
-				}
-			);
-
-			return Errors.Union(FailedRoles.Select(
-				R => {
-					return string.Format("Abnormal Exit: Reason={0}, ExitCode={1}, Log={2}", R.Summary, R.ExitCode, Path.GetFileName(R.Artifacts.LogPath));
-				}
+			return Errors.Union(RoleResults.Where(R => R.ProcessResult != UnrealProcessResult.ExitOk).Select(
+				R => string.Format("Abnormal Exit: Reason={0}, ExitCode={1}, Log={2}", R.Summary, R.ExitCode, Path.GetFileName(R.Artifacts.LogPath))
 			));
 		}
 
@@ -1377,24 +1359,6 @@ namespace Gauntlet
 		/// <returns></returns>
 		protected virtual UnrealProcessResult GetExitCodeAndReason(StopReason InReason, UnrealLog InLog, UnrealRoleArtifacts InArtifacts, out string ExitReason, out int ExitCode)
 		{
-			// Gauntlet killed the process. This can be valid in many scenarios (e.g. shutting down an ancillary 
-			// process, but if there was a timeout it will be handled at a higher level
-			if (InArtifacts.AppInstance.WasKilled)
-			{
-				if (InReason == StopReason.MaxDuration)
-				{
-					ExitReason = "Process was killed by Gauntlet due to a timeout";
-					ExitCode = -1;
-					return UnrealProcessResult.TimeOut;
-				}
-				else
-				{
-					ExitReason = "Process was killed by Gauntlet";
-					ExitCode = 0;
-					return UnrealProcessResult.ExitOk;
-				}
-			}
-
 			// first check for fatal issues
 			if (InLog.FatalError != null)
 			{
@@ -1419,6 +1383,24 @@ namespace Gauntlet
 				return UnrealProcessResult.EncounteredEnsure;
 			}
 
+			// Gauntlet killed the process. This can be valid in many scenarios (e.g. shutting down an ancillary 
+			// process, but if there was a timeout it will be handled at a higher level
+			if (InArtifacts.AppInstance.WasKilled)
+			{
+				if (InReason == StopReason.MaxDuration)
+				{
+					ExitReason = "Process was killed by Gauntlet due to a timeout";
+					ExitCode = -1;
+					return UnrealProcessResult.TimeOut;
+				}
+				else
+				{
+					ExitReason = string.Format("Process was killed by Gauntlet [Reason={0}]", InReason.ToString());
+					ExitCode = 0;
+					return UnrealProcessResult.ExitOk;
+				}
+			}
+
 			// If we found a valid exit code with test markup, return it
 			if (InLog.HasTestExitCode)
 			{
