UAT: Add support for async build commands and BuildGraph tasks.

diff --git a/Engine/Source/Programs/AutomationTool/AutomationUtils/Automation.cs b/Engine/Source/Programs/AutomationTool/AutomationUtils/Automation.cs
--- a/Engine/Source/Programs/AutomationTool/AutomationUtils/Automation.cs
+++ b/Engine/Source/Programs/AutomationTool/AutomationUtils/Automation.cs
@@ -10,6 +10,7 @@ using UnrealBuildTool;
 using EpicGames.Core;
 using OpenTracing.Util;
 using UnrealBuildBase;
+using System.Threading.Tasks;
 
 namespace AutomationTool
 {
@@ -46,7 +47,7 @@ namespace AutomationTool
 		/// Main method.
 		/// </summary>
 		/// <param name="Arguments">Command line</param>
-		public static ExitCode Process(ParsedCommandLine AutomationToolCommandLine, StartupTraceListener StartupListener, HashSet<FileReference> ScriptModuleAssemblies)
+		public static async Task<ExitCode> ProcessAsync(ParsedCommandLine AutomationToolCommandLine, StartupTraceListener StartupListener, HashSet<FileReference> ScriptModuleAssemblies)
 		{
 			GlobalCommandLine.Initialize(AutomationToolCommandLine);
 
@@ -150,7 +151,7 @@ namespace AutomationTool
 				try
 				{
 					// Find and execute commands.
-					ExitCode Result = Execute(AutomationToolCommandLine.CommandsToExecute, ScriptManager.Commands);
+					ExitCode Result = await ExecuteAsync(AutomationToolCommandLine.CommandsToExecute, ScriptManager.Commands);
 					if (TelemetryFile != null)
 					{
 						Directory.CreateDirectory(Path.GetDirectoryName(TelemetryFile));
@@ -234,7 +235,7 @@ namespace AutomationTool
 		/// </summary>
 		/// <param name="CommandsToExecute"></param>
 		/// <param name="Commands"></param>
-		public static ExitCode Execute(List<CommandInfo> CommandsToExecute, Dictionary<string, Type> Commands)
+		public static async Task<ExitCode> ExecuteAsync(List<CommandInfo> CommandsToExecute, Dictionary<string, Type> Commands)
 		{
 			for (int CommandIndex = 0; CommandIndex < CommandsToExecute.Count; ++CommandIndex)
 			{
@@ -250,7 +251,7 @@ namespace AutomationTool
 				Command.Params = CommandInfo.Arguments.ToArray();
 				try
 				{
-					ExitCode Result = Command.Execute();
+					ExitCode Result = await Command.ExecuteAsync();
 					if(Result != ExitCode.Success)
 					{
 						return Result;
diff --git a/Engine/Source/Programs/AutomationTool/AutomationUtils/BuildCommand.cs b/Engine/Source/Programs/AutomationTool/AutomationUtils/BuildCommand.cs
--- a/Engine/Source/Programs/AutomationTool/AutomationUtils/BuildCommand.cs
+++ b/Engine/Source/Programs/AutomationTool/AutomationUtils/BuildCommand.cs
@@ -3,6 +3,7 @@ using System;
 using System.Collections.Generic;
 using System.Linq;
 using System.Text;
+using System.Threading.Tasks;
 using EpicGames.Core;
 using UnrealBuildBase;
 
@@ -344,6 +345,14 @@ namespace AutomationTool
 			return ExitCode.Success;
 		}
 
+		/// <summary>
+		/// Async command entry point.
+		/// </summary>
+		public virtual Task<ExitCode> ExecuteAsync()
+		{
+			return Task.FromResult(Execute());
+		}
+
 		/// <summary>
 		/// Executes a new command as a child of another command.
 		/// </summary>
diff --git a/Engine/Source/Programs/AutomationTool/BuildGraph/BuildGraph.cs b/Engine/Source/Programs/AutomationTool/BuildGraph/BuildGraph.cs
--- a/Engine/Source/Programs/AutomationTool/BuildGraph/BuildGraph.cs
+++ b/Engine/Source/Programs/AutomationTool/BuildGraph/BuildGraph.cs
@@ -160,7 +160,7 @@ namespace AutomationTool
 	/// Build graphs are declared using an XML script using syntax similar to MSBuild, ANT or NAnt, and consist of the following components:
 	///
 	/// - Tasks:        Building blocks which can be executed as part of the build process. Many predefined tasks are provided ('Cook', 'Compile', 'Copy', 'Stage', 'Log', 'PakFile', etc...), and additional tasks may be 
-	///                 added be declaring classes derived from AutomationTool.CustomTask in other UAT modules. 
+	///                 added be declaring classes derived from AutomationTool.BuildTask in other UAT modules. 
 	/// - Nodes:        A named sequence of tasks which are executed in order to produce outputs. Nodes may have dependencies on other nodes for their outputs before they can be executed. Declared with the 'Node' element.
 	/// - Agents:		A machine which can execute a sequence of nodes, if running as part of a build system. Has no effect when building locally. Declared with the 'Agent' element.
 	/// - Triggers:     Container for agents which should only be executed when explicitly triggered (using the -Trigger=... or -SkipTriggers command line argument). Declared with the 'Trigger' element.
@@ -221,7 +221,7 @@ namespace AutomationTool
 		/// <summary>
 		/// Main entry point for the BuildGraph command
 		/// </summary>
-		public override ExitCode Execute()
+		public override async Task<ExitCode> ExecuteAsync()
 		{
 			// Parse the command line parameters
 			string ScriptFileName = ParseParamValue("Script", null);
@@ -633,7 +633,7 @@ namespace AutomationTool
 			}
 
 			// Create tasks for the entire graph
-			Dictionary<BgTask, CustomTask> TaskInfoToTask = new Dictionary<BgTask, CustomTask>();
+			Dictionary<BgTask, BgTaskImpl> TaskInfoToTask = new Dictionary<BgTask, BgTaskImpl>();
 			if (bSkipValidation && SingleNode != null)
 			{
 				if (!CreateTaskInstances(Graph, SingleNode, NameToTask, TaskInfoToTask))
@@ -654,14 +654,14 @@ namespace AutomationTool
 			{
 				if(SingleNode != null)
 				{
-					if(!BuildNode(new JobContext(this), Graph, SingleNode, TaskInfoToTask, Storage, bWithBanner: true))
+					if(!await BuildNodeAsync(new JobContext(this), Graph, SingleNode, TaskInfoToTask, Storage, bWithBanner: true))
 					{
 						return ExitCode.Error_Unknown;
 					}
 				}
 				else
 				{
-					if(!BuildAllNodes(new JobContext(this), Graph, TaskInfoToTask, Storage))
+					if(!await BuildAllNodesAsync(new JobContext(this), Graph, TaskInfoToTask, Storage))
 					{
 						return ExitCode.Error_Unknown;
 					}
@@ -670,7 +670,7 @@ namespace AutomationTool
 			return ExitCode.Success;
 		}
 
-		bool CreateTaskInstances(BgGraph Graph, Dictionary<string, ScriptTaskBinding> NameToTask, Dictionary<BgTask, CustomTask> TaskInfoToTask)
+		bool CreateTaskInstances(BgGraph Graph, Dictionary<string, ScriptTaskBinding> NameToTask, Dictionary<BgTask, BgTaskImpl> TaskInfoToTask)
 		{
 			bool bResult = true;
 			foreach (BgAgent Agent in Graph.Agents)
@@ -683,12 +683,12 @@ namespace AutomationTool
 			return bResult;
 		}
 
-		bool CreateTaskInstances(BgGraph Graph, BgNode Node, Dictionary<string, ScriptTaskBinding> NameToTask, Dictionary<BgTask, CustomTask> TaskInfoToTask)
+		bool CreateTaskInstances(BgGraph Graph, BgNode Node, Dictionary<string, ScriptTaskBinding> NameToTask, Dictionary<BgTask, BgTaskImpl> TaskInfoToTask)
 		{
 			bool bResult = true;
 			foreach (BgTask TaskInfo in Node.Tasks)
 			{
-				CustomTask Task = BindTask(Node, TaskInfo, NameToTask, Graph.TagNameToNodeOutput);
+				BgTaskImpl Task = BindTask(Node, TaskInfo, NameToTask, Graph.TagNameToNodeOutput);
 				if (Task == null)
 				{
 					bResult = false;
@@ -701,7 +701,7 @@ namespace AutomationTool
 			return bResult;
 		}
 
-		CustomTask BindTask(BgNode Node, BgTask TaskInfo, Dictionary<string, ScriptTaskBinding> NameToTask, IReadOnlyDictionary<string, BgNodeOutput> TagNameToNodeOutput)
+		BgTaskImpl BindTask(BgNode Node, BgTask TaskInfo, Dictionary<string, ScriptTaskBinding> NameToTask, IReadOnlyDictionary<string, BgNodeOutput> TagNameToNodeOutput)
 		{
 			// Get the reflection info for this element
 			ScriptTaskBinding Task;
@@ -752,7 +752,7 @@ namespace AutomationTool
 					}
 
 					// Parse the values and add them to the collection
-					List<string> ValueStrings = CustomTask.SplitDelimitedList(Value);
+					List<string> ValueStrings = BgTaskImpl.SplitDelimitedList(Value);
 					foreach (string ValueString in ValueStrings)
 					{
 						object ElementValue = ParseValue(ValueString, Parameter.ValueType);
@@ -768,7 +768,7 @@ namespace AutomationTool
 			}
 
 			// Add it to the list
-			CustomTask NewTask = (CustomTask)Activator.CreateInstance(Task.TaskClass, ParametersObject);
+			BgTaskImpl NewTask = (BgTaskImpl)Activator.CreateInstance(Task.TaskClass, ParametersObject);
 
 			// Set up the source location for diagnostics
 			NewTask.SourceLocation = TaskInfo.Location;
@@ -820,11 +820,11 @@ namespace AutomationTool
 			}
 			else if (ValueType == typeof(FileReference))
 			{
-				return CustomTask.ResolveFile(ValueText);
+				return BgTaskImpl.ResolveFile(ValueText);
 			}
 			else if (ValueType == typeof(DirectoryReference))
 			{
-				return CustomTask.ResolveDirectory(ValueText);
+				return BgTaskImpl.ResolveDirectory(ValueText);
 			}
 
 			TypeConverter Converter = TypeDescriptor.GetConverter(ValueType);
@@ -873,9 +873,9 @@ namespace AutomationTool
 				{
 					foreach(TaskElementAttribute ElementAttribute in Type.GetCustomAttributes<TaskElementAttribute>())
 					{
-						if(!Type.IsSubclassOf(typeof(CustomTask)))
+						if(!Type.IsSubclassOf(typeof(BgTaskImpl)))
 						{
-							CommandUtils.LogError("Class '{0}' has TaskElementAttribute, but is not derived from 'Task'", Type.Name);
+							CommandUtils.LogError("Class '{0}' has TaskElementAttribute, but is not derived from 'BgTaskImpl'", Type.Name);
 							return false;
 						}
 						if(NameToTask.ContainsKey(ElementAttribute.Name))
@@ -1004,7 +1004,7 @@ namespace AutomationTool
 		/// <param name="Graph">The graph instance</param>
 		/// <param name="Storage">The temp storage backend which stores the shared state</param>
 		/// <returns>True if everything built successfully</returns>
-		bool BuildAllNodes(JobContext Job, BgGraph Graph, Dictionary<BgTask, CustomTask> TaskInfoToTask, TempStorage Storage)
+		async Task<bool> BuildAllNodesAsync(JobContext Job, BgGraph Graph, Dictionary<BgTask, BgTaskImpl> TaskInfoToTask, TempStorage Storage)
 		{
 			// Build a flat list of nodes to execute, in order
 			BgNode[] NodesToExecute = Graph.Agents.SelectMany(x => x.Nodes).ToArray();
@@ -1029,7 +1029,7 @@ namespace AutomationTool
 				if(!Storage.IsComplete(NodeToExecute.Name))
 				{
 					LogInformation("");
-					if(!BuildNode(Job, Graph, NodeToExecute, TaskInfoToTask, Storage, bWithBanner: false))
+					if(!await BuildNodeAsync(Job, Graph, NodeToExecute, TaskInfoToTask, Storage, bWithBanner: false))
 					{
 						return false;
 					} 
@@ -1049,7 +1049,7 @@ namespace AutomationTool
 		/// <param name="Storage">The temp storage backend which stores the shared state</param>
 		/// <param name="bWithBanner">Whether to write a banner before and after this node's log output</param>
 		/// <returns>True if the node built successfully, false otherwise.</returns>
-		bool BuildNode(JobContext Job, BgGraph Graph, BgNode Node, Dictionary<BgTask, CustomTask> TaskInfoToTask, TempStorage Storage, bool bWithBanner)
+		async Task<bool> BuildNodeAsync(JobContext Job, BgGraph Graph, BgNode Node, Dictionary<BgTask, BgTaskImpl> TaskInfoToTask, TempStorage Storage, bool bWithBanner)
 		{
 			DirectoryReference RootDir = new DirectoryReference(CommandUtils.CmdEnv.LocalRoot);
 
@@ -1106,7 +1106,7 @@ namespace AutomationTool
 				Console.WriteLine();
 				CommandUtils.LogInformation("========== Starting: {0} ==========", Node.Name);
 			}
-			if(!ExecuteTasks(Node, Job, TaskInfoToTask, TagNameToFileSet))
+			if(!await ExecuteTasksAsync(Node, Job, TaskInfoToTask, TagNameToFileSet))
 			{
 				return false;
 			}
@@ -1237,9 +1237,9 @@ namespace AutomationTool
 		/// <param name="TaskInfoToTask">Map from TaskInfo to Task object</param>
 		/// <param name="TagNameToFileSet">Mapping from tag names to the set of files they include. Should be set to contain the node inputs on entry.</param>
 		/// <returns>Whether the task succeeded or not. Exiting with an exception will be caught and treated as a failure.</returns>
-		bool ExecuteTasks(BgNode Node, JobContext Job, Dictionary<BgTask, CustomTask> TaskInfoToTask, Dictionary<string, HashSet<FileReference>> TagNameToFileSet)
+		async Task<bool> ExecuteTasksAsync(BgNode Node, JobContext Job, Dictionary<BgTask, BgTaskImpl> TaskInfoToTask, Dictionary<string, HashSet<FileReference>> TagNameToFileSet)
 		{
-			List<CustomTask> Tasks = Node.Tasks.ConvertAll(x => TaskInfoToTask[x]);
+			List<BgTaskImpl> Tasks = Node.Tasks.ConvertAll(x => TaskInfoToTask[x]);
 
 			// Run each of the tasks in order
 			HashSet<FileReference> BuildProducts = TagNameToFileSet[Node.DefaultOutput.TagName];
@@ -1254,7 +1254,7 @@ namespace AutomationTool
 						try
 						{
 							Tasks[Idx].GetTraceMetadata(Scope.Span, "");
-							Tasks[Idx].Execute(Job, BuildProducts, TagNameToFileSet);
+							await Tasks[Idx].ExecuteAsync(Job, BuildProducts, TagNameToFileSet);
 						}
 						catch (Exception Ex)
 						{
@@ -1279,7 +1279,7 @@ namespace AutomationTool
 						}
 						try
 						{
-							Executor.Execute(Job, BuildProducts, TagNameToFileSet);
+							await Executor.ExecuteAsync(Job, BuildProducts, TagNameToFileSet);
 						}
 						catch (Exception Ex)
 						{
diff --git a/Engine/Source/Programs/AutomationTool/BuildGraph/Task.cs b/Engine/Source/Programs/AutomationTool/BuildGraph/Task.cs
--- a/Engine/Source/Programs/AutomationTool/BuildGraph/Task.cs
+++ b/Engine/Source/Programs/AutomationTool/BuildGraph/Task.cs
@@ -80,7 +80,7 @@ namespace AutomationTool
 		/// </summary>
 		/// <param name="Task">Task to add</param>
 		/// <returns>True if the task could be added, false otherwise</returns>
-		bool Add(CustomTask Task);
+		bool Add(BgTaskImpl Task);
 
 		/// <summary>
 		/// Execute all the tasks added to this executor.
@@ -89,13 +89,13 @@ namespace AutomationTool
 		/// <param name="BuildProducts">Set of build products produced by this node.</param>
 		/// <param name="TagNameToFileSet">Mapping from tag names to the set of files they include</param>
 		/// <returns>Whether the task succeeded or not. Exiting with an exception will be caught and treated as a failure.</returns>
-		void Execute(JobContext Job, HashSet<FileReference> BuildProducts, Dictionary<string, HashSet<FileReference>> TagNameToFileSet);
+		Task ExecuteAsync(JobContext Job, HashSet<FileReference> BuildProducts, Dictionary<string, HashSet<FileReference>> TagNameToFileSet);
 	}
 
 	/// <summary>
 	/// Base class for all custom build tasks
 	/// </summary>
-	public abstract class CustomTask
+	public abstract class BgTaskImpl
 	{
 		/// <summary>
 		/// Line number in a source file that this task was declared. Optional; used for log messages.
@@ -109,7 +109,7 @@ namespace AutomationTool
 		/// <param name="BuildProducts">Set of build products produced by this node.</param>
 		/// <param name="TagNameToFileSet">Mapping from tag names to the set of files they include</param>
 		/// <returns>Whether the task succeeded or not. Exiting with an exception will be caught and treated as a failure.</returns>
-		public abstract void Execute(JobContext Job, HashSet<FileReference> BuildProducts, Dictionary<string, HashSet<FileReference>> TagNameToFileSet);
+		public abstract Task ExecuteAsync(JobContext Job, HashSet<FileReference> BuildProducts, Dictionary<string, HashSet<FileReference>> TagNameToFileSet);
 
 		/// <summary>
 		/// Creates a proxy to execute this node.
@@ -425,4 +425,26 @@ namespace AutomationTool
 			return Text.Split(';').Select(x => x.Trim()).Where(x => x.Length > 0).ToList();
 		}
 	}
+
+	/// <summary>
+	/// Legacy implementation of <see cref="BgTaskImpl"/> which operates synchronously
+	/// </summary>
+	public abstract class CustomTask : BgTaskImpl
+	{
+		/// <inheritdoc/>
+		public sealed override Task ExecuteAsync(JobContext Job, HashSet<FileReference> BuildProducts, Dictionary<string, HashSet<FileReference>> TagNameToFileSet)
+		{
+			Execute(Job, BuildProducts, TagNameToFileSet);
+			return Task.CompletedTask;
+		}
+
+		/// <summary>
+		/// Execute this node.
+		/// </summary>
+		/// <param name="Job">Information about the current job</param>
+		/// <param name="BuildProducts">Set of build products produced by this node.</param>
+		/// <param name="TagNameToFileSet">Mapping from tag names to the set of files they include</param>
+		/// <returns>Whether the task succeeded or not. Exiting with an exception will be caught and treated as a failure.</returns>
+		public abstract void Execute(JobContext Job, HashSet<FileReference> BuildProducts, Dictionary<string, HashSet<FileReference>> TagNameToFileSet);
+	}
 }
diff --git a/Engine/Source/Programs/AutomationTool/BuildGraph/Tasks/AgeStoreTask.cs b/Engine/Source/Programs/AutomationTool/BuildGraph/Tasks/AgeStoreTask.cs
--- a/Engine/Source/Programs/AutomationTool/BuildGraph/Tasks/AgeStoreTask.cs
+++ b/Engine/Source/Programs/AutomationTool/BuildGraph/Tasks/AgeStoreTask.cs
@@ -8,6 +8,7 @@ using System.Linq;
 using System.Xml;
 using EpicGames.Core;
 using UnrealBuildTool;
+using System.Threading.Tasks;
 
 namespace Win.Automation
 {
@@ -54,7 +55,7 @@ namespace Win.Automation
 	/// difference is that it uses the last modified time rather than last access time to determine which files to delete.
 	/// </summary>
 	[TaskElement("AgeStore", typeof(AgeStoreTaskParameters))]
-	public class AgeStoreTask : CustomTask
+	public class AgeStoreTask : BgTaskImpl
 	{
 		/// <summary>
 		/// Parameters for this task
@@ -177,7 +178,7 @@ namespace Win.Automation
 		/// <param name="Job">Information about the current job</param>
 		/// <param name="BuildProducts">Set of build products produced by this node.</param>
 		/// <param name="TagNameToFileSet">Mapping from tag names to the set of files they include</param>
-		public override void Execute(JobContext Job, HashSet<FileReference> BuildProducts, Dictionary<string, HashSet<FileReference>> TagNameToFileSet)
+		public override Task ExecuteAsync(JobContext Job, HashSet<FileReference> BuildProducts, Dictionary<string, HashSet<FileReference>> TagNameToFileSet)
 		{
 			// Get the list of symbol file name patterns from the platform.
 			Platform TargetPlatform = Platform.GetPlatform(Parameters.Platform);
@@ -220,6 +221,7 @@ namespace Win.Automation
 			{
 				RecurseDirectory(ExpireTimeUtc, new DirectoryInfo(SymbolServerDirectory.FullName), DirectoryStructure, 0, Filter, ExistingBuilds, TargetPlatform.SymbolServerDeleteIndividualFiles);
 			});
+			return Task.CompletedTask;
 		}
 
 		/// <summary>
diff --git a/Engine/Source/Programs/AutomationTool/BuildGraph/Tasks/AwsAssumeRoleTask.cs b/Engine/Source/Programs/AutomationTool/BuildGraph/Tasks/AwsAssumeRoleTask.cs
--- a/Engine/Source/Programs/AutomationTool/BuildGraph/Tasks/AwsAssumeRoleTask.cs
+++ b/Engine/Source/Programs/AutomationTool/BuildGraph/Tasks/AwsAssumeRoleTask.cs
@@ -8,6 +8,7 @@ using System.IO;
 using System.Linq;
 using System.Text;
 using System.Text.Json;
+using System.Threading.Tasks;
 using System.Xml;
 
 namespace BuildGraph.Tasks
@@ -92,7 +93,7 @@ namespace BuildGraph.Tasks
 		/// <param name="Job">Information about the current job</param>
 		/// <param name="BuildProducts">Set of build products produced by this node.</param>
 		/// <param name="TagNameToFileSet">Mapping from tag names to the set of files they include</param>
-		public override void Execute(JobContext Job, HashSet<FileReference> BuildProducts, Dictionary<string, HashSet<FileReference>> TagNameToFileSet)
+		public override async Task ExecuteAsync(JobContext Job, HashSet<FileReference> BuildProducts, Dictionary<string, HashSet<FileReference>> TagNameToFileSet)
 		{
 			StringBuilder Arguments = new StringBuilder("sts assume-role");
 			if(Parameters.Arn != null)
@@ -106,7 +107,7 @@ namespace BuildGraph.Tasks
 			Arguments.Append($" --duration-seconds {Parameters.Duration}");
 
 			Dictionary<string, string> Environment = SpawnTaskBase.ParseEnvVars(Parameters.Environment, Parameters.EnvironmentFile);
-			IProcessResult Result = SpawnTaskBase.Execute("aws", Arguments.ToString(), EnvVars: Environment, LogOutput: false);
+			IProcessResult Result = await SpawnTaskBase.ExecuteAsync("aws", Arguments.ToString(), EnvVars: Environment, LogOutput: false);
 
 			JsonSerializerOptions Options = new JsonSerializerOptions();
 			Options.PropertyNameCaseInsensitive = true;
diff --git a/Engine/Source/Programs/AutomationTool/BuildGraph/Tasks/AwsEcsDeployTask.cs b/Engine/Source/Programs/AutomationTool/BuildGraph/Tasks/AwsEcsDeployTask.cs
--- a/Engine/Source/Programs/AutomationTool/BuildGraph/Tasks/AwsEcsDeployTask.cs
+++ b/Engine/Source/Programs/AutomationTool/BuildGraph/Tasks/AwsEcsDeployTask.cs
@@ -7,6 +7,7 @@ using System.Collections.Generic;
 using System.IO;
 using System.Text;
 using System.Text.Json;
+using System.Threading.Tasks;
 using System.Xml;
 using UnrealBuildBase;
 
@@ -92,7 +93,7 @@ namespace BuildGraph.Tasks
 		/// <param name="Job">Information about the current job</param>
 		/// <param name="BuildProducts">Set of build products produced by this node.</param>
 		/// <param name="TagNameToFileSet">Mapping from tag names to the set of files they include</param>
-		public override void Execute(JobContext Job, HashSet<FileReference> BuildProducts, Dictionary<string, HashSet<FileReference>> TagNameToFileSet)
+		public override async Task ExecuteAsync(JobContext Job, HashSet<FileReference> BuildProducts, Dictionary<string, HashSet<FileReference>> TagNameToFileSet)
 		{
 			string TaskDefTemplate = File.ReadAllText(ResolveFile(Parameters.TaskDefinitionFile).FullName);
 			string TaskDefRendered = TaskDefTemplate.Replace("%%DOCKER_IMAGE%%", Parameters.DockerImage);
@@ -105,14 +106,14 @@ namespace BuildGraph.Tasks
 			DirectoryReference.CreateDirectory(TempTaskDefFile.Directory);
 			File.WriteAllText(TempTaskDefFile.FullName, TaskDefRendered, new UTF8Encoding(encoderShouldEmitUTF8Identifier: false));
 			
-			IProcessResult CreateTaskDefResult = SpawnTaskBase.Execute("aws", $"ecs register-task-definition --cli-input-json \"file://{TempTaskDefFile.FullName}\"", EnvVars: ParseEnvVars(Parameters.Environment, Parameters.EnvironmentFile), LogOutput: Parameters.LogOutput);
+			IProcessResult CreateTaskDefResult = await SpawnTaskBase.ExecuteAsync("aws", $"ecs register-task-definition --cli-input-json \"file://{TempTaskDefFile.FullName}\"", EnvVars: ParseEnvVars(Parameters.Environment, Parameters.EnvironmentFile), LogOutput: Parameters.LogOutput);
 
 			JsonDocument TaskDefJson = JsonDocument.Parse(CreateTaskDefResult.Output);
 			string TaskDefFamily = TaskDefJson.RootElement.GetProperty("taskDefinition").GetProperty("family").GetString();
 			string TaskDefRevision = TaskDefJson.RootElement.GetProperty("taskDefinition").GetProperty("revision").ToString();
 
 			string Params = $"ecs update-service --cluster {Parameters.Cluster} --service {Parameters.Service} --task-definition {TaskDefFamily}:{TaskDefRevision}";
-			SpawnTaskBase.Execute("aws", Params, EnvVars: ParseEnvVars(Parameters.Environment, Parameters.EnvironmentFile), LogOutput: Parameters.LogOutput);
+			await SpawnTaskBase.ExecuteAsync("aws", Params, EnvVars: ParseEnvVars(Parameters.Environment, Parameters.EnvironmentFile), LogOutput: Parameters.LogOutput);
 
 			Log.TraceInformation($"Service {Parameters.Service} updated to use new task def {TaskDefFamily}:{TaskDefRevision}");
 		}
diff --git a/Engine/Source/Programs/AutomationTool/BuildGraph/Tasks/AwsTask.cs b/Engine/Source/Programs/AutomationTool/BuildGraph/Tasks/AwsTask.cs
--- a/Engine/Source/Programs/AutomationTool/BuildGraph/Tasks/AwsTask.cs
+++ b/Engine/Source/Programs/AutomationTool/BuildGraph/Tasks/AwsTask.cs
@@ -7,6 +7,7 @@ using System.Collections.Generic;
 using System.IO;
 using System.Text;
 using System.Text.Json;
+using System.Threading.Tasks;
 using System.Xml;
 
 namespace BuildGraph.Tasks
@@ -67,9 +68,9 @@ namespace BuildGraph.Tasks
 		/// <param name="Job">Information about the current job</param>
 		/// <param name="BuildProducts">Set of build products produced by this node.</param>
 		/// <param name="TagNameToFileSet">Mapping from tag names to the set of files they include</param>
-		public override void Execute(JobContext Job, HashSet<FileReference> BuildProducts, Dictionary<string, HashSet<FileReference>> TagNameToFileSet)
+		public override async Task ExecuteAsync(JobContext Job, HashSet<FileReference> BuildProducts, Dictionary<string, HashSet<FileReference>> TagNameToFileSet)
 		{
-			SpawnTaskBase.Execute("aws", Parameters.Arguments, EnvVars: ParseEnvVars(Parameters.Environment, Parameters.EnvironmentFile), LogOutput: Parameters.LogOutput);
+			await SpawnTaskBase.ExecuteAsync("aws", Parameters.Arguments, EnvVars: ParseEnvVars(Parameters.Environment, Parameters.EnvironmentFile), LogOutput: Parameters.LogOutput);
 		}
 
 		/// <summary>
diff --git a/Engine/Source/Programs/AutomationTool/BuildGraph/Tasks/CommandTask.cs b/Engine/Source/Programs/AutomationTool/BuildGraph/Tasks/CommandTask.cs
--- a/Engine/Source/Programs/AutomationTool/BuildGraph/Tasks/CommandTask.cs
+++ b/Engine/Source/Programs/AutomationTool/BuildGraph/Tasks/CommandTask.cs
@@ -50,7 +50,7 @@ namespace BuildGraph.Tasks
 	/// Invokes an AutomationTool child process to run the given command.
 	/// </summary>
 	[TaskElement("Command", typeof(CommandTaskParameters))]
-	public class CommandTask : CustomTask
+	public class CommandTask : BgTaskImpl
 	{
 		/// <summary>
 		/// Parameters for this task
@@ -72,7 +72,7 @@ namespace BuildGraph.Tasks
 		/// <param name="Job">Information about the current job</param>
 		/// <param name="BuildProducts">Set of build products produced by this node.</param>
 		/// <param name="TagNameToFileSet">Mapping from tag names to the set of files they include</param>
-		public override void Execute(JobContext Job, HashSet<FileReference> BuildProducts, Dictionary<string, HashSet<FileReference>> TagNameToFileSet)
+		public override Task ExecuteAsync(JobContext Job, HashSet<FileReference> BuildProducts, Dictionary<string, HashSet<FileReference>> TagNameToFileSet)
 		{
 			// If we're merging telemetry from the child process, get a temp filename for it
 			FileReference TelemetryFile = null;
@@ -126,6 +126,7 @@ namespace BuildGraph.Tasks
 					Log.TraceWarning("Unable to read UAT telemetry file from {0}", TelemetryFile);
 				}
 			}
+			return Task.CompletedTask;
 		}
 
 		/// <summary>
diff --git a/Engine/Source/Programs/AutomationTool/BuildGraph/Tasks/CommandletTask.cs b/Engine/Source/Programs/AutomationTool/BuildGraph/Tasks/CommandletTask.cs
--- a/Engine/Source/Programs/AutomationTool/BuildGraph/Tasks/CommandletTask.cs
+++ b/Engine/Source/Programs/AutomationTool/BuildGraph/Tasks/CommandletTask.cs
@@ -54,7 +54,7 @@ namespace AutomationTool.Tasks
 	/// Spawns the editor to run a commandlet.
 	/// </summary>
 	[TaskElement("Commandlet", typeof(CommandletTaskParameters))]
-	public class CommandletTask : CustomTask
+	public class CommandletTask : BgTaskImpl
 	{
 		/// <summary>
 		/// Parameters for this task
@@ -76,7 +76,7 @@ namespace AutomationTool.Tasks
 		/// <param name="Job">Information about the current job</param>
 		/// <param name="BuildProducts">Set of build products produced by this node.</param>
 		/// <param name="TagNameToFileSet">Mapping from tag names to the set of files they include</param>
-		public override void Execute(JobContext Job, HashSet<FileReference> BuildProducts, Dictionary<string, HashSet<FileReference>> TagNameToFileSet)
+		public override Task ExecuteAsync(JobContext Job, HashSet<FileReference> BuildProducts, Dictionary<string, HashSet<FileReference>> TagNameToFileSet)
 		{
 			// Get the full path to the project file
 			FileReference ProjectFile = null;
@@ -114,6 +114,7 @@ namespace AutomationTool.Tasks
 
 			// Run the commandlet
 			CommandUtils.RunCommandlet(ProjectFile, EditorExe.FullName, Parameters.Name, Parameters.Arguments, Parameters.ErrorLevel);
+			return Task.CompletedTask;
 		}
 
 		/// <summary>
diff --git a/Engine/Source/Programs/AutomationTool/BuildGraph/Tasks/CompileTask.cs b/Engine/Source/Programs/AutomationTool/BuildGraph/Tasks/CompileTask.cs
--- a/Engine/Source/Programs/AutomationTool/BuildGraph/Tasks/CompileTask.cs
+++ b/Engine/Source/Programs/AutomationTool/BuildGraph/Tasks/CompileTask.cs
@@ -113,7 +113,7 @@ namespace AutomationTool
 		/// </summary>
 		/// <param name="Task">Task to add</param>
 		/// <returns>True if the task could be added, false otherwise</returns>
-		public bool Add(CustomTask Task)
+		public bool Add(BgTaskImpl Task)
 		{
 			CompileTask CompileTask = Task as CompileTask;
 			if(CompileTask == null)
@@ -154,7 +154,7 @@ namespace AutomationTool
 		/// <param name="BuildProducts">Set of build products produced by this node.</param>
 		/// <param name="TagNameToFileSet">Mapping from tag names to the set of files they include</param>
 		/// <returns>Whether the task succeeded or not. Exiting with an exception will be caught and treated as a failure.</returns>
-		public void Execute(JobContext Job, HashSet<FileReference> BuildProducts, Dictionary<string, HashSet<FileReference>> TagNameToFileSet)
+		public Task ExecuteAsync(JobContext Job, HashSet<FileReference> BuildProducts, Dictionary<string, HashSet<FileReference>> TagNameToFileSet)
 		{
 			// Create the agenda
 			UnrealBuild.BuildAgenda Agenda = new UnrealBuild.BuildAgenda();
@@ -187,6 +187,7 @@ namespace AutomationTool
 
 			// Add everything to the list of build products
 			BuildProducts.UnionWith(Builder.BuildProductFiles.Select(x => new FileReference(x)));
+			return Task.CompletedTask;
 		}
 	}
 
@@ -194,7 +195,7 @@ namespace AutomationTool
 	/// Compiles a target with UnrealBuildTool.
 	/// </summary>
 	[TaskElement("Compile", typeof(CompileTaskParameters))]
-	public class CompileTask : CustomTask
+	public class CompileTask : BgTaskImpl
 	{
 		/// <summary>
 		/// Parameters for this task
@@ -249,13 +250,13 @@ namespace AutomationTool
 		/// <param name="Job">Information about the current job</param>
 		/// <param name="BuildProducts">Set of build products produced by this node.</param>
 		/// <param name="TagNameToFileSet">Mapping from tag names to the set of files they include</param>
-		public override void Execute(JobContext Job, HashSet<FileReference> BuildProducts, Dictionary<string, HashSet<FileReference>> TagNameToFileSet)
+		public override Task ExecuteAsync(JobContext Job, HashSet<FileReference> BuildProducts, Dictionary<string, HashSet<FileReference>> TagNameToFileSet)
 		{
 			//
 			// Don't do any logic here. You have to do it in the ctor or a getter
 			//  otherwise you break the ParallelExecutor pathway, which doesn't call this function!
 			//
-			GetExecutor().Execute(Job, BuildProducts, TagNameToFileSet);
+			return GetExecutor().ExecuteAsync(Job, BuildProducts, TagNameToFileSet);
 		}
 
 		/// <summary>
diff --git a/Engine/Source/Programs/AutomationTool/BuildGraph/Tasks/CookTask.cs b/Engine/Source/Programs/AutomationTool/BuildGraph/Tasks/CookTask.cs
--- a/Engine/Source/Programs/AutomationTool/BuildGraph/Tasks/CookTask.cs
+++ b/Engine/Source/Programs/AutomationTool/BuildGraph/Tasks/CookTask.cs
@@ -69,7 +69,7 @@ namespace BuildGraph.Tasks
 	/// Cook a selection of maps for a certain platform
 	/// </summary>
 	[TaskElement("Cook", typeof(CookTaskParameters))]
-	public class CookTask : CustomTask
+	public class CookTask : BgTaskImpl
 	{
 		/// <summary>
 		/// Parameters for the task
@@ -91,7 +91,7 @@ namespace BuildGraph.Tasks
 		/// <param name="Job">Information about the current job</param>
 		/// <param name="BuildProducts">Set of build products produced by this node.</param>
 		/// <param name="TagNameToFileSet">Mapping from tag names to the set of files they include</param>
-		public override void Execute(JobContext Job, HashSet<FileReference> BuildProducts, Dictionary<string, HashSet<FileReference>> TagNameToFileSet)
+		public override Task ExecuteAsync(JobContext Job, HashSet<FileReference> BuildProducts, Dictionary<string, HashSet<FileReference>> TagNameToFileSet)
 		{
 			// Figure out the project that this target belongs to
 			FileReference ProjectFile = null;
@@ -146,6 +146,7 @@ namespace BuildGraph.Tasks
 
 			// Add them to the set of build products
 			BuildProducts.UnionWith(CookedFiles);
+			return Task.CompletedTask;
 		}
 
 		/// <summary>
diff --git a/Engine/Source/Programs/AutomationTool/BuildGraph/Tasks/CopyTask.cs b/Engine/Source/Programs/AutomationTool/BuildGraph/Tasks/CopyTask.cs
--- a/Engine/Source/Programs/AutomationTool/BuildGraph/Tasks/CopyTask.cs
+++ b/Engine/Source/Programs/AutomationTool/BuildGraph/Tasks/CopyTask.cs
@@ -61,7 +61,7 @@ namespace BuildGraph.Tasks
 	/// Copies files from one directory to another.
 	/// </summary>
 	[TaskElement("Copy", typeof(CopyTaskParameters))]
-	public class CopyTask : CustomTask
+	public class CopyTask : BgTaskImpl
 	{
 		/// <summary>
 		/// Parameters for this task
@@ -83,7 +83,7 @@ namespace BuildGraph.Tasks
 		/// <param name="Job">Information about the current job</param>
 		/// <param name="BuildProducts">Set of build products produced by this node.</param>
 		/// <param name="TagNameToFileSet">Mapping from tag names to the set of files they include</param>
-		public override void Execute(JobContext Job, HashSet<FileReference> BuildProducts, Dictionary<string, HashSet<FileReference>> TagNameToFileSet)
+		public override Task ExecuteAsync(JobContext Job, HashSet<FileReference> BuildProducts, Dictionary<string, HashSet<FileReference>> TagNameToFileSet)
 		{
 			// Parse all the source patterns
 			FilePattern SourcePattern = new FilePattern(Unreal.RootDirectory, Parameters.From);
@@ -127,7 +127,7 @@ namespace BuildGraph.Tasks
 				{
 					CommandUtils.LogInformation("No files found matching '{0}'", SourcePattern);
 				}
-				return;
+				return Task.CompletedTask;
 			}
 
 			// If the target is on a network share, retry creating the first directory until it succeeds
@@ -186,6 +186,7 @@ namespace BuildGraph.Tasks
 			{
 				FindOrAddTagSet(TagNameToFileSet, TagName).UnionWith(TargetFileToSourceFile.Keys);
 			}
+			return Task.CompletedTask;
 		}
 
 		/// <summary>
diff --git a/Engine/Source/Programs/AutomationTool/BuildGraph/Tasks/CsCompileTask.cs b/Engine/Source/Programs/AutomationTool/BuildGraph/Tasks/CsCompileTask.cs
--- a/Engine/Source/Programs/AutomationTool/BuildGraph/Tasks/CsCompileTask.cs
+++ b/Engine/Source/Programs/AutomationTool/BuildGraph/Tasks/CsCompileTask.cs
@@ -78,7 +78,7 @@ namespace AutomationTool.Tasks
 	/// Compiles C# project files, and their dependencies.
 	/// </summary>
 	[TaskElement("CsCompile", typeof(CsCompileTaskParameters))]
-	public class CsCompileTask : CustomTask
+	public class CsCompileTask : BgTaskImpl
 	{
 		/// <summary>
 		/// Parameters for the task
@@ -100,7 +100,7 @@ namespace AutomationTool.Tasks
 		/// <param name="Job">Information about the current job</param>
 		/// <param name="BuildProducts">Set of build products produced by this node.</param>
 		/// <param name="TagNameToFileSet">Mapping from tag names to the set of files they include</param>
-		public override void Execute(JobContext Job, HashSet<FileReference> BuildProducts, Dictionary<string, HashSet<FileReference>> TagNameToFileSet)
+		public override Task ExecuteAsync(JobContext Job, HashSet<FileReference> BuildProducts, Dictionary<string, HashSet<FileReference>> TagNameToFileSet)
 		{
 			// Get the project file
 			HashSet<FileReference> ProjectFiles = ResolveFilespec(Unreal.RootDirectory, Parameters.Project, TagNameToFileSet);
@@ -197,6 +197,7 @@ namespace AutomationTool.Tasks
 			// Merge them into the standard set of build products
 			BuildProducts.UnionWith(ProjectBuildProducts);
 			BuildProducts.UnionWith(ProjectReferences);
+			return Task.CompletedTask;
 		}
 
 		/// <summary>
diff --git a/Engine/Source/Programs/AutomationTool/BuildGraph/Tasks/DeleteTask.cs b/Engine/Source/Programs/AutomationTool/BuildGraph/Tasks/DeleteTask.cs
--- a/Engine/Source/Programs/AutomationTool/BuildGraph/Tasks/DeleteTask.cs
+++ b/Engine/Source/Programs/AutomationTool/BuildGraph/Tasks/DeleteTask.cs
@@ -49,7 +49,7 @@ namespace BuildGraph.Tasks
 	/// Delete a set of files.
 	/// </summary>
 	[TaskElement("Delete", typeof(DeleteTaskParameters))]
-	public class DeleteTask : CustomTask
+	public class DeleteTask : BgTaskImpl
 	{
 		/// <summary>
 		/// Parameters for this task
@@ -71,7 +71,7 @@ namespace BuildGraph.Tasks
 		/// <param name="Job">Information about the current job</param>
 		/// <param name="BuildProducts">Set of build products produced by this node.</param>
 		/// <param name="TagNameToFileSet">Mapping from tag names to the set of files they include</param>
-		public override void Execute(JobContext Job, HashSet<FileReference> BuildProducts, Dictionary<string, HashSet<FileReference>> TagNameToFileSet)
+		public override Task ExecuteAsync(JobContext Job, HashSet<FileReference> BuildProducts, Dictionary<string, HashSet<FileReference>> TagNameToFileSet)
 		{
 			if (Parameters.Files != null)
 			{
@@ -130,6 +130,7 @@ namespace BuildGraph.Tasks
 					}
 				}
 			}
+			return Task.CompletedTask;
 		}
 
 		/// <summary>
diff --git a/Engine/Source/Programs/AutomationTool/BuildGraph/Tasks/DockerBuildTask.cs b/Engine/Source/Programs/AutomationTool/BuildGraph/Tasks/DockerBuildTask.cs
--- a/Engine/Source/Programs/AutomationTool/BuildGraph/Tasks/DockerBuildTask.cs
+++ b/Engine/Source/Programs/AutomationTool/BuildGraph/Tasks/DockerBuildTask.cs
@@ -9,6 +9,7 @@ using System.Text;
 using System.Xml;
 using AutomationTool;
 using UnrealBuildBase;
+using System.Threading.Tasks;
 
 namespace BuildGraph.Tasks
 {
@@ -92,7 +93,7 @@ namespace BuildGraph.Tasks
 		/// <param name="Job">Information about the current job</param>
 		/// <param name="BuildProducts">Set of build products produced by this node.</param>
 		/// <param name="TagNameToFileSet">Mapping from tag names to the set of files they include</param>
-		public override void Execute(JobContext Job, HashSet<FileReference> BuildProducts, Dictionary<string, HashSet<FileReference>> TagNameToFileSet)
+		public override async Task ExecuteAsync(JobContext Job, HashSet<FileReference> BuildProducts, Dictionary<string, HashSet<FileReference>> TagNameToFileSet)
 		{
 			Log.TraceInformation("Building Docker image");
 			using (LogIndentScope Scope = new LogIndentScope("  "))
@@ -133,7 +134,7 @@ namespace BuildGraph.Tasks
 					Arguments.Append($" {Parameters.Arguments}");
 				}
 
-				SpawnTaskBase.Execute("docker", Arguments.ToString(), EnvVars: ParseEnvVars(Parameters.Environment, Parameters.EnvironmentFile), WorkingDir: StagingDir.FullName);
+				await SpawnTaskBase.ExecuteAsync("docker", Arguments.ToString(), EnvVars: ParseEnvVars(Parameters.Environment, Parameters.EnvironmentFile), WorkingDir: StagingDir.FullName);
 			}
 		}
 
diff --git a/Engine/Source/Programs/AutomationTool/BuildGraph/Tasks/DockerComposeDownTask.cs b/Engine/Source/Programs/AutomationTool/BuildGraph/Tasks/DockerComposeDownTask.cs
--- a/Engine/Source/Programs/AutomationTool/BuildGraph/Tasks/DockerComposeDownTask.cs
+++ b/Engine/Source/Programs/AutomationTool/BuildGraph/Tasks/DockerComposeDownTask.cs
@@ -9,6 +9,7 @@ using System.Text;
 using System.Xml;
 using AutomationTool;
 using UnrealBuildBase;
+using System.Threading.Tasks;
 
 namespace BuildGraph.Tasks
 {
@@ -56,7 +57,7 @@ namespace BuildGraph.Tasks
 		/// <param name="Job">Information about the current job</param>
 		/// <param name="BuildProducts">Set of build products produced by this node.</param>
 		/// <param name="TagNameToFileSet">Mapping from tag names to the set of files they include</param>
-		public override void Execute(JobContext Job, HashSet<FileReference> BuildProducts, Dictionary<string, HashSet<FileReference>> TagNameToFileSet)
+		public override async Task ExecuteAsync(JobContext Job, HashSet<FileReference> BuildProducts, Dictionary<string, HashSet<FileReference>> TagNameToFileSet)
 		{
 			StringBuilder Arguments = new StringBuilder("--ansi never ");
 			if (!String.IsNullOrEmpty(Parameters.File))
@@ -72,7 +73,7 @@ namespace BuildGraph.Tasks
 			Log.TraceInformation("Running docker compose {0}", Arguments.ToString());
 			using (LogIndentScope Scope = new LogIndentScope("  "))
 			{
-				SpawnTaskBase.Execute("docker-compose", Arguments.ToString());
+				await SpawnTaskBase.ExecuteAsync("docker-compose", Arguments.ToString());
 			}
 		}
 
diff --git a/Engine/Source/Programs/AutomationTool/BuildGraph/Tasks/DockerComposeUpTask.cs b/Engine/Source/Programs/AutomationTool/BuildGraph/Tasks/DockerComposeUpTask.cs
--- a/Engine/Source/Programs/AutomationTool/BuildGraph/Tasks/DockerComposeUpTask.cs
+++ b/Engine/Source/Programs/AutomationTool/BuildGraph/Tasks/DockerComposeUpTask.cs
@@ -9,6 +9,7 @@ using System.Text;
 using System.Xml;
 using AutomationTool;
 using UnrealBuildBase;
+using System.Threading.Tasks;
 
 namespace BuildGraph.Tasks
 {
@@ -56,7 +57,7 @@ namespace BuildGraph.Tasks
 		/// <param name="Job">Information about the current job</param>
 		/// <param name="BuildProducts">Set of build products produced by this node.</param>
 		/// <param name="TagNameToFileSet">Mapping from tag names to the set of files they include</param>
-		public override void Execute(JobContext Job, HashSet<FileReference> BuildProducts, Dictionary<string, HashSet<FileReference>> TagNameToFileSet)
+		public override async Task ExecuteAsync(JobContext Job, HashSet<FileReference> BuildProducts, Dictionary<string, HashSet<FileReference>> TagNameToFileSet)
 		{
 			StringBuilder Arguments = new StringBuilder("--ansi never ");
 			if (!String.IsNullOrEmpty(Parameters.File))
@@ -72,7 +73,7 @@ namespace BuildGraph.Tasks
 			Log.TraceInformation("Running docker compose {0}", Arguments.ToString());
 			using (LogIndentScope Scope = new LogIndentScope("  "))
 			{
-				SpawnTaskBase.Execute("docker-compose", Arguments.ToString());
+				await SpawnTaskBase.ExecuteAsync("docker-compose", Arguments.ToString());
 			}
 		}
 
diff --git a/Engine/Source/Programs/AutomationTool/BuildGraph/Tasks/DockerPushTask.cs b/Engine/Source/Programs/AutomationTool/BuildGraph/Tasks/DockerPushTask.cs
--- a/Engine/Source/Programs/AutomationTool/BuildGraph/Tasks/DockerPushTask.cs
+++ b/Engine/Source/Programs/AutomationTool/BuildGraph/Tasks/DockerPushTask.cs
@@ -8,6 +8,7 @@ using System.IO;
 using System.Linq;
 using System.Text;
 using System.Text.Json;
+using System.Threading.Tasks;
 using System.Xml;
 
 namespace BuildGraph.Tasks
@@ -80,7 +81,7 @@ namespace BuildGraph.Tasks
 		/// <param name="Job">Information about the current job</param>
 		/// <param name="BuildProducts">Set of build products produced by this node.</param>
 		/// <param name="TagNameToFileSet">Mapping from tag names to the set of files they include</param>
-		public override void Execute(JobContext Job, HashSet<FileReference> BuildProducts, Dictionary<string, HashSet<FileReference>> TagNameToFileSet)
+		public override async Task ExecuteAsync(JobContext Job, HashSet<FileReference> BuildProducts, Dictionary<string, HashSet<FileReference>> TagNameToFileSet)
 		{
 			Log.TraceInformation("Pushing Docker image");
 			using (LogIndentScope Scope = new LogIndentScope("  "))
@@ -89,13 +90,13 @@ namespace BuildGraph.Tasks
 
 				if (Parameters.AwsEcr)
 				{
-					IProcessResult Result = SpawnTaskBase.Execute("aws", "ecr get-login-password", EnvVars: Environment, LogOutput: false);
-					Execute("docker", $"login {Parameters.Repository} --username AWS --password-stdin", Input: Result.Output);
+					IProcessResult Result = await SpawnTaskBase.ExecuteAsync("aws", "ecr get-login-password", EnvVars: Environment, LogOutput: false);
+					await ExecuteAsync("docker", $"login {Parameters.Repository} --username AWS --password-stdin", Input: Result.Output);
 				}
 
 				string TargetImage = Parameters.TargetImage ?? Parameters.Image;
-				Execute("docker", $"tag {Parameters.Image} {Parameters.Repository}/{TargetImage}", EnvVars: Environment);
-				Execute("docker", $"push {Parameters.Repository}/{TargetImage}", EnvVars: Environment);
+				await ExecuteAsync("docker", $"tag {Parameters.Image} {Parameters.Repository}/{TargetImage}", EnvVars: Environment);
+				await ExecuteAsync("docker", $"push {Parameters.Repository}/{TargetImage}", EnvVars: Environment);
 			}
 		}
 
diff --git a/Engine/Source/Programs/AutomationTool/BuildGraph/Tasks/DockerTask.cs b/Engine/Source/Programs/AutomationTool/BuildGraph/Tasks/DockerTask.cs
--- a/Engine/Source/Programs/AutomationTool/BuildGraph/Tasks/DockerTask.cs
+++ b/Engine/Source/Programs/AutomationTool/BuildGraph/Tasks/DockerTask.cs
@@ -6,6 +6,7 @@ using System;
 using System.Collections.Generic;
 using System.IO;
 using System.Text;
+using System.Threading.Tasks;
 using System.Xml;
 
 namespace BuildGraph.Tasks
@@ -66,9 +67,9 @@ namespace BuildGraph.Tasks
 		/// <param name="Job">Information about the current job</param>
 		/// <param name="BuildProducts">Set of build products produced by this node.</param>
 		/// <param name="TagNameToFileSet">Mapping from tag names to the set of files they include</param>
-		public override void Execute(JobContext Job, HashSet<FileReference> BuildProducts, Dictionary<string, HashSet<FileReference>> TagNameToFileSet)
+		public override async Task ExecuteAsync(JobContext Job, HashSet<FileReference> BuildProducts, Dictionary<string, HashSet<FileReference>> TagNameToFileSet)
 		{
-			Execute("docker", Parameters.Arguments, EnvVars: ParseEnvVars(Parameters.Environment, Parameters.EnvironmentFile), WorkingDir: Parameters.WorkingDir);
+			await ExecuteAsync("docker", Parameters.Arguments, EnvVars: ParseEnvVars(Parameters.Environment, Parameters.EnvironmentFile), WorkingDir: Parameters.WorkingDir);
 		}
 
 		/// <summary>
diff --git a/Engine/Source/Programs/AutomationTool/BuildGraph/Tasks/DotNetTask.cs b/Engine/Source/Programs/AutomationTool/BuildGraph/Tasks/DotNetTask.cs
--- a/Engine/Source/Programs/AutomationTool/BuildGraph/Tasks/DotNetTask.cs
+++ b/Engine/Source/Programs/AutomationTool/BuildGraph/Tasks/DotNetTask.cs
@@ -6,6 +6,7 @@ using System;
 using System.Collections.Generic;
 using System.IO;
 using System.Text;
+using System.Threading.Tasks;
 using System.Xml;
 
 namespace BuildGraph.Tasks
@@ -78,7 +79,7 @@ namespace BuildGraph.Tasks
 		/// <param name="Job">Information about the current job</param>
 		/// <param name="BuildProducts">Set of build products produced by this node.</param>
 		/// <param name="TagNameToFileSet">Mapping from tag names to the set of files they include</param>
-		public override void Execute(JobContext Job, HashSet<FileReference> BuildProducts, Dictionary<string, HashSet<FileReference>> TagNameToFileSet)
+		public override async Task ExecuteAsync(JobContext Job, HashSet<FileReference> BuildProducts, Dictionary<string, HashSet<FileReference>> TagNameToFileSet)
 		{
 			FileReference DotNetFile = Parameters.DotNetPath == null ? HostPlatform.Current.GetDotnetExe() : Parameters.DotNetPath;
 			if(!FileReference.Exists(DotNetFile))
@@ -86,7 +87,7 @@ namespace BuildGraph.Tasks
 				throw new AutomationException("DotNet is missing from {0}", DotNetFile);
 			}
 
-			IProcessResult Result = Execute(DotNetFile.FullName, Parameters.Arguments, WorkingDir: Parameters.BaseDir, EnvVars: ParseEnvVars(Parameters.Environment, Parameters.EnvironmentFile));
+			IProcessResult Result = await ExecuteAsync(DotNetFile.FullName, Parameters.Arguments, WorkingDir: Parameters.BaseDir, EnvVars: ParseEnvVars(Parameters.Environment, Parameters.EnvironmentFile));
 			if (Result.ExitCode < 0 || Result.ExitCode >= Parameters.ErrorLevel)
 			{
 				throw new AutomationException("Docker terminated with an exit code indicating an error ({0})", Result.ExitCode);
diff --git a/Engine/Source/Programs/AutomationTool/BuildGraph/Tasks/GatherBuildProductsFromFileTask.cs b/Engine/Source/Programs/AutomationTool/BuildGraph/Tasks/GatherBuildProductsFromFileTask.cs
--- a/Engine/Source/Programs/AutomationTool/BuildGraph/Tasks/GatherBuildProductsFromFileTask.cs
+++ b/Engine/Source/Programs/AutomationTool/BuildGraph/Tasks/GatherBuildProductsFromFileTask.cs
@@ -4,6 +4,7 @@ using System;
 using System.Collections.Generic;
 using System.IO;
 using System.Text;
+using System.Threading.Tasks;
 using System.Xml;
 
 using EpicGames.Core;
@@ -18,14 +19,14 @@ namespace AutomationTool.Tasks
 	}
 
 	[TaskElement("GatherBuildProductsFromFile", typeof(GatherBuildProductsFromFileTaskParameters))]
-	class GatherBuildProductsFromFileTask : CustomTask
+	class GatherBuildProductsFromFileTask : BgTaskImpl
 	{
 		public GatherBuildProductsFromFileTask(GatherBuildProductsFromFileTaskParameters InParameters)
 		{
 			Parameters = InParameters;
 		}
 
-		public override void Execute(JobContext Job, HashSet<FileReference> BuildProducts, Dictionary<string, HashSet<FileReference>> TagNameToFileSet)
+		public override Task ExecuteAsync(JobContext Job, HashSet<FileReference> BuildProducts, Dictionary<string, HashSet<FileReference>> TagNameToFileSet)
 		{
 			List<FileReference> CleanupFiles = new List<FileReference>();
 
@@ -44,6 +45,8 @@ namespace AutomationTool.Tasks
 			{
 				CommandUtils.LogInformation("Failed to gather build products: {0}", Ex.Message);
 			}
+
+			return Task.CompletedTask;
 		}
 		
 		public override void Write(XmlWriter Writer)
diff --git a/Engine/Source/Programs/AutomationTool/BuildGraph/Tasks/GitCloneTask.cs b/Engine/Source/Programs/AutomationTool/BuildGraph/Tasks/GitCloneTask.cs
--- a/Engine/Source/Programs/AutomationTool/BuildGraph/Tasks/GitCloneTask.cs
+++ b/Engine/Source/Programs/AutomationTool/BuildGraph/Tasks/GitCloneTask.cs
@@ -4,6 +4,7 @@ using EpicGames.Core;
 using System;
 using System.Collections.Generic;
 using System.Text;
+using System.Threading.Tasks;
 using System.Xml;
 using UnrealBuildBase;
 
@@ -43,7 +44,7 @@ namespace AutomationTool.Tasks
 	/// Spawns Git and waits for it to complete.
 	/// </summary>
 	[TaskElement("Git-Clone", typeof(GitCloneTaskParameters))]
-	public class GitCloneTask : CustomTask
+	public class GitCloneTask : BgTaskImpl
 	{
 		/// <summary>
 		/// Parameters for this task
@@ -65,7 +66,7 @@ namespace AutomationTool.Tasks
 		/// <param name="Job">Information about the current job</param>
 		/// <param name="BuildProducts">Set of build products produced by this node.</param>
 		/// <param name="TagNameToFileSet">Mapping from tag names to the set of files they include</param>
-		public override void Execute(JobContext Job, HashSet<FileReference> BuildProducts, Dictionary<string, HashSet<FileReference>> TagNameToFileSet)
+		public override async Task ExecuteAsync(JobContext Job, HashSet<FileReference> BuildProducts, Dictionary<string, HashSet<FileReference>> TagNameToFileSet)
 		{
 			FileReference GitExe = CommandUtils.FindToolInPath("git");
 			if(GitExe == null)
@@ -80,7 +81,7 @@ namespace AutomationTool.Tasks
 				DirectoryReference GitDir = DirectoryReference.Combine(Dir, ".git");
 				if (!FileReference.Exists(FileReference.Combine(GitDir, "HEAD")))
 				{
-					RunGit(GitExe, $"init \"{Dir}\"", Unreal.RootDirectory);
+					await RunGitAsync(GitExe, $"init \"{Dir}\"", Unreal.RootDirectory);
 				}
 
 				if (Parameters.ConfigFile != null)
@@ -90,12 +91,12 @@ namespace AutomationTool.Tasks
 
 				if (Parameters.Remote != null)
 				{
-					RunGit(GitExe, $"remote add origin {Parameters.Remote}", Dir);
+					await RunGitAsync(GitExe, $"remote add origin {Parameters.Remote}", Dir);
 				}
 
-				RunGit(GitExe, "clean -dxf", Dir);
-				RunGit(GitExe, "fetch --all", Dir);
-				RunGit(GitExe, $"reset --hard {Parameters.Branch}", Dir);
+				await RunGitAsync(GitExe, "clean -dxf", Dir);
+				await RunGitAsync(GitExe, "fetch --all", Dir);
+				await RunGitAsync(GitExe, $"reset --hard {Parameters.Branch}", Dir);
 			}
 		}
 
@@ -105,13 +106,14 @@ namespace AutomationTool.Tasks
 		/// <param name="ToolFile"></param>
 		/// <param name="Arguments"></param>
 		/// <param name="WorkingDir"></param>
-		void RunGit(FileReference ToolFile, string Arguments, DirectoryReference WorkingDir)
+		Task RunGitAsync(FileReference ToolFile, string Arguments, DirectoryReference WorkingDir)
 		{
 			IProcessResult Result = CommandUtils.Run(ToolFile.FullName, Arguments, WorkingDir: WorkingDir.FullName);
 			if (Result.ExitCode != 0)
 			{
 				throw new AutomationException("Git terminated with an exit code indicating an error ({0})", Result.ExitCode);
 			}
+			return Task.CompletedTask;
 		}
 
 		/// <summary>
diff --git a/Engine/Source/Programs/AutomationTool/BuildGraph/Tasks/GitTask.cs b/Engine/Source/Programs/AutomationTool/BuildGraph/Tasks/GitTask.cs
--- a/Engine/Source/Programs/AutomationTool/BuildGraph/Tasks/GitTask.cs
+++ b/Engine/Source/Programs/AutomationTool/BuildGraph/Tasks/GitTask.cs
@@ -4,6 +4,7 @@ using EpicGames.Core;
 using System;
 using System.Collections.Generic;
 using System.Text;
+using System.Threading.Tasks;
 using System.Xml;
 
 namespace AutomationTool.Tasks
@@ -36,7 +37,7 @@ namespace AutomationTool.Tasks
 	/// Spawns Git and waits for it to complete.
 	/// </summary>
 	[TaskElement("Git", typeof(GitTaskParameters))]
-	public class GitTask : CustomTask
+	public class GitTask : BgTaskImpl
 	{
 		/// <summary>
 		/// Parameters for this task
@@ -58,7 +59,7 @@ namespace AutomationTool.Tasks
 		/// <param name="Job">Information about the current job</param>
 		/// <param name="BuildProducts">Set of build products produced by this node.</param>
 		/// <param name="TagNameToFileSet">Mapping from tag names to the set of files they include</param>
-		public override void Execute(JobContext Job, HashSet<FileReference> BuildProducts, Dictionary<string, HashSet<FileReference>> TagNameToFileSet)
+		public override Task ExecuteAsync(JobContext Job, HashSet<FileReference> BuildProducts, Dictionary<string, HashSet<FileReference>> TagNameToFileSet)
 		{
 			FileReference ToolFile = CommandUtils.FindToolInPath("git");
 			if(ToolFile == null)
@@ -71,6 +72,8 @@ namespace AutomationTool.Tasks
 			{
 				throw new AutomationException("Git terminated with an exit code indicating an error ({0})", Result.ExitCode);
 			}
+
+			return Task.CompletedTask;
 		}
 
 		/// <summary>
diff --git a/Engine/Source/Programs/AutomationTool/BuildGraph/Tasks/HelmTask.cs b/Engine/Source/Programs/AutomationTool/BuildGraph/Tasks/HelmTask.cs
--- a/Engine/Source/Programs/AutomationTool/BuildGraph/Tasks/HelmTask.cs
+++ b/Engine/Source/Programs/AutomationTool/BuildGraph/Tasks/HelmTask.cs
@@ -6,6 +6,7 @@ using System;
 using System.Collections.Generic;
 using System.IO;
 using System.Text;
+using System.Threading.Tasks;
 using System.Xml;
 using UnrealBuildBase;
 
@@ -109,7 +110,7 @@ namespace BuildGraph.Tasks
 		/// <param name="Job">Information about the current job</param>
 		/// <param name="BuildProducts">Set of build products produced by this node.</param>
 		/// <param name="TagNameToFileSet">Mapping from tag names to the set of files they include</param>
-		public override void Execute(JobContext Job, HashSet<FileReference> BuildProducts, Dictionary<string, HashSet<FileReference>> TagNameToFileSet)
+		public override async Task ExecuteAsync(JobContext Job, HashSet<FileReference> BuildProducts, Dictionary<string, HashSet<FileReference>> TagNameToFileSet)
 		{
 
 			// Build the argument list
@@ -152,7 +153,7 @@ namespace BuildGraph.Tasks
 			}
 
 			string AdditionalArguments = String.IsNullOrEmpty(Parameters.Arguments) ? String.Empty : $" {Parameters.Arguments}";
-			SpawnTaskBase.Execute("helm", CommandLineArguments.Join(Arguments) + AdditionalArguments, WorkingDir: Parameters.WorkingDir, EnvVars: ParseEnvVars(Parameters.Environment, Parameters.EnvironmentFile));
+			await SpawnTaskBase.ExecuteAsync("helm", CommandLineArguments.Join(Arguments) + AdditionalArguments, WorkingDir: Parameters.WorkingDir, EnvVars: ParseEnvVars(Parameters.Environment, Parameters.EnvironmentFile));
 		}
 
 		/// <summary>
diff --git a/Engine/Source/Programs/AutomationTool/BuildGraph/Tasks/KubectlTask.cs b/Engine/Source/Programs/AutomationTool/BuildGraph/Tasks/KubectlTask.cs
--- a/Engine/Source/Programs/AutomationTool/BuildGraph/Tasks/KubectlTask.cs
+++ b/Engine/Source/Programs/AutomationTool/BuildGraph/Tasks/KubectlTask.cs
@@ -5,6 +5,7 @@ using System;
 using System.Collections.Generic;
 using System.IO;
 using System.Text;
+using System.Threading.Tasks;
 using System.Xml;
 
 namespace AutomationTool.Tasks
@@ -31,7 +32,7 @@ namespace AutomationTool.Tasks
 	/// Spawns Kubectland waits for it to complete.
 	/// </summary>
 	[TaskElement("Kubectl", typeof(KubectlTaskParameters))]
-	public class KubectlTask : CustomTask
+	public class KubectlTask : BgTaskImpl
 	{
 		/// <summary>
 		/// Parameters for this task
@@ -53,7 +54,7 @@ namespace AutomationTool.Tasks
 		/// <param name="Job">Information about the current job</param>
 		/// <param name="BuildProducts">Set of build products produced by this node.</param>
 		/// <param name="TagNameToFileSet">Mapping from tag names to the set of files they include</param>
-		public override void Execute(JobContext Job, HashSet<FileReference> BuildProducts, Dictionary<string, HashSet<FileReference>> TagNameToFileSet)
+		public override Task ExecuteAsync(JobContext Job, HashSet<FileReference> BuildProducts, Dictionary<string, HashSet<FileReference>> TagNameToFileSet)
 		{
 			FileReference KubectlExe = CommandUtils.FindToolInPath("kubectl");
 			if (KubectlExe == null)
@@ -66,6 +67,8 @@ namespace AutomationTool.Tasks
 			{
 				throw new AutomationException("Kubectl terminated with an exit code indicating an error ({0})", Result.ExitCode);
 			}
+
+			return Task.CompletedTask;
 		}
 
 		/// <summary>
diff --git a/Engine/Source/Programs/AutomationTool/BuildGraph/Tasks/LogTask.cs b/Engine/Source/Programs/AutomationTool/BuildGraph/Tasks/LogTask.cs
--- a/Engine/Source/Programs/AutomationTool/BuildGraph/Tasks/LogTask.cs
+++ b/Engine/Source/Programs/AutomationTool/BuildGraph/Tasks/LogTask.cs
@@ -43,7 +43,7 @@ namespace BuildGraph.Tasks
 	/// Print a message (and other optional diagnostic information) to the output log.
 	/// </summary>
 	[TaskElement("Log", typeof(LogTaskParameters))]
-	public class LogTask : CustomTask
+	public class LogTask : BgTaskImpl
 	{
 		/// <summary>
 		/// Parameters for the task
@@ -65,7 +65,7 @@ namespace BuildGraph.Tasks
 		/// <param name="Job">Information about the current job</param>
 		/// <param name="BuildProducts">Set of build products produced by this node.</param>
 		/// <param name="TagNameToFileSet">Mapping from tag names to the set of files they include</param>
-		public override void Execute(JobContext Job, HashSet<FileReference> BuildProducts, Dictionary<string, HashSet<FileReference>> TagNameToFileSet)
+		public override Task ExecuteAsync(JobContext Job, HashSet<FileReference> BuildProducts, Dictionary<string, HashSet<FileReference>> TagNameToFileSet)
 		{
 			// Print the message
 			if(!String.IsNullOrEmpty(Parameters.Message))
@@ -89,6 +89,8 @@ namespace BuildGraph.Tasks
 					}
 				}
 			}
+
+			return Task.CompletedTask;
 		}
 
 		/// <summary>
diff --git a/Engine/Source/Programs/AutomationTool/BuildGraph/Tasks/ModifyConfigTask.cs b/Engine/Source/Programs/AutomationTool/BuildGraph/Tasks/ModifyConfigTask.cs
--- a/Engine/Source/Programs/AutomationTool/BuildGraph/Tasks/ModifyConfigTask.cs
+++ b/Engine/Source/Programs/AutomationTool/BuildGraph/Tasks/ModifyConfigTask.cs
@@ -53,7 +53,7 @@ namespace BuildGraph.Tasks
 	/// Modifies a config file
 	/// </summary>
 	[TaskElement("ModifyConfig", typeof(ModifyConfigTaskParameters))]
-	public class ModifyConfigTask : CustomTask
+	public class ModifyConfigTask : BgTaskImpl
 	{
 		/// <summary>
 		/// Parameters for this task
@@ -75,7 +75,7 @@ namespace BuildGraph.Tasks
 		/// <param name="Job">Information about the current job</param>
 		/// <param name="BuildProducts">Set of build products produced by this node.</param>
 		/// <param name="TagNameToFileSet">Mapping from tag names to the set of files they include</param>
-		public override void Execute(JobContext Job, HashSet<FileReference> BuildProducts, Dictionary<string, HashSet<FileReference>> TagNameToFileSet)
+		public override Task ExecuteAsync(JobContext Job, HashSet<FileReference> BuildProducts, Dictionary<string, HashSet<FileReference>> TagNameToFileSet)
 		{
 			FileReference ConfigFileLocation = ResolveFile(Parameters.File);
 
@@ -104,6 +104,7 @@ namespace BuildGraph.Tasks
 
 			// Add the archive to the set of build products
 			BuildProducts.Add(ConfigFileLocation);
+			return Task.CompletedTask;
 		}
 
 		/// <summary>
diff --git a/Engine/Source/Programs/AutomationTool/BuildGraph/Tasks/ModifyJsonValue.cs b/Engine/Source/Programs/AutomationTool/BuildGraph/Tasks/ModifyJsonValue.cs
--- a/Engine/Source/Programs/AutomationTool/BuildGraph/Tasks/ModifyJsonValue.cs
+++ b/Engine/Source/Programs/AutomationTool/BuildGraph/Tasks/ModifyJsonValue.cs
@@ -11,6 +11,7 @@ using System.Text;
 using EpicGames.Core;
 using UnrealBuildBase;
 using System.Text.Json;
+using System.Threading.Tasks;
 
 namespace BuildGraph.Tasks
 {
@@ -42,7 +43,7 @@ namespace BuildGraph.Tasks
 	/// Modifies json files by setting a value specified in the key path
 	/// </summary>
 	[TaskElement("ModifyJsonValue", typeof(ModifyJsonValueParameters))]
-	public class ModifyJsonValue : CustomTask
+	public class ModifyJsonValue : BgTaskImpl
 	{
 		ModifyJsonValueParameters Parameters;
 
@@ -58,7 +59,7 @@ namespace BuildGraph.Tasks
 		/// <summary>
 		/// Placeholder comment
 		/// </summary>
-		public override void Execute(JobContext Job, HashSet<FileReference> BuildProducts, Dictionary<string, HashSet<FileReference>> TagNameToFileSet)
+		public override async Task ExecuteAsync(JobContext Job, HashSet<FileReference> BuildProducts, Dictionary<string, HashSet<FileReference>> TagNameToFileSet)
 		{
 			string[] Keys = Parameters.KeyPath.Split('.');
 			if (Keys.Length == 0)
@@ -81,7 +82,7 @@ namespace BuildGraph.Tasks
 				CurrObj[Keys[Keys.Length - 1]] = Parameters.NewValue;
 
 				var NewContents = JsonSerializer.Serialize(ParamObj, new JsonSerializerOptions { WriteIndented = true });
-				File.WriteAllText(JsonFile, NewContents, new UTF8Encoding(encoderShouldEmitUTF8Identifier: false));
+				await File.WriteAllTextAsync(JsonFile, NewContents, new UTF8Encoding(encoderShouldEmitUTF8Identifier: false));
 			}
 		}
 
diff --git a/Engine/Source/Programs/AutomationTool/BuildGraph/Tasks/MoveTask.cs b/Engine/Source/Programs/AutomationTool/BuildGraph/Tasks/MoveTask.cs
--- a/Engine/Source/Programs/AutomationTool/BuildGraph/Tasks/MoveTask.cs
+++ b/Engine/Source/Programs/AutomationTool/BuildGraph/Tasks/MoveTask.cs
@@ -54,7 +54,7 @@ namespace BuildGraph.Tasks
 	/// Moves files from one directory to another.
 	/// </summary>
 	[TaskElement("Move", typeof(MoveTaskParameters))]
-	public class MoveTask : CustomTask
+	public class MoveTask : BgTaskImpl
 	{
 		/// <summary>
 		/// Parameters for this task
@@ -76,7 +76,7 @@ namespace BuildGraph.Tasks
 		/// <param name="Job">Information about the current job</param>
 		/// <param name="BuildProducts">Set of build products produced by this node.</param>
 		/// <param name="TagNameToFileSet">Mapping from tag names to the set of files they include</param>
-		public override void Execute(JobContext Job, HashSet<FileReference> BuildProducts, Dictionary<string, HashSet<FileReference>> TagNameToFileSet)
+		public override Task ExecuteAsync(JobContext Job, HashSet<FileReference> BuildProducts, Dictionary<string, HashSet<FileReference>> TagNameToFileSet)
 		{
 			// Parse all the source patterns
 			FilePattern SourcePattern = new FilePattern(Unreal.RootDirectory, Parameters.From);
@@ -99,7 +99,7 @@ namespace BuildGraph.Tasks
 			if(TargetFileToSourceFile.Count == 0)
 			{
 				CommandUtils.LogInformation("No files found matching '{0}'", SourcePattern);
-				return;
+				return Task.CompletedTask;
 			}
 
 			// Copy them all
@@ -114,6 +114,8 @@ namespace BuildGraph.Tasks
 			{
 				FindOrAddTagSet(TagNameToFileSet, TagName).UnionWith(TargetFileToSourceFile.Keys);
 			}
+
+			return Task.CompletedTask;
 		}
 
 		/// <summary>
diff --git a/Engine/Source/Programs/AutomationTool/BuildGraph/Tasks/MsBuildTask.cs b/Engine/Source/Programs/AutomationTool/BuildGraph/Tasks/MsBuildTask.cs
--- a/Engine/Source/Programs/AutomationTool/BuildGraph/Tasks/MsBuildTask.cs
+++ b/Engine/Source/Programs/AutomationTool/BuildGraph/Tasks/MsBuildTask.cs
@@ -51,7 +51,7 @@ namespace AutomationTool.Tasks
 	/// Executes MsBuild
 	/// </summary>
 	[TaskElement("MsBuild", typeof(MsBuildTaskParameters))]
-	public class MsBuildTask : CustomTask
+	public class MsBuildTask : BgTaskImpl
 	{
 		/// <summary>
 		/// Parameters for the task
@@ -73,7 +73,7 @@ namespace AutomationTool.Tasks
 		/// <param name="Job">Information about the current job</param>
 		/// <param name="BuildProducts">Set of build products produced by this node.</param>
 		/// <param name="TagNameToFileSet">Mapping from tag names to the set of files they include</param>
-		public override void Execute(JobContext Job, HashSet<FileReference> BuildProducts, Dictionary<string, HashSet<FileReference>> TagNameToFileSet)
+		public override Task ExecuteAsync(JobContext Job, HashSet<FileReference> BuildProducts, Dictionary<string, HashSet<FileReference>> TagNameToFileSet)
 		{
 			// Get the project file
 			HashSet<FileReference> ProjectFiles = ResolveFilespec(Unreal.RootDirectory, Parameters.Project, TagNameToFileSet);
@@ -110,6 +110,8 @@ namespace AutomationTool.Tasks
 			{
 				CommandUtils.MsBuild(CommandUtils.CmdEnv, ProjectFile.FullName, String.Join(" ", Arguments), null);
 			}
+
+			return Task.CompletedTask;
 		}
 
 		/// <summary>
diff --git a/Engine/Source/Programs/AutomationTool/BuildGraph/Tasks/NotarizeTask.cs b/Engine/Source/Programs/AutomationTool/BuildGraph/Tasks/NotarizeTask.cs
--- a/Engine/Source/Programs/AutomationTool/BuildGraph/Tasks/NotarizeTask.cs
+++ b/Engine/Source/Programs/AutomationTool/BuildGraph/Tasks/NotarizeTask.cs
@@ -53,7 +53,7 @@ namespace BuildGraph.Tasks
 	}
 
 	[TaskElement("Notarize", typeof(NotarizeTaskParameters))]
-	class NotarizeTask : CustomTask
+	class NotarizeTask : BgTaskImpl
 	{
 		/// <summary>
 		/// Parameters for the task
@@ -75,7 +75,7 @@ namespace BuildGraph.Tasks
 		/// <param name="Job">Information about the current job</param>
 		/// <param name="BuildProducts">Set of build products produced by this node.</param>
 		/// <param name="TagNameToFileSet">Mapping from tag names to the set of files they include</param>
-		public override void Execute(JobContext Job, HashSet<FileReference> BuildProducts, Dictionary<string, HashSet<FileReference>> TagNameToFileSet)
+		public override Task ExecuteAsync(JobContext Job, HashSet<FileReference> BuildProducts, Dictionary<string, HashSet<FileReference>> TagNameToFileSet)
 		{
 			// Ensure running on a mac.
 			if(BuildHostPlatform.Current.Platform != UnrealTargetPlatform.Mac)
@@ -193,6 +193,8 @@ namespace BuildGraph.Tasks
 					throw new AutomationException(Ex, "Querying for the notarization progress failed, output: {0}", Output);
 				}
 			}
+
+			return Task.CompletedTask;
 		}
 
 		private string GetLogFile(string Url)
diff --git a/Engine/Source/Programs/AutomationTool/BuildGraph/Tasks/PakFileTask.cs b/Engine/Source/Programs/AutomationTool/BuildGraph/Tasks/PakFileTask.cs
--- a/Engine/Source/Programs/AutomationTool/BuildGraph/Tasks/PakFileTask.cs
+++ b/Engine/Source/Programs/AutomationTool/BuildGraph/Tasks/PakFileTask.cs
@@ -79,7 +79,7 @@ namespace BuildGraph.Tasks
 	/// Creates a PAK file from a given set of files.
 	/// </summary>
 	[TaskElement("PakFile", typeof(PakFileTaskParameters))]
-	public class PakFileTask : CustomTask
+	public class PakFileTask : BgTaskImpl
 	{
 		/// <summary>
 		/// Parameters for the task
@@ -101,7 +101,7 @@ namespace BuildGraph.Tasks
 		/// <param name="Job">Information about the current job</param>
 		/// <param name="BuildProducts">Set of build products produced by this node.</param>
 		/// <param name="TagNameToFileSet">Mapping from tag names to the set of files they include</param>
-		public override void Execute(JobContext Job, HashSet<FileReference> BuildProducts, Dictionary<string, HashSet<FileReference>> TagNameToFileSet)
+		public override Task ExecuteAsync(JobContext Job, HashSet<FileReference> BuildProducts, Dictionary<string, HashSet<FileReference>> TagNameToFileSet)
 		{
 			// Find the directories we're going to rebase relative to
 			HashSet<DirectoryReference> RebaseDirs = new HashSet<DirectoryReference>{ Unreal.RootDirectory };
@@ -181,6 +181,8 @@ namespace BuildGraph.Tasks
 			{
 				FindOrAddTagSet(TagNameToFileSet, TagName).Add(OutputFile);
 			}
+
+			return Task.CompletedTask;
 		}
 
 		/// <summary>
diff --git a/Engine/Source/Programs/AutomationTool/BuildGraph/Tasks/RenameTask.cs b/Engine/Source/Programs/AutomationTool/BuildGraph/Tasks/RenameTask.cs
--- a/Engine/Source/Programs/AutomationTool/BuildGraph/Tasks/RenameTask.cs
+++ b/Engine/Source/Programs/AutomationTool/BuildGraph/Tasks/RenameTask.cs
@@ -50,7 +50,7 @@ namespace BuildGraph.Tasks
 	/// Renames a file, or group of files.
 	/// </summary>
 	[TaskElement("Rename", typeof(RenameTaskParameters))]
-	public class RenameTask : CustomTask
+	public class RenameTask : BgTaskImpl
 	{
 		/// <summary>
 		/// Parameters for this task
@@ -72,7 +72,7 @@ namespace BuildGraph.Tasks
 		/// <param name="Job">Information about the current job</param>
 		/// <param name="BuildProducts">Set of build products produced by this node.</param>
 		/// <param name="TagNameToFileSet">Mapping from tag names to the set of files they include</param>
-		public override void Execute(JobContext Job, HashSet<FileReference> BuildProducts, Dictionary<string, HashSet<FileReference>> TagNameToFileSet)
+		public override Task ExecuteAsync(JobContext Job, HashSet<FileReference> BuildProducts, Dictionary<string, HashSet<FileReference>> TagNameToFileSet)
 		{
 			// Get the pattern to match against. If it's a simple pattern (eg. *.cpp, Engine/Build/...), automatically infer the source wildcard
 			string FromPattern = Parameters.From;
@@ -145,6 +145,8 @@ namespace BuildGraph.Tasks
 			{
 				FindOrAddTagSet(TagNameToFileSet, TagName).UnionWith(RenameFiles.Values);
 			}
+
+			return Task.CompletedTask;
 		}
 
 		/// <summary>
diff --git a/Engine/Source/Programs/AutomationTool/BuildGraph/Tasks/SanitizeReceiptTask.cs b/Engine/Source/Programs/AutomationTool/BuildGraph/Tasks/SanitizeReceiptTask.cs
--- a/Engine/Source/Programs/AutomationTool/BuildGraph/Tasks/SanitizeReceiptTask.cs
+++ b/Engine/Source/Programs/AutomationTool/BuildGraph/Tasks/SanitizeReceiptTask.cs
@@ -35,7 +35,7 @@ namespace AutomationTool.Tasks
 	/// Task that tags build products and/or runtime dependencies by reading from *.target files.
 	/// </summary>
 	[TaskElement("SanitizeReceipt", typeof(SanitizeReceiptTaskParameters))]
-	class SanitizeReceiptTask : CustomTask
+	class SanitizeReceiptTask : BgTaskImpl
 	{
 		/// <summary>
 		/// Parameters to this task
@@ -57,7 +57,7 @@ namespace AutomationTool.Tasks
 		/// <param name="Job">Information about the current job</param>
 		/// <param name="BuildProducts">Set of build products produced by this node.</param>
 		/// <param name="TagNameToFileSet">Mapping from tag names to the set of files they include</param>
-		public override void Execute(JobContext Job, HashSet<FileReference> BuildProducts, Dictionary<string, HashSet<FileReference>> TagNameToFileSet)
+		public override Task ExecuteAsync(JobContext Job, HashSet<FileReference> BuildProducts, Dictionary<string, HashSet<FileReference>> TagNameToFileSet)
 		{
 			// Set the Engine directory
 			DirectoryReference EngineDir = Parameters.EngineDir ?? Unreal.EngineDirectory;
@@ -118,6 +118,7 @@ namespace AutomationTool.Tasks
 					Receipt.Write(TargetFile, EngineDir);
 				}
 			}
+			return Task.CompletedTask;
 		}
 
 		/// <summary>
diff --git a/Engine/Source/Programs/AutomationTool/BuildGraph/Tasks/SetVersionTask.cs b/Engine/Source/Programs/AutomationTool/BuildGraph/Tasks/SetVersionTask.cs
--- a/Engine/Source/Programs/AutomationTool/BuildGraph/Tasks/SetVersionTask.cs
+++ b/Engine/Source/Programs/AutomationTool/BuildGraph/Tasks/SetVersionTask.cs
@@ -70,7 +70,7 @@ namespace AutomationTool.Tasks
 	/// Updates the local version files (Engine/Source/Runtime/Launch/Resources/Version.h, Engine/Build/Build.version, and Engine/Source/Programs/Shared/Metadata.cs) with the given version information.
 	/// </summary>
 	[TaskElement("SetVersion", typeof(SetVersionTaskParameters))]
-	public class SetVersionTask : CustomTask
+	public class SetVersionTask : BgTaskImpl
 	{
 		/// <summary>
 		/// Parameters for the task
@@ -92,7 +92,7 @@ namespace AutomationTool.Tasks
 		/// <param name="Job">Information about the current job</param>
 		/// <param name="BuildProducts">Set of build products produced by this node.</param>
 		/// <param name="TagNameToFileSet">Mapping from tag names to the set of files they include</param>
-		public override void Execute(JobContext Job, HashSet<FileReference> BuildProducts, Dictionary<string, HashSet<FileReference>> TagNameToFileSet)
+		public override Task ExecuteAsync(JobContext Job, HashSet<FileReference> BuildProducts, Dictionary<string, HashSet<FileReference>> TagNameToFileSet)
 		{
 			// Update the version files
 			List<FileReference> VersionFiles = UnrealBuild.StaticUpdateVersionFiles(Parameters.Change, Parameters.CompatibleChange, Parameters.Branch, Parameters.Build, Parameters.Licensee, Parameters.Promoted, !Parameters.SkipWrite);
@@ -105,6 +105,7 @@ namespace AutomationTool.Tasks
 
 			// Add them to the list of build products
 			BuildProducts.UnionWith(VersionFiles);
+			return Task.CompletedTask;
 		}
 
 		/// <summary>
diff --git a/Engine/Source/Programs/AutomationTool/BuildGraph/Tasks/SignTask.cs b/Engine/Source/Programs/AutomationTool/BuildGraph/Tasks/SignTask.cs
--- a/Engine/Source/Programs/AutomationTool/BuildGraph/Tasks/SignTask.cs
+++ b/Engine/Source/Programs/AutomationTool/BuildGraph/Tasks/SignTask.cs
@@ -36,7 +36,7 @@ namespace BuildGraph.Tasks
 	/// Signs a set of executable files with an installed certificate.
 	/// </summary>
 	[TaskElement("Sign", typeof(SignTaskParameters))]
-	public class SignTask : CustomTask
+	public class SignTask : BgTaskImpl
 	{
 		/// <summary>
 		/// Parameters for this task
@@ -58,7 +58,7 @@ namespace BuildGraph.Tasks
 		/// <param name="Job">Information about the current job</param>
 		/// <param name="BuildProducts">Set of build products produced by this node.</param>
 		/// <param name="TagNameToFileSet">Mapping from tag names to the set of files they include</param>
-		public override void Execute(JobContext Job, HashSet<FileReference> BuildProducts, Dictionary<string, HashSet<FileReference>> TagNameToFileSet)
+		public override Task ExecuteAsync(JobContext Job, HashSet<FileReference> BuildProducts, Dictionary<string, HashSet<FileReference>> TagNameToFileSet)
 		{
 			// Find the matching files
 			FileReference[] Files = ResolveFilespec(Unreal.RootDirectory, Parameters.Files, TagNameToFileSet).OrderBy(x => x.FullName).ToArray();
@@ -74,6 +74,7 @@ namespace BuildGraph.Tasks
 
 			// Add them to the list of build products
 			BuildProducts.UnionWith(Files);
+			return Task.CompletedTask;
 		}
 
 		/// <summary>
diff --git a/Engine/Source/Programs/AutomationTool/BuildGraph/Tasks/SpawnTask.cs b/Engine/Source/Programs/AutomationTool/BuildGraph/Tasks/SpawnTask.cs
--- a/Engine/Source/Programs/AutomationTool/BuildGraph/Tasks/SpawnTask.cs
+++ b/Engine/Source/Programs/AutomationTool/BuildGraph/Tasks/SpawnTask.cs
@@ -64,12 +64,12 @@ namespace BuildGraph.Tasks
 	/// <summary>
 	/// Base class for tasks that run an external tool
 	/// </summary>
-	public abstract class SpawnTaskBase : CustomTask
+	public abstract class SpawnTaskBase : BgTaskImpl
 	{
 		/// <summary>
 		/// Execute a command
 		/// </summary>
-		protected static IProcessResult Execute(string Exe, string Arguments, string WorkingDir = null, Dictionary<string, string> EnvVars = null, bool LogOutput = true, int ErrorLevel = 1, string Input = null)
+		protected static Task<IProcessResult> ExecuteAsync(string Exe, string Arguments, string WorkingDir = null, Dictionary<string, string> EnvVars = null, bool LogOutput = true, int ErrorLevel = 1, string Input = null)
 		{
 			if (WorkingDir != null)
 			{
@@ -88,7 +88,7 @@ namespace BuildGraph.Tasks
 				throw new AutomationException("{0} terminated with an exit code indicating an error ({1})", Path.GetFileName(Exe), Result.ExitCode);
 			}
 
-			return Result;
+			return Task.FromResult(Result);
 		}
 
 		/// <summary>
@@ -168,9 +168,9 @@ namespace BuildGraph.Tasks
 		/// <param name="Job">Information about the current job</param>
 		/// <param name="BuildProducts">Set of build products produced by this node.</param>
 		/// <param name="TagNameToFileSet">Mapping from tag names to the set of files they include</param>
-		public override void Execute(JobContext Job, HashSet<FileReference> BuildProducts, Dictionary<string, HashSet<FileReference>> TagNameToFileSet)
+		public override async Task ExecuteAsync(JobContext Job, HashSet<FileReference> BuildProducts, Dictionary<string, HashSet<FileReference>> TagNameToFileSet)
 		{
-			Execute(Parameters.Exe, Parameters.Arguments, Parameters.WorkingDir, EnvVars: ParseEnvVars(Parameters.Environment, Parameters.EnvironmentFile), LogOutput: Parameters.LogOutput, ErrorLevel: Parameters.ErrorLevel);
+			await ExecuteAsync(Parameters.Exe, Parameters.Arguments, Parameters.WorkingDir, EnvVars: ParseEnvVars(Parameters.Environment, Parameters.EnvironmentFile), LogOutput: Parameters.LogOutput, ErrorLevel: Parameters.ErrorLevel);
 		}
 
 		/// <summary>
diff --git a/Engine/Source/Programs/AutomationTool/BuildGraph/Tasks/StageTask.cs b/Engine/Source/Programs/AutomationTool/BuildGraph/Tasks/StageTask.cs
--- a/Engine/Source/Programs/AutomationTool/BuildGraph/Tasks/StageTask.cs
+++ b/Engine/Source/Programs/AutomationTool/BuildGraph/Tasks/StageTask.cs
@@ -72,7 +72,7 @@ namespace BuildGraph.Tasks
 	/// Stages files listed in a build receipt to an output directory.
 	/// </summary>
 	[TaskElement("Stage", typeof(StageTaskParameters))]
-	public class StageTask : CustomTask
+	public class StageTask : BgTaskImpl
 	{
 		/// <summary>
 		/// Parameters for the task
@@ -94,7 +94,7 @@ namespace BuildGraph.Tasks
 		/// <param name="Job">Information about the current job</param>
 		/// <param name="BuildProducts">Set of build products produced by this node.</param>
 		/// <param name="TagNameToFileSet">Mapping from tag names to the set of files they include</param>
-		public override void Execute(JobContext Job, HashSet<FileReference> BuildProducts, Dictionary<string, HashSet<FileReference>> TagNameToFileSet)
+		public override Task ExecuteAsync(JobContext Job, HashSet<FileReference> BuildProducts, Dictionary<string, HashSet<FileReference>> TagNameToFileSet)
 		{
 			// Get the project path, and check it exists
 			FileReference ProjectFile = Parameters.Project;
@@ -169,6 +169,7 @@ namespace BuildGraph.Tasks
 
 			// Add the target file to the list of build products
 			BuildProducts.UnionWith(TargetFiles);
+			return Task.CompletedTask;
 		}
 
 		/// <summary>
diff --git a/Engine/Source/Programs/AutomationTool/BuildGraph/Tasks/StripTask.cs b/Engine/Source/Programs/AutomationTool/BuildGraph/Tasks/StripTask.cs
--- a/Engine/Source/Programs/AutomationTool/BuildGraph/Tasks/StripTask.cs
+++ b/Engine/Source/Programs/AutomationTool/BuildGraph/Tasks/StripTask.cs
@@ -53,7 +53,7 @@ namespace BuildGraph.Tasks
 	/// Strips debugging information from a set of files.
 	/// </summary>
 	[TaskElement("Strip", typeof(StripTaskParameters))]
-	public class StripTask : CustomTask
+	public class StripTask : BgTaskImpl
 	{
 		/// <summary>
 		/// Parameters for this task
@@ -75,7 +75,7 @@ namespace BuildGraph.Tasks
 		/// <param name="Job">Information about the current job</param>
 		/// <param name="BuildProducts">Set of build products produced by this node.</param>
 		/// <param name="TagNameToFileSet">Mapping from tag names to the set of files they include</param>
-		public override void Execute(JobContext Job, HashSet<FileReference> BuildProducts, Dictionary<string, HashSet<FileReference>> TagNameToFileSet)
+		public override Task ExecuteAsync(JobContext Job, HashSet<FileReference> BuildProducts, Dictionary<string, HashSet<FileReference>> TagNameToFileSet)
 		{
 			// Get the base directory
 			DirectoryReference BaseDir = Parameters.BaseDir;
@@ -113,6 +113,7 @@ namespace BuildGraph.Tasks
 
 			// Add the target files to the set of build products
 			BuildProducts.UnionWith(TargetFiles);
+			return Task.CompletedTask;
 		}
 
 		/// <summary>
diff --git a/Engine/Source/Programs/AutomationTool/BuildGraph/Tasks/SubmitTask.cs b/Engine/Source/Programs/AutomationTool/BuildGraph/Tasks/SubmitTask.cs
--- a/Engine/Source/Programs/AutomationTool/BuildGraph/Tasks/SubmitTask.cs
+++ b/Engine/Source/Programs/AutomationTool/BuildGraph/Tasks/SubmitTask.cs
@@ -84,7 +84,7 @@ namespace AutomationTool.Tasks
 	/// Creates a new changelist and submits a set of files to a Perforce stream.
 	/// </summary>
 	[TaskElement("Submit", typeof(SubmitTaskParameters))]
-	public class SubmitTask : CustomTask
+	public class SubmitTask : BgTaskImpl
 	{
 		/// <summary>
 		/// Parameters for the task
@@ -106,7 +106,7 @@ namespace AutomationTool.Tasks
 		/// <param name="Job">Information about the current job</param>
 		/// <param name="BuildProducts">Set of build products produced by this node.</param>
 		/// <param name="TagNameToFileSet">Mapping from tag names to the set of files they include</param>
-		public override void Execute(JobContext Job, HashSet<FileReference> BuildProducts, Dictionary<string, HashSet<FileReference>> TagNameToFileSet)
+		public override Task ExecuteAsync(JobContext Job, HashSet<FileReference> BuildProducts, Dictionary<string, HashSet<FileReference>> TagNameToFileSet)
 		{
 			HashSet<FileReference> Files = ResolveFilespec(Unreal.RootDirectory, Parameters.Files, TagNameToFileSet);
 			if (Files.Count == 0)
@@ -166,7 +166,7 @@ namespace AutomationTool.Tasks
 					if(SubmitP4.TryDeleteEmptyChange(NewCL))
 					{
 						CommandUtils.LogInformation("No files to submit; ignored.");
-						return;
+						return Task.CompletedTask;
 					}
 				}
 
@@ -179,6 +179,7 @@ namespace AutomationTool.Tasks
 				}
 				CommandUtils.LogInformation("Submitted in changelist {0}", SubmittedCL);
 			}
+			return Task.CompletedTask;
 		}
 
 		/// <summary>
diff --git a/Engine/Source/Programs/AutomationTool/BuildGraph/Tasks/SymStoreTask.cs b/Engine/Source/Programs/AutomationTool/BuildGraph/Tasks/SymStoreTask.cs
--- a/Engine/Source/Programs/AutomationTool/BuildGraph/Tasks/SymStoreTask.cs
+++ b/Engine/Source/Programs/AutomationTool/BuildGraph/Tasks/SymStoreTask.cs
@@ -6,6 +6,7 @@ using System.IO;
 using System.Linq;
 using System.Text;
 using System.Threading;
+using System.Threading.Tasks;
 using System.Xml;
 using EpicGames.Core;
 using UnrealBuildBase;
@@ -53,7 +54,7 @@ namespace AutomationTool
     /// Task that strips symbols from a set of files.
     /// </summary>
     [TaskElement("SymStore", typeof(SymStoreTaskParameters))]
-    public class SymStoreTask : CustomTask
+    public class SymStoreTask : BgTaskImpl
     {
         /// <summary>
         /// Parameters for this task
@@ -75,7 +76,7 @@ namespace AutomationTool
         /// <param name="Job">Information about the current job</param>
         /// <param name="BuildProducts">Set of build products produced by this node.</param>
         /// <param name="TagNameToFileSet">Mapping from tag names to the set of files they include</param>
-        public override void Execute(JobContext Job, HashSet<FileReference> BuildProducts, Dictionary<string, HashSet<FileReference>> TagNameToFileSet)
+        public override Task ExecuteAsync(JobContext Job, HashSet<FileReference> BuildProducts, Dictionary<string, HashSet<FileReference>> TagNameToFileSet)
         {
 			// Find the matching files
 			List<FileReference> Files = ResolveFilespec(Unreal.RootDirectory, Parameters.Files, TagNameToFileSet).ToList();
@@ -92,7 +93,9 @@ namespace AutomationTool
 					throw new AutomationException("Failure publishing symbol files.");
 				}
 			});
-        }
+
+			return Task.CompletedTask;
+		}
 
         /// <summary>
         /// Output this task out to an XML writer.
diff --git a/Engine/Source/Programs/AutomationTool/BuildGraph/Tasks/TagReceiptTask.cs b/Engine/Source/Programs/AutomationTool/BuildGraph/Tasks/TagReceiptTask.cs
--- a/Engine/Source/Programs/AutomationTool/BuildGraph/Tasks/TagReceiptTask.cs
+++ b/Engine/Source/Programs/AutomationTool/BuildGraph/Tasks/TagReceiptTask.cs
@@ -72,7 +72,7 @@ namespace AutomationTool.Tasks
 	/// Task that tags build products and/or runtime dependencies by reading from *.target files.
 	/// </summary>
 	[TaskElement("TagReceipt", typeof(TagReceiptTaskParameters))]
-	class TagReceiptTask : CustomTask
+	class TagReceiptTask : BgTaskImpl
 	{
 		/// <summary>
 		/// Parameters to this task
@@ -113,7 +113,7 @@ namespace AutomationTool.Tasks
 		/// <param name="Job">Information about the current job</param>
 		/// <param name="BuildProducts">Set of build products produced by this node.</param>
 		/// <param name="TagNameToFileSet">Mapping from tag names to the set of files they include</param>
-		public override void Execute(JobContext Job, HashSet<FileReference> BuildProducts, Dictionary<string, HashSet<FileReference>> TagNameToFileSet)
+		public override Task ExecuteAsync(JobContext Job, HashSet<FileReference> BuildProducts, Dictionary<string, HashSet<FileReference>> TagNameToFileSet)
 		{
 			// Output a warning if the project directory is specified
 			if(Parameters.ProjectDir != null)
@@ -190,6 +190,7 @@ namespace AutomationTool.Tasks
 
 			// Apply the tag to all the matching files
 			FindOrAddTagSet(TagNameToFileSet, Parameters.With).UnionWith(Files);
+			return Task.CompletedTask;
 		}
 
 		/// <summary>
diff --git a/Engine/Source/Programs/AutomationTool/BuildGraph/Tasks/TagTask.cs b/Engine/Source/Programs/AutomationTool/BuildGraph/Tasks/TagTask.cs
--- a/Engine/Source/Programs/AutomationTool/BuildGraph/Tasks/TagTask.cs
+++ b/Engine/Source/Programs/AutomationTool/BuildGraph/Tasks/TagTask.cs
@@ -61,7 +61,7 @@ namespace AutomationTool.Tasks
 	/// parameter. From this list, any files not matched by the 'Filter' parameter are removed, followed by any files matched by the 'Except' parameter.
 	/// </summary>
 	[TaskElement("Tag", typeof(TagTaskParameters))]
-	class TagTask : CustomTask
+	class TagTask : BgTaskImpl
 	{
 		/// <summary>
 		/// Parameters to this task
@@ -83,7 +83,7 @@ namespace AutomationTool.Tasks
 		/// <param name="Job">Information about the current job</param>
 		/// <param name="BuildProducts">Set of build products produced by this node.</param>
 		/// <param name="TagNameToFileSet">Mapping from tag names to the set of files they include</param>
-		public override void Execute(JobContext Job, HashSet<FileReference> BuildProducts, Dictionary<string, HashSet<FileReference>> TagNameToFileSet)
+		public override Task ExecuteAsync(JobContext Job, HashSet<FileReference> BuildProducts, Dictionary<string, HashSet<FileReference>> TagNameToFileSet)
 		{
 			// Get the base directory
 			DirectoryReference BaseDir = Parameters.BaseDir ?? Unreal.RootDirectory;
@@ -134,6 +134,7 @@ namespace AutomationTool.Tasks
 			{
 				FindOrAddTagSet(TagNameToFileSet, TagName).UnionWith(Files);
 			}
+			return Task.CompletedTask;
 		}
 
 		/// <summary>
diff --git a/Engine/Source/Programs/AutomationTool/BuildGraph/Tasks/UnzipTask.cs b/Engine/Source/Programs/AutomationTool/BuildGraph/Tasks/UnzipTask.cs
--- a/Engine/Source/Programs/AutomationTool/BuildGraph/Tasks/UnzipTask.cs
+++ b/Engine/Source/Programs/AutomationTool/BuildGraph/Tasks/UnzipTask.cs
@@ -48,7 +48,7 @@ namespace BuildGraph.Tasks
 	/// Extract files from a zip archive.
 	/// </summary>
 	[TaskElement("Unzip", typeof(UnzipTaskParameters))]
-	public class UnzipTask : CustomTask
+	public class UnzipTask : BgTaskImpl
 	{
 		/// <summary>
 		/// Parameters for this task
@@ -70,7 +70,7 @@ namespace BuildGraph.Tasks
 		/// <param name="Job">Information about the current job</param>
 		/// <param name="BuildProducts">Set of build products produced by this node.</param>
 		/// <param name="TagNameToFileSet">Mapping from tag names \to the set of files they include</param>
-		public override void Execute(JobContext Job, HashSet<FileReference> BuildProducts, Dictionary<string, HashSet<FileReference>> TagNameToFileSet)
+		public override Task ExecuteAsync(JobContext Job, HashSet<FileReference> BuildProducts, Dictionary<string, HashSet<FileReference>> TagNameToFileSet)
 		{
 			DirectoryReference ToDir = Parameters.ToDir;
 
@@ -99,6 +99,7 @@ namespace BuildGraph.Tasks
 
 			// Add the archive to the set of build products
 			BuildProducts.UnionWith(OutputFiles);
+			return Task.CompletedTask;
 		}
 
 		/// <summary>
diff --git a/Engine/Source/Programs/AutomationTool/BuildGraph/Tasks/UploadTask.cs b/Engine/Source/Programs/AutomationTool/BuildGraph/Tasks/UploadTask.cs
--- a/Engine/Source/Programs/AutomationTool/BuildGraph/Tasks/UploadTask.cs
+++ b/Engine/Source/Programs/AutomationTool/BuildGraph/Tasks/UploadTask.cs
@@ -92,7 +92,7 @@ namespace BuildGraph.Tasks
 	/// Uploads a set of files to Jupiter for future retrival
 	/// </summary>
 	[TaskElement("Upload", typeof(UploadTaskParameters))]
-	public class UploadTask : CustomTask
+	public class UploadTask : BgTaskImpl
 	{
 		/// <summary>
 		/// Parameters for this task
@@ -114,7 +114,7 @@ namespace BuildGraph.Tasks
 		/// <param name="Job">Information about the current job</param>
 		/// <param name="BuildProducts">Set of build products produced by this node.</param>
 		/// <param name="TagNameToFileSet">Mapping from tag names to the set of files they include</param>
-		public override void Execute(JobContext Job, HashSet<FileReference> BuildProducts, Dictionary<string, HashSet<FileReference>> TagNameToFileSet)
+		public override async Task ExecuteAsync(JobContext Job, HashSet<FileReference> BuildProducts, Dictionary<string, HashSet<FileReference>> TagNameToFileSet)
 		{
 			// Find all the input files
 			List<FileReference> Files;
@@ -160,9 +160,7 @@ namespace BuildGraph.Tasks
 				}
 			}
 			// Upload the tree to Jupiter
-			Task<Dictionary<FileReference, List<string>>> Task = FileTree.UploadToJupiter(Parameters.JupiterUrl, Parameters.JupiterNamespace, Parameters.JupiterKey, Metadata);
-			Task.Wait();
-			Dictionary<FileReference, List<string>> Mapping = Task.Result;
+			Dictionary<FileReference, List<string>> Mapping = await FileTree.UploadToJupiter(Parameters.JupiterUrl, Parameters.JupiterNamespace, Parameters.JupiterKey, Metadata);
 
 			// Debug output of which files mapped to which blobs, can be useful to determine which files are constantly being uploaded
 			// Json.Save(FileReference.Combine(AutomationTool.Unreal.RootDirectory, "JupiterUpload.json"), Mapping);
diff --git a/Engine/Source/Programs/AutomationTool/BuildGraph/Tasks/WaitTask.cs b/Engine/Source/Programs/AutomationTool/BuildGraph/Tasks/WaitTask.cs
--- a/Engine/Source/Programs/AutomationTool/BuildGraph/Tasks/WaitTask.cs
+++ b/Engine/Source/Programs/AutomationTool/BuildGraph/Tasks/WaitTask.cs
@@ -4,6 +4,7 @@ using System;
 using System.Collections.Generic;
 using System.IO;
 using System.Threading;
+using System.Threading.Tasks;
 using System.Xml;
 using AutomationTool;
 using EpicGames.Core;
@@ -26,7 +27,7 @@ namespace BuildGraph.Tasks
 	/// Waits a defined number of seconds.
 	/// </summary>
 	[TaskElement("Wait", typeof(WaitTaskParameters))]
-	public class WaitTask : CustomTask
+	public class WaitTask : BgTaskImpl
 	{
 		/// <summary>
 		/// Parameters for this task
@@ -48,9 +49,9 @@ namespace BuildGraph.Tasks
 		/// <param name="Job">Information about the current job</param>
 		/// <param name="BuildProducts">Set of build products produced by this node.</param>
 		/// <param name="TagNameToFileSet">Mapping from tag names to the set of files they include</param>
-		public override void Execute(JobContext Job, HashSet<FileReference> BuildProducts, Dictionary<string, HashSet<FileReference>> TagNameToFileSet)
+		public override async Task ExecuteAsync(JobContext Job, HashSet<FileReference> BuildProducts, Dictionary<string, HashSet<FileReference>> TagNameToFileSet)
 		{
-			Thread.Sleep(TimeSpan.FromSeconds(Parameters.Seconds));
+			await Task.Delay(TimeSpan.FromSeconds(Parameters.Seconds));
 		}
 
 		/// <summary>
diff --git a/Engine/Source/Programs/AutomationTool/BuildGraph/Tasks/WriteTextFileTask.cs b/Engine/Source/Programs/AutomationTool/BuildGraph/Tasks/WriteTextFileTask.cs
--- a/Engine/Source/Programs/AutomationTool/BuildGraph/Tasks/WriteTextFileTask.cs
+++ b/Engine/Source/Programs/AutomationTool/BuildGraph/Tasks/WriteTextFileTask.cs
@@ -8,6 +8,7 @@ using System.Linq;
 using System.Xml;
 using EpicGames.Core;
 using UnrealBuildBase;
+using System.Threading.Tasks;
 
 namespace BuildGraph.Tasks
 {
@@ -51,7 +52,7 @@ namespace BuildGraph.Tasks
 	/// Writes text to a file.
 	/// </summary>
 	[TaskElement("WriteTextFile", typeof(WriteTextFileTaskParameters))]
-	public class WriteTextFileTask : CustomTask
+	public class WriteTextFileTask : BgTaskImpl
 	{
 		/// <summary>
 		/// Parameters for this task.
@@ -73,7 +74,7 @@ namespace BuildGraph.Tasks
 		/// <param name="Job">Information about the current job.</param>
 		/// <param name="BuildProducts">Set of build products produced by this node.</param>
 		/// <param name="TagNameToFileSet">Mapping from tag names to the set of files they include.</param>
-		public override void Execute(JobContext Job, HashSet<FileReference> BuildProducts, Dictionary<string, HashSet<FileReference>> TagNameToFileSet)
+		public override async Task ExecuteAsync(JobContext Job, HashSet<FileReference> BuildProducts, Dictionary<string, HashSet<FileReference>> TagNameToFileSet)
 		{
 			string FileText = Parameters.Text;
 
@@ -101,12 +102,12 @@ namespace BuildGraph.Tasks
 			if (Parameters.Append)
 			{
 				CommandUtils.LogInformation(string.Format("Appending text to file '{0}': {1}", Parameters.File, FileText));
-				FileReference.AppendAllText(Parameters.File, Environment.NewLine + FileText);
+				await FileReference.AppendAllTextAsync(Parameters.File, Environment.NewLine + FileText);
 			}
 			else
 			{
 				CommandUtils.LogInformation(string.Format("Writing text to file '{0}': {1}", Parameters.File, FileText));
-				FileReference.WriteAllText(Parameters.File, FileText);
+				await FileReference.WriteAllTextAsync(Parameters.File, FileText);
 			}
 		}
 
diff --git a/Engine/Source/Programs/AutomationTool/BuildGraph/Tasks/ZipTask.cs b/Engine/Source/Programs/AutomationTool/BuildGraph/Tasks/ZipTask.cs
--- a/Engine/Source/Programs/AutomationTool/BuildGraph/Tasks/ZipTask.cs
+++ b/Engine/Source/Programs/AutomationTool/BuildGraph/Tasks/ZipTask.cs
@@ -47,7 +47,7 @@ namespace BuildGraph.Tasks
 	/// Compresses files into a zip archive.
 	/// </summary>
 	[TaskElement("Zip", typeof(ZipTaskParameters))]
-	public class ZipTask : CustomTask
+	public class ZipTask : BgTaskImpl
 	{
 		/// <summary>
 		/// Parameters for this task
@@ -69,7 +69,7 @@ namespace BuildGraph.Tasks
 		/// <param name="Job">Information about the current job</param>
 		/// <param name="BuildProducts">Set of build products produced by this node.</param>
 		/// <param name="TagNameToFileSet">Mapping from tag names to the set of files they include</param>
-		public override void Execute(JobContext Job, HashSet<FileReference> BuildProducts, Dictionary<string, HashSet<FileReference>> TagNameToFileSet)
+		public override Task ExecuteAsync(JobContext Job, HashSet<FileReference> BuildProducts, Dictionary<string, HashSet<FileReference>> TagNameToFileSet)
 		{
 			// Find all the input files
 			List<FileReference> Files;
@@ -94,6 +94,7 @@ namespace BuildGraph.Tasks
 
 			// Add the archive to the set of build products
 			BuildProducts.Add(Parameters.ZipFile);
+			return Task.CompletedTask;
 		}
 
 		/// <summary>
diff --git a/Engine/Source/Programs/AutomationTool/Program.cs b/Engine/Source/Programs/AutomationTool/Program.cs
--- a/Engine/Source/Programs/AutomationTool/Program.cs
+++ b/Engine/Source/Programs/AutomationTool/Program.cs
@@ -14,6 +14,7 @@ using System.Text;
 using Microsoft.Build.Execution;
 using Microsoft.Build.Framework.Profiler;
 using UnrealBuildBase;
+using System.Threading.Tasks;
 
 namespace AutomationToolDriver
 {
@@ -312,7 +313,7 @@ namespace AutomationToolDriver
 		static StartupTraceListener StartupListener = new StartupTraceListener();
 
 		// Do not add [STAThread] here. It will cause deadlocks in platform automation code.
-		public static int Main(string[] Arguments)
+		public static async Task<int> Main(string[] Arguments)
 		{
 			// Initialize the log system, buffering the output until we can create the log file
 			Log.AddTraceListener(StartupListener);
@@ -390,7 +391,7 @@ namespace AutomationToolDriver
 				bool bWaitForUATMutex = AutomationToolCommandLine.IsSetGlobal("-WaitForUATMutex");
 
 				// Don't allow simultaneous execution of AT (in the same branch)
-				ReturnCode = ProcessSingleton.RunSingleInstance(MainProc, bWaitForUATMutex);
+				ReturnCode = await ProcessSingleton.RunSingleInstanceAsync(MainProc, bWaitForUATMutex);
 			}
 			catch (Exception Ex)
             {
@@ -408,7 +409,7 @@ namespace AutomationToolDriver
             return (int)ReturnCode;
         }
 
-		static ExitCode MainProc()
+		static async Task<ExitCode> MainProc()
 		{
 			string ScriptsForProject = (string)AutomationToolCommandLine.GetValueUnchecked("-ScriptsForProject");
 			List<string> AdditionalScriptDirs = (List<string>) AutomationToolCommandLine.GetValueUnchecked("-ScriptDir");
@@ -447,8 +448,8 @@ namespace AutomationToolDriver
 			// Call into AutomationTool.Automation.Process()
 
 			Type AutomationTools_Automation = AutomationUtilsAssembly.GetType("AutomationTool.Automation");
-			MethodInfo Automation_Process = AutomationTools_Automation.GetMethod("Process");
-			return (ExitCode) Automation_Process.Invoke(null,
+			MethodInfo Automation_Process = AutomationTools_Automation.GetMethod("ProcessAsync");
+			return await (Task<ExitCode>) Automation_Process.Invoke(null,
 				new object[] {AutomationToolCommandLine, StartupListener, ScriptModuleAssemblyPaths});
 		}
 	}
diff --git a/Engine/Source/Programs/AutomationTool/Turnkey/Commands/ExecuteBuild.cs b/Engine/Source/Programs/AutomationTool/Turnkey/Commands/ExecuteBuild.cs
--- a/Engine/Source/Programs/AutomationTool/Turnkey/Commands/ExecuteBuild.cs
+++ b/Engine/Source/Programs/AutomationTool/Turnkey/Commands/ExecuteBuild.cs
@@ -157,7 +157,7 @@ namespace Turnkey.Commands
 			BCRCommand.Arguments = Arguments.Select(x => x.Substring(1)).ToList();
 
 			// use the BCR's exitcode as Turnkey's exitcode
-			TurnkeyUtils.ExitCode = Automation.Execute(new List<UnrealBuildBase.CommandInfo>() { BCRCommand }, ScriptManager.Commands);
+			TurnkeyUtils.ExitCode = Automation.ExecuteAsync(new List<UnrealBuildBase.CommandInfo>() { BCRCommand }, ScriptManager.Commands).Result;
 		}
 
 		string GetIniSetting(string Spec, DirectoryReference ProjectDir, UnrealTargetPlatform Platform)
diff --git a/Engine/Source/Programs/Shared/EpicGames.Build/Automation/ProcessSingleton.cs b/Engine/Source/Programs/Shared/EpicGames.Build/Automation/ProcessSingleton.cs
--- a/Engine/Source/Programs/Shared/EpicGames.Build/Automation/ProcessSingleton.cs
+++ b/Engine/Source/Programs/Shared/EpicGames.Build/Automation/ProcessSingleton.cs
@@ -7,6 +7,7 @@ using System.IO;
 using System.Reflection;
 using System.Text;
 using System.Threading;
+using System.Threading.Tasks;
 
 namespace UnrealBuildBase
 {
@@ -28,8 +29,8 @@ namespace UnrealBuildBase
 		/// Runs the specified delegate checking if this is the only instance of the application.
 		/// </summary>
 		/// <param name="Main"></param>
-		/// <param name="Param"></param>
-		public static ExitCode RunSingleInstance(Func<ExitCode> Main, bool bWaitForUATMutex)
+		/// <param name="bWaitForUATMutex"></param>
+		public static async Task<ExitCode> RunSingleInstanceAsync(Func<Task<ExitCode>> Main, bool bWaitForUATMutex)
 		{
 			bool AllowMultipleInsances = (Environment.GetEnvironmentVariable("uebp_UATMutexNoWait") == "1");
 	
@@ -58,7 +59,7 @@ namespace UnrealBuildBase
 					}
 				}
 
-				ExitCode Result = Main();
+				ExitCode Result = await Main();
 
 				if (IsSoleInstance)
 				{
diff --git a/Engine/Source/Programs/Shared/EpicGames.Core/FileReference.cs b/Engine/Source/Programs/Shared/EpicGames.Core/FileReference.cs
--- a/Engine/Source/Programs/Shared/EpicGames.Core/FileReference.cs
+++ b/Engine/Source/Programs/Shared/EpicGames.Core/FileReference.cs
@@ -763,6 +763,16 @@ namespace EpicGames.Core
 			File.AppendAllLines(Location.FullName, Contents);
 		}
 
+		/// <summary>
+		/// Appends the contents to a file
+		/// </summary>
+		/// <param name="Location">Location of the file</param>
+		/// <param name="Contents">Contents to append to the file</param>
+		public static Task AppendAllLinesAsync(FileReference Location, IEnumerable<string> Contents)
+		{
+			return File.AppendAllLinesAsync(Location.FullName, Contents);
+		}
+
 		/// <summary>
 		/// Appends the contents to a file
 		/// </summary>
@@ -774,6 +784,17 @@ namespace EpicGames.Core
 			File.AppendAllLines(Location.FullName, Contents, Encoding);
 		}
 
+		/// <summary>
+		/// Appends the contents to a file
+		/// </summary>
+		/// <param name="Location">Location of the file</param>
+		/// <param name="Contents">Contents to append to the file</param>
+		/// <param name="Encoding">The encoding to use when parsing the file</param>
+		public static Task AppendAllLinesAsync(FileReference Location, IEnumerable<string> Contents, Encoding Encoding)
+		{
+			return File.AppendAllLinesAsync(Location.FullName, Contents, Encoding);
+		}
+
 		/// <summary>
 		/// Appends the contents to a file
 		/// </summary>
@@ -784,6 +805,16 @@ namespace EpicGames.Core
 			File.AppendAllLines(Location.FullName, Contents);
 		}
 
+		/// <summary>
+		/// Appends the contents to a file
+		/// </summary>
+		/// <param name="Location">Location of the file</param>
+		/// <param name="Contents">Contents to append to the file</param>
+		public static Task AppendAllLinesAsync(FileReference Location, string[] Contents)
+		{
+			return File.AppendAllLinesAsync(Location.FullName, Contents);
+		}
+
 		/// <summary>
 		/// Appends the contents to a file
 		/// </summary>
@@ -795,6 +826,17 @@ namespace EpicGames.Core
 			File.AppendAllLines(Location.FullName, Contents, Encoding);
 		}
 
+		/// <summary>
+		/// Appends the contents to a file
+		/// </summary>
+		/// <param name="Location">Location of the file</param>
+		/// <param name="Contents">Contents to append to the file</param>
+		/// <param name="Encoding">The encoding to use when parsing the file</param>
+		public static Task AppendAllLinesAsync(FileReference Location, string[] Contents, Encoding Encoding)
+		{
+			return File.AppendAllLinesAsync(Location.FullName, Contents, Encoding);
+		}
+
 		/// <summary>
 		/// Appends the contents to a file
 		/// </summary>
@@ -805,6 +847,16 @@ namespace EpicGames.Core
 			File.AppendAllText(Location.FullName, Contents);
 		}
 
+		/// <summary>
+		/// Appends the contents to a file
+		/// </summary>
+		/// <param name="Location">Location of the file</param>
+		/// <param name="Contents">Contents to append to the file</param>
+		public static Task AppendAllTextAsync(FileReference Location, string Contents)
+		{
+			return File.AppendAllTextAsync(Location.FullName, Contents);
+		}
+
 		/// <summary>
 		/// Appends the contents to a file
 		/// </summary>
@@ -816,6 +868,17 @@ namespace EpicGames.Core
 			File.AppendAllText(Location.FullName, Contents, Encoding);
 		}
 
+		/// <summary>
+		/// Appends the contents to a file
+		/// </summary>
+		/// <param name="Location">Location of the file</param>
+		/// <param name="Contents">Contents to append to the file</param>
+		/// <param name="Encoding">The encoding to use when parsing the file</param>
+		public static Task AppendAllTextAsync(FileReference Location, string Contents, Encoding Encoding)
+		{
+			return File.AppendAllTextAsync(Location.FullName, Contents, Encoding);
+		}
+
 		#endregion
 	}
 
