UAT: Move all BuildGraph tasks into the AutomationTool.Tasks namespace (currently inconsistently in multiple other namespaces).

diff --git a/Engine/Source/Programs/AutomationTool/BuildGraph/Tasks/AgeStoreTask.cs b/Engine/Source/Programs/AutomationTool/BuildGraph/Tasks/AgeStoreTask.cs
--- a/Engine/Source/Programs/AutomationTool/BuildGraph/Tasks/AgeStoreTask.cs
+++ b/Engine/Source/Programs/AutomationTool/BuildGraph/Tasks/AgeStoreTask.cs
@@ -10,7 +10,7 @@ using EpicGames.Core;
 using UnrealBuildTool;
 using System.Threading.Tasks;
 
-namespace Win.Automation
+namespace AutomationTool.Tasks
 {
 	/// <summary>
 	/// Parameters for a task that purges data from a symbol store after a given age
diff --git a/Engine/Source/Programs/AutomationTool/BuildGraph/Tasks/AwsAssumeRoleTask.cs b/Engine/Source/Programs/AutomationTool/BuildGraph/Tasks/AwsAssumeRoleTask.cs
--- a/Engine/Source/Programs/AutomationTool/BuildGraph/Tasks/AwsAssumeRoleTask.cs
+++ b/Engine/Source/Programs/AutomationTool/BuildGraph/Tasks/AwsAssumeRoleTask.cs
@@ -11,7 +11,7 @@ using System.Text.Json;
 using System.Threading.Tasks;
 using System.Xml;
 
-namespace BuildGraph.Tasks
+namespace AutomationTool.Tasks
 {
 	/// <summary>
 	/// Parameters for an AWS CLI task
diff --git a/Engine/Source/Programs/AutomationTool/BuildGraph/Tasks/AwsEcsDeployTask.cs b/Engine/Source/Programs/AutomationTool/BuildGraph/Tasks/AwsEcsDeployTask.cs
--- a/Engine/Source/Programs/AutomationTool/BuildGraph/Tasks/AwsEcsDeployTask.cs
+++ b/Engine/Source/Programs/AutomationTool/BuildGraph/Tasks/AwsEcsDeployTask.cs
@@ -11,7 +11,7 @@ using System.Threading.Tasks;
 using System.Xml;
 using UnrealBuildBase;
 
-namespace BuildGraph.Tasks
+namespace AutomationTool.Tasks
 {
 	/// <summary>
 	/// Parameters for a AWS ECS deploy task
diff --git a/Engine/Source/Programs/AutomationTool/BuildGraph/Tasks/AwsTask.cs b/Engine/Source/Programs/AutomationTool/BuildGraph/Tasks/AwsTask.cs
--- a/Engine/Source/Programs/AutomationTool/BuildGraph/Tasks/AwsTask.cs
+++ b/Engine/Source/Programs/AutomationTool/BuildGraph/Tasks/AwsTask.cs
@@ -10,7 +10,7 @@ using System.Text.Json;
 using System.Threading.Tasks;
 using System.Xml;
 
-namespace BuildGraph.Tasks
+namespace AutomationTool.Tasks
 {
 	/// <summary>
 	/// Parameters for a spawn task
diff --git a/Engine/Source/Programs/AutomationTool/BuildGraph/Tasks/CommandTask.cs b/Engine/Source/Programs/AutomationTool/BuildGraph/Tasks/CommandTask.cs
--- a/Engine/Source/Programs/AutomationTool/BuildGraph/Tasks/CommandTask.cs
+++ b/Engine/Source/Programs/AutomationTool/BuildGraph/Tasks/CommandTask.cs
@@ -12,7 +12,7 @@ using System.IO;
 using EpicGames.Core;
 using UnrealBuildBase;
 
-namespace BuildGraph.Tasks
+namespace AutomationTool.Tasks
 {
 	static class StringExtensions
 	{
diff --git a/Engine/Source/Programs/AutomationTool/BuildGraph/Tasks/CompileTask.cs b/Engine/Source/Programs/AutomationTool/BuildGraph/Tasks/CompileTask.cs
--- a/Engine/Source/Programs/AutomationTool/BuildGraph/Tasks/CompileTask.cs
+++ b/Engine/Source/Programs/AutomationTool/BuildGraph/Tasks/CompileTask.cs
@@ -12,7 +12,7 @@ using System.Xml;
 using EpicGames.Core;
 using OpenTracing;
 
-namespace AutomationTool
+namespace AutomationTool.Tasks
 {
 	/// <summary>
 	/// Parameters for a compile task
diff --git a/Engine/Source/Programs/AutomationTool/BuildGraph/Tasks/CookTask.cs b/Engine/Source/Programs/AutomationTool/BuildGraph/Tasks/CookTask.cs
--- a/Engine/Source/Programs/AutomationTool/BuildGraph/Tasks/CookTask.cs
+++ b/Engine/Source/Programs/AutomationTool/BuildGraph/Tasks/CookTask.cs
@@ -14,7 +14,7 @@ using OpenTracing;
 using OpenTracing.Util;
 using UnrealBuildTool;
 
-namespace BuildGraph.Tasks
+namespace AutomationTool.Tasks
 {
 	/// <summary>
 	/// Parameters for a task that runs the cooker
diff --git a/Engine/Source/Programs/AutomationTool/BuildGraph/Tasks/CopyTask.cs b/Engine/Source/Programs/AutomationTool/BuildGraph/Tasks/CopyTask.cs
--- a/Engine/Source/Programs/AutomationTool/BuildGraph/Tasks/CopyTask.cs
+++ b/Engine/Source/Programs/AutomationTool/BuildGraph/Tasks/CopyTask.cs
@@ -13,7 +13,7 @@ using EpicGames.Core;
 using UnrealBuildTool;
 using UnrealBuildBase;
 
-namespace BuildGraph.Tasks
+namespace AutomationTool.Tasks
 {
 	/// <summary>
 	/// Parameters for a copy task
diff --git a/Engine/Source/Programs/AutomationTool/BuildGraph/Tasks/DeleteTask.cs b/Engine/Source/Programs/AutomationTool/BuildGraph/Tasks/DeleteTask.cs
--- a/Engine/Source/Programs/AutomationTool/BuildGraph/Tasks/DeleteTask.cs
+++ b/Engine/Source/Programs/AutomationTool/BuildGraph/Tasks/DeleteTask.cs
@@ -13,7 +13,7 @@ using EpicGames.Core;
 using UnrealBuildTool;
 using UnrealBuildBase;
 
-namespace BuildGraph.Tasks
+namespace AutomationTool.Tasks
 {
 	/// <summary>
 	/// Parameters for a copy task
diff --git a/Engine/Source/Programs/AutomationTool/BuildGraph/Tasks/DockerBuildTask.cs b/Engine/Source/Programs/AutomationTool/BuildGraph/Tasks/DockerBuildTask.cs
--- a/Engine/Source/Programs/AutomationTool/BuildGraph/Tasks/DockerBuildTask.cs
+++ b/Engine/Source/Programs/AutomationTool/BuildGraph/Tasks/DockerBuildTask.cs
@@ -11,7 +11,7 @@ using AutomationTool;
 using UnrealBuildBase;
 using System.Threading.Tasks;
 
-namespace BuildGraph.Tasks
+namespace AutomationTool.Tasks
 {
 	/// <summary>
 	/// Parameters for a Docker-Build task
diff --git a/Engine/Source/Programs/AutomationTool/BuildGraph/Tasks/DockerComposeDownTask.cs b/Engine/Source/Programs/AutomationTool/BuildGraph/Tasks/DockerComposeDownTask.cs
--- a/Engine/Source/Programs/AutomationTool/BuildGraph/Tasks/DockerComposeDownTask.cs
+++ b/Engine/Source/Programs/AutomationTool/BuildGraph/Tasks/DockerComposeDownTask.cs
@@ -11,7 +11,7 @@ using AutomationTool;
 using UnrealBuildBase;
 using System.Threading.Tasks;
 
-namespace BuildGraph.Tasks
+namespace AutomationTool.Tasks
 {
 	/// <summary>
 	/// Parameters for a Docker-Compose task
diff --git a/Engine/Source/Programs/AutomationTool/BuildGraph/Tasks/DockerComposeUpTask.cs b/Engine/Source/Programs/AutomationTool/BuildGraph/Tasks/DockerComposeUpTask.cs
--- a/Engine/Source/Programs/AutomationTool/BuildGraph/Tasks/DockerComposeUpTask.cs
+++ b/Engine/Source/Programs/AutomationTool/BuildGraph/Tasks/DockerComposeUpTask.cs
@@ -11,7 +11,7 @@ using AutomationTool;
 using UnrealBuildBase;
 using System.Threading.Tasks;
 
-namespace BuildGraph.Tasks
+namespace AutomationTool.Tasks
 {
 	/// <summary>
 	/// Parameters for a Docker-Compose task
diff --git a/Engine/Source/Programs/AutomationTool/BuildGraph/Tasks/DockerPushTask.cs b/Engine/Source/Programs/AutomationTool/BuildGraph/Tasks/DockerPushTask.cs
--- a/Engine/Source/Programs/AutomationTool/BuildGraph/Tasks/DockerPushTask.cs
+++ b/Engine/Source/Programs/AutomationTool/BuildGraph/Tasks/DockerPushTask.cs
@@ -11,7 +11,7 @@ using System.Text.Json;
 using System.Threading.Tasks;
 using System.Xml;
 
-namespace BuildGraph.Tasks
+namespace AutomationTool.Tasks
 {
 	/// <summary>
 	/// Parameters for a Docker-Build task
diff --git a/Engine/Source/Programs/AutomationTool/BuildGraph/Tasks/DockerTask.cs b/Engine/Source/Programs/AutomationTool/BuildGraph/Tasks/DockerTask.cs
--- a/Engine/Source/Programs/AutomationTool/BuildGraph/Tasks/DockerTask.cs
+++ b/Engine/Source/Programs/AutomationTool/BuildGraph/Tasks/DockerTask.cs
@@ -9,7 +9,7 @@ using System.Text;
 using System.Threading.Tasks;
 using System.Xml;
 
-namespace BuildGraph.Tasks
+namespace AutomationTool.Tasks
 {
 	/// <summary>
 	/// Parameters for a Docker task
diff --git a/Engine/Source/Programs/AutomationTool/BuildGraph/Tasks/DotNetTask.cs b/Engine/Source/Programs/AutomationTool/BuildGraph/Tasks/DotNetTask.cs
--- a/Engine/Source/Programs/AutomationTool/BuildGraph/Tasks/DotNetTask.cs
+++ b/Engine/Source/Programs/AutomationTool/BuildGraph/Tasks/DotNetTask.cs
@@ -9,7 +9,7 @@ using System.Text;
 using System.Threading.Tasks;
 using System.Xml;
 
-namespace BuildGraph.Tasks
+namespace AutomationTool.Tasks
 {
 	/// <summary>
 	/// Parameters for a DotNet task
diff --git a/Engine/Source/Programs/AutomationTool/BuildGraph/Tasks/HelmTask.cs b/Engine/Source/Programs/AutomationTool/BuildGraph/Tasks/HelmTask.cs
--- a/Engine/Source/Programs/AutomationTool/BuildGraph/Tasks/HelmTask.cs
+++ b/Engine/Source/Programs/AutomationTool/BuildGraph/Tasks/HelmTask.cs
@@ -10,7 +10,7 @@ using System.Threading.Tasks;
 using System.Xml;
 using UnrealBuildBase;
 
-namespace BuildGraph.Tasks
+namespace AutomationTool.Tasks
 {
 	/// <summary>
 	/// Parameters for a Helm task
diff --git a/Engine/Source/Programs/AutomationTool/BuildGraph/Tasks/LogTask.cs b/Engine/Source/Programs/AutomationTool/BuildGraph/Tasks/LogTask.cs
--- a/Engine/Source/Programs/AutomationTool/BuildGraph/Tasks/LogTask.cs
+++ b/Engine/Source/Programs/AutomationTool/BuildGraph/Tasks/LogTask.cs
@@ -13,7 +13,7 @@ using EpicGames.Core;
 using UnrealBuildTool;
 using UnrealBuildBase;
 
-namespace BuildGraph.Tasks
+namespace AutomationTool.Tasks
 {
 	/// <summary>
 	/// Parameters for the log task
diff --git a/Engine/Source/Programs/AutomationTool/BuildGraph/Tasks/ModifyConfigTask.cs b/Engine/Source/Programs/AutomationTool/BuildGraph/Tasks/ModifyConfigTask.cs
--- a/Engine/Source/Programs/AutomationTool/BuildGraph/Tasks/ModifyConfigTask.cs
+++ b/Engine/Source/Programs/AutomationTool/BuildGraph/Tasks/ModifyConfigTask.cs
@@ -11,7 +11,7 @@ using System.Xml;
 using EpicGames.Core;
 using UnrealBuildTool;
 
-namespace BuildGraph.Tasks
+namespace AutomationTool.Tasks
 {
 	/// <summary>
 	/// Parameters for a ModifyConfig task
diff --git a/Engine/Source/Programs/AutomationTool/BuildGraph/Tasks/ModifyJsonValue.cs b/Engine/Source/Programs/AutomationTool/BuildGraph/Tasks/ModifyJsonValue.cs
--- a/Engine/Source/Programs/AutomationTool/BuildGraph/Tasks/ModifyJsonValue.cs
+++ b/Engine/Source/Programs/AutomationTool/BuildGraph/Tasks/ModifyJsonValue.cs
@@ -13,7 +13,7 @@ using UnrealBuildBase;
 using System.Text.Json;
 using System.Threading.Tasks;
 
-namespace BuildGraph.Tasks
+namespace AutomationTool.Tasks
 {
 	/// <summary>
 	/// Parameters for a ModifyJsonValue task
diff --git a/Engine/Source/Programs/AutomationTool/BuildGraph/Tasks/MoveTask.cs b/Engine/Source/Programs/AutomationTool/BuildGraph/Tasks/MoveTask.cs
--- a/Engine/Source/Programs/AutomationTool/BuildGraph/Tasks/MoveTask.cs
+++ b/Engine/Source/Programs/AutomationTool/BuildGraph/Tasks/MoveTask.cs
@@ -12,7 +12,7 @@ using EpicGames.Core;
 using UnrealBuildTool;
 using UnrealBuildBase;
 
-namespace BuildGraph.Tasks
+namespace AutomationTool.Tasks
 {
 	/// <summary>
 	/// Parameters for a move task
diff --git a/Engine/Source/Programs/AutomationTool/BuildGraph/Tasks/NotarizeTask.cs b/Engine/Source/Programs/AutomationTool/BuildGraph/Tasks/NotarizeTask.cs
--- a/Engine/Source/Programs/AutomationTool/BuildGraph/Tasks/NotarizeTask.cs
+++ b/Engine/Source/Programs/AutomationTool/BuildGraph/Tasks/NotarizeTask.cs
@@ -14,7 +14,7 @@ using System.Xml;
 using EpicGames.Core;
 using UnrealBuildTool;
 
-namespace BuildGraph.Tasks
+namespace AutomationTool.Tasks
 {
 	/// <summary>
 	/// Parameters for a task that notarizes a dmg via the apple notarization process
diff --git a/Engine/Source/Programs/AutomationTool/BuildGraph/Tasks/PakFileTask.cs b/Engine/Source/Programs/AutomationTool/BuildGraph/Tasks/PakFileTask.cs
--- a/Engine/Source/Programs/AutomationTool/BuildGraph/Tasks/PakFileTask.cs
+++ b/Engine/Source/Programs/AutomationTool/BuildGraph/Tasks/PakFileTask.cs
@@ -13,7 +13,7 @@ using EpicGames.Core;
 using UnrealBuildTool;
 using UnrealBuildBase;
 
-namespace BuildGraph.Tasks
+namespace AutomationTool.Tasks
 {
 	/// <summary>
 	/// Parameters for a task that runs the cooker
diff --git a/Engine/Source/Programs/AutomationTool/BuildGraph/Tasks/RenameTask.cs b/Engine/Source/Programs/AutomationTool/BuildGraph/Tasks/RenameTask.cs
--- a/Engine/Source/Programs/AutomationTool/BuildGraph/Tasks/RenameTask.cs
+++ b/Engine/Source/Programs/AutomationTool/BuildGraph/Tasks/RenameTask.cs
@@ -14,7 +14,7 @@ using EpicGames.Core;
 using UnrealBuildTool;
 using UnrealBuildBase;
 
-namespace BuildGraph.Tasks
+namespace AutomationTool.Tasks
 {
 	/// <summary>
 	/// Parameters for a copy task
diff --git a/Engine/Source/Programs/AutomationTool/BuildGraph/Tasks/SignTask.cs b/Engine/Source/Programs/AutomationTool/BuildGraph/Tasks/SignTask.cs
--- a/Engine/Source/Programs/AutomationTool/BuildGraph/Tasks/SignTask.cs
+++ b/Engine/Source/Programs/AutomationTool/BuildGraph/Tasks/SignTask.cs
@@ -12,7 +12,7 @@ using EpicGames.Core;
 using UnrealBuildTool;
 using UnrealBuildBase;
 
-namespace BuildGraph.Tasks
+namespace AutomationTool.Tasks
 {
 	/// <summary>
 	/// Parameters for a task that strips symbols from a set of files
diff --git a/Engine/Source/Programs/AutomationTool/BuildGraph/Tasks/SpawnTask.cs b/Engine/Source/Programs/AutomationTool/BuildGraph/Tasks/SpawnTask.cs
--- a/Engine/Source/Programs/AutomationTool/BuildGraph/Tasks/SpawnTask.cs
+++ b/Engine/Source/Programs/AutomationTool/BuildGraph/Tasks/SpawnTask.cs
@@ -11,7 +11,7 @@ using System.Xml;
 using EpicGames.Core;
 using UnrealBuildTool;
 
-namespace BuildGraph.Tasks
+namespace AutomationTool.Tasks
 {
 	/// <summary>
 	/// Parameters for a spawn task
diff --git a/Engine/Source/Programs/AutomationTool/BuildGraph/Tasks/StageTask.cs b/Engine/Source/Programs/AutomationTool/BuildGraph/Tasks/StageTask.cs
--- a/Engine/Source/Programs/AutomationTool/BuildGraph/Tasks/StageTask.cs
+++ b/Engine/Source/Programs/AutomationTool/BuildGraph/Tasks/StageTask.cs
@@ -12,7 +12,7 @@ using EpicGames.Core;
 using UnrealBuildTool;
 using UnrealBuildBase;
 
-namespace BuildGraph.Tasks
+namespace AutomationTool.Tasks
 {
 	/// <summary>
 	/// Parameters for the staging task
diff --git a/Engine/Source/Programs/AutomationTool/BuildGraph/Tasks/StripTask.cs b/Engine/Source/Programs/AutomationTool/BuildGraph/Tasks/StripTask.cs
--- a/Engine/Source/Programs/AutomationTool/BuildGraph/Tasks/StripTask.cs
+++ b/Engine/Source/Programs/AutomationTool/BuildGraph/Tasks/StripTask.cs
@@ -11,7 +11,7 @@ using System.Xml;
 using EpicGames.Core;
 using UnrealBuildTool;
 
-namespace BuildGraph.Tasks
+namespace AutomationTool.Tasks
 {
 	/// <summary>
 	/// Parameters for a task that strips symbols from a set of files
diff --git a/Engine/Source/Programs/AutomationTool/BuildGraph/Tasks/SymStoreTask.cs b/Engine/Source/Programs/AutomationTool/BuildGraph/Tasks/SymStoreTask.cs
--- a/Engine/Source/Programs/AutomationTool/BuildGraph/Tasks/SymStoreTask.cs
+++ b/Engine/Source/Programs/AutomationTool/BuildGraph/Tasks/SymStoreTask.cs
@@ -12,7 +12,7 @@ using EpicGames.Core;
 using UnrealBuildBase;
 using UnrealBuildTool;
 
-namespace AutomationTool
+namespace AutomationTool.Tasks
 {
     /// <summary>
     /// Parameters for a task that uploads symbols to a symbol server
diff --git a/Engine/Source/Programs/AutomationTool/BuildGraph/Tasks/UnzipTask.cs b/Engine/Source/Programs/AutomationTool/BuildGraph/Tasks/UnzipTask.cs
--- a/Engine/Source/Programs/AutomationTool/BuildGraph/Tasks/UnzipTask.cs
+++ b/Engine/Source/Programs/AutomationTool/BuildGraph/Tasks/UnzipTask.cs
@@ -12,7 +12,7 @@ using EpicGames.Core;
 using UnrealBuildTool;
 using UnrealBuildBase;
 
-namespace BuildGraph.Tasks
+namespace AutomationTool.Tasks
 {
 	/// <summary>
 	/// Parameters for a zip task
diff --git a/Engine/Source/Programs/AutomationTool/BuildGraph/Tasks/UploadTask.cs b/Engine/Source/Programs/AutomationTool/BuildGraph/Tasks/UploadTask.cs
--- a/Engine/Source/Programs/AutomationTool/BuildGraph/Tasks/UploadTask.cs
+++ b/Engine/Source/Programs/AutomationTool/BuildGraph/Tasks/UploadTask.cs
@@ -14,7 +14,7 @@ using System.Xml;
 using EpicGames.Core;
 using EpicGames.Jupiter;
 
-namespace BuildGraph.Tasks
+namespace AutomationTool.Tasks
 {
 	/// <summary>
 	/// Parameters for a zip task
diff --git a/Engine/Source/Programs/AutomationTool/BuildGraph/Tasks/WaitTask.cs b/Engine/Source/Programs/AutomationTool/BuildGraph/Tasks/WaitTask.cs
--- a/Engine/Source/Programs/AutomationTool/BuildGraph/Tasks/WaitTask.cs
+++ b/Engine/Source/Programs/AutomationTool/BuildGraph/Tasks/WaitTask.cs
@@ -9,7 +9,7 @@ using System.Xml;
 using AutomationTool;
 using EpicGames.Core;
 
-namespace BuildGraph.Tasks
+namespace AutomationTool.Tasks
 {
 	/// <summary>
 	/// Parameters for a wait task
diff --git a/Engine/Source/Programs/AutomationTool/BuildGraph/Tasks/WriteTextFileTask.cs b/Engine/Source/Programs/AutomationTool/BuildGraph/Tasks/WriteTextFileTask.cs
--- a/Engine/Source/Programs/AutomationTool/BuildGraph/Tasks/WriteTextFileTask.cs
+++ b/Engine/Source/Programs/AutomationTool/BuildGraph/Tasks/WriteTextFileTask.cs
@@ -10,7 +10,7 @@ using EpicGames.Core;
 using UnrealBuildBase;
 using System.Threading.Tasks;
 
-namespace BuildGraph.Tasks
+namespace AutomationTool.Tasks
 {
 	/// <summary>
 	/// Parameters for a <see cref="WriteTextFileTask"/>.
diff --git a/Engine/Source/Programs/AutomationTool/BuildGraph/Tasks/ZipTask.cs b/Engine/Source/Programs/AutomationTool/BuildGraph/Tasks/ZipTask.cs
--- a/Engine/Source/Programs/AutomationTool/BuildGraph/Tasks/ZipTask.cs
+++ b/Engine/Source/Programs/AutomationTool/BuildGraph/Tasks/ZipTask.cs
@@ -11,7 +11,7 @@ using System.Xml;
 using EpicGames.Core;
 using UnrealBuildTool;
 
-namespace BuildGraph.Tasks
+namespace AutomationTool.Tasks
 {
 	/// <summary>
 	/// Parameters for a zip task
