Ensure that Zen ProjectIds are unique for projects with the same name in differing roots (or in the same root) while attempting to maintain some easily recognizable desriptors for the project.
Staged builds are explicitly assigned a project id via the persistent UECommandline.txt file.  This works on mobile/consoles, but isn't implemented for Windows/Linux/Mac, so staged builds for those platforms don't have a persistent assignment of project id and instead rely on the launch by the editor including the appropriate commandlines.
Unstaged builds on Windows/Mac/Linux now determine project id independently after they have determined the project file path.  This happens after the project file has been determined and has had case corrected.  This ensures it is different/unique for multiple blueprint projects that share the same executable.
Fix a bug where second cook in editor to zen would stall because the queue thinks all additions are finished due to a boolean that isn't reset.

diff --git a/Engine/Source/Developer/IoStoreUtilities/Internal/ZenStoreHttpClient.h b/Engine/Source/Developer/IoStoreUtilities/Internal/ZenStoreHttpClient.h
--- a/Engine/Source/Developer/IoStoreUtilities/Internal/ZenStoreHttpClient.h
+++ b/Engine/Source/Developer/IoStoreUtilities/Internal/ZenStoreHttpClient.h
@@ -64,6 +64,7 @@ public:
 	static const UTF8CHAR* FindOrAddAttachmentId(FUtf8StringView AttachmentText);
 	static const UTF8CHAR* FindAttachmentId(FUtf8StringView AttachmentText);
 
+	static FString GenerateDefaultProjectId();
 private:
 	TIoStatusOr<FIoBuffer> ReadOpLogUri(FStringBuilderBase& ChunkUri, uint64 Offset = 0, uint64 Size = ~0ull);
 
diff --git a/Engine/Source/Developer/IoStoreUtilities/Private/ZenStoreHttpClient.cpp b/Engine/Source/Developer/IoStoreUtilities/Private/ZenStoreHttpClient.cpp
--- a/Engine/Source/Developer/IoStoreUtilities/Private/ZenStoreHttpClient.cpp
+++ b/Engine/Source/Developer/IoStoreUtilities/Private/ZenStoreHttpClient.cpp
@@ -12,7 +12,10 @@
 #include "Async/Async.h"
 #include "HAL/PlatformFileManager.h"
 #include "IO/IoHash.h"
+#include "Misc/App.h"
+#include "Misc/Paths.h"
 #include "Misc/ScopeLock.h"
+#include "Misc/SecureHash.h"
 #include "Misc/StringBuilder.h"
 #include "ProfilingDebugging/CpuProfilerTrace.h"
 #include "Serialization/BufferArchive.h"
@@ -763,4 +766,21 @@ const UTF8CHAR* FZenStoreHttpClient::FindAttachmentId(FUtf8StringView Attachment
 	return Existing;
 }
 
+// Duplicated in StorageServerPlatformFile.cpp to avoid having a public API in a shared module
+static FString GetProjectPathId()
+{
+	FString ProjectFilePath = FPaths::GetProjectFilePath();
+	FPaths::NormalizeFilename(ProjectFilePath);
+	FString AbsProjectFilePath = IFileManager::Get().ConvertToAbsolutePathForExternalAppForRead(*ProjectFilePath);
+	FTCHARToUTF8 AbsProjectFilePathUTF8(*AbsProjectFilePath);
+
+	FString HashString = FMD5::HashBytes((unsigned char*)AbsProjectFilePathUTF8.Get(), AbsProjectFilePathUTF8.Length()).Left(8);
+	return FString::Printf(TEXT("%s.%.8s"), FApp::GetProjectName(), *HashString);
+}
+
+FString FZenStoreHttpClient::GenerateDefaultProjectId()
+{
+	return GetProjectPathId();
+}
+
 }
diff --git a/Engine/Source/Developer/IoStoreUtilities/Private/ZenStoreWriter.cpp b/Engine/Source/Developer/IoStoreUtilities/Private/ZenStoreWriter.cpp
--- a/Engine/Source/Developer/IoStoreUtilities/Private/ZenStoreWriter.cpp
+++ b/Engine/Source/Developer/IoStoreUtilities/Private/ZenStoreWriter.cpp
@@ -104,6 +104,11 @@ public:
 		NewCommitEvent->Trigger();
 	}
 
+	void ResetAdding()
+	{
+		bCompleteAdding = false;
+	}
+
 	void WaitUntilEmpty()
 	{
 		bool bWait = false;
@@ -162,7 +167,7 @@ FZenStoreWriter::FZenStoreWriter(
 {
 	StaticInit();
 
-	FString ProjectId = FApp::GetProjectName();
+	FString ProjectId = FZenStoreHttpClient::GenerateDefaultProjectId();
 	FString OplogId = InTargetPlatform->PlatformName();
 
 	HttpClient = MakeUnique<UE::FZenStoreHttpClient>();
@@ -351,7 +356,7 @@ void FZenStoreWriter::Initialize(const FCookInfo& Info)
 			IFileManager::Get().DeleteDirectory(*OutputPath, bRequireExists, bTree);
 		}
 
-		FString ProjectId = FApp::GetProjectName();
+		FString ProjectId = FZenStoreHttpClient::GenerateDefaultProjectId();
 		FString OplogId = TargetPlatform.PlatformName();
 		bool bOplogEstablished = HttpClient->TryCreateOplog(ProjectId, OplogId, Info.bFullBuild);
 		UE_CLOG(!bOplogEstablished, LogZenStoreWriter, Fatal, TEXT("Failed to establish oplog on the ZenServer"));
@@ -490,6 +495,7 @@ void FZenStoreWriter::BeginCook()
 
 	if (FPlatformProcess::SupportsMultithreading())
 	{
+		CommitQueue->ResetAdding();
 		CommitThread = AsyncThread([this]()
 		{ 
 			for(;;)
diff --git a/Engine/Source/Editor/UnrealEd/Private/TargetDomain/TargetDomainUtils.cpp b/Engine/Source/Editor/UnrealEd/Private/TargetDomain/TargetDomainUtils.cpp
--- a/Engine/Source/Editor/UnrealEd/Private/TargetDomain/TargetDomainUtils.cpp
+++ b/Engine/Source/Editor/UnrealEd/Private/TargetDomain/TargetDomainUtils.cpp
@@ -416,7 +416,7 @@ FEditorDomainOplog::FEditorDomainOplog()
 {
 	StaticInit();
 
-	FString ProjectId = FApp::GetProjectName();
+	FString ProjectId = FZenStoreHttpClient::GenerateDefaultProjectId();
 	FString OplogId = TEXT("EditorDomain");
 
 	FString RootDir = FPaths::RootDir();
diff --git a/Engine/Source/Programs/AutomationTool/AutomationUtils/ProjectUtils.cs b/Engine/Source/Programs/AutomationTool/AutomationUtils/ProjectUtils.cs
--- a/Engine/Source/Programs/AutomationTool/AutomationUtils/ProjectUtils.cs
+++ b/Engine/Source/Programs/AutomationTool/AutomationUtils/ProjectUtils.cs
@@ -88,6 +88,18 @@ namespace AutomationTool
 			return CommandUtils.GetFilenameWithoutAnyExtensions(RawProjectPath.FullName);
 		}
 
+		/// <summary>
+		/// Gets a short alphanumeric identifier for the project path.
+		/// </summary>
+		/// <param name="RawProjectPath">Full project path.</param>
+		/// <returns>Project path identifier</returns>
+		public static string GetProjectPathId(FileReference RawProjectPath)
+		{
+			string UniformProjectPath = FileReference.FindCorrectCase(RawProjectPath).ToNormalizedPath();
+			string ProjectPathHash = ContentHash.MD5(Encoding.UTF8.GetBytes(UniformProjectPath)).ToString();
+			return String.Format("{0}.{1}", GetShortProjectName(RawProjectPath), ProjectPathHash.Substring(0, 8));
+		}
+
 		/// <summary>
 		/// Gets project properties.
 		/// </summary>
diff --git a/Engine/Source/Programs/AutomationTool/Scripts/CopyBuildToStagingDirectory.Automation.cs b/Engine/Source/Programs/AutomationTool/Scripts/CopyBuildToStagingDirectory.Automation.cs
--- a/Engine/Source/Programs/AutomationTool/Scripts/CopyBuildToStagingDirectory.Automation.cs
+++ b/Engine/Source/Programs/AutomationTool/Scripts/CopyBuildToStagingDirectory.Automation.cs
@@ -4010,185 +4010,201 @@ namespace AutomationScripts
 			return DirectoryReference.Combine(SC.ProjectRoot, "Intermediate", "UAT", SC.FinalCookPlatform);
 		}
 
-		public static void WriteStageCommandline(FileReference IntermediateCmdLineFile, ProjectParams Params, DeploymentContext SC)
+		private static string GetFileHostCommandline(ProjectParams Params, DeploymentContext SC)
 		{
-			// this file needs to be treated as a UFS file for casing, but NonUFS for being put into the .pak file. 
-			// @todo: Maybe there should be a new category - UFSNotForPak
-			if (SC.StageTargetPlatform.DeployLowerCaseFilenames(StagedFileType.SystemNonUFS))
-			{
-				IntermediateCmdLineFile = new FileReference(CombinePaths(Path.GetDirectoryName(IntermediateCmdLineFile.FullName), Path.GetFileName(IntermediateCmdLineFile.FullName).ToLowerInvariant()));
-			}
-			if (FileReference.Exists(IntermediateCmdLineFile))
+			string FileHostParams = " ";
+			if (Params.CookOnTheFly || Params.FileServer)
 			{
-				FileReference.Delete(IntermediateCmdLineFile);
-			}
+				if (Params.ZenStore)
+				{
+					if (Params.CookOnTheFly)
+					{
+						FileHostParams += "-cookonthefly ";
+					}
+					FileHostParams += "-zenstoreproject=" + ProjectUtils.GetProjectPathId(SC.RawProjectPath) + " ";
 
-			if (!SC.StageTargetPlatform.ShouldStageCommandLine(Params, SC))
-			{
-				return;
-			}
+					FileHostParams += "-zenstorehost=";
+				}
+				else
+				{
+					FileHostParams += "-filehostip=";
+				}
 
-			LogInformation("Creating UECommandLine.txt");
-			if (!string.IsNullOrEmpty(Params.StageCommandline) || !string.IsNullOrEmpty(Params.RunCommandline))
-			{
-				string FileHostParams = " ";
-				if (Params.CookOnTheFly || Params.FileServer)
+				// add localhost first for platforms using redirection
+				const string LocalHost = "127.0.0.1";
+				if (!IsNullOrEmpty(Params.Port))
 				{
-					if (Params.ZenStore)
+					bool FirstParam = true;
+					foreach (var Port in Params.Port)
 					{
-						if (Params.CookOnTheFly)
+						if (!FirstParam)
 						{
-							FileHostParams += "-cookonthefly ";
+							FileHostParams += "+";
 						}
-						FileHostParams += "-zenstorehost=";
-					}
-					else
-					{
-						FileHostParams += "-filehostip=";
-					}
-
-					// add localhost first for platforms using redirection
-					const string LocalHost = "127.0.0.1";
-					if (!IsNullOrEmpty(Params.Port))
-					{
-						bool FirstParam = true;
-						foreach (var Port in Params.Port)
+						FirstParam = false;
+						string[] PortProtocol = Port.Split(new char[] { ':' });
+						if (PortProtocol.Length > 1)
 						{
-							if (!FirstParam)
-							{
-								FileHostParams += "+";
-							}
-							FirstParam = false;
-							string[] PortProtocol = Port.Split(new char[] { ':' });
-							if (PortProtocol.Length > 1)
-							{
-								FileHostParams += String.Format("{0}://{1}:{2}", PortProtocol[0], LocalHost, PortProtocol[1]);
-							}
-							else
-							{
-								FileHostParams += LocalHost;
-								FileHostParams += ":";
-								FileHostParams += Params.Port;
-							}
-
+							FileHostParams += String.Format("{0}://{1}:{2}", PortProtocol[0], LocalHost, PortProtocol[1]);
 						}
+						else
+						{
+							FileHostParams += LocalHost;
+							FileHostParams += ":";
+							FileHostParams += Params.Port;
+						}
+
 					}
-					else
-					{
-						// use default port
-						FileHostParams += LocalHost;
-					}
-					if (UnrealBuildTool.BuildHostPlatform.Current.Platform == UnrealTargetPlatform.Mac)
+				}
+				else
+				{
+					// use default port
+					FileHostParams += LocalHost;
+				}
+
+				if (UnrealBuildTool.BuildHostPlatform.Current.Platform == UnrealTargetPlatform.Mac)
+				{
+					NetworkInterface[] Interfaces = NetworkInterface.GetAllNetworkInterfaces();
+					foreach (NetworkInterface adapter in Interfaces)
 					{
-						NetworkInterface[] Interfaces = NetworkInterface.GetAllNetworkInterfaces();
-						foreach (NetworkInterface adapter in Interfaces)
+						if (adapter.NetworkInterfaceType != NetworkInterfaceType.Loopback)
 						{
-							if (adapter.NetworkInterfaceType != NetworkInterfaceType.Loopback)
+							IPInterfaceProperties IP = adapter.GetIPProperties();
+							for (int Index = 0; Index < IP.UnicastAddresses.Count; ++Index)
 							{
-								IPInterfaceProperties IP = adapter.GetIPProperties();
-								for (int Index = 0; Index < IP.UnicastAddresses.Count; ++Index)
+								if (InternalUtils.IsDnsEligible(IP.UnicastAddresses[Index]) && IP.UnicastAddresses[Index].Address.AddressFamily == System.Net.Sockets.AddressFamily.InterNetwork)
 								{
-									if (InternalUtils.IsDnsEligible(IP.UnicastAddresses[Index]) && IP.UnicastAddresses[Index].Address.AddressFamily == System.Net.Sockets.AddressFamily.InterNetwork)
+									if (!IsNullOrEmpty(Params.Port))
 									{
-										if (!IsNullOrEmpty(Params.Port))
+										foreach (var Port in Params.Port)
 										{
-											foreach (var Port in Params.Port)
+											FileHostParams += "+";
+											string[] PortProtocol = Port.Split(new char[] { ':' });
+											if (PortProtocol.Length > 1)
 											{
-												FileHostParams += "+";
-												string[] PortProtocol = Port.Split(new char[] { ':' });
-												if (PortProtocol.Length > 1)
-												{
-													FileHostParams += String.Format("{0}://{1}:{2}", PortProtocol[0], IP.UnicastAddresses[Index].Address.ToString(), PortProtocol[1]);
-												}
-												else
-												{
-													FileHostParams += IP.UnicastAddresses[Index].Address.ToString();
-													FileHostParams += ":";
-													FileHostParams += Params.Port;
-												}
-
+												FileHostParams += String.Format("{0}://{1}:{2}", PortProtocol[0], IP.UnicastAddresses[Index].Address.ToString(), PortProtocol[1]);
+											}
+											else
+											{
+												FileHostParams += IP.UnicastAddresses[Index].Address.ToString();
+												FileHostParams += ":";
+												FileHostParams += Params.Port;
 											}
-										}
-										else
-										{
-											FileHostParams += "+";
-											// use default port
-											FileHostParams += IP.UnicastAddresses[Index].Address.ToString();
-										}
 
+										}
+									}
+									else
+									{
+										FileHostParams += "+";
+										// use default port
+										FileHostParams += IP.UnicastAddresses[Index].Address.ToString();
 									}
+
 								}
 							}
 						}
 					}
-					else
+				}
+				else
+				{
+					NetworkInterface[] Interfaces = NetworkInterface.GetAllNetworkInterfaces();
+					foreach (NetworkInterface adapter in Interfaces)
 					{
-						NetworkInterface[] Interfaces = NetworkInterface.GetAllNetworkInterfaces();
-						foreach (NetworkInterface adapter in Interfaces)
+						if (adapter.OperationalStatus == OperationalStatus.Up)
 						{
-							if (adapter.OperationalStatus == OperationalStatus.Up)
+							IPInterfaceProperties IP = adapter.GetIPProperties();
+							for (int Index = 0; Index < IP.UnicastAddresses.Count; ++Index)
 							{
-								IPInterfaceProperties IP = adapter.GetIPProperties();
-								for (int Index = 0; Index < IP.UnicastAddresses.Count; ++Index)
+								if (InternalUtils.IsDnsEligible(IP.UnicastAddresses[Index]))
 								{
-									if (InternalUtils.IsDnsEligible(IP.UnicastAddresses[Index]))
+									if (!IsNullOrEmpty(Params.Port))
 									{
-										if (!IsNullOrEmpty(Params.Port))
+										foreach (var Port in Params.Port)
 										{
-											foreach (var Port in Params.Port)
+											FileHostParams += "+";
+											string[] PortProtocol = Port.Split(new char[] { ':' });
+											if (PortProtocol.Length > 1)
 											{
-												FileHostParams += "+";
-												string[] PortProtocol = Port.Split(new char[] { ':' });
-												if (PortProtocol.Length > 1)
-												{
-													FileHostParams += String.Format("{0}://{1}:{2}", PortProtocol[0], IP.UnicastAddresses[Index].Address.ToString(), PortProtocol[1]);
-												}
-												else
-												{
-													FileHostParams += IP.UnicastAddresses[Index].Address.ToString();
-													FileHostParams += ":";
-													FileHostParams += Params.Port;
-												}
+												FileHostParams += String.Format("{0}://{1}:{2}", PortProtocol[0], IP.UnicastAddresses[Index].Address.ToString(), PortProtocol[1]);
+											}
+											else
+											{
+												FileHostParams += IP.UnicastAddresses[Index].Address.ToString();
+												FileHostParams += ":";
+												FileHostParams += Params.Port;
 											}
-										}
-										else
-										{
-											FileHostParams += "+";
-											// use default port
-											FileHostParams += IP.UnicastAddresses[Index].Address.ToString();
 										}
 									}
+									else
+									{
+										FileHostParams += "+";
+										// use default port
+										FileHostParams += IP.UnicastAddresses[Index].Address.ToString();
+									}
 								}
 							}
 						}
 					}
-
-					FileHostParams += " ";
 				}
 
-				String ProjectFile = String.Format("{0} ", SC.ProjectArgForCommandLines);
-				if (SC.StageTargetPlatform.PlatformType == UnrealTargetPlatform.Mac || SC.StageTargetPlatform.PlatformType == UnrealTargetPlatform.Win64 || SC.StageTargetPlatform.PlatformType == UnrealTargetPlatform.Linux)
-				{
-					ProjectFile = "";
-				}
-				DirectoryReference.CreateDirectory(GetIntermediateCommandlineDir(SC));
-				string CommandLine = String.Format("{0} {1} {2} {3}\n", ProjectFile, Params.StageCommandline.Trim(new char[] { '\"' }), Params.RunCommandline.Trim(new char[] { '\"' }), FileHostParams).Trim();
-				if (Params.IterativeDeploy)
-				{
-					CommandLine += " -iterative";
-				}
-				File.WriteAllText(IntermediateCmdLineFile.FullName, CommandLine);
+				FileHostParams += " ";
 			}
-			else
+
+			return FileHostParams;
+		}
+
+		public static void WriteStageCommandline(FileReference IntermediateCmdLineFile, ProjectParams Params, DeploymentContext SC)
+		{
+			// this file needs to be treated as a UFS file for casing, but NonUFS for being put into the .pak file. 
+			// @todo: Maybe there should be a new category - UFSNotForPak
+			if (SC.StageTargetPlatform.DeployLowerCaseFilenames(StagedFileType.SystemNonUFS))
 			{
-				String ProjectFile = String.Format("{0} ", SC.ProjectArgForCommandLines);
-				if (SC.StageTargetPlatform.PlatformType == UnrealTargetPlatform.Mac || SC.StageTargetPlatform.PlatformType == UnrealTargetPlatform.Win64 || SC.StageTargetPlatform.PlatformType == UnrealTargetPlatform.Linux)
-				{
-					ProjectFile = "";
-				}
-				DirectoryReference.CreateDirectory(GetIntermediateCommandlineDir(SC));
-				File.WriteAllText(IntermediateCmdLineFile.FullName, ProjectFile);
+				IntermediateCmdLineFile = new FileReference(CombinePaths(Path.GetDirectoryName(IntermediateCmdLineFile.FullName), Path.GetFileName(IntermediateCmdLineFile.FullName).ToLowerInvariant()));
+			}
+			if (FileReference.Exists(IntermediateCmdLineFile))
+			{
+				FileReference.Delete(IntermediateCmdLineFile);
+			}
+
+			if (!SC.StageTargetPlatform.ShouldStageCommandLine(Params, SC))
+			{
+				return;
+			}
+
+			StringBuilder CommandLineBuilder = new StringBuilder();
+			CommandLineBuilder.AppendFormat("{0} ", SC.ProjectArgForCommandLines);
+
+			if (!string.IsNullOrEmpty(Params.StageCommandline))
+			{
+				CommandLineBuilder.AppendFormat("{0} ", Params.StageCommandline.Trim(new char[] { '\"' }));
+			}
+
+			if (!string.IsNullOrEmpty(Params.RunCommandline))
+			{
+				CommandLineBuilder.AppendFormat("{0} ", Params.RunCommandline.Trim(new char[] { '\"' }));
+			}
+
+			String FileHostCommandline = GetFileHostCommandline(Params, SC).Trim();
+			if (!string.IsNullOrEmpty(FileHostCommandline))
+			{
+				CommandLineBuilder.AppendFormat("{0} ", FileHostCommandline);
+			}
+
+			if (Params.IterativeDeploy)
+			{
+				CommandLineBuilder.Append("-iterative ");
 			}
+
+			String CompleteCommandLine = CommandLineBuilder.ToString().Trim();
+
+			if (string.IsNullOrEmpty(CompleteCommandLine))
+			{
+				return;
+			}
+
+			LogInformation("Creating UECommandLine.txt");
+			DirectoryReference.CreateDirectory(GetIntermediateCommandlineDir(SC));
+			File.WriteAllText(IntermediateCmdLineFile.FullName, CompleteCommandLine);
 		}
 
 		private static void WriteStageCommandline(ProjectParams Params, DeploymentContext SC)
diff --git a/Engine/Source/Programs/AutomationTool/Scripts/RunProjectCommand.Automation.cs b/Engine/Source/Programs/AutomationTool/Scripts/RunProjectCommand.Automation.cs
--- a/Engine/Source/Programs/AutomationTool/Scripts/RunProjectCommand.Automation.cs
+++ b/Engine/Source/Programs/AutomationTool/Scripts/RunProjectCommand.Automation.cs
@@ -637,6 +637,7 @@ namespace AutomationScripts
 						{
 							TempCmdLine += "-cookonthefly ";
 						}
+						TempCmdLine += "-zenstoreproject=" + ProjectUtils.GetProjectPathId(SC.RawProjectPath) + " ";
 						TempCmdLine += "-zenstorehost=";
 					}
 					else
diff --git a/Engine/Source/Runtime/Core/Public/GenericPlatform/GenericPlatformFile.h b/Engine/Source/Runtime/Core/Public/GenericPlatform/GenericPlatformFile.h
--- a/Engine/Source/Runtime/Core/Public/GenericPlatform/GenericPlatformFile.h
+++ b/Engine/Source/Runtime/Core/Public/GenericPlatform/GenericPlatformFile.h
@@ -252,6 +252,11 @@ public:
 	 */
 	virtual void		InitializeAfterSetActive() { }
 
+	/**
+	 * Performs initialization of the platform file after the project path has been set.
+	 */
+	virtual void		InitializeAfterProjectFilePath() { }
+
 	/**
 	 * Build an in memory unique pak file from a subset of files in this pak file
 	 */
diff --git a/Engine/Source/Runtime/Launch/Private/LaunchEngineLoop.cpp b/Engine/Source/Runtime/Launch/Private/LaunchEngineLoop.cpp
--- a/Engine/Source/Runtime/Launch/Private/LaunchEngineLoop.cpp
+++ b/Engine/Source/Runtime/Launch/Private/LaunchEngineLoop.cpp
@@ -2177,8 +2177,14 @@ int32 FEngineLoop::PreInitPreStartupScreen(const TCHAR* CmdLine)
 		const FString ProjectFilePath = FPaths::Combine(*FPaths::ProjectDir(), *FString::Printf(TEXT("%s.%s"), FApp::GetProjectName(), *FProjectDescriptor::GetExtension()));
 		FPaths::SetProjectFilePath(ProjectFilePath);
 	}
+
+	// Fix the project file path case before we attempt to fix the game name
+	LaunchFixProjectPathCase();
 #endif
 
+	// Last chance to initialize platform file with possible knowledge of the project file path
+	FPlatformFileManager::Get().GetPlatformFile().InitializeAfterProjectFilePath();
+
 	// Now verify the project file if we have one
 	if (FPaths::IsProjectFilePathSet()
 #if IS_PROGRAM
@@ -2204,9 +2210,6 @@ int32 FEngineLoop::PreInitPreStartupScreen(const TCHAR* CmdLine)
 	}
 
 #if !IS_PROGRAM
-	// Fix the project file path case before we attempt to fix the game name
-	LaunchFixProjectPathCase();
-
 	if (FApp::HasProjectName())
 	{
 		// Tell the module manager what the game binaries folder is
diff --git a/Engine/Source/Runtime/StorageServerClient/Private/StorageServerConnection.cpp b/Engine/Source/Runtime/StorageServerClient/Private/StorageServerConnection.cpp
--- a/Engine/Source/Runtime/StorageServerClient/Private/StorageServerConnection.cpp
+++ b/Engine/Source/Runtime/StorageServerClient/Private/StorageServerConnection.cpp
@@ -6,6 +6,7 @@
 #include "Sockets.h"
 #include "Misc/App.h"
 #include "Misc/Paths.h"
+#include "Misc/SecureHash.h"
 #include "IO/IoDispatcher.h"
 #include "Misc/StringBuilder.h"
 #include "Misc/ScopeExit.h"
@@ -64,6 +65,18 @@ static uint64 GetCompressedOffset(const FCompressedBuffer& Buffer, uint64 RawOff
 	return 0;
 }
 
+// Duplicated in ZenStoreHttpClient.cpp to avoid having a public API in a shared module
+static FString GetProjectPathId()
+{
+	FString ProjectFilePath = FPaths::GetProjectFilePath();
+	FPaths::NormalizeFilename(ProjectFilePath);
+	FString AbsProjectFilePath = IFileManager::Get().ConvertToAbsolutePathForExternalAppForRead(*ProjectFilePath);
+	FTCHARToUTF8 AbsProjectFilePathUTF8(*AbsProjectFilePath);
+
+	FString HashString = FMD5::HashBytes((unsigned char*)AbsProjectFilePathUTF8.Get(), AbsProjectFilePathUTF8.Length()).Left(8);
+	return FString::Printf(TEXT("%s.%.8s"), FApp::GetProjectName(), *HashString);
+}
+
 FStorageServerRequest::FStorageServerRequest(FAnsiStringView Verb, FAnsiStringView Resource, FAnsiStringView Hostname, EStorageServerContentType Accept)
 : AcceptType(Accept)
 {
@@ -437,7 +450,7 @@ bool FStorageServerConnection::Initialize(TArrayView<const FString> InHostAddres
 	}
 	else
 	{
-		OplogPath.Append(TCHAR_TO_ANSI(FApp::GetProjectName()));
+		OplogPath.Append(TCHAR_TO_ANSI(*GetProjectPathId()));
 	}
 	OplogPath.Append("/oplog/");
 	if (InPlatformNameOverride)
diff --git a/Engine/Source/Runtime/StorageServerClient/Private/StorageServerPlatformFile.cpp b/Engine/Source/Runtime/StorageServerClient/Private/StorageServerPlatformFile.cpp
--- a/Engine/Source/Runtime/StorageServerClient/Private/StorageServerPlatformFile.cpp
+++ b/Engine/Source/Runtime/StorageServerClient/Private/StorageServerPlatformFile.cpp
@@ -267,57 +267,54 @@ bool FStorageServerPlatformFile::ShouldBeUsed(IPlatformFile* Inner, const TCHAR*
 bool FStorageServerPlatformFile::Initialize(IPlatformFile* Inner, const TCHAR* CmdLine)
 {
 	LowerLevel = Inner;
-	bool bResult = false;
 	if (HostAddrs.Num() > 0)
 	{
-		Connection.Reset(new FStorageServerConnection());
-		FString StorageServerProject;
-		FParse::Value(CmdLine, TEXT("-ZenStoreProject="), StorageServerProject);
-		const TCHAR* ProjectOverride = StorageServerProject.IsEmpty() ? nullptr : *StorageServerProject;
-		FString StorageServerPlatform;
-		FParse::Value(CmdLine, TEXT("-ZenStorePlatform="), StorageServerPlatform);
-		const TCHAR* PlatformOverride = StorageServerPlatform.IsEmpty() ? nullptr : *StorageServerPlatform;
-		if (Connection->Initialize(HostAddrs, 1337, ProjectOverride, PlatformOverride))
+		// Don't initialize the connection yet because we want to incorporate project file path information into the initialization.
+		FParse::Value(CmdLine, TEXT("-ZenStoreProject="), ServerProject);
+		FParse::Value(CmdLine, TEXT("-ZenStorePlatform="), ServerPlatform);
+		return true;
+	}
+	return false;
+}
+
+void FStorageServerPlatformFile::InitializeAfterProjectFilePath()
+{
+	Connection.Reset(new FStorageServerConnection());
+	const TCHAR* ProjectOverride = ServerProject.IsEmpty() ? nullptr : *ServerProject;
+	const TCHAR* PlatformOverride = ServerPlatform.IsEmpty() ? nullptr : *ServerPlatform;
+	if (Connection->Initialize(HostAddrs, 1337, ProjectOverride, PlatformOverride))
+	{
+		if (SendGetFileListMessage())
 		{
-			if (SendGetFileListMessage())
+			FIoDispatcher& IoDispatcher = FIoDispatcher::Get();
+			TSharedRef<FStorageServerIoDispatcherBackend> IoDispatcherBackend = MakeShared<FStorageServerIoDispatcherBackend>(*Connection.Get());
+			IoDispatcher.Mount(IoDispatcherBackend);
+			FCoreDelegates::CreatePackageStore.BindLambda([this]() -> TSharedPtr<IPackageStore>
 			{
-				FIoDispatcher& IoDispatcher = FIoDispatcher::Get();
-				TSharedRef<FStorageServerIoDispatcherBackend> IoDispatcherBackend = MakeShared<FStorageServerIoDispatcherBackend>(*Connection.Get());
-				IoDispatcher.Mount(IoDispatcherBackend);
-				bResult = true;
-				FCoreDelegates::CreatePackageStore.BindLambda([this]() -> TSharedPtr<IPackageStore>
-				{
 #if WITH_COTF
-					if (IsRunningCookOnTheFly())
-					{
-						UE::Cook::ICookOnTheFlyModule& CookOnTheFlyModule = FModuleManager::LoadModuleChecked<UE::Cook::ICookOnTheFlyModule>(TEXT("CookOnTheFly"));
-						UE::Cook::ICookOnTheFlyServerConnection& CotfConnection = CookOnTheFlyModule.GetServerConnection();
-						CotfConnection.OnMessage().AddRaw(this, &FStorageServerPlatformFile::OnCookOnTheFlyMessage);
-						return MakeCookOnTheFlyPackageStore(CotfConnection);
-					}
-					else
+				if (IsRunningCookOnTheFly())
+				{
+					UE::Cook::ICookOnTheFlyModule& CookOnTheFlyModule = FModuleManager::LoadModuleChecked<UE::Cook::ICookOnTheFlyModule>(TEXT("CookOnTheFly"));
+					UE::Cook::ICookOnTheFlyServerConnection& CotfConnection = CookOnTheFlyModule.GetServerConnection();
+					CotfConnection.OnMessage().AddRaw(this, &FStorageServerPlatformFile::OnCookOnTheFlyMessage);
+					return MakeCookOnTheFlyPackageStore(CotfConnection);
+				}
+				else
 #endif
-					{
-						return MakeShared<FStorageServerPackageStore>(*Connection.Get());
-					}
-				});
-			}
-			else
-			{
-				UE_LOG(LogStorageServerPlatformFile, Fatal, TEXT("Failed to get file list from Zen at '%s'"), *Connection->GetHostAddr());
-			}
+				{
+					return MakeShared<FStorageServerPackageStore>(*Connection.Get());
+				}
+			});
 		}
 		else
 		{
-			UE_LOG(LogStorageServerPlatformFile, Fatal, TEXT("Failed to initialize connection"));
+			UE_LOG(LogStorageServerPlatformFile, Fatal, TEXT("Failed to get file list from Zen at '%s'"), *Connection->GetHostAddr());
 		}
 	}
-	if (!bResult)
+	else
 	{
-		Connection.Reset(nullptr);
-		return false;
+		UE_LOG(LogStorageServerPlatformFile, Fatal, TEXT("Failed to initialize connection"));
 	}
-	return true;
 }
 
 bool FStorageServerPlatformFile::FileExists(const TCHAR* Filename)
diff --git a/Engine/Source/Runtime/StorageServerClient/Private/StorageServerPlatformFile.h b/Engine/Source/Runtime/StorageServerClient/Private/StorageServerPlatformFile.h
--- a/Engine/Source/Runtime/StorageServerClient/Private/StorageServerPlatformFile.h
+++ b/Engine/Source/Runtime/StorageServerClient/Private/StorageServerPlatformFile.h
@@ -60,6 +60,7 @@ public:
 	virtual ~FStorageServerPlatformFile();
 	virtual bool ShouldBeUsed(IPlatformFile* Inner, const TCHAR* CmdLine) const override;
 	virtual bool Initialize(IPlatformFile* Inner, const TCHAR* CmdLine) override;
+	virtual void InitializeAfterProjectFilePath() override;
 
 	virtual IPlatformFile* GetLowerLevel() override
 	{
@@ -113,6 +114,8 @@ private:
 	FStringView ServerProjectDirView = FStringView(TEXT("/{project}/"));
 	TUniquePtr<FStorageServerConnection> Connection;
 	FStorageServerFileSystemTOC ServerToc;
+	FString ServerProject;
+	FString ServerPlatform;
 	mutable TArray<FString> HostAddrs;
 };
 
