Fixing logging format call to include the file name

#rnx

diff --git a/Engine/Source/Programs/AutomationTool/BuildGraph/TempStorage.cs b/Engine/Source/Programs/AutomationTool/BuildGraph/TempStorage.cs
--- a/Engine/Source/Programs/AutomationTool/BuildGraph/TempStorage.cs
+++ b/Engine/Source/Programs/AutomationTool/BuildGraph/TempStorage.cs
@@ -1130,7 +1130,7 @@ namespace AutomationTool
 						{
 							if (UnzipFileAttempts == 0)
 							{
-								Log.TraceError("All retries exhausted attempting to unzip entries from '{0}'. Terminating.");
+								Log.TraceError("All retries exhausted attempting to unzip entries from '{0}'. Terminating.", ZipFile.FullName);
 								throw;
 							}
 
