Add and Utilize UseChunkDB setting on bundle config to no longer generate chunkDBs on any flagged bundles
#swarm 18391661
#tests Android / Switch preflight

diff --git a/Engine/Source/Programs/AutomationTool/AutomationUtils/BundleUtils.cs b/Engine/Source/Programs/AutomationTool/AutomationUtils/BundleUtils.cs
--- a/Engine/Source/Programs/AutomationTool/AutomationUtils/BundleUtils.cs
+++ b/Engine/Source/Programs/AutomationTool/AutomationUtils/BundleUtils.cs
@@ -22,6 +22,7 @@ namespace AutomationUtils.Automation
 			public bool bFoundParent { get; set; }
 			public bool bContainsShaderLibrary { get; set; }
 			public int Order { get; set; }
+			public bool	UseChunkDBs { get; set; }
 			public string ExecFileName { get; set; } // TODO: We never used this.  Clean this up.
 		}
 
@@ -104,6 +105,17 @@ namespace AutomationUtils.Automation
 					}
 				}
 
+				{
+					bool bUseChunkDBs;
+					if (!BundleConfig.GetBool(SectionName, "UseChunkDBs", out bUseChunkDBs))
+					{
+						Bundle.UseChunkDBs = false;
+					}
+					else
+					{
+						Bundle.UseChunkDBs = bUseChunkDBs;
+					}
+				}
 				GetPlatformSettings(Bundle, BundleConfig, BundleDefinitionPrefix + Bundle.Name);
 
 				Results.Add(Bundle);
diff --git a/Engine/Source/Runtime/InstallBundleManager/Public/InstallBundleTypes.h b/Engine/Source/Runtime/InstallBundleManager/Public/InstallBundleTypes.h
--- a/Engine/Source/Runtime/InstallBundleManager/Public/InstallBundleTypes.h
+++ b/Engine/Source/Runtime/InstallBundleManager/Public/InstallBundleTypes.h
@@ -236,6 +236,7 @@ struct FInstallBundleSourceUpdateBundleInfo
 	FDateTime LastAccessTime = FDateTime::MinValue(); // If cached, used to decide eviction order
 	EInstallBundleInstallState BundleContentState = EInstallBundleInstallState::NotInstalled; // Whether this bundle is up to date
 	bool bIsCached = false; // Whether this bundle should be cached if this source has a bundle cache
+	bool bUseChunkDBs = false; // Whether this bundle should attempt to use ChunkDBs at all
 };
 
 struct FInstallBundleSourceUpdateBundleInfoResult
