Added a debug statement and exception handler to try and track down a periodic exception in logging from tests

diff --git a/Engine/Source/Programs/AutomationTool/Gauntlet/Framework/Gauntlet.Utils.cs b/Engine/Source/Programs/AutomationTool/Gauntlet/Framework/Gauntlet.Utils.cs
--- a/Engine/Source/Programs/AutomationTool/Gauntlet/Framework/Gauntlet.Utils.cs
+++ b/Engine/Source/Programs/AutomationTool/Gauntlet/Framework/Gauntlet.Utils.cs
@@ -347,15 +347,22 @@ namespace Gauntlet
 			// TODO - Remove all Gauntlet logging and switch to UBT log?
 			CommandUtils.LogInformation(Message);
 
-			if (LogFile != null)
+			try
 			{
-				LogFile.WriteLine(Message);
-			}
+				if (LogFile != null)
+				{
+					LogFile.WriteLine(Message);
+				}
 
-			if (Callbacks != null)
-			{
-				Callbacks.ForEach(A => A(Message));
+				if (Callbacks != null)
+				{
+					Callbacks.ForEach(A => A(Message));
+				}
 			}
+			catch (Exception Ex)
+			{
+				CommandUtils.LogWarning("Exception logging '{0}'. {1}", Message, Ex.ToString());
+			}			
 		}	
 
 		static public void Verbose(string Format, params object[] Args)
diff --git a/Engine/Source/Programs/AutomationTool/Gauntlet/Unreal/BuildSource/Gauntlet.UnrealBuildSource.cs b/Engine/Source/Programs/AutomationTool/Gauntlet/Unreal/BuildSource/Gauntlet.UnrealBuildSource.cs
--- a/Engine/Source/Programs/AutomationTool/Gauntlet/Unreal/BuildSource/Gauntlet.UnrealBuildSource.cs
+++ b/Engine/Source/Programs/AutomationTool/Gauntlet/Unreal/BuildSource/Gauntlet.UnrealBuildSource.cs
@@ -516,6 +516,7 @@ namespace Gauntlet
 			}
 
 			// Cleanup the commandline
+			Log.Info("Processing CommandLine {0}", Config.CommandLine);
 			Config.CommandLine = GenerateProcessedCommandLine(Config.CommandLine);
 
 			// Now add the project (the above code doesn't handle arguments without a leading - so do this last
