DDC: Fixed the S3 upload to use relative paths

#rnx

diff --git a/Engine/Source/Programs/AutomationTool/Scripts/UploadDDCToAWS.cs b/Engine/Source/Programs/AutomationTool/Scripts/UploadDDCToAWS.cs
--- a/Engine/Source/Programs/AutomationTool/Scripts/UploadDDCToAWS.cs
+++ b/Engine/Source/Programs/AutomationTool/Scripts/UploadDDCToAWS.cs
@@ -645,7 +645,7 @@ namespace AutomationTool
 				if (Info.Exists)
 				{
 					string Key = new FileReference(Info).MakeRelativeTo(BaseDir).Replace('\\', '/');
-					Files.Add(new DerivedDataFile(Info, Path));
+					Files.Add(new DerivedDataFile(Info, Key));
 				}
 			}
 			catch (Exception Ex)
