DDC: Extended the S3 cache with the ability to do local generation and decompression of bundles

diff --git a/Engine/Source/Developer/DerivedDataCache/Private/S3DerivedDataBackend.cpp b/Engine/Source/Developer/DerivedDataCache/Private/S3DerivedDataBackend.cpp
--- a/Engine/Source/Developer/DerivedDataCache/Private/S3DerivedDataBackend.cpp
+++ b/Engine/Source/Developer/DerivedDataCache/Private/S3DerivedDataBackend.cpp
@@ -313,7 +313,14 @@ public:
 		// Print any diagnostic output
 		if (!(ResponseCode >= 200 && ResponseCode <= 299))
 		{
-			Log->Logf(ELogVerbosity::Error, TEXT("Download failed for %s (response %d):\n%s\n%s"), ANSI_TO_TCHAR(Url), ResponseCode, *CallbackData.ResponseHeader.ToWideString(), ANSI_TO_TCHAR((const ANSICHAR*)OutResponseBody.GetData()));
+			if (FAnsiStringView(Url).StartsWith("file://"))
+			{
+				ResponseCode = 200;
+			}
+			else
+			{
+				Log->Logf(ELogVerbosity::Error, TEXT("Download failed for %s (response %d):\n%s\n%s"), ANSI_TO_TCHAR(Url), ResponseCode, *CallbackData.ResponseHeader.ToWideString(), ANSI_TO_TCHAR((const ANSICHAR*)OutResponseBody.GetData()));
+			}
 		}
 		return ResponseCode;
 	}
@@ -786,23 +793,32 @@ FS3DerivedDataBackend::FS3DerivedDataBackend(const TCHAR* InRootManifestPath, co
 	{
 		RequestPool.Reset(new FRequestPool(TCHAR_TO_ANSI(InRegion), TCHAR_TO_ANSI(*RootManifest.AccessKey), TCHAR_TO_ANSI(*RootManifest.SecretKey)));
 
-		// Test whether we can reach the canary URL
-		bool bCanaryValid = true;
-		if (GIsBuildMachine)
+		FString LocalManifestPath;
+		if (RootManifest.Keys.Last().StartsWith(TEXT("file://")))
 		{
-			UE_LOG(LogDerivedDataCache, Log, TEXT("S3DerivedDataBackend: Disabling on build machine"));
-			bCanaryValid = false;
+			LocalManifestPath = FPaths::GetPath(RootManifest.Keys.Last());
 		}
-		else if (CanaryObjectKey.Len() > 0)
-		{
-			TArray<uint8> Data;
 
-			FStringOutputDevice DummyOutputDevice;
-			if (!IsSuccessfulHttpResponse(RequestPool->Download(*(BaseUrl / CanaryObjectKey), nullptr, Data, &DummyOutputDevice)))
+		// Test whether we can reach the canary URL
+		bool bCanaryValid = true;
+		if (LocalManifestPath.IsEmpty())
+		{
+			if (GIsBuildMachine)
 			{
-				UE_LOG(LogDerivedDataCache, Log, TEXT("S3DerivedDataBackend: Unable to download canary file. Disabling."));
+				UE_LOG(LogDerivedDataCache, Log, TEXT("S3DerivedDataBackend: Disabling on build machine"));
 				bCanaryValid = false;
 			}
+			else if (CanaryObjectKey.Len() > 0)
+			{
+				TArray<uint8> Data;
+
+				FStringOutputDevice DummyOutputDevice;
+				if (!IsSuccessfulHttpResponse(RequestPool->Download(*(BaseUrl / CanaryObjectKey), nullptr, Data, &DummyOutputDevice)))
+				{
+					UE_LOG(LogDerivedDataCache, Log, TEXT("S3DerivedDataBackend: Unable to download canary file. Disabling."));
+					bCanaryValid = false;
+				}
+			}
 		}
 
 		// Allow the user to override it from the editor
@@ -841,7 +857,8 @@ FS3DerivedDataBackend::FS3DerivedDataBackend(const TCHAR* InRootManifestPath, co
 				{
 					if (!FPaths::FileExists(Bundle.LocalFile))
 					{
-						TSharedPtr<FBundleDownload> Download(new FBundleDownload(CriticalSection, Bundle, BaseUrl + Bundle.ObjectKey, *RequestPool.Get(), Context));
+						FString BundleUrl = LocalManifestPath.IsEmpty() ? BaseUrl + Bundle.ObjectKey : LocalManifestPath / Bundle.Name + TEXT(".gz");
+						TSharedPtr<FBundleDownload> Download(new FBundleDownload(CriticalSection, Bundle, BundleUrl, *RequestPool.Get(), Context));
 						Download->Event = FFunctionGraphTask::CreateAndDispatchWhenReady([Download]() { Download->Execute(); }, TStatId());
 						Downloads.Add(MoveTemp(Download));
 					}
@@ -1007,8 +1024,9 @@ bool FS3DerivedDataBackend::DownloadManifest(const FRootManifest& RootManifest,
 	FString BundleManifestKey = RootManifest.Keys.Last();
 
 	// Download the bundle manifest
+	FString BundleManifestUrl = BundleManifestKey.StartsWith(TEXT("file://")) ? BundleManifestKey : BaseUrl + BundleManifestKey;
 	TArray<uint8> BundleManifestData;
-	long ResponseCode = RequestPool->Download(*(BaseUrl + BundleManifestKey), nullptr, BundleManifestData, Context);
+	long ResponseCode = RequestPool->Download(*BundleManifestUrl, nullptr, BundleManifestData, Context);
 	if (!IsSuccessfulHttpResponse(ResponseCode))
 	{
 		Context->Logf(ELogVerbosity::Warning, TEXT("Unable to download bundle manifest from %s (%d)"), *BundleManifestKey, (int)ResponseCode);
diff --git a/Engine/Source/Programs/AutomationTool/Scripts/UploadDDCToAWS.cs b/Engine/Source/Programs/AutomationTool/Scripts/UploadDDCToAWS.cs
--- a/Engine/Source/Programs/AutomationTool/Scripts/UploadDDCToAWS.cs
+++ b/Engine/Source/Programs/AutomationTool/Scripts/UploadDDCToAWS.cs
@@ -172,9 +172,10 @@ namespace AutomationTool
 
 		public override void ExecuteBuild()
 		{
-			string BucketName = ParseRequiredStringParam("Bucket");
-			FileReference CredentialsFile = ParseRequiredFileReferenceParam("CredentialsFile");
-			string CredentialsKey = ParseRequiredStringParam("CredentialsKey");
+			bool bLocal = ParseParam("Local");
+			string BucketName = !bLocal ? ParseRequiredStringParam("Bucket") : "";
+			FileReference CredentialsFile = !bLocal ? ParseRequiredFileReferenceParam("CredentialsFile") : null;
+			string CredentialsKey = !bLocal ? ParseRequiredStringParam("CredentialsKey") : "";
 			DirectoryReference CacheDir = ParseRequiredDirectoryReferenceParam("CacheDir");
 			DirectoryReference FilterDir = ParseRequiredDirectoryReferenceParam("FilterDir");
 			int Days = ParseParamInt("Days", 7);
@@ -183,18 +184,19 @@ namespace AutomationTool
 			string KeyPrefix = ParseParamValue("KeyPrefix", "");
 			bool bReset = ParseParam("Reset");
 
-			// The credentials to upload with
-			AWSCredentials Credentials;
-
 			// Try to get the credentials by the key passed in from the script
-			CredentialProfileStoreChain CredentialsChain = new CredentialProfileStoreChain(CredentialsFile.FullName);
-			if (!CredentialsChain.TryGetAWSCredentials(CredentialsKey, out Credentials))
+			AWSCredentials Credentials = null;
+			if (!bLocal)
 			{
-				throw new AutomationException("Unknown credentials key: {0}", CredentialsKey);
+				CredentialProfileStoreChain CredentialsChain = new CredentialProfileStoreChain(CredentialsFile.FullName);
+				if (!CredentialsChain.TryGetAWSCredentials(CredentialsKey, out Credentials))
+				{
+					throw new AutomationException("Unknown credentials key: {0}", CredentialsKey);
+				}
 			}
 
 			// Create the new client
-			using (AmazonS3Client Client = new AmazonS3Client(Credentials, Region))
+			using (AmazonS3Client Client = !bLocal ? new AmazonS3Client(Credentials, Region) : null)
 			{
 				using (SemaphoreSlim RequestSemaphore = new SemaphoreSlim(4))
 				{
@@ -520,13 +522,23 @@ namespace AutomationTool
 								UploadFile(Client, BucketName, NewBundleManifestFile, 0, Entry.Key, RequestSemaphore, null);
 							}
 
+							string LocalBundleManifestPath = String.Format("file://{0}", NewBundleManifestFile.FullName.Replace('\\', '/'));
+							if (bLocal && NewRootManifest.Entries.LastOrDefault()?.Key != LocalBundleManifestPath)
+							{
+								RootManifest.Entry NewEntry = new RootManifest.Entry();
+								NewEntry.CreateTime = UtcNow;
+								NewEntry.Key = LocalBundleManifestPath;
+								NewRootManifest.Entries.Add(NewEntry);
+							}
+
 							// Overwrite all the existing manifests
-							if (AllowSubmit)
+							int ChangeNumber = 0;
+							if (AllowSubmit || (P4Enabled && bLocal))
 							{
 								List<string> ExistingFiles = P4.Files(CommandUtils.MakePathSafeToUseWithCommandLine(RootManifestFile.FullName));
 
 								// Create a changelist containing the new manifest
-								int ChangeNumber = P4.CreateChange(Description: "Updating DDC bundle manifest");
+								ChangeNumber = P4.CreateChange(Description: "Updating DDC bundle manifest");
 								if (ExistingFiles.Count > 0)
 								{
 									P4.Edit(ChangeNumber, CommandUtils.MakePathSafeToUseWithCommandLine(RootManifestFile.FullName));
@@ -537,7 +549,14 @@ namespace AutomationTool
 									NewRootManifest.Save(RootManifestFile);
 									P4.Add(ChangeNumber, CommandUtils.MakePathSafeToUseWithCommandLine(RootManifestFile.FullName));
 								}
+							}
+							else if (bLocal)
+							{
+								NewRootManifest.Save(RootManifestFile);
+							}
 
+							if (AllowSubmit)
+							{
 								// Submit it
 								int SubmittedChangeNumber;
 								P4.Submit(ChangeNumber, out SubmittedChangeNumber, true);
@@ -597,6 +616,10 @@ namespace AutomationTool
 
 		private bool TryDownloadFile(AmazonS3Client Client, string BucketName, string Key, FileReference OutputFile)
 		{
+			if (Client == null)
+			{
+				return false;
+			}
 			try
 			{
 				GetObjectRequest Request = new GetObjectRequest();
@@ -702,6 +725,10 @@ namespace AutomationTool
 
 		protected void UploadFile(AmazonS3Client Client, string BucketName, FileReference File, long Length, string Key, SemaphoreSlim Semaphore, UploadState State)
 		{
+			if (Client == null)
+			{
+				return;
+			}
 			try
 			{
 				Retry(() => UploadFileInner(Client, BucketName, File, Length, Key, Semaphore, State), 10);
