BuildGraph: New mechanism for declaring graphs using C# code (WIP).

Nodes can now be implemented by arbitary C# methods. Graph structure is specified through expression trees implemented using Bg* types, which are not substituted with values until execution time. Doing so allows determination of node and option dependencies for a particular target, allowing us to generate dynamic UI for presenting relevant settings to the user.

Includes partial implementation of Installed Build script as an example implementation.

diff --git a/Engine/Source/Programs/AutomationTool/BuildGraph/BgGraphBuilder.cs b/Engine/Source/Programs/AutomationTool/BuildGraph/BgGraphBuilder.cs
new file mode 100644
--- /dev/null
+++ b/Engine/Source/Programs/AutomationTool/BuildGraph/BgGraphBuilder.cs
@@ -0,0 +1,21 @@
+// Copyright Epic Games, Inc. All Rights Reserved.
+
+using EpicGames.BuildGraph;
+using System;
+using System.Collections.Generic;
+using System.Text;
+
+namespace AutomationTool
+{
+	/// <summary>
+	/// Base class for any user defined graphs
+	/// </summary>
+	public abstract class BgGraphBuilder
+	{
+		/// <summary>
+		/// Callback used to instantiate the graph
+		/// </summary>
+		/// <param name="Graph">The graph object</param>
+		public abstract void SetupGraph(BgGraphSpec Graph);
+	}
+}
diff --git a/Engine/Source/Programs/AutomationTool/BuildGraph/BuildGraph.cs b/Engine/Source/Programs/AutomationTool/BuildGraph/BuildGraph.cs
--- a/Engine/Source/Programs/AutomationTool/BuildGraph/BuildGraph.cs
+++ b/Engine/Source/Programs/AutomationTool/BuildGraph/BuildGraph.cs
@@ -18,6 +18,7 @@ using OpenTracing.Util;
 using UnrealBuildBase;
 using UnrealBuildTool;
 using System.Threading.Tasks;
+using EpicGames.BuildGraph.Expressions;
 
 namespace AutomationTool
 {
@@ -213,6 +214,26 @@ namespace AutomationTool
 	[Help("WriteToSharedStorage", "Allow writing to shared storage. If not set, but -SharedStorageDir is specified, build products will read but not written")]
 	public class BuildGraph : BuildCommand
 	{
+		class BgGraphConstants : IBgGraphConstants
+		{
+			public BgString Stream { get; set; }
+			public BgInt Change { get; }
+			public BgInt CodeChange { get; }
+			public BgBool IsBuildMachine { get; }
+			public (BgInt Major, BgInt Minor, BgInt Patch) EngineVersion { get; }
+
+			public BgGraphConstants(P4Environment P4Env)
+			{
+				this.Stream = P4Env?.Branch ?? "Unknown";
+				this.Change = P4Env?.Changelist ?? 0;
+				this.CodeChange = P4Env?.CodeChangelist ?? 0;
+				this.IsBuildMachine = CommandUtils.IsBuildMachine;
+
+				ReadOnlyBuildVersion Version = ReadOnlyBuildVersion.Current;
+				this.EngineVersion = (Version.MajorVersion, Version.MinorVersion, Version.PatchVersion);
+			}
+		}
+
 		/// <summary>
 		/// Context object for reading scripts and evaluating conditions
 		/// </summary>
@@ -224,6 +245,7 @@ namespace AutomationTool
 		public override async Task<ExitCode> ExecuteAsync()
 		{
 			// Parse the command line parameters
+			string ClassName = ParseParamValue("Class", null);
 			string ScriptFileName = ParseParamValue("Script", null);
 			string[] TargetNames = ParseParamValues("Target").SelectMany(x => x.Split(';', '+').Select(y => y.Trim()).Where(y => y.Length > 0)).ToArray();
 			string DocumentationFileName = ParseParamValue("Documentation", null);
@@ -369,56 +391,87 @@ namespace AutomationTool
 				return ExitCode.Success;
 			}
 
-			// Import schema if one is passed in
-			BgScriptSchema Schema;
-			if (ImportSchemaFileName != null)
+			// Create the graph
+			BgGraph Graph;
+			if (ClassName != null)
 			{
-				Schema = BgScriptSchema.Import(FileReference.FromString(ImportSchemaFileName));
+				// Find all the graph builders
+				Dictionary<string, Type> NameToType = new Dictionary<string, Type>();
+				FindAvailableGraphs(NameToType);
+
+				Type BuilderType;
+				if (!NameToType.TryGetValue(ClassName, out BuilderType))
+				{
+					Logger.LogError("Unable to find graph '{GraphName}'", ClassName);
+					Logger.LogInformation("");
+					Logger.LogInformation("Available graphs:");
+					foreach (string Name in NameToType.Keys.OrderBy(x => x))
+					{
+						Logger.LogInformation("  {GraphName}", Name);
+					}
+					return ExitCode.Error_Unknown;
+				}
+
+				BgGraphSpec GraphSpec = new BgGraphSpec(new BgGraphConstants(P4Enabled? P4Env : null));
+
+				BgGraphBuilder Builder = (BgGraphBuilder)Activator.CreateInstance(BuilderType);
+				Builder.SetupGraph(GraphSpec);
+
+				Graph = GraphSpec.CreateGraph(new HashSet<string>(), Arguments);
 			}
 			else
 			{
-				// Add any primitive types
-				List<(Type, ScriptSchemaStandardType)> PrimitiveTypes = new List<(Type, ScriptSchemaStandardType)>();
-				PrimitiveTypes.Add((typeof(FileReference), ScriptSchemaStandardType.BalancedString));
-				PrimitiveTypes.Add((typeof(DirectoryReference), ScriptSchemaStandardType.BalancedString));
-				PrimitiveTypes.Add((typeof(UnrealTargetPlatform), ScriptSchemaStandardType.BalancedString));
-				PrimitiveTypes.Add((typeof(MCPPlatform), ScriptSchemaStandardType.BalancedString));
-
-				// Create a schema for the given tasks
-				Schema = new BgScriptSchema(NameToTask.Values, PrimitiveTypes);
-				if (SchemaFileName != null)
+				// Import schema if one is passed in
+				BgScriptSchema Schema;
+				if (ImportSchemaFileName != null)
+				{
+					Schema = BgScriptSchema.Import(FileReference.FromString(ImportSchemaFileName));
+				}
+				else
 				{
-					FileReference FullSchemaFileName = new FileReference(SchemaFileName);
-					LogInformation("Writing schema to {0}...", FullSchemaFileName.FullName);
-					Schema.Export(FullSchemaFileName);
-					if (ScriptFileName == null)
+					// Add any primitive types
+					List<(Type, ScriptSchemaStandardType)> PrimitiveTypes = new List<(Type, ScriptSchemaStandardType)>();
+					PrimitiveTypes.Add((typeof(FileReference), ScriptSchemaStandardType.BalancedString));
+					PrimitiveTypes.Add((typeof(DirectoryReference), ScriptSchemaStandardType.BalancedString));
+					PrimitiveTypes.Add((typeof(UnrealTargetPlatform), ScriptSchemaStandardType.BalancedString));
+					PrimitiveTypes.Add((typeof(MCPPlatform), ScriptSchemaStandardType.BalancedString));
+
+					// Create a schema for the given tasks
+					Schema = new BgScriptSchema(NameToTask.Values, PrimitiveTypes);
+					if (SchemaFileName != null)
 					{
-						return ExitCode.Success;
+						FileReference FullSchemaFileName = new FileReference(SchemaFileName);
+						LogInformation("Writing schema to {0}...", FullSchemaFileName.FullName);
+						Schema.Export(FullSchemaFileName);
+						if (ScriptFileName == null)
+						{
+							return ExitCode.Success;
+						}
 					}
 				}
-			}
 
-			// Check there was a script specified
-			if(ScriptFileName == null)
-			{
-				LogError("Missing -Script= parameter for BuildGraph");
-				return ExitCode.Error_Unknown;
-			}
+				// Check there was a script specified
+				if (ScriptFileName == null)
+				{
+					LogError("Missing -Script= parameter for BuildGraph");
+					return ExitCode.Error_Unknown;
+				}
 
-			// Normalize the script filename
-			FileReference FullScriptFile = FileReference.Combine(Unreal.RootDirectory, ScriptFileName);
-			if (!FullScriptFile.IsUnderDirectory(Unreal.RootDirectory))
-			{
-				LogError("BuildGraph scripts must be under the UE root directory");
-				return ExitCode.Error_Unknown;
-			}
-			ScriptFileName = FullScriptFile.MakeRelativeTo(Unreal.RootDirectory).Replace('\\', '/');
+				// Normalize the script filename
+				FileReference FullScriptFile = FileReference.Combine(Unreal.RootDirectory, ScriptFileName);
+				if (!FullScriptFile.IsUnderDirectory(Unreal.RootDirectory))
+				{
+					LogError("BuildGraph scripts must be under the UE root directory");
+					return ExitCode.Error_Unknown;
+				}
+				ScriptFileName = FullScriptFile.MakeRelativeTo(Unreal.RootDirectory).Replace('\\', '/');
 
-			// Read the script from disk
-			BgGraph Graph = BgScriptReader.ReadAsync(Context, ScriptFileName, Arguments, DefaultProperties, Schema, Logger, SingleNodeName).Result;
-			if(Graph == null)
-			{
-				return ExitCode.Error_Unknown;
+				// Read the script from disk
+				Graph = BgScriptReader.ReadAsync(Context, ScriptFileName, Arguments, DefaultProperties, Schema, Logger, SingleNodeName).Result;
+				if (Graph == null)
+				{
+					return ExitCode.Error_Unknown;
+				}
 			}
 
 			// Export the graph for Horde
@@ -843,6 +896,31 @@ namespace AutomationTool
 			Logger.LogScriptError(Task.Location, Format, Args);
 		}
 
+		static void FindAvailableGraphs(Dictionary<string, Type> NameToType)
+		{
+			foreach (Assembly LoadedAssembly in ScriptManager.AllScriptAssemblies)
+			{
+				Type[] Types;
+				try
+				{
+					Types = LoadedAssembly.GetTypes();
+				}
+				catch (ReflectionTypeLoadException ex)
+				{
+					LogWarning("Exception {0} while trying to get types from assembly {1}. LoaderExceptions: {2}", ex, LoadedAssembly, string.Join("\n", ex.LoaderExceptions.Select(x => x.Message)));
+					continue;
+				}
+
+				foreach (Type Type in Types)
+				{
+					if (Type.IsSubclassOf(typeof(BgGraphBuilder)))
+					{
+						NameToType.Add(Type.Name, Type);
+					}
+				}
+			}
+		}
+
 		/// <summary>
 		/// Find all the tasks which are available from the loaded assemblies
 		/// </summary>
diff --git a/Engine/Source/Programs/AutomationTool/BuildGraph/Scripts/InstalledBuild.cs b/Engine/Source/Programs/AutomationTool/BuildGraph/Scripts/InstalledBuild.cs
new file mode 100644
--- /dev/null
+++ b/Engine/Source/Programs/AutomationTool/BuildGraph/Scripts/InstalledBuild.cs
@@ -0,0 +1,502 @@
+// Copyright Epic Games, Inc. All Rights Reserved.
+
+using EpicGames.Core;
+using EpicGames.BuildGraph;
+using EpicGames.BuildGraph.Expressions;
+using System;
+using UnrealBuildTool;
+using AutomationTool.Tasks;
+using System.Collections.Generic;
+using System.Linq;
+using UnrealBuildBase;
+using System.Threading.Tasks;
+
+using static AutomationTool.Tasks.StandardTasks;
+
+#nullable enable
+
+namespace AutomationTool
+{
+	class InstalledBuild : BgGraphBuilder
+	{
+		static FileSet Workspace { get; } = FileSet.FromDirectory(Unreal.RootDirectory);
+		static DirectoryReference RootDir { get; } = new DirectoryReference(CommandUtils.CmdEnv.LocalRoot);
+		static DirectoryReference IntermediateDir => DirectoryReference.Combine(RootDir, "Engine", "Intermediate", "Installed");
+		static DirectoryReference TempMiscDir => DirectoryReference.Combine(IntermediateDir, "General");
+		static DirectoryReference TempCsToolsDir => DirectoryReference.Combine(IntermediateDir, "CsTools");
+		static DirectoryReference TempDdcDir = DirectoryReference.Combine(IntermediateDir, "DDC");
+
+		static string[] PluginsExceptions =
+		{
+			"Engine/Plugins/Enterprise/DatasmithCADImporter/...",
+			"Engine/Plugins/Enterprise/DatasmithIFCImporter/...",
+			"Engine/Plugins/Enterprise/DatasmithC4DImporter/...",
+			"Engine/Plugins/Enterprise/AxFImporter/...",
+			"Engine/Plugins/Enterprise/MDLImporter/..."
+		};
+
+		static string[] WinSignFilter =
+		{
+			"*.exe",
+			"*.dll"
+		};
+
+		static string[] WinStripFilter =
+		{
+			"*.pdb",
+			"-/Engine/Binaries/Win64/UnrealEditor*.pdb",
+			"-/Engine/Plugins/.../Binaries/Win64/UnrealEditor*.pdb",
+		};
+
+		static string[] ProjectsToBuildDdc =
+		{
+		};
+
+		static List<string> GetDdcProjects(UnrealTargetPlatform Platform)
+		{
+			List<string> Projects = new List<string>();
+			Projects.Add("Templates/TP_AEC_ArchvisBP/TP_AEC_ArchvisBP.uproject");
+			Projects.Add("Templates/TP_AEC_BlankBP/TP_AEC_BlankBP.uproject");
+			Projects.Add("Templates/TP_AEC_CollabBP/TP_AEC_CollabBP.uproject");
+			Projects.Add("Templates/TP_AEC_ProdConfigBP/TP_AEC_ProdConfigBP.uproject");
+			Projects.Add("Templates/TP_FirstPersonBP/TP_FirstPersonBP.uproject");
+			Projects.Add("Templates/TP_HandheldARBP/TP_HandheldARBP.uproject");
+			Projects.Add("Templates/TP_AEC_HandheldARBP/TP_AEC_HandheldARBP.uproject");
+			Projects.Add("Templates/TP_MFG_HandheldARBP/TP_MFG_HandheldARBP.uproject");
+			Projects.Add("Templates/TP_MFG_CollabBP/TP_MFG_CollabBP.uproject");
+			Projects.Add("Templates/TP_MFG_ProdConfigBP/TP_MFG_ProdConfigBP.uproject");
+			Projects.Add("Templates/TP_PhotoStudioBP/TP_PhotoStudioBP.uproject");
+			Projects.Add("Templates/TP_PuzzleBP/TP_PuzzleBP.uproject");
+			Projects.Add("Templates/TP_ThirdPersonBP/TP_ThirdPersonBP.uproject");
+			Projects.Add("Templates/TP_TopDownBP/TP_TopDownBP.uproject");
+			Projects.Add("Templates/TP_VehicleAdvBP/TP_VehicleAdvBP.uproject");
+			Projects.Add("Templates/TP_VirtualRealityBP/TP_VirtualRealityBP.uproject");
+//			Projects.Add("Samples/StarterContent/StarterContent.uproject");
+
+			if (Platform == UnrealTargetPlatform.Win64)
+			{
+				Projects.Add("Templates/TP_InCamVFXBP/TP_InCamVFXBP.uproject");
+				Projects.Add("Templates/TP_DMXBP/TP_DMXBP.uproject");
+			}
+			return Projects;
+		}
+
+		public override void SetupGraph(BgGraphSpec Graph)
+		{
+			string RootDir = CommandUtils.CmdEnv.LocalRoot;
+
+			BgBoolOption HostPlatformOnly = Graph.AddOption("HostPlatformOnly", "A helper option to make an installed build for your host platform only, so that you don't have to disable each platform individually", false);
+			BgBoolOption HostPlatformEditorOnly = Graph.AddOption("HostPlatformEditorOnly", "A helper option to make an installed build for your host platform only, so that you don't have to disable each platform individually", false);
+			BgBoolOption AllPlatforms = Graph.AddOption("AllPlatforms", "Include all target platforms by default", false);
+			BgBoolOption CompileDatasmithPlugins = Graph.AddOption("CompileDatasmithPlugins", "If Datasmith plugins should be compiled on a separate node.", false);
+
+			UnrealTargetPlatform CurrentHostPlatform = HostPlatform.Current.HostEditorPlatform;
+
+			BgBool DefaultWithWin64 = !(HostPlatformEditorOnly | (HostPlatformOnly & (CurrentHostPlatform != UnrealTargetPlatform.Win64)));
+			BgBool DefaultWithMac = !(HostPlatformEditorOnly | (HostPlatformOnly & (CurrentHostPlatform != UnrealTargetPlatform.Mac)));
+			BgBool DefaultWithLinux = !(HostPlatformEditorOnly | (HostPlatformOnly & (CurrentHostPlatform != UnrealTargetPlatform.Linux)));
+			BgBool DefaultWithLinuxArm64 = !(HostPlatformEditorOnly | (HostPlatformOnly & (CurrentHostPlatform != UnrealTargetPlatform.Linux)));
+			BgBool DefaultWithPlatform = !(HostPlatformEditorOnly | HostPlatformOnly);
+			BgBool DefaultWithIOS = !((CurrentHostPlatform != UnrealTargetPlatform.Mac) & !AllPlatforms);
+			BgBool DefaultWithHoloLens = !((CurrentHostPlatform != UnrealTargetPlatform.Win64) & !AllPlatforms);
+
+			BgBoolOption WithWin64 = Graph.AddOption("WithWin64", "Include the Win64 target platform", DefaultWithWin64);
+			BgBoolOption WithMac = Graph.AddOption("WithMac", "Include the Mac target platform", DefaultWithMac);
+			BgBoolOption WithAndroid = Graph.AddOption("WithAndroid", "Include the Android target platform", DefaultWithPlatform);
+			BgBoolOption WithIOS = Graph.AddOption("WithIOS", "Include the iOS target platform", DefaultWithIOS);
+			BgBoolOption WithTVOS = Graph.AddOption("WithTVOS", "Include the tvOS target platform", DefaultWithIOS);
+			BgBoolOption WithLinux = Graph.AddOption("WithLinux", "Include the Linux target platform", DefaultWithLinux);
+			BgBoolOption WithLinuxArm64 = Graph.AddOption("WithLinuxArm64", "Include the Linux AArch64 target platform", DefaultWithLinuxArm64);
+			BgBoolOption WithHoloLens = Graph.AddOption("WithHoloLens", "Include the HoloLens target platform", DefaultWithHoloLens);
+
+			BgBoolOption WithClient = Graph.AddOption("WithClient", "Include precompiled client targets", false);
+			BgBoolOption WithServer = Graph.AddOption("WithServer", "Include precompiled server targets", false);
+			BgBoolOption WithDDC = Graph.AddOption("WithDDC", "Build a standalone derived-data cache for the engine content and templates", true);
+			BgBoolOption HostPlatformDDCOnly = Graph.AddOption("HostPlatformDDCOnly", "Whether to include DDC for the host platform only", true);
+			BgBoolOption SignExecutables = Graph.AddOption("SignExecutables", "Sign the executables produced where signing is available", false);
+
+			BgStringOption AnalyticsTypeOverride = Graph.AddOption("AnalyticsTypeOverride", "Identifier for analytic events to send", "");
+			BgBool EmbedSrcSrvInfo = Graph.AddOption("EmbedSrcSrvInfo", "Whether to add Source indexing to Windows game apps so they can be added to a symbol server", false);
+
+			BgList<BgEnum<UnrealTargetConfiguration>> DefaultGameConfigurations = BgList<BgEnum<UnrealTargetConfiguration>>.Create(UnrealTargetConfiguration.DebugGame, UnrealTargetConfiguration.Development, UnrealTargetConfiguration.Shipping);
+			BgList<BgEnum<UnrealTargetConfiguration>> GameConfigurations = Graph.AddOption("GameConfigurations", "Which game configurations to include for packaged applications", DefaultGameConfigurations);
+			BgBoolOption WithFullDebugInfo = Graph.AddOption("WithFullDebugInfo", "Generate full debug info for binary editor and packaged application builds", false);
+
+			BgStringOption BuiltDirectory = Graph.AddOption("BuiltDirectory", "Directory for outputting the built engine", RootDir + "/LocalBuilds/Engine");
+
+			BgStringOption CrashReporterAPIURL = Graph.AddOption("CrashReporterAPIURL", "The URL to use to talk to the CrashReporterClient API.", "");
+			BgStringOption CrashReporterAPIKey = Graph.AddOption("CrashReporterAPIKey", "The API key to use to talk to the CrashReporterClient API.", "");
+
+			BgString CrashReporterCompileArgs = "";
+			CrashReporterCompileArgs = CrashReporterCompileArgs.IfThen(CrashReporterAPIURL != "" & CrashReporterAPIKey != "", BgString.Format("-define:CRC_TELEMETRY_URL=\"{0}\" -define:CRC_TELEMETRY_KEY_DEV=\"{1}\" -define:CRC_TELEMETRY_KEY_RELEASE=\"{1}\" -OverrideBuildEnvironment", CrashReporterAPIURL, CrashReporterAPIKey));
+
+			/////// EDITORS ////////////////////////////////////////////////
+
+			//// Windows ////
+			BgAgentSpec EditorWin64 = Graph.AddAgent("Editor Win64").Type("Win64_Licensee");
+
+			BgNodeSpec VersionFilesNode = EditorWin64
+				.AddNode(x => UpdateVersionFilesAsync(x));
+
+			BgNodeSpec<BgFileSet> WinUhtNode = EditorWin64
+				.AddNode(x => CompileUnrealHeaderToolAsync(x, UnrealTargetPlatform.Win64))
+				.Requires(VersionFilesNode);
+
+			BgNodeSpec<BgFileSet> WinEditorNode = EditorWin64
+				.AddNode(x => CompileUnrealEditorWin64Async(x, CrashReporterCompileArgs, EmbedSrcSrvInfo, CompileDatasmithPlugins, WithFullDebugInfo, SignExecutables))
+				.Requires(WinUhtNode);
+			
+			Graph.AddAggregate("Win64 Editor", WinEditorNode, Label: "Editors/Win64");
+
+
+			/////// TARGET PLATFORMS ////////////////////////////////////////////////
+
+			//// Win64 ////
+
+			BgAgentSpec TargetWin64 = Graph.AddAgent("Target Win64").Type("Win64_Licensee");
+
+			BgNodeSpec<BgFileSet> WinGame = TargetWin64
+				.AddNode(x => CompileUnrealGameWin64(x, GameConfigurations, EmbedSrcSrvInfo, WithFullDebugInfo, SignExecutables))
+				.Requires(WinUhtNode);
+
+			Graph.AddAggregate("TargetPlatforms_Win64", WinGame);
+
+			/////// TOOLS //////////////////////////////////////////////////////////
+
+			//// Build Rules ////
+
+			BgAgentSpec BuildRules = Graph.AddAgent("BuildRules").Type("Win64_Licensee");
+
+			BgNodeSpec RulesAssemblies = BuildRules
+				.AddNode(x => CompileRulesAssemblies(x));
+
+			//// Win Tools ////
+
+			BgAgentSpec ToolsGroupWin64 = Graph.AddAgent("Tools Group Win64").Type("Win64_Licensee");
+
+			BgNodeSpec<BgFileSet> WinTools = ToolsGroupWin64
+				.AddNode(x => BuildToolsWin64Async(x, CrashReporterCompileArgs))
+				.Requires(WinUhtNode);
+
+			BgNodeSpec<BgFileSet> CsTools = ToolsGroupWin64
+				.AddNode(x => BuildToolsCSAsync(x, SignExecutables))
+				.Requires(VersionFilesNode);
+
+
+			/////// DDC //////////////////////////////////////////////////////////
+
+			BgAgentSpec DDCGroupWin64 = Graph.AddAgent("DDC Group Win64").Type("Win64_Licensee");
+
+			BgList<BgString> DDCPlatformsWin64 = BgList<BgString>.Create("WindowsEditor");
+			DDCPlatformsWin64 = DDCPlatformsWin64.AddIf(WithWin64, "Windows");
+
+			BgNodeSpec DdcNode = DDCGroupWin64
+				.AddNode(x => BuildDDCWin64Async(x, DDCPlatformsWin64, BgList<BgFileSet>.Create(WinUhtNode.Output, WinEditorNode.Output, WinTools.Output)))
+				.Requires(WinEditorNode, WinTools);
+
+
+			/////// STAGING ///////
+
+			// Windows 
+			BgAgentSpec WinStageAgent = Graph.AddAgent("Installed Build Group Win64").Type("Win64_Licensee");
+
+			BgList<BgFileSet> WinInstalledFiles = BgList<BgFileSet>.Empty;
+			WinInstalledFiles = WinInstalledFiles.Add(WinUhtNode.Output);
+			WinInstalledFiles = WinInstalledFiles.Add(WinEditorNode.Output);
+			WinInstalledFiles = WinInstalledFiles.Add(WinTools.Output);
+			WinInstalledFiles = WinInstalledFiles.Add(CsTools.Output);
+			WinInstalledFiles = WinInstalledFiles.AddIf(WithWin64, WinGame.Output);
+
+			BgList<BgString> WinPlatforms = BgList<BgString>.Empty;
+			WinPlatforms = WinPlatforms.AddIf(WithWin64, "Win64");
+
+			BgList<BgString> WinContentOnlyPlatforms = BgList<BgString>.Empty;
+
+			BgString WinOutputDir = "LocalBuilds/Engine/Windows";
+			BgString WinFinalizeArgs = BgString.Format("-OutputDir=\"{0}\" -Platforms={1} -ContentOnlyPlatforms={2}", WinOutputDir, BgString.Join(";", WinPlatforms), BgString.Join(";", WinContentOnlyPlatforms));
+
+			BgNodeSpec WinInstalledNode = WinStageAgent
+				.AddNode(x => MakeInstalledBuildWin64Async(x, WinInstalledFiles, WinFinalizeArgs, WinOutputDir))
+				.Requires(WinInstalledFiles);
+
+			Graph.AddAggregate("HostPlatforms_Win64", WinInstalledNode, Label: "Builds/Win64");
+		}
+
+		/// <summary>
+		/// Update the build version
+		/// </summary>
+		static async Task UpdateVersionFilesAsync(BgContext State)
+		{
+			if (State.IsBuildMachine)
+			{
+				await SetVersionAsync(State.Change, State.Stream.Replace('/', '+'));
+			}
+		}
+
+		/// <summary>
+		/// Builds UHT for any platform
+		/// </summary>
+		[BgNodeName("Compile UnrealHeaderTool {Platform}")]
+		static async Task<BgFileSet> CompileUnrealHeaderToolAsync(BgContext State, UnrealTargetPlatform Platform)
+		{
+			return await CompileAsync("UnrealHeaderTool", Platform, UnrealTargetConfiguration.Development, Arguments: "-precompile -allmodules");
+		}
+
+		/// <summary>
+		/// Builds the Windows editor
+		/// </summary>
+		[BgNodeName("Compile UnrealEditor Win64")]
+		static async Task<BgFileSet> CompileUnrealEditorWin64Async(BgContext State, BgString CrashReporterCompileArgs, BgBool EmbedSrcSrvInfo, BgBool CompileDatasmithPlugins, BgBool WithFullDebugInfo, BgBool SignExecutables)
+		{
+			FileSet OutputFiles = FileSet.Empty;
+			OutputFiles += await CompileAsync("UnrealEditor", UnrealTargetPlatform.Win64, UnrealTargetConfiguration.DebugGame, Arguments: $"-precompile -allmodules {State.Get(CrashReporterCompileArgs)}");
+			OutputFiles += await CompileAsync("UnrealEditor", UnrealTargetPlatform.Win64, UnrealTargetConfiguration.Development, Arguments: $"-precompile -allmodules {State.Get(CrashReporterCompileArgs)}");
+
+			if (State.Get(EmbedSrcSrvInfo))
+			{
+				// Embed source info into the PDB files. Should be done from this machine to ensure that paths are correct.
+				Log.TraceInformation("Embedding source file information into PDB files...");
+				FileSet SourceFiles = Workspace.Filter("Engine/Source/...;Engine/Plugins/...").Except("Engine/Source/ThirdParty/...").Filter("*.c;*.h;*.cpp;*.hpp;*.inl");
+//				State.SrcSrv(BinaryFiles: Full, SourceFiles: SourceFiles);
+			}
+
+			if (!State.Get(WithFullDebugInfo))
+			{
+				FileSet UnstrippedFiles = OutputFiles.Filter(WinStripFilter);
+				OutputFiles += await StripAsync(UnstrippedFiles, UnrealTargetPlatform.Win64, BaseDir: RootDir, OutputDir: TempMiscDir);
+				OutputFiles -= UnstrippedFiles;
+			}
+
+			if(State.Get(SignExecutables))
+			{
+//				FileSet UnsignedFiles = OutputFiles.Filter(WinSignFilter);
+//				FileSet FilesToCopy = UnsignedFiles.Except((new HashSet<FileReference>(UnsignedFiles.Where(x => !x.IsUnderDirectory(TempMiscDir)));
+//				UnsignedFiles += State.Copy(FilesToCopy.Flatten(TempMiscDir));
+//				UnsignedFiles.Sign();
+//				OutputFiles -= FilesToCopy;
+//				OutputFiles += UnsignedFiles;
+			}
+
+			return OutputFiles;
+		}
+
+		/// <summary>
+		/// Builds the game target
+		/// </summary>
+		[BgNodeName("Compile UnrealGame Win64")]
+		static async Task<BgFileSet> CompileUnrealGameWin64(BgContext State, BgList<BgEnum<UnrealTargetConfiguration>> Configurations, BgBool EmbedSrcSrvInfo, BgBool WithFullDebugInfo, BgBool SignExecutables)
+		{
+			FileSet Files = FileSet.Empty;
+
+			List<UnrealTargetConfiguration> ConfigurationsValue = State.Get(Configurations);
+			foreach (UnrealTargetConfiguration Configuration in ConfigurationsValue)
+			{
+				Files += await CompileAsync("UnrealGame", UnrealTargetPlatform.Win64, Configuration, Arguments: "-precompile -allmodules -nolink");
+				Files += await CompileAsync("UnrealGame", UnrealTargetPlatform.Win64, Configuration, Arguments: "-precompile", Clean: false);
+			}
+
+			return Files;
+		}
+
+		static async Task<BgFileSet> CompileRulesAssemblies(BgContext State)
+		{
+			FileReference UbtPath = FileReference.Combine(RootDir, "Engine/Binaries/DotNET/UnrealBuildTool/UnrealBuildTool");
+			if (HostPlatform.Current.HostEditorPlatform == UnrealTargetPlatform.Win64)
+			{
+				UbtPath = UbtPath.ChangeExtension(".exe");
+			}
+
+			await SpawnAsync(UbtPath.FullName, Arguments: "-Mode=QueryTargets");
+
+			return Workspace.Filter("Engine/Intermediate/Build/BuildRules/...");
+		}
+
+		static async Task<BgFileSet> BuildToolsWin64Async(BgContext State, BgString CrashReporterCompileArgs)
+		{
+			string CrashReportClientArgs = State.Get(CrashReporterCompileArgs);
+
+			// State.Tag(Files: "#NotForLicensee Build Tools Win64", With: Files);
+
+			FileSet Files = FileSet.Empty;
+			Files += await CompileAsync("CrashReportClient", UnrealTargetPlatform.Win64, UnrealTargetConfiguration.Shipping, Arguments: CrashReportClientArgs);
+			Files += await CompileAsync("CrashReportClientEditor", UnrealTargetPlatform.Win64, UnrealTargetConfiguration.Shipping, Arguments: CrashReportClientArgs);
+			Files += await CompileAsync("ShaderCompileWorker", UnrealTargetPlatform.Win64, UnrealTargetConfiguration.Development);
+			Files += await CompileAsync("EpicWebHelper", UnrealTargetPlatform.Win64, UnrealTargetConfiguration.Development);
+			Files += await CompileAsync("UnrealInsights", UnrealTargetPlatform.Win64, UnrealTargetConfiguration.Development);
+			Files += await CompileAsync("UnrealFrontend", UnrealTargetPlatform.Win64, UnrealTargetConfiguration.Development);
+			Files += await CompileAsync("UnrealLightmass", UnrealTargetPlatform.Win64, UnrealTargetConfiguration.Development);
+			Files += await CompileAsync("UnrealPak", UnrealTargetPlatform.Win64, UnrealTargetConfiguration.Development);
+			Files += await CompileAsync("UnrealMultiUserServer", UnrealTargetPlatform.Win64, UnrealTargetConfiguration.Development);
+			Files += await CompileAsync("UnrealRecoverySvc", UnrealTargetPlatform.Win64, UnrealTargetConfiguration.Development);
+			Files += await CompileAsync("LiveCodingConsole", UnrealTargetPlatform.Win64, UnrealTargetConfiguration.Development);
+			Files += await CompileAsync("BootstrapPackagedGame", UnrealTargetPlatform.Win64, UnrealTargetConfiguration.Shipping);
+			Files += await CompileAsync("BuildPatchTool", UnrealTargetPlatform.Win64, UnrealTargetConfiguration.Shipping);
+			Files += await CompileAsync("SwitchboardListener", UnrealTargetPlatform.Win64, UnrealTargetConfiguration.Development);
+			Files += FileSet.FromFile(RootDir, "Engine/Binaries/Win64/XGEControlWorker.exe");
+			Files += FileSet.FromFile(RootDir, "Engine/Saved/UnrealBuildTool/BuildConfiguration.Schema.xsd");
+
+			return Files;
+		}
+
+		[BgNodeName("Build Tools CS")]
+		static async Task<BgFileSet> BuildToolsCSAsync(BgContext State, BgBool SignExecutables)
+		{
+			FileUtils.ForceDeleteDirectory(TempCsToolsDir);
+
+			// Copy Source and referenced libraries to a new location with Confidential folders removed
+			FileSet UatProjects = Workspace.Filter("Engine/Source/Programs/AutomationTool/....csproj");
+
+			FileSet RedistUatSource = Workspace.Filter("Engine/Binaries/DotNET/...;Engine/Binaries/ThirdParty/Newtonsoft/...;Engine/Binaries/ThirdParty/IOS/...;Engine/Binaries/ThirdParty/VisualStudio/...;Engine/Source/Programs/...;Engine/Platforms/*/Source/Programs/...;Engine/Source/Editor/SwarmInterface/...");
+			await RedistUatSource.CopyToAsync(Unreal.RootDirectory);
+
+			// Compile all the tools
+			CsCompileOutput Output = CsCompileOutput.Empty;
+			if (HostPlatform.Current.HostEditorPlatform == UnrealTargetPlatform.Win64)
+			{
+				Output += await CsCompileAsync(FileReference.Combine(TempCsToolsDir, "Engine/Source/Programs/UnrealSwarm/SwarmCoordinator/SwarmCoordinator.csproj"), Platform: "AnyCPU", Configuration: "Development");
+				Output += await CsCompileAsync(FileReference.Combine(TempCsToolsDir, "Engine/Source/Programs/UnrealSwarm/Agent/Agent.csproj"), Platform: "AnyCPU", Configuration: "Development");
+				Output += await CsCompileAsync(FileReference.Combine(TempCsToolsDir, "Engine/Source/Editor/SwarmInterface/DotNET/SwarmInterface.csproj"), Platform: "AnyCPU", Configuration: "Development");
+				Output += await CsCompileAsync(FileReference.Combine(TempCsToolsDir, "Engine/Source/Programs/UnrealControls/UnrealControls.csproj"), Platform: "AnyCPU", Configuration: "Development");
+				Output += await CsCompileAsync(FileReference.Combine(TempCsToolsDir, "Engine/Source/Programs/IOS/iPhonePackager/iPhonePackager.csproj"), Platform: "AnyCPU", Configuration: "Development", Arguments: "/verbosity:minimal /target:Rebuild");
+			}
+
+			if (HostPlatform.Current.HostEditorPlatform != UnrealTargetPlatform.Linux)
+			{
+				Output += await CsCompileAsync(FileReference.Combine(TempCsToolsDir, "Engine/Source/Programs/NetworkProfiler/NetworkProfiler/NetworkProfiler.csproj"), Platform: "AnyCPU", Configuration: "Development");
+			}
+
+			Output += await CsCompileAsync(FileReference.Combine(RootDir, "Engine/Source/Programs/IOS/MobileDeviceInterface/MobileDeviceInterface.csproj"), Platform: "AnyCPU", Configuration: "Development", Arguments: "/verbosity:minimal /target:Rebuild");
+
+			FileSet Binaries = Output.Binaries + Output.References;
+
+			// Tag AutomationTool and UnrealBuildTool folders recursively as NET Core dependencies are not currently handled by CsCompile
+			Binaries += FileSet.FromDirectory(TempCsToolsDir).Filter("Engine/Binaries/DotNET/AutomationTool/...");
+			Binaries += FileSet.FromDirectory(TempCsToolsDir).Filter("Engine/Binaries/DotNET/UnrealBuildTool/...");
+
+			// Tag AutomationTool Script module build records, so that prebuilt modules may be discovered in the absence of source code
+			Binaries += FileSet.FromDirectory(TempCsToolsDir).Filter("Engine/Intermediate/ScriptModules/...");
+
+			if(State.Get(SignExecutables))
+			{
+//				Binaries.Sign();
+			}
+
+			return Binaries;
+		}
+
+		/// <summary>
+		/// Creates a DDC pack file for the supported platforms
+		/// </summary>
+		[BgNodeName("Build DDC Win64")]
+		static async Task<BgFileSet> BuildDDCWin64Async(BgContext State, BgList<BgString> Platforms, BgList<BgFileSet> Dependencies)
+		{
+			// Build up a list of files needed to build DDC
+			FileSet ToCopy = State.Get(Dependencies);
+			ToCopy += await CsCompileAsync(FileReference.Combine(RootDir, "Engine/Source/Programs/UnrealBuildTool/UnrealBuildTool.csproj"), Platform: "AnyCPU", Configuration: "Development", EnumerateOnly: true).MergeAsync();
+			ToCopy += await ToCopy.Filter("*.target").TagReceiptsAsync(RuntimeDependencies: true);
+			ToCopy += FileSet.FromFile(RootDir, "Engine/Binaries/DotNET/Ionic.Zip.Reduced.dll");
+			ToCopy += FileSet.FromFile(RootDir, "Engine/Binaries/DotNET/OneSky.dll");
+			ToCopy += FileSet.FromDirectory(RootDir).Filter("Templates/TemplateResources/...");
+
+			foreach (DirectoryReference ExtensionDir in Unreal.GetExtensionDirs(Unreal.EngineDirectory))
+			{
+				ToCopy += FileSet.FromDirectory(ExtensionDir).Filter("Content/...").Except("....psd;....pdn;....fbx;....po");
+				ToCopy += FileSet.FromDirectory(ExtensionDir).Filter("Config/...").Except("....vdf");
+				ToCopy += FileSet.FromDirectory(ExtensionDir).Filter("Shaders/...");
+
+				DirectoryReference ExtensionPluginsDir = DirectoryReference.Combine(ExtensionDir, "Plugins");
+				ToCopy += FileSet.FromDirectory(ExtensionPluginsDir).Filter("....uplugin;.../Config/...;.../Content/...;.../Resources/...;.../Shaders/...;.../Templates/...").Except(".../TwitchLiveStreaming/...");
+			}
+
+			// Filter out the files not needed to build DDC. Removing confidential folders can affect DDC keys, so we want to be sure that we're making DDC with a build that can use it.
+			FileSet FilteredCopyList = ToCopy.Except(".../Source/...;.../Intermediate/...");
+
+			Dictionary<FileReference, FileReference> TargetFileToSourceFile = new Dictionary<FileReference, FileReference>();
+			MapFilesToOutputDir(FilteredCopyList, TempDdcDir, TargetFileToSourceFile);
+
+			FileUtils.ForceDeleteDirectoryContents(TempDdcDir);
+			await FilteredCopyList.CopyToAsync(TempDdcDir);
+
+			// Run the DDC commandlet
+			List<string> Arguments = new List<string>();
+			Arguments.Add($"-TempDir=\"{TempDdcDir}\"");
+			Arguments.Add($"-FeaturePacks=\"{String.Join(";", GetDdcProjects(UnrealTargetPlatform.Win64))}\"");
+			Arguments.Add($"-TargetPlatforms={String.Join("+", State.Get(Platforms))}");
+			Arguments.Add($"-HostPlatform=Win64");
+			await State.CommandAsync("BuildDerivedDataCache", Arguments: String.Join(" ", Arguments));
+
+			// Return a tag for the output file
+			return FileSet.FromFile(TempDdcDir, "Engine/DerivedDataCache/Compressed.ddp");
+		}
+
+		/// <summary>
+		/// Copy all the build artifacts to the output folder
+		/// </summary>
+		[BgNodeName("Make Installed Build Win64")]
+		static async Task<BgFileSet> MakeInstalledBuildWin64Async(BgContext State, BgList<BgFileSet> InputFiles, BgString FinalizeArgs, BgString OutputDir)
+		{
+			// Find all the input files, and add any runtime dependencies from the receipts into the list
+			FileSet SourceFiles = State.Get(InputFiles);
+			SourceFiles += await SourceFiles.Filter("*.target").TagReceiptsAsync(RuntimeDependencies: true);
+
+			// Include any files referenced by dependency lists
+			FileSet DependencyListFiles = SourceFiles.Filter(".../DependencyList.txt;.../DependencyList-AllModules.txt");
+			foreach (FileReference DependencyListFile in DependencyListFiles)
+			{
+				string[] Lines = FileReference.ReadAllLines(DependencyListFile);
+				foreach (string Line in Lines)
+				{
+					string TrimLine = Line.Trim();
+					if (TrimLine.Length > 0)
+					{
+						SourceFiles += FileSet.FromFile(Unreal.RootDirectory, TrimLine);
+					}
+				}
+			}
+			SourceFiles -= DependencyListFiles;
+
+			// Clear the output directory
+			DirectoryReference OutputDirRef = new DirectoryReference(State.Get(OutputDir));
+			FileUtils.ForceDeleteDirectoryContents(OutputDirRef);
+			FileSet InstalledFiles = await SourceFiles.CopyToAsync(OutputDirRef);
+
+			// Run the finalize command 
+			await State.CommandAsync("FinalizeInstalledBuild", Arguments: State.Get(FinalizeArgs));
+
+			// Get the final list of files
+			InstalledFiles += FileSet.FromFile(OutputDirRef, "Engine/Build/InstalledEngine.txt");
+
+			// Sanitize any receipts in the output directory
+			await InstalledFiles.Filter("*.target").SanitizeReceiptsAsync();
+
+			return InstalledFiles;
+		}
+
+		static void MapFilesToOutputDir(IEnumerable<FileReference> SourceFiles, DirectoryReference TargetDir, Dictionary<FileReference, FileReference> TargetFileToSourceFile)
+		{
+			foreach (FileReference SourceFile in SourceFiles)
+			{
+				DirectoryReference BaseDir;
+				if (SourceFile.IsUnderDirectory(TempCsToolsDir))
+				{
+					BaseDir = TempCsToolsDir;
+				}
+				else if (SourceFile.IsUnderDirectory(TempMiscDir))
+				{
+					BaseDir = TempMiscDir;
+				}
+				else if (SourceFile.IsUnderDirectory(TempDdcDir))
+				{
+					BaseDir = TempDdcDir;
+				}
+				else
+				{
+					BaseDir = Unreal.RootDirectory;
+				}
+
+				FileReference TargetFile = FileReference.Combine(TargetDir, SourceFile.MakeRelativeTo(BaseDir));
+				TargetFileToSourceFile[TargetFile] = SourceFile;
+			}
+		}
+	}
+}
diff --git a/Engine/Source/Programs/AutomationTool/BuildGraph/Tasks/CallMethodTask.cs b/Engine/Source/Programs/AutomationTool/BuildGraph/Tasks/CallMethodTask.cs
new file mode 100644
--- /dev/null
+++ b/Engine/Source/Programs/AutomationTool/BuildGraph/Tasks/CallMethodTask.cs
@@ -0,0 +1,298 @@
+// Copyright Epic Games, Inc. All Rights Reserved.
+
+using EpicGames.BuildGraph;
+using System;
+using System.Collections.Generic;
+using System.IO;
+using System.Linq;
+using System.Text;
+using System.Threading.Tasks;
+using System.Xml;
+using EpicGames.Core;
+using UnrealBuildBase;
+using UnrealBuildTool;
+using System.Reflection;
+using EpicGames.BuildGraph.Expressions;
+using Microsoft.Extensions.Logging;
+using System.Runtime.CompilerServices;
+using System.ComponentModel;
+
+#nullable enable
+
+namespace AutomationTool.Tasks
+{
+	/// <summary>
+	/// Parameters for a task which calls a C# method
+	/// </summary>
+	public class CallMethodTaskParameters
+	{
+		/// <summary>
+		/// Name of the class to run a method in
+		/// </summary>
+		[TaskParameter]
+		public string Class = String.Empty;
+
+		/// <summary>
+		/// The method name to execute.
+		/// </summary>
+		[TaskParameter]
+		public string Method = String.Empty;
+
+		/// <summary>
+		/// Argument value
+		/// </summary>
+		[TaskParameter(Optional = true)]
+		public string? Arg1;
+
+		/// <summary>
+		/// Argument value
+		/// </summary>
+		[TaskParameter(Optional = true)]
+		public string? Arg2;
+
+		/// <summary>
+		/// Argument value
+		/// </summary>
+		[TaskParameter(Optional = true)]
+		public string? Arg3;
+
+		/// <summary>
+		/// Argument value
+		/// </summary>
+		[TaskParameter(Optional = true)]
+		public string? Arg4;
+
+		/// <summary>
+		/// Argument value
+		/// </summary>
+		[TaskParameter(Optional = true)]
+		public string? Arg5;
+
+		/// <summary>
+		/// Argument value
+		/// </summary>
+		[TaskParameter(Optional = true)]
+		public string? Arg6;
+
+		/// <summary>
+		/// Argument value
+		/// </summary>
+		[TaskParameter(Optional = true)]
+		public string? Arg7;
+
+		/// <summary>
+		/// Argument value
+		/// </summary>
+		[TaskParameter(Optional = true)]
+		public string? Arg8;
+
+		/// <summary>
+		/// Argument value
+		/// </summary>
+		[TaskParameter(Optional = true)]
+		public string? Arg9;
+
+		/// <summary>
+		/// List of tags for outputs
+		/// </summary>
+		[TaskParameter(Optional = true)]
+		public string? Tags;
+	}
+
+	/// <summary>
+	/// Spawns the editor to run a commandlet.
+	/// </summary>
+	[TaskElement("CallMethod", typeof(CallMethodTaskParameters))]
+	public class CallMethodTask : BgTaskImpl
+	{
+		class BgStateImpl : BgContext
+		{
+			JobContext JobContext;
+
+			public BgStateImpl(JobContext JobContext, Dictionary<string, HashSet<FileReference>> TagNameToFileSet)
+				: base(CreateExprContext(TagNameToFileSet))
+			{
+				this.JobContext = JobContext;
+			}
+
+			static BgExprContext CreateExprContext(Dictionary<string, HashSet<FileReference>> TagNameToFileSet)
+			{
+				BgExprContext Context = new BgExprContext();
+				foreach (KeyValuePair<string, HashSet<FileReference>> Pair in TagNameToFileSet)
+				{
+					Context.TagNameToFileSet[Pair.Key] = FileSet.FromFiles(Unreal.RootDirectory, Pair.Value);
+				}
+				return Context;
+			}
+
+			public override string Stream => CommandUtils.P4Enabled ? CommandUtils.P4Env.Branch : "";
+
+			public override int Change => CommandUtils.P4Enabled ? CommandUtils.P4Env.Changelist : 0;
+
+			public override int CodeChange => CommandUtils.P4Enabled ? CommandUtils.P4Env.CodeChangelist : 0;
+
+			public override (int Major, int Minor, int Patch) EngineVersion
+			{
+				get
+				{
+					ReadOnlyBuildVersion Current = ReadOnlyBuildVersion.Current;
+					return (Current.MajorVersion, Current.MinorVersion, Current.PatchVersion);
+				}
+			}
+
+			public override bool IsBuildMachine => CommandUtils.IsBuildMachine;
+		}
+				
+		/// <summary>
+		/// Parameters for this task
+		/// </summary>
+		CallMethodTaskParameters Parameters;
+
+		/// <summary>
+		/// Construct a new CommandletTask.
+		/// </summary>
+		/// <param name="InParameters">Parameters for this task</param>
+		public CallMethodTask(CallMethodTaskParameters InParameters)
+		{
+			Parameters = InParameters;
+		}
+
+		/// <summary>
+		/// Execute the task.
+		/// </summary>
+		/// <param name="Job">Information about the current job</param>
+		/// <param name="BuildProducts">Set of build products produced by this node.</param>
+		/// <param name="TagNameToFileSet">Mapping from tag names to the set of files they include</param>
+		public override async Task ExecuteAsync(JobContext Job, HashSet<FileReference> BuildProducts, Dictionary<string, HashSet<FileReference>> TagNameToFileSet)
+		{
+			Type? Type = Type.GetType(Parameters.Class);
+			if (Type == null)
+			{
+				throw new AutomationException($"Unable to find class '{Parameters.Class}'");
+			}
+
+			MethodInfo? Method = Type.GetMethod(Parameters.Method, BindingFlags.Public | BindingFlags.NonPublic | BindingFlags.Static);
+			if (Method == null)
+			{
+				throw new AutomationException($"Unable to find method '{Parameters.Method}' on class {Parameters.Class}");
+			}
+
+			List<string?> ArgumentStrings = new List<string?>();
+			ArgumentStrings.Add(Parameters.Arg1!);
+			ArgumentStrings.Add(Parameters.Arg2!);
+			ArgumentStrings.Add(Parameters.Arg3!);
+			ArgumentStrings.Add(Parameters.Arg4!);
+			ArgumentStrings.Add(Parameters.Arg5!);
+			ArgumentStrings.Add(Parameters.Arg6!);
+			ArgumentStrings.Add(Parameters.Arg7!);
+			ArgumentStrings.Add(Parameters.Arg8!);
+			ArgumentStrings.Add(Parameters.Arg9!);
+
+			ParameterInfo[] MethodParameters = Method.GetParameters();
+
+			BgStateImpl State = new BgStateImpl(Job, TagNameToFileSet);
+
+			object[] Arguments = new object[MethodParameters.Length];
+			for (int Idx = 0; Idx < MethodParameters.Length; Idx++)
+			{
+				string? ArgumentString = ArgumentStrings[Idx];
+
+				Type ParameterType = MethodParameters[Idx].ParameterType;
+				if (ParameterType == typeof(BgContext))
+				{
+					Arguments[Idx] = State;
+				}
+				else if (ArgumentString == null)
+				{
+					throw new AutomationException($"Missing argument for parameter {MethodParameters[Idx].Name}");
+				}
+				else if (ParameterType == typeof(bool))
+				{
+					Arguments[Idx] = bool.Parse(ArgumentString);
+				}
+				else if (ParameterType == typeof(int))
+				{
+					Arguments[Idx] = int.Parse(ArgumentString);
+				}
+				else if (ParameterType == typeof(string))
+				{
+					Arguments[Idx] = ArgumentString;
+				}
+				else if (typeof(IBgExpr).IsAssignableFrom(ParameterType))
+				{
+					Arguments[Idx] = BgType.Get(ParameterType).DeserializeArgument(ArgumentString);
+				}
+				else
+				{
+					Arguments[Idx] = TypeDescriptor.GetConverter(ParameterType).ConvertFromString(ArgumentString);
+				}
+			}
+
+			Task Task = (Task)Method.Invoke(null, Arguments)!;
+			await Task;
+
+			if(Parameters.Tags != null)
+			{
+				object? Result = null;
+				if (Method.ReturnType.IsGenericType && Method.ReturnType.GetGenericTypeDefinition() == typeof(Task<>))
+				{
+					Type TaskType = Task.GetType();
+					PropertyInfo? Property = TaskType.GetProperty(nameof(Task<int>.Result));
+					Result = Property!.GetValue(Task);
+				}
+
+				object?[] OutputValues;
+				if (Result is BgFileSet)
+				{
+					OutputValues = new[] { Result };
+				}
+				else if (Result is ITuple Tuple)
+				{
+					OutputValues = Enumerable.Range(0, Tuple.Length).Select(x => Tuple[x]).ToArray();
+				}
+				else
+				{
+					OutputValues = Array.Empty<object?>();
+				}
+
+				string[] TagNames = Parameters.Tags.Split(';');
+				for (int Idx = 0; Idx < OutputValues.Length; Idx++)
+				{
+					if (OutputValues[Idx] is BgFileSet Token)
+					{
+						string TagName = TagNames[Idx];
+						FileSet FileSet = Token.ComputeValue(State.Context);
+						TagNameToFileSet[TagName] = new HashSet<FileReference>(FileSet.Flatten().Values);
+						BuildProducts.UnionWith(TagNameToFileSet[TagName]);
+					}
+				}
+			}
+		}
+
+		/// <summary>
+		/// Output this task out to an XML writer.
+		/// </summary>
+		public override void Write(XmlWriter Writer)
+		{
+			Write(Writer, Parameters);
+		}
+
+		/// <summary>
+		/// Find all the tags which are used as inputs to this task
+		/// </summary>
+		/// <returns>The tag names which are read by this task</returns>
+		public override IEnumerable<string> FindConsumedTagNames()
+		{
+			yield break;
+		}
+
+		/// <summary>
+		/// Find all the tags which are modified by this task
+		/// </summary>
+		/// <returns>The tag names which are modified by this task</returns>
+		public override IEnumerable<string> FindProducedTagNames()
+		{
+			yield break;
+		}
+	}
+}
diff --git a/Engine/Source/Programs/AutomationTool/BuildGraph/Tasks/CommandTask.cs b/Engine/Source/Programs/AutomationTool/BuildGraph/Tasks/CommandTask.cs
--- a/Engine/Source/Programs/AutomationTool/BuildGraph/Tasks/CommandTask.cs
+++ b/Engine/Source/Programs/AutomationTool/BuildGraph/Tasks/CommandTask.cs
@@ -11,6 +11,8 @@ using System.Xml;
 using System.IO;
 using EpicGames.Core;
 using UnrealBuildBase;
+using EpicGames.BuildGraph;
+using AutomationTool.Tasks;
 
 namespace AutomationTool.Tasks
 {
@@ -164,4 +166,24 @@ namespace AutomationTool.Tasks
 			yield break;
 		}
 	}
+
+	public static partial class StandardTasks
+	{
+		/// <summary>
+		/// Runs another UAT command
+		/// </summary>
+		/// <param name="State">The execution state</param>
+		/// <param name="Name">Name of the command to run</param>
+		/// <param name="Arguments">Arguments for the command</param>
+		/// <param name="MergeTelemetryWithPrefix">If non-null, instructs telemetry from the command to be merged into the telemetry for this UAT instance with the given prefix. May be an empty (non-null) string.</param>
+		public static async Task CommandAsync(this BgContext State, string Name, string Arguments = null, string MergeTelemetryWithPrefix = null)
+		{
+			CommandTaskParameters Parameters = new CommandTaskParameters();
+			Parameters.Name = Name;
+			Parameters.Arguments = Arguments ?? Parameters.Arguments;
+			Parameters.MergeTelemetryWithPrefix = MergeTelemetryWithPrefix ?? Parameters.MergeTelemetryWithPrefix;
+
+			await ExecuteAsync(new CommandTask(Parameters));
+		}
+	}
 }
diff --git a/Engine/Source/Programs/AutomationTool/BuildGraph/Tasks/CommandletTask.cs b/Engine/Source/Programs/AutomationTool/BuildGraph/Tasks/CommandletTask.cs
--- a/Engine/Source/Programs/AutomationTool/BuildGraph/Tasks/CommandletTask.cs
+++ b/Engine/Source/Programs/AutomationTool/BuildGraph/Tasks/CommandletTask.cs
@@ -143,4 +143,27 @@ namespace AutomationTool.Tasks
 			yield break;
 		}
 	}
+
+	public static partial class StandardTasks
+	{
+		/// <summary>
+		/// Task which runs a UE commandlet
+		/// </summary>
+		/// <param name="State"></param>
+		/// <param name="Name">The commandlet name to execute.</param>
+		/// <param name="Project">The project to run the editor with.</param>
+		/// <param name="Arguments">Arguments to be passed to the commandlet.</param>
+		/// <param name="EditorExe">The editor executable to use. Defaults to the development UnrealEditor executable for the current platform.</param>
+		/// <param name="ErrorLevel">The minimum exit code, which is treated as an error.</param>
+		public static async Task CommandletAsync(this BgContext State, string Name, FileReference Project = null, string Arguments = null, FileReference EditorExe = null, int ErrorLevel = 1)
+		{
+			CommandletTaskParameters Parameters = new CommandletTaskParameters();
+			Parameters.Name = Name;
+			Parameters.Project = Project?.FullName;
+			Parameters.Arguments = Arguments;
+			Parameters.EditorExe = EditorExe;
+			Parameters.ErrorLevel = ErrorLevel;
+			await ExecuteAsync(new CommandletTask(Parameters));
+		}
+	}
 }
diff --git a/Engine/Source/Programs/AutomationTool/BuildGraph/Tasks/CompileTask.cs b/Engine/Source/Programs/AutomationTool/BuildGraph/Tasks/CompileTask.cs
--- a/Engine/Source/Programs/AutomationTool/BuildGraph/Tasks/CompileTask.cs
+++ b/Engine/Source/Programs/AutomationTool/BuildGraph/Tasks/CompileTask.cs
@@ -332,4 +332,33 @@ namespace AutomationTool.Tasks
 			return FindTagNamesFromList(Parameters.Tag);
 		}
 	}
+
+	public static partial class StandardTasks
+	{
+		/// <summary>
+		/// Compiles a target
+		/// </summary>
+		/// <param name="Target">The target to compile</param>
+		/// <param name="Configuration">The configuration to compile</param>
+		/// <param name="Platform">The platform to compile for</param>
+		/// <param name="Project">The project to compile with</param>
+		/// <param name="Arguments">Additional arguments for UnrealBuildTool</param>
+		/// <param name="AllowXGE">Whether to allow using XGE for compilation</param>
+		/// <param name="AllowParallelExecutor">Whether to allow using the parallel executor for this compile</param>
+		/// <param name="Clean">Whether to allow cleaning this target. If unspecified, targets are cleaned if the -Clean argument is passed on the command line</param>
+		/// <returns>Build products from the compile</returns>
+		public static async Task<FileSet> CompileAsync(string Target, UnrealTargetPlatform Platform, UnrealTargetConfiguration Configuration, FileReference Project = null, string Arguments = null, bool AllowXGE = true, bool AllowParallelExecutor = true, bool? Clean = null)
+		{
+			CompileTaskParameters Parameters = new CompileTaskParameters();
+			Parameters.Target = Target;
+			Parameters.Platform = Platform;
+			Parameters.Configuration = Configuration;
+			Parameters.Project = Project?.FullName;
+			Parameters.Arguments = Arguments;
+			Parameters.AllowXGE = AllowXGE;
+			Parameters.AllowParallelExecutor = AllowParallelExecutor;
+			Parameters.Clean = Clean;
+			return await ExecuteAsync(new CompileTask(Parameters));
+		}
+	}
 }
diff --git a/Engine/Source/Programs/AutomationTool/BuildGraph/Tasks/CopyTask.cs b/Engine/Source/Programs/AutomationTool/BuildGraph/Tasks/CopyTask.cs
--- a/Engine/Source/Programs/AutomationTool/BuildGraph/Tasks/CopyTask.cs
+++ b/Engine/Source/Programs/AutomationTool/BuildGraph/Tasks/CopyTask.cs
@@ -12,6 +12,7 @@ using System.Xml;
 using EpicGames.Core;
 using UnrealBuildTool;
 using UnrealBuildBase;
+using Microsoft.Extensions.Logging;
 
 namespace AutomationTool.Tasks
 {
@@ -83,7 +84,7 @@ namespace AutomationTool.Tasks
 		/// <param name="Job">Information about the current job</param>
 		/// <param name="BuildProducts">Set of build products produced by this node.</param>
 		/// <param name="TagNameToFileSet">Mapping from tag names to the set of files they include</param>
-		public override Task ExecuteAsync(JobContext Job, HashSet<FileReference> BuildProducts, Dictionary<string, HashSet<FileReference>> TagNameToFileSet)
+		public override async Task ExecuteAsync(JobContext Job, HashSet<FileReference> BuildProducts, Dictionary<string, HashSet<FileReference>> TagNameToFileSet)
 		{
 			// Parse all the source patterns
 			FilePattern SourcePattern = new FilePattern(Unreal.RootDirectory, Parameters.From);
@@ -93,7 +94,7 @@ namespace AutomationTool.Tasks
 
 			// Apply the filter to the source files
 			HashSet<FileReference> Files = null;
-			if(!String.IsNullOrEmpty(Parameters.Files))
+			if (!String.IsNullOrEmpty(Parameters.Files))
 			{
 				SourcePattern = SourcePattern.AsDirectoryPattern();
 				Files = ResolveFilespec(SourcePattern.BaseDirectory, Parameters.Files, TagNameToFileSet);
@@ -102,8 +103,38 @@ namespace AutomationTool.Tasks
 			// Build the file mapping
 			Dictionary<FileReference, FileReference> TargetFileToSourceFile = FilePattern.CreateMapping(Files, ref SourcePattern, ref TargetPattern);
 
+			// Check we got some files
+			if (TargetFileToSourceFile.Count == 0)
+			{
+				if (Parameters.ErrorIfNotFound)
+				{
+					CommandUtils.LogError("No files found matching '{0}'", SourcePattern);
+				}
+				else
+				{
+					CommandUtils.LogInformation("No files found matching '{0}'", SourcePattern);
+				}
+				return;
+			}
+
+			// Run the copy
+			CommandUtils.LogInformation("Copying {0} file{1} from {2} to {3}...", TargetFileToSourceFile.Count, (TargetFileToSourceFile.Count == 1) ? "" : "s", SourcePattern.BaseDirectory, TargetPattern.BaseDirectory);
+			await ExecuteAsync(TargetFileToSourceFile, Parameters.Overwrite);
+
+			// Update the list of build products
+			BuildProducts.UnionWith(TargetFileToSourceFile.Keys);
+
+			// Apply the optional output tag to them
+			foreach (string TagName in FindTagNamesFromList(Parameters.Tag))
+			{
+				FindOrAddTagSet(TagNameToFileSet, TagName).UnionWith(TargetFileToSourceFile.Keys);
+			}
+		}
+
+		public static Task ExecuteAsync(Dictionary<FileReference, FileReference> TargetFileToSourceFile, bool Overwrite)
+		{
 			//  If we're not overwriting, remove any files where the destination file already exists.
-			if(!Parameters.Overwrite)
+			if (!Overwrite)
 			{
 				TargetFileToSourceFile = TargetFileToSourceFile.Where(File =>
 				{
@@ -116,20 +147,6 @@ namespace AutomationTool.Tasks
 				}).ToDictionary(Pair => Pair.Key, Pair => Pair.Value);
 			}
 
-			// Check we got some files
-			if(TargetFileToSourceFile.Count == 0)
-			{
-				if (Parameters.ErrorIfNotFound)
-				{
-					CommandUtils.LogError("No files found matching '{0}'", SourcePattern);
-				}
-				else
-				{
-					CommandUtils.LogInformation("No files found matching '{0}'", SourcePattern);
-				}
-				return Task.CompletedTask;
-			}
-
 			// If the target is on a network share, retry creating the first directory until it succeeds
 			DirectoryReference FirstTargetDirectory = TargetFileToSourceFile.First().Key.Directory;
 			if(!DirectoryReference.Exists(FirstTargetDirectory))
@@ -171,21 +188,11 @@ namespace AutomationTool.Tasks
 
 			// Copy them all
 			KeyValuePair<FileReference, FileReference>[] FilePairs = TargetFileToSourceFile.ToArray();
-			CommandUtils.LogInformation("Copying {0} file{1} from {2} to {3}...", FilePairs.Length, (FilePairs.Length == 1)? "" : "s", SourcePattern.BaseDirectory, TargetPattern.BaseDirectory);
 			foreach(KeyValuePair<FileReference, FileReference> FilePair in FilePairs)
 			{
 				CommandUtils.LogLog("  {0} -> {1}", FilePair.Value, FilePair.Key);
 			}
 			CommandUtils.ThreadedCopyFiles(FilePairs.Select(x => x.Value.FullName).ToList(), FilePairs.Select(x => x.Key.FullName).ToList(), bQuiet: true, bRetry: true);
-
-			// Update the list of build products
-			BuildProducts.UnionWith(TargetFileToSourceFile.Keys);
-
-			// Apply the optional output tag to them
-			foreach(string TagName in FindTagNamesFromList(Parameters.Tag))
-			{
-				FindOrAddTagSet(TagNameToFileSet, TagName).UnionWith(TargetFileToSourceFile.Keys);
-			}
 			return Task.CompletedTask;
 		}
 
@@ -218,4 +225,27 @@ namespace AutomationTool.Tasks
 			return FindTagNamesFromList(Parameters.Tag);
 		}
 	}
+
+	public static partial class BgStateExtensions
+	{
+		/// <summary>
+		/// Copy files from one location to another
+		/// </summary>
+		/// <param name="Files">The files to copy</param>
+		/// <param name="TargetDir"></param>
+		/// <param name="Overwrite">Whether or not to overwrite existing files.</param>
+		public static async Task<FileSet> CopyToAsync(this FileSet Files, DirectoryReference TargetDir, bool? Overwrite = null)
+		{
+			// Run the copy
+			Dictionary<FileReference, FileReference> TargetFileToSourceFile = Files.Flatten(TargetDir);
+			if (TargetFileToSourceFile.Count == 0)
+			{
+				return FileSet.Empty;
+			}
+
+			Log.Logger.LogInformation("Copying {NumFiles} file(s) to {TargetDir}...", TargetFileToSourceFile.Count, TargetDir);
+			await CopyTask.ExecuteAsync(TargetFileToSourceFile, Overwrite ?? true);
+			return FileSet.FromFiles(TargetFileToSourceFile.Keys.Select(x => (x.MakeRelativeTo(TargetDir), x)));
+		}
+	}
 }
diff --git a/Engine/Source/Programs/AutomationTool/BuildGraph/Tasks/CsCompileTask.cs b/Engine/Source/Programs/AutomationTool/BuildGraph/Tasks/CsCompileTask.cs
--- a/Engine/Source/Programs/AutomationTool/BuildGraph/Tasks/CsCompileTask.cs
+++ b/Engine/Source/Programs/AutomationTool/BuildGraph/Tasks/CsCompileTask.cs
@@ -306,5 +306,81 @@ namespace AutomationTool.Tasks
 				}
 			}
 		}
-	}	
+	}
+
+	public class CsCompileOutput
+	{
+		public static CsCompileOutput Empty { get; } = new CsCompileOutput(FileSet.Empty, FileSet.Empty);
+
+		public FileSet Binaries { get; }
+		public FileSet References { get; }
+
+		public CsCompileOutput(FileSet Binaries, FileSet References)
+		{
+			this.Binaries = Binaries;
+			this.References = References;
+		}
+
+		public FileSet Merge()
+		{
+			return Binaries + References;
+		}
+
+		public static CsCompileOutput operator +(CsCompileOutput Lhs, CsCompileOutput Rhs)
+		{
+			return new CsCompileOutput(Lhs.Binaries + Rhs.Binaries, Lhs.References + Rhs.References);
+		}
+	}
+
+	public static class CsCompileOutputExtensions
+	{
+		public static async Task<FileSet> MergeAsync(this Task<CsCompileOutput> Task)
+		{
+			return (await Task).Merge();
+		}
+	}
+
+	public static partial class StandardTasks
+	{
+		/// <summary>
+		/// Compile a C# project
+		/// </summary>
+		/// <param name="Project">The C# project files to compile.</param>
+		/// <param name="Platform">The platform to compile.</param>
+		/// <param name="Configuration">The configuration to compile.</param>
+		/// <param name="Target">The target to build.</param>
+		/// <param name="Properties">Properties for the command.</param>
+		/// <param name="Arguments">Additional options to pass to the compiler.</param>
+		/// <param name="EnumerateOnly">Only enumerate build products -- do not actually compile the projects.</param>
+		public static async Task<CsCompileOutput> CsCompileAsync(FileReference Project, string Platform = null, string Configuration = null, string Target = null, string Properties = null, string Arguments = null, bool? EnumerateOnly = null)
+		{
+			CsCompileTaskParameters Parameters = new CsCompileTaskParameters();
+			Parameters.Project = Project.FullName;
+			Parameters.Platform = Platform;
+			Parameters.Configuration = Configuration;
+			Parameters.Target = Target;
+			Parameters.Properties = Properties;
+			Parameters.Arguments = Arguments;
+			Parameters.EnumerateOnly = EnumerateOnly ?? Parameters.EnumerateOnly;
+			Parameters.Tag = "#Out";
+			Parameters.TagReferences = "#Refs";
+
+			HashSet<FileReference> BuildProducts = new HashSet<FileReference>();
+			Dictionary<string, HashSet<FileReference>> TagNameToFileSet = new Dictionary<string, HashSet<FileReference>>();
+			await new CsCompileTask(Parameters).ExecuteAsync(new JobContext(null!), BuildProducts, TagNameToFileSet);
+
+			FileSet Binaries = FileSet.Empty;
+			FileSet References = FileSet.Empty;
+			if (TagNameToFileSet.TryGetValue(Parameters.Tag, out HashSet<FileReference> BinaryFiles))
+			{
+				Binaries = FileSet.FromFiles(Unreal.RootDirectory, BinaryFiles);
+			}
+			if (TagNameToFileSet.TryGetValue(Parameters.TagReferences, out HashSet<FileReference> ReferenceFiles))
+			{
+				References = FileSet.FromFiles(Unreal.RootDirectory, ReferenceFiles);
+			}
+
+			return new CsCompileOutput(Binaries, References);
+		}
+	}
 }
diff --git a/Engine/Source/Programs/AutomationTool/BuildGraph/Tasks/DotNetTask.cs b/Engine/Source/Programs/AutomationTool/BuildGraph/Tasks/DotNetTask.cs
--- a/Engine/Source/Programs/AutomationTool/BuildGraph/Tasks/DotNetTask.cs
+++ b/Engine/Source/Programs/AutomationTool/BuildGraph/Tasks/DotNetTask.cs
@@ -1,6 +1,8 @@
 // Copyright Epic Games, Inc. All Rights Reserved.
 
 using AutomationTool;
+using AutomationTool.Tasks;
+using EpicGames.BuildGraph;
 using EpicGames.Core;
 using System;
 using System.Collections.Generic;
@@ -120,4 +122,29 @@ namespace AutomationTool.Tasks
 			yield break;
 		}
 	}
+
+	public static partial class StandardTasks
+	{
+		/// <summary>
+		/// Runs a command using dotnet.
+		/// </summary>
+		/// <param name="Arguments">Command-line arguments.</param>
+		/// <param name="BaseDir">Base directory for running the command.</param>
+		/// <param name="Environment">Environment variables to set.</param>
+		/// <param name="EnvironmentFile">File to read environment variables from.</param>
+		/// <param name="ErrorLevel">The minimum exit code, which is treated as an error.</param>
+		/// <param name="DotNetPath">Override path to dotnet executable.</param>
+		public static async Task DotNetAsync(string Arguments = null, DirectoryReference BaseDir = null, string Environment = null, FileReference EnvironmentFile = null, int ErrorLevel = 1, FileReference DotNetPath = null)
+		{
+			DotNetTaskParameters Parameters = new DotNetTaskParameters();
+			Parameters.Arguments = Arguments;
+			Parameters.BaseDir = BaseDir?.FullName;
+			Parameters.Environment = Environment;
+			Parameters.EnvironmentFile = EnvironmentFile?.FullName;
+			Parameters.ErrorLevel = ErrorLevel;
+			Parameters.DotNetPath = DotNetPath;
+			
+			await ExecuteAsync(new DotNetTask(Parameters));
+		}
+	}
 }
diff --git a/Engine/Source/Programs/AutomationTool/BuildGraph/Tasks/SanitizeReceiptTask.cs b/Engine/Source/Programs/AutomationTool/BuildGraph/Tasks/SanitizeReceiptTask.cs
--- a/Engine/Source/Programs/AutomationTool/BuildGraph/Tasks/SanitizeReceiptTask.cs
+++ b/Engine/Source/Programs/AutomationTool/BuildGraph/Tasks/SanitizeReceiptTask.cs
@@ -57,13 +57,20 @@ namespace AutomationTool.Tasks
 		/// <param name="Job">Information about the current job</param>
 		/// <param name="BuildProducts">Set of build products produced by this node.</param>
 		/// <param name="TagNameToFileSet">Mapping from tag names to the set of files they include</param>
-		public override Task ExecuteAsync(JobContext Job, HashSet<FileReference> BuildProducts, Dictionary<string, HashSet<FileReference>> TagNameToFileSet)
+		public override async Task ExecuteAsync(JobContext Job, HashSet<FileReference> BuildProducts, Dictionary<string, HashSet<FileReference>> TagNameToFileSet)
 		{
 			// Set the Engine directory
 			DirectoryReference EngineDir = Parameters.EngineDir ?? Unreal.EngineDirectory;
 
 			// Resolve the input list
 			IEnumerable<FileReference> TargetFiles = ResolveFilespec(Unreal.RootDirectory, Parameters.Files, TagNameToFileSet);
+			await ExecuteAsync(TargetFiles, EngineDir);
+		}
+
+		public static Task ExecuteAsync(IEnumerable<FileReference> TargetFiles, DirectoryReference EngineDir)
+		{
+			EngineDir ??= Unreal.EngineDirectory;
+
 			foreach(FileReference TargetFile in TargetFiles)
 			{
 				// check all files are .target files
@@ -147,4 +154,18 @@ namespace AutomationTool.Tasks
 			return new string[0];
 		}
 	}
+
+	/// <summary>
+	/// Extension methods
+	/// </summary>
+	public static class SanitizeReceiptExtensions
+	{
+		/// <summary>
+		/// Sanitize the given receipt files, removing any files that don't exist in the current workspace
+		/// </summary>
+		public static async Task SanitizeReceiptsAsync(this FileSet TargetFiles, DirectoryReference EngineDir = null)
+		{
+			await SanitizeReceiptTask.ExecuteAsync(TargetFiles, EngineDir);
+		}
+	}
 }
diff --git a/Engine/Source/Programs/AutomationTool/BuildGraph/Tasks/SetVersionTask.cs b/Engine/Source/Programs/AutomationTool/BuildGraph/Tasks/SetVersionTask.cs
--- a/Engine/Source/Programs/AutomationTool/BuildGraph/Tasks/SetVersionTask.cs
+++ b/Engine/Source/Programs/AutomationTool/BuildGraph/Tasks/SetVersionTask.cs
@@ -9,6 +9,7 @@ using System.Threading.Tasks;
 using System.Xml;
 using EpicGames.Core;
 using UnrealBuildTool;
+using UnrealBuildBase;
 
 namespace AutomationTool.Tasks
 {
@@ -134,4 +135,27 @@ namespace AutomationTool.Tasks
 			return FindTagNamesFromList(Parameters.Tag);
 		}
 	}
+
+	public static partial class StandardTasks
+	{
+		public static async Task<FileSet> ExecuteAsync(BgTaskImpl Task)
+		{
+			HashSet<FileReference> BuildProducts = new HashSet<FileReference>();
+			await Task.ExecuteAsync(new JobContext(null!), BuildProducts, new Dictionary<string, HashSet<FileReference>>());
+			return FileSet.FromFiles(Unreal.RootDirectory, BuildProducts);
+		}
+
+		public static async Task<FileSet> SetVersionAsync(int Change, string Branch, int? CompatibleChange = null, string Build = null, bool? Licensee = null, bool? Promoted = null, bool? SkipWrite = null)
+		{
+			SetVersionTaskParameters Parameters = new SetVersionTaskParameters();
+			Parameters.Change = Change;
+			Parameters.CompatibleChange = CompatibleChange ?? Parameters.CompatibleChange;
+			Parameters.Branch = Branch ?? Parameters.Branch;
+			Parameters.Branch = Build ?? Parameters.Build;
+			Parameters.Licensee = Licensee ?? Parameters.Licensee;
+			Parameters.Promoted = Promoted ?? Parameters.Promoted;
+			Parameters.SkipWrite = SkipWrite ?? Parameters.SkipWrite;
+			return await ExecuteAsync(new SetVersionTask(Parameters));
+		}
+	}
 }
diff --git a/Engine/Source/Programs/AutomationTool/BuildGraph/Tasks/SpawnTask.cs b/Engine/Source/Programs/AutomationTool/BuildGraph/Tasks/SpawnTask.cs
--- a/Engine/Source/Programs/AutomationTool/BuildGraph/Tasks/SpawnTask.cs
+++ b/Engine/Source/Programs/AutomationTool/BuildGraph/Tasks/SpawnTask.cs
@@ -10,6 +10,8 @@ using System.Threading.Tasks;
 using System.Xml;
 using EpicGames.Core;
 using UnrealBuildTool;
+using EpicGames.BuildGraph;
+using AutomationTool.Tasks;
 
 namespace AutomationTool.Tasks
 {
@@ -199,4 +201,32 @@ namespace AutomationTool.Tasks
 			yield break;
 		}
 	}
+
+	public static partial class StandardTasks
+	{
+		/// <summary>
+		/// Execute an external program
+		/// </summary>
+		/// <param name="State">Execution state</param>
+		/// <param name="Exe">Executable to spawn.</param>
+		/// <param name="Arguments">Arguments for the newly created process.</param>
+		/// <param name="WorkingDir">Working directory for spawning the new task.</param>
+		/// <param name="Environment">Environment variables to set.</param>
+		/// <param name="EnvironmentFile">File to read environment from.</param>
+		/// <param name="LogOutput">Write output to the log.</param>
+		/// <param name="ErrorLevel">The minimum exit code which is treated as an error.</param>
+		public static async Task SpawnAsync(string Exe, string Arguments = null, string WorkingDir = null, string Environment = null, string EnvironmentFile = null, bool? LogOutput = null, int? ErrorLevel = null)
+		{
+			SpawnTaskParameters Parameters = new SpawnTaskParameters();
+			Parameters.Exe = Exe;
+			Parameters.Arguments = Arguments;
+			Parameters.WorkingDir = WorkingDir;
+			Parameters.Environment = Environment;
+			Parameters.EnvironmentFile = EnvironmentFile;
+			Parameters.LogOutput = LogOutput ?? Parameters.LogOutput;
+			Parameters.ErrorLevel = ErrorLevel ?? Parameters.ErrorLevel;
+
+			await ExecuteAsync(new SpawnTask(Parameters));
+		}
+	}
 }
diff --git a/Engine/Source/Programs/AutomationTool/BuildGraph/Tasks/StripTask.cs b/Engine/Source/Programs/AutomationTool/BuildGraph/Tasks/StripTask.cs
--- a/Engine/Source/Programs/AutomationTool/BuildGraph/Tasks/StripTask.cs
+++ b/Engine/Source/Programs/AutomationTool/BuildGraph/Tasks/StripTask.cs
@@ -1,6 +1,7 @@
 // Copyright Epic Games, Inc. All Rights Reserved.
 
 using AutomationTool;
+using AutomationTool.Tasks;
 using EpicGames.BuildGraph;
 using System;
 using System.Collections.Generic;
@@ -142,4 +143,25 @@ namespace AutomationTool.Tasks
 			return FindTagNamesFromList(Parameters.Tag);
 		}
 	}
+
+	public static partial class StandardTasks
+	{
+		/// <summary>
+		/// Strips symbols from a set of files.
+		/// </summary>
+		/// <param name="Files"></param>
+		/// <param name="Platform">The platform toolchain to strip binaries.</param>
+		/// <param name="BaseDir">The directory to find files in.</param>
+		/// <param name="OutputDir">Output directory for the stripped files. Defaults to the input path, overwriting the input files.</param>
+		/// <returns></returns>
+		public static async Task<FileSet> StripAsync(FileSet Files, UnrealTargetPlatform Platform, DirectoryReference BaseDir = null, DirectoryReference OutputDir = null)
+		{
+			StripTaskParameters Parameters = new StripTaskParameters();
+			Parameters.Platform = Platform;
+			Parameters.BaseDir = BaseDir;
+			Parameters.Files = String.Join(";", Files.Flatten().Values.Select(x => x.FullName));
+			Parameters.OutputDir = OutputDir;
+			return await ExecuteAsync(new StripTask(Parameters));
+		}
+	}
 }
diff --git a/Engine/Source/Programs/AutomationTool/BuildGraph/Tasks/TagReceiptTask.cs b/Engine/Source/Programs/AutomationTool/BuildGraph/Tasks/TagReceiptTask.cs
--- a/Engine/Source/Programs/AutomationTool/BuildGraph/Tasks/TagReceiptTask.cs
+++ b/Engine/Source/Programs/AutomationTool/BuildGraph/Tasks/TagReceiptTask.cs
@@ -113,10 +113,10 @@ namespace AutomationTool.Tasks
 		/// <param name="Job">Information about the current job</param>
 		/// <param name="BuildProducts">Set of build products produced by this node.</param>
 		/// <param name="TagNameToFileSet">Mapping from tag names to the set of files they include</param>
-		public override Task ExecuteAsync(JobContext Job, HashSet<FileReference> BuildProducts, Dictionary<string, HashSet<FileReference>> TagNameToFileSet)
+		public override async Task ExecuteAsync(JobContext Job, HashSet<FileReference> BuildProducts, Dictionary<string, HashSet<FileReference>> TagNameToFileSet)
 		{
 			// Output a warning if the project directory is specified
-			if(Parameters.ProjectDir != null)
+			if (Parameters.ProjectDir != null)
 			{
 				CommandUtils.LogWarning("The ProjectDir argument to the TagReceipt parameter is deprecated. This path is now determined automatically from the receipt.");
 			}
@@ -126,6 +126,16 @@ namespace AutomationTool.Tasks
 
 			// Resolve the input list
 			IEnumerable<FileReference> TargetFiles = ResolveFilespec(Unreal.RootDirectory, Parameters.Files, TagNameToFileSet);
+
+			// Filter the files
+			HashSet<FileReference> Files = await ExecuteAsync(EngineDir, TargetFiles, Parameters.BuildProducts, BuildProductType, Parameters.RuntimeDependencies, StagedFileType);
+
+			// Apply the tag to all the matching files
+			FindOrAddTagSet(TagNameToFileSet, Parameters.With).UnionWith(Files);
+		}
+
+		public static Task<HashSet<FileReference>> ExecuteAsync(DirectoryReference EngineDir, IEnumerable<FileReference> TargetFiles, bool BuildProducts, BuildProductType? BuildProductType, bool RuntimeDependencies, StagedFileType? StagedFileType = null)
+		{
 			HashSet<FileReference> Files = new HashSet<FileReference>();
 
 			foreach (FileReference TargetFile in TargetFiles)
@@ -144,7 +154,7 @@ namespace AutomationTool.Tasks
 					continue;
 				}
 
-				if (Parameters.BuildProducts)
+				if (BuildProducts)
 				{
 					foreach (BuildProduct BuildProduct in Receipt.BuildProducts)
 					{
@@ -160,7 +170,7 @@ namespace AutomationTool.Tasks
 					}
 				}
 
-				if (Parameters.RuntimeDependencies)
+				if (RuntimeDependencies)
 				{
 					foreach (RuntimeDependency RuntimeDependency in Receipt.RuntimeDependencies)
 					{
@@ -188,9 +198,7 @@ namespace AutomationTool.Tasks
 				}
 			}
 
-			// Apply the tag to all the matching files
-			FindOrAddTagSet(TagNameToFileSet, Parameters.With).UnionWith(Files);
-			return Task.CompletedTask;
+			return Task.FromResult(Files);
 		}
 
 		/// <summary>
@@ -219,4 +227,13 @@ namespace AutomationTool.Tasks
 			return FindTagNamesFromList(Parameters.With);
 		}
 	}
+
+	public static class TaskExtensions
+	{
+		public static async Task<FileSet> TagReceiptsAsync(this FileSet Files, DirectoryReference EngineDir = null, bool BuildProducts = false, BuildProductType? BuildProductType = null, bool RuntimeDependencies = false, StagedFileType? StagedFileType = null)
+		{
+			HashSet<FileReference> Result = await TagReceiptTask.ExecuteAsync(EngineDir ?? Unreal.EngineDirectory, Files, BuildProducts, BuildProductType, RuntimeDependencies, StagedFileType);
+			return FileSet.FromFiles(Unreal.RootDirectory, Result);
+		}
+	}
 }
diff --git a/Engine/Source/Programs/AutomationTool/Win/Tasks/SrcSrvTask.cs b/Engine/Source/Programs/AutomationTool/Win/Tasks/SrcSrvTask.cs
--- a/Engine/Source/Programs/AutomationTool/Win/Tasks/SrcSrvTask.cs
+++ b/Engine/Source/Programs/AutomationTool/Win/Tasks/SrcSrvTask.cs
@@ -15,6 +15,8 @@ using System.Diagnostics;
 using Action = System.Action;
 using EpicGames.Core;
 using UnrealBuildBase;
+using EpicGames.BuildGraph;
+using AutomationTool.Tasks;
 
 namespace Win.Automation
 {
@@ -76,11 +78,14 @@ namespace Win.Automation
 		/// <param name="TagNameToFileSet">Mapping from tag names to the set of files they include</param>
 		public override void Execute(JobContext Job, HashSet<FileReference> BuildProducts, Dictionary<string, HashSet<FileReference>> TagNameToFileSet)
 		{
-			// Find the matching files
-			FileReference[] PdbFiles = ResolveFilespec(Unreal.RootDirectory, Parameters.BinaryFiles, TagNameToFileSet).Where(x => x.HasExtension(".pdb")).ToArray();
-
-			// Find all the matching source files
+			FileReference[] BinaryFiles = ResolveFilespec(Unreal.RootDirectory, Parameters.BinaryFiles, TagNameToFileSet).ToArray();
 			FileReference[] SourceFiles = ResolveFilespec(Unreal.RootDirectory, Parameters.SourceFiles, TagNameToFileSet).ToArray();
+			Execute(BinaryFiles, SourceFiles, Parameters.Branch, Parameters.Change);
+		}
+
+		internal static void Execute(FileReference[] BinaryFiles, FileReference[] SourceFiles, string Branch, int Change)
+		{
+			FileReference[] PdbFiles = BinaryFiles.Where(x => x.HasExtension(".pdb")).ToArray();
 
 			// Get the PDBSTR.EXE path, using the latest SDK version we can find.
 			FileReference PdbStrExe = GetPdbStrExe();
@@ -98,8 +103,8 @@ namespace Win.Automation
 				Writer.WriteLine("SRCSRV: variables------------------------------------------");
 				Writer.WriteLine("SRCSRVTRG=%sdtrg%");
 				Writer.WriteLine("SRCSRVCMD=%sdcmd%");
-				Writer.WriteLine("SDCMD=p4.exe print -o %srcsrvtrg% \"{0}/%var2%@{1}\"", Parameters.Branch.TrimEnd('/'), Parameters.Change);
-				Writer.WriteLine("SDTRG=%targ%\\{0}\\{1}\\%fnbksl%(%var2%)", Parameters.Branch.Replace('/', '+'), Parameters.Change);
+				Writer.WriteLine("SDCMD=p4.exe print -o %srcsrvtrg% \"{0}/%var2%@{1}\"", Branch.TrimEnd('/'), Change);
+				Writer.WriteLine("SDTRG=%targ%\\{0}\\{1}\\%fnbksl%(%var2%)", Branch.Replace('/', '+'), Change);
 				Writer.WriteLine("SRCSRV: source files ---------------------------------------");
 				foreach (FileReference SourceFile in SourceFiles)
 				{
@@ -121,7 +126,7 @@ namespace Win.Automation
         /// <param name="SrcSrvIni">Ini file containing settings to embed</param>
         /// <param name="State">The current loop state</param>
         /// <returns>True if the tool executed succesfully</returns>
-        void ExecuteTool(FileReference PdbStrExe, FileReference PdbFile, FileReference SrcSrvIni, ParallelLoopState State)
+        static void ExecuteTool(FileReference PdbStrExe, FileReference PdbFile, FileReference SrcSrvIni, ParallelLoopState State)
         {
             using (Process Process = new Process())
             {
@@ -149,7 +154,7 @@ namespace Win.Automation
                     State.Break();
                 }
 
-                lock (this)
+                lock (State)
                 {
                     foreach (string Message in Messages)
                     {
@@ -221,4 +226,20 @@ namespace Win.Automation
 			throw new AutomationException("Unable to find a Windows SDK installation containing PDBSTR.EXE");
 		}
 	}
+
+	public static partial class BgStateExtensions
+	{
+		/// <summary>
+		/// Uploads symbols to a symbol server
+		/// </summary>
+		/// <param name="BinaryFiles">List of output files. PDBs will be extracted from this list.</param>
+		/// <param name="SourceFiles">List of source files to index and embed into the PDBs.</param>
+		/// <param name="Branch">Branch to base all the depot source files from.</param>
+		/// <param name="Change">Changelist to sync files from.</param>
+		/// <returns></returns>
+		public static void SrcSrv(this BgContext State, HashSet<FileReference> BinaryFiles, HashSet<FileReference> SourceFiles, string Branch, int Change)
+		{
+			SrcSrvTask.Execute(BinaryFiles.ToArray(), SourceFiles.ToArray(), Branch, Change);
+		}
+	}
 }
diff --git a/Engine/Source/Programs/Shared/EpicGames.BuildGraph/BgAgentSpec.cs b/Engine/Source/Programs/Shared/EpicGames.BuildGraph/BgAgentSpec.cs
new file mode 100644
--- /dev/null
+++ b/Engine/Source/Programs/Shared/EpicGames.BuildGraph/BgAgentSpec.cs
@@ -0,0 +1,159 @@
+// Copyright Epic Games, Inc. All Rights Reserved.
+
+using EpicGames.BuildGraph.Expressions;
+using Microsoft.Extensions.Logging;
+using System;
+using System.Collections.Generic;
+using System.Linq;
+using System.Linq.Expressions;
+using System.Reflection;
+using System.Text;
+using System.Threading.Tasks;
+
+namespace EpicGames.BuildGraph
+{
+	/// <summary>
+	/// Configuration object for agents
+	/// </summary>
+	public class BgAgentConfig
+	{
+		/// <summary>
+		/// List of agent types to select from, in order of preference. The first agent type supported by a stream will be used.
+		/// </summary>
+		public BgList<BgString> Types { get; set; } = BgList<BgString>.Empty;
+
+		/// <summary>
+		/// List of diagnostics for this agent to run
+		/// </summary>
+		public BgList<BgDiagnosticSpec> Diagnostics { get; set; } = BgList<BgDiagnosticSpec>.Empty;
+
+		/// <summary>
+		/// Add agent types to the agent definition
+		/// </summary>
+		public BgAgentConfig AddTypes(params BgString[] TypeNames)
+		{
+			Types = Types.Add(TypeNames);
+			return this;
+		}
+
+		/// <summary>
+		/// Check that a certain condition is true before executing the node
+		/// </summary>
+		public BgAgentConfig Check(BgBool Condition, LogLevel Level, BgString Message)
+		{
+			BgDiagnosticSpec Diagnostic = new BgDiagnosticSpec(Level, Message);
+			Diagnostics = Diagnostics.AddIf(Condition, Diagnostic);
+			return this;
+		}
+	}
+
+	/// <summary>
+	/// Describes an agent that can execute execute build steps
+	/// </summary>
+	public class BgAgentSpec
+	{
+		/// <summary>
+		/// Name of the agent
+		/// </summary>
+		public BgString Name { get; }
+
+		/// <summary>
+		/// List of agent types to select from, in order of preference. The first agent type supported by a stream will be used.
+		/// </summary>
+		public BgList<BgString> TypeNames { get; set; } = BgList<BgString>.Empty;
+
+		/// <summary>
+		/// List of diagnostics for this agent to run
+		/// </summary>
+		public BgList<BgDiagnosticSpec> Diagnostics { get; set; } = BgList<BgDiagnosticSpec>.Empty;
+
+		/// <summary>
+		/// Nodes that have been added to the agent
+		/// </summary>
+		List<BgNodeSpec> NodeSpecs { get; set; } = new List<BgNodeSpec>();
+
+		/// <summary>
+		/// Private constructor. Use <see cref="BgGraphSpec.AddAgent(BgString)"/> to construct an agent
+		/// </summary>
+		/// <param name="Name"></param>
+		internal BgAgentSpec(BgString Name)
+		{
+			this.Name = Name;
+		}
+
+		/// <summary>
+		/// Adds a node to be executed on this agent
+		/// </summary>
+		/// <param name="Function">Lambda expression which calls a static method to be used for this node</param>
+		/// <returns>New node spec instance</returns>
+		public BgNodeSpec AddNode(Expression<Func<BgContext, Task>> Function)
+		{
+			BgNodeSpec NodeSpec = BgNodeSpec.Create(Function);
+			NodeSpecs.Add(NodeSpec);
+			return NodeSpec;
+		}
+
+		/// <summary>
+		/// Adds a node to be executed on this agent
+		/// </summary>
+		/// <param name="Function">Lambda expression which calls a static method to be used for this node</param>
+		/// <returns>New node spec instance</returns>
+		public BgNodeSpec<T> AddNode<T>(Expression<Func<BgContext, Task<T>>> Function)
+		{
+			BgNodeSpec<T> NodeSpec = BgNodeSpec.Create<T>(Function);
+			NodeSpecs.Add(NodeSpec);
+			return NodeSpec;
+		}
+
+		/// <summary>
+		/// Add agent types to the agent definition
+		/// </summary>
+		public BgAgentSpec Type(params BgString[] NewTypeNames)
+		{
+			TypeNames = TypeNames.Add(NewTypeNames);
+			return this;
+		}
+
+		/// <summary>
+		/// Check that a certain condition is true before executing the node
+		/// </summary>
+		public BgAgentSpec WarnIf(BgBool Condition, BgString Message)
+		{
+			BgDiagnosticSpec Diagnostic = new BgDiagnosticSpec(LogLevel.Warning, Message);
+			Diagnostics = Diagnostics.AddIf(Condition, Diagnostic);
+			return this;
+		}
+
+		/// <summary>
+		/// Check that a certain condition is true before executing the node
+		/// </summary>
+		public BgAgentSpec ErrorIf(BgBool Condition, BgString Message)
+		{
+			BgDiagnosticSpec Diagnostic = new BgDiagnosticSpec(LogLevel.Error, Message);
+			Diagnostics = Diagnostics.AddIf(Condition, Diagnostic);
+			return this;
+		}
+
+		/// <summary>
+		/// Creates a concrete <see cref="BgAgent"/> object from the specification
+		/// </summary>
+		public void AddToGraph(BgExprContext Context, BgGraph Graph)
+		{
+			string[] Types = TypeNames.Compute(Context).ToArray();
+
+			BgAgent Agent = new BgAgent(Name.Compute(Context), Types);
+			Graph.NameToAgent.Add(Agent.Name, Agent);
+			Graph.Agents.Add(Agent);
+
+			foreach (BgDiagnosticSpec Precondition in Diagnostics.GetEnumerable(Context))
+			{
+				Precondition.AddToGraph(Context, Graph, Agent, null);
+			}
+
+			foreach (BgNodeSpec NodeSpec in NodeSpecs)
+			{
+				NodeSpec.AddToGraph(Context, Graph, Agent);
+			}
+		}
+	}
+}
diff --git a/Engine/Source/Programs/Shared/EpicGames.BuildGraph/BgAggregateSpec.cs b/Engine/Source/Programs/Shared/EpicGames.BuildGraph/BgAggregateSpec.cs
new file mode 100644
--- /dev/null
+++ b/Engine/Source/Programs/Shared/EpicGames.BuildGraph/BgAggregateSpec.cs
@@ -0,0 +1,54 @@
+// Copyright Epic Games, Inc. All Rights Reserved.
+
+using EpicGames.BuildGraph.Expressions;
+using System;
+using System.Collections.Generic;
+using System.Linq;
+using System.Text;
+
+namespace EpicGames.BuildGraph
+{
+	/// <summary>
+	/// Specification for an aggregate target in the graph
+	/// </summary>
+	public class BgAggregateSpec
+	{
+		/// <summary>
+		/// Name of the aggregate
+		/// </summary>
+		public BgString Name { get; }
+
+		/// <summary>
+		/// Outputs required for the aggregate
+		/// </summary>
+		public BgList<BgFileSet> RequiredOutputs { get; set; }
+
+		/// <summary>
+		/// Internal constructor. Use <see cref="BgGraphSpec.AddAggregate(BgString, BgList{BgFileSet})"/> to create an aggregate.
+		/// </summary>
+		internal BgAggregateSpec(BgString Name, BgList<BgFileSet> RequiredOutputs)
+		{
+			this.Name = Name;
+			this.RequiredOutputs = RequiredOutputs;
+		}
+
+		/// <summary>
+		/// Creates a concrete aggregate object from this specification.
+		/// </summary>
+		internal void AddToGraph(BgExprContext Context, BgGraph Graph)
+		{
+			BgAggregate Aggregate = new BgAggregate(Name.Compute(Context));
+			Aggregate.RequiredNodes.UnionWith(RequiredOutputs.ComputeTags(Context).Select(x => Graph.TagNameToNodeOutput[x].ProducingNode));
+			Graph.NameToAggregate.Add(Aggregate.Name, Aggregate);
+		}
+
+		/// <summary>
+		/// Adds a set of dependencies to this aggregate
+		/// </summary>
+		/// <param name="Tokens">List of token dependencies</param>
+		public void Requires(params BgFileSet[] Tokens)
+		{
+			RequiredOutputs.Add(Tokens);
+		}
+	}
+}
diff --git a/Engine/Source/Programs/Shared/EpicGames.BuildGraph/BgContext.cs b/Engine/Source/Programs/Shared/EpicGames.BuildGraph/BgContext.cs
new file mode 100644
--- /dev/null
+++ b/Engine/Source/Programs/Shared/EpicGames.BuildGraph/BgContext.cs
@@ -0,0 +1,141 @@
+// Copyright Epic Games, Inc. All Rights Reserved.
+
+using EpicGames.BuildGraph.Expressions;
+using EpicGames.Core;
+using System;
+using System.Collections.Generic;
+using System.Linq;
+using System.Text;
+
+namespace EpicGames.BuildGraph
+{
+	/// <summary>
+	/// Context for executing a node method
+	/// </summary>
+	public abstract class BgContext
+	{
+		/// <summary>
+		/// The stream executing the current build
+		/// </summary>
+		public abstract string Stream { get; }
+
+		/// <summary>
+		/// Changelist being built
+		/// </summary>
+		public abstract int Change { get; }
+
+		/// <summary>
+		/// The code changelist currently being built
+		/// </summary>
+		public abstract int CodeChange { get; }
+
+		/// <summary>
+		/// Version number for the engine
+		/// </summary>
+		public abstract (int Major, int Minor, int Patch) EngineVersion { get; }
+
+		/// <summary>
+		/// Whether this machine is a builder
+		/// </summary>
+		public abstract bool IsBuildMachine { get; }
+
+		/// <summary>
+		/// Context for evaluating expressions
+		/// </summary>
+		public BgExprContext Context { get; }
+
+		/// <summary>
+		/// All outputs for the node
+		/// </summary>
+		public HashSet<FileReference> BuildProducts = new HashSet<FileReference>();
+
+		/// <summary>
+		/// Constructor
+		/// </summary>
+		public BgContext(BgExprContext Context)
+		{
+			this.Context = Context;
+		}
+
+		/// <summary>
+		/// Resolve a boolean expression to a value
+		/// </summary>
+		/// <param name="Bool">The boolean expression</param>
+		/// <returns>Value of the expression</returns>
+		public bool Get(BgBool Bool) => Bool.Compute(Context);
+
+		/// <summary>
+		/// Resolve an integer expression to a value
+		/// </summary>
+		/// <param name="Integer">The integer expression</param>
+		/// <returns>Value of the expression</returns>
+		public int Get(BgInt Integer) => Integer.Compute(Context);
+
+		/// <summary>
+		/// Resolve a string expression to a value
+		/// </summary>
+		/// <param name="String">The string expression</param>
+		/// <returns>Value of the expression</returns>
+		public string Get(BgString String) => String.Compute(Context);
+
+		/// <summary>
+		/// Resolves an enum expression to a value
+		/// </summary>
+		/// <typeparam name="TEnum">The enum type</typeparam>
+		/// <param name="Enum">Enum expression</param>
+		/// <returns>The enum value</returns>
+		public TEnum Get<TEnum>(BgEnum<TEnum> Enum) where TEnum : struct => Enum.Compute(Context);
+
+		/// <summary>
+		/// Resolve a list of enums to a value
+		/// </summary>
+		/// <typeparam name="TEnum">The enum type</typeparam>
+		/// <param name="List">Enum expression</param>
+		/// <returns>The enum value</returns>
+		public List<TEnum> Get<TEnum>(BgList<BgEnum<TEnum>> List) where TEnum : struct => List.GetEnumerable(Context).Select(x => Get(x)).ToList();
+
+		/// <summary>
+		/// Resolve a list of strings
+		/// </summary>
+		/// <param name="List">List expression</param>
+		/// <returns></returns>
+		public List<string> Get(BgList<BgString> List) => List.Compute(Context);
+
+		/// <summary>
+		/// Resolve a file set
+		/// </summary>
+		/// <param name="FileSet">The token expression</param>
+		/// <returns>Set of files for the token</returns>
+		public FileSet Get(BgFileSet FileSet) => FileSet.ComputeValue(Context);
+
+		/// <summary>
+		/// Resolve a file set
+		/// </summary>
+		/// <param name="FileSets">The token expression</param>
+		/// <returns>Set of files for the token</returns>
+		public FileSet Get(BgList<BgFileSet> FileSets)
+		{
+			FileSet Result = FileSet.Empty;
+			foreach (BgFileSet FileSet in FileSets.GetEnumerable(Context))
+			{
+				Result += Get(FileSet);
+			}
+			return Result;
+		}
+
+		/// <summary>
+		/// Resolve a file set
+		/// </summary>
+		/// <param name="FileSets">The token expression</param>
+		/// <returns>Set of files for the token</returns>
+		public FileSet Get(IEnumerable<BgFileSet> FileSets)
+		{
+			FileSet Result = FileSet.Empty;
+			foreach (BgFileSet FileSet in FileSets)
+			{
+				Result += FileSet.ComputeValue(Context);
+			}
+			return Result;
+		}
+	}
+}
diff --git a/Engine/Source/Programs/Shared/EpicGames.BuildGraph/BgDiagnostic.cs b/Engine/Source/Programs/Shared/EpicGames.BuildGraph/BgDiagnostic.cs
new file mode 100644
--- /dev/null
+++ b/Engine/Source/Programs/Shared/EpicGames.BuildGraph/BgDiagnostic.cs
@@ -0,0 +1,53 @@
+// Copyright Epic Games, Inc. All Rights Reserved.
+
+using EpicGames.Core;
+using System;
+using System.Collections.Generic;
+using System.Text;
+
+namespace EpicGames.BuildGraph
+{
+	/// <summary>
+	/// Diagnostic message from the graph script. These messages are parsed at startup, then culled along with the rest of the graph nodes before output. Doing so
+	/// allows errors and warnings which are only output if a node is part of the graph being executed.
+	/// </summary>
+	public class BgGraphDiagnostic
+	{
+		/// <summary>
+		/// Location of the diagnostic
+		/// </summary>
+		public BgScriptLocation Location;
+
+		/// <summary>
+		/// The diagnostic event type
+		/// </summary>
+		public LogEventType EventType;
+
+		/// <summary>
+		/// The message to display
+		/// </summary>
+		public string Message;
+
+		/// <summary>
+		/// The node which this diagnostic is declared in. If the node is culled from the graph, the message will not be displayed.
+		/// </summary>
+		public BgNode? EnclosingNode;
+
+		/// <summary>
+		/// The agent that this diagnostic is declared in. If the entire agent is culled from the graph, the message will not be displayed.
+		/// </summary>
+		public BgAgent? EnclosingAgent;
+
+		/// <summary>
+		/// Constructor
+		/// </summary>
+		public BgGraphDiagnostic(BgScriptLocation Location, LogEventType EventType, string Message, BgNode? EnclosingNode, BgAgent? EnclosingAgent)
+		{
+			this.Location = Location;
+			this.EventType = EventType;
+			this.Message = Message;
+			this.EnclosingNode = EnclosingNode;
+			this.EnclosingAgent = EnclosingAgent;
+		}
+	}
+}
diff --git a/Engine/Source/Programs/Shared/EpicGames.BuildGraph/BgDiagnosticSpec.cs b/Engine/Source/Programs/Shared/EpicGames.BuildGraph/BgDiagnosticSpec.cs
new file mode 100644
--- /dev/null
+++ b/Engine/Source/Programs/Shared/EpicGames.BuildGraph/BgDiagnosticSpec.cs
@@ -0,0 +1,57 @@
+// Copyright Epic Games, Inc. All Rights Reserved.
+
+using EpicGames.BuildGraph.Expressions;
+using EpicGames.Core;
+using Microsoft.Extensions.Logging;
+using System;
+using System.Collections.Generic;
+using System.Text;
+
+namespace EpicGames.BuildGraph
+{
+	/// <summary>
+	/// Specification for a diagnostic message
+	/// </summary>
+	public class BgDiagnosticSpec : IBgExpr<BgDiagnosticSpec>
+	{
+		/// <summary>
+		/// Verbosity to output the message at
+		/// </summary>
+		public LogLevel Level;
+
+		/// <summary>
+		/// The message to display
+		/// </summary>
+		public BgString Message;
+
+		/// <summary>
+		/// Constructor
+		/// </summary>
+		/// <param name="Level">Output level</param>
+		/// <param name="Message"></param>
+		public BgDiagnosticSpec(LogLevel Level, BgString Message)
+		{
+			this.Level = Level;
+			this.Message = Message;
+		}
+
+		/// <summary>
+		/// Adds a diagnostic to the graph
+		/// </summary>
+		internal void AddToGraph(BgExprContext Context, BgGraph Graph, BgAgent? EnclosingAgent, BgNode? EnclosingNode)
+		{
+			BgScriptLocation Location = new BgScriptLocation("(unknown)", "(unknown)", 1);
+			string MessageValue = Message.Compute(Context);
+			Graph.Diagnostics.Add(new BgGraphDiagnostic(Location, (LogEventType)Level, MessageValue, EnclosingNode, EnclosingAgent));
+		}
+
+		/// <inheritdoc/>
+		public BgDiagnosticSpec IfThen(BgBool Condition, BgDiagnosticSpec ValueIfTrue) => throw new NotImplementedException();
+
+		/// <inheritdoc/>
+		object IBgExpr.Compute(BgExprContext Context) => this;
+
+		/// <inheritdoc/>
+		public BgString ToBgString() => Message;
+	}
+}
diff --git a/Engine/Source/Programs/Shared/EpicGames.BuildGraph/BgGraph.cs b/Engine/Source/Programs/Shared/EpicGames.BuildGraph/BgGraph.cs
--- a/Engine/Source/Programs/Shared/EpicGames.BuildGraph/BgGraph.cs
+++ b/Engine/Source/Programs/Shared/EpicGames.BuildGraph/BgGraph.cs
@@ -35,93 +35,6 @@ namespace EpicGames.BuildGraph
 		ShowNotifications = 0x4,
 	}
 
-	/// <summary>
-	/// Diagnostic message from the graph script. These messages are parsed at startup, then culled along with the rest of the graph nodes before output. Doing so
-	/// allows errors and warnings which are only output if a node is part of the graph being executed.
-	/// </summary>
-	public class BgGraphDiagnostic
-	{
-		/// <summary>
-		/// Location of the diagnostic
-		/// </summary>
-		public BgScriptLocation Location;
-
-		/// <summary>
-		/// The diagnostic event type
-		/// </summary>
-		public LogEventType EventType;
-
-		/// <summary>
-		/// The message to display
-		/// </summary>
-		public string Message;
-
-		/// <summary>
-		/// The node which this diagnostic is declared in. If the node is culled from the graph, the message will not be displayed.
-		/// </summary>
-		public BgNode? EnclosingNode;
-
-		/// <summary>
-		/// The agent that this diagnostic is declared in. If the entire agent is culled from the graph, the message will not be displayed.
-		/// </summary>
-		public BgAgent? EnclosingAgent;
-
-		/// <summary>
-		/// Constructor
-		/// </summary>
-		public BgGraphDiagnostic(BgScriptLocation Location, LogEventType EventType, string Message, BgNode? EnclosingNode, BgAgent? EnclosingAgent)
-		{
-			this.Location = Location;
-			this.EventType = EventType;
-			this.Message = Message;
-			this.EnclosingNode = EnclosingNode;
-			this.EnclosingAgent = EnclosingAgent;
-		}
-	}
-
-	/// <summary>
-	/// Represents a graph option. These are expanded during preprocessing, but are retained in order to display help messages.
-	/// </summary>
-	public class BgScriptOption
-	{
-		/// <summary>
-		/// Name of this option
-		/// </summary>
-		public string Name;
-
-		/// <summary>
-		/// Description for this option
-		/// </summary>
-		public string Description;
-
-		/// <summary>
-		/// Default value for this option
-		/// </summary>
-		public string DefaultValue;
-
-		/// <summary>
-		/// Constructor
-		/// </summary>
-		/// <param name="Name">The name of this option</param>
-		/// <param name="Description">Description of the option, for display on help pages</param>
-		/// <param name="DefaultValue">Default value for the option</param>
-		public BgScriptOption(string Name, string Description, string DefaultValue)
-		{
-			this.Name = Name;
-			this.Description = Description;
-			this.DefaultValue = DefaultValue;
-		}
-
-		/// <summary>
-		/// Returns a name of this option for debugging
-		/// </summary>
-		/// <returns>Name of the option</returns>
-		public override string ToString()
-		{
-			return Name;
-		}
-	}
-
 	/// <summary>
 	/// Definition of a graph.
 	/// </summary>
@@ -130,7 +43,7 @@ namespace EpicGames.BuildGraph
 		/// <summary>
 		/// List of options, in the order they were specified
 		/// </summary>
-		public List<BgScriptOption> Options = new List<BgScriptOption>();
+		public List<BgOption> Options = new List<BgOption>();
 
 		/// <summary>
 		/// List of agents containing nodes to execute
@@ -288,6 +201,17 @@ namespace EpicGames.BuildGraph
 			return false;
 		}
 
+		static void AddDependencies(BgNode Node, HashSet<BgNode> RetainNodes)
+		{
+			if (RetainNodes.Add(Node))
+			{
+				foreach(BgNode InputDependency in Node.InputDependencies)
+				{
+					AddDependencies(InputDependency, RetainNodes);
+				}
+			}
+		}
+
 		/// <summary>
 		/// Cull the graph to only include the given nodes and their dependencies
 		/// </summary>
@@ -295,10 +219,10 @@ namespace EpicGames.BuildGraph
 		public void Select(IEnumerable<BgNode> TargetNodes)
 		{
 			// Find this node and all its dependencies
-			HashSet<BgNode> RetainNodes = new HashSet<BgNode>(TargetNodes);
+			HashSet<BgNode> RetainNodes = new HashSet<BgNode>();
 			foreach (BgNode TargetNode in TargetNodes)
 			{
-				RetainNodes.UnionWith(TargetNode.InputDependencies);
+				AddDependencies(TargetNode, RetainNodes);
 			}
 
 			// Remove all the nodes which are not marked to be kept
@@ -715,7 +639,7 @@ namespace EpicGames.BuildGraph
 			{
 				// Get the list of messages
 				List<KeyValuePair<string, string>> Parameters = new List<KeyValuePair<string, string>>();
-				foreach (BgScriptOption Option in Options)
+				foreach (BgOption Option in Options)
 				{
 					string Name = String.Format("-set:{0}=...", Option.Name);
 
diff --git a/Engine/Source/Programs/Shared/EpicGames.BuildGraph/BgGraphSpec.cs b/Engine/Source/Programs/Shared/EpicGames.BuildGraph/BgGraphSpec.cs
new file mode 100644
--- /dev/null
+++ b/Engine/Source/Programs/Shared/EpicGames.BuildGraph/BgGraphSpec.cs
@@ -0,0 +1,233 @@
+// Copyright Epic Games, Inc. All Rights Reserved.
+
+using EpicGames.BuildGraph.Expressions;
+using Microsoft.Extensions.Logging;
+using System;
+using System.Collections.Generic;
+using System.Text;
+using System.Linq;
+
+namespace EpicGames.BuildGraph
+{
+	/// <summary>
+	/// Constant values for evaluating the graph
+	/// </summary>
+	public interface IBgGraphConstants
+	{
+		/// <summary>
+		/// The current stream
+		/// </summary>
+		public BgString Stream { get; }
+
+		/// <summary>
+		/// The changelist being built
+		/// </summary>
+		public BgInt Change { get; }
+
+		/// <summary>
+		/// Current code changelist
+		/// </summary>
+		public BgInt CodeChange { get; }
+
+		/// <summary>
+		/// The current engnine version
+		/// </summary>
+		public (BgInt Major, BgInt Minor, BgInt Patch) EngineVersion { get; }
+	}
+
+	/// <summary>
+	/// Specification for a graph in fluent syntax
+	/// </summary>
+	public class BgGraphSpec
+	{
+		/// <summary>
+		/// Constant properties for evaluating the graph
+		/// </summary>
+		public IBgGraphConstants Constants { get; }
+
+		List<IBgOption> Options = new List<IBgOption>();
+		List<BgAgentSpec> AgentSpecs = new List<BgAgentSpec>();
+		List<BgAggregateSpec> AggregateSpecs = new List<BgAggregateSpec>();
+		List<BgLabelSpec> LabelSpecs = new List<BgLabelSpec>();
+		BgList<BgDiagnosticSpec> DiagnosticSpecs = new List<BgDiagnosticSpec>();
+
+		/// <summary>
+		/// Constructor
+		/// </summary>
+		/// <param name="Constants">Constant values for use during graph evaluation</param>
+		public BgGraphSpec(IBgGraphConstants Constants)
+		{
+			this.Constants = Constants;
+		}
+
+		/// <summary>
+		/// Adds a boolean option
+		/// </summary>
+		/// <param name="Name">Name of the option. Must be unique.</param>
+		/// <param name="Description">Description for the option</param>
+		/// <param name="DefaultValue">Default value for the option.</param>
+		/// <returns>Boolean expression</returns>
+		public BgBoolOption AddOption(string Name, BgString Description, BgBool DefaultValue)
+		{
+			BgBoolOption Option = new BgBoolOption(Name, Description, DefaultValue);
+			Options.Add(Option);
+			return Option;
+		}
+
+		/// <summary>
+		/// Adds an integer option
+		/// </summary>
+		/// <param name="Name">Name of the option. Must be unique.</param>
+		/// <param name="Description">Description for the option</param>
+		/// <param name="DefaultValue">Default value for the option.</param>
+		/// <returns>Integer expression</returns>
+		public BgIntOption AddIntOption(string Name, BgString Description, BgInt DefaultValue)
+		{
+			BgIntOption Option = new BgIntOption(Name, Description, DefaultValue);
+			Options.Add(Option);
+			return Option;
+		}
+
+		/// <summary>
+		/// Adds a string option
+		/// </summary>
+		/// <param name="Name">Name of the option. Must be unique.</param>
+		/// <param name="Description">Description for the option</param>
+		/// <param name="DefaultValue">Default value for the option.</param>
+		/// <returns>String expression</returns>
+		public BgStringOption AddOption(string Name, BgString Description, BgString DefaultValue)
+		{
+			BgStringOption Option = new BgStringOption(Name, Description, DefaultValue);
+			Options.Add(Option);
+			return Option;
+		}
+
+		/// <summary>
+		/// Adds an enum option
+		/// </summary>
+		/// <typeparam name="TEnum">Type of enum value</typeparam>
+		/// <param name="Name">Name of the option. Must be unique.</param>
+		/// <param name="Description">Description for the option</param>
+		/// <param name="DefaultValue">Default value for the option.</param>
+		/// <returns>Enum option</returns>
+		public BgEnumOption<TEnum> AddOption<TEnum>(string Name, BgString Description, BgEnum<TEnum> DefaultValue) where TEnum : struct
+		{
+			BgEnumOption<TEnum> Option = new BgEnumOption<TEnum>(Name, Description, DefaultValue);
+			Options.Add(Option);
+			return Option;
+		}
+
+		/// <summary>
+		/// Adds an option which may be a list of enum values
+		/// </summary>
+		/// <typeparam name="TEnum">Type of enum value</typeparam>
+		/// <param name="Name">Name of the option. Must be unique.</param>
+		/// <param name="Description">Description for the option</param>
+		/// <param name="DefaultValue">Default value for the option.</param>
+		/// <returns></returns>
+		public BgEnumListOption<TEnum> AddOption<TEnum>(string Name, BgString Description, BgList<BgEnum<TEnum>> DefaultValue) where TEnum : struct
+		{
+			BgEnumListOption<TEnum> Option = new BgEnumListOption<TEnum>(Name, Description, DefaultValue);
+			Options.Add(Option);
+			return Option;
+		}
+
+		/// <summary>
+		/// Adds a new agent to the graph
+		/// </summary>
+		/// <param name="Name">Name of the agent</param>
+		/// <returns>Agent specification object</returns>
+		public BgAgentSpec AddAgent(BgString Name)
+		{
+			BgAgentSpec Agent = new BgAgentSpec(Name);
+			AgentSpecs.Add(Agent);
+			return Agent;
+		}
+
+		/// <summary>
+		/// Adds a new aggregate to the graph
+		/// </summary>
+		/// <param name="Name">Name of the agent</param>
+		/// <param name="Outputs">Initial set of tokens to include in the aggregate</param>
+		/// <returns>Agent specification object</returns>
+		public BgAggregateSpec AddAggregate(BgString Name, BgList<BgFileSet> Outputs)
+		{
+			BgAggregateSpec Aggregate = new BgAggregateSpec(Name, Outputs);
+			AggregateSpecs.Add(Aggregate);
+			return Aggregate;
+		}
+
+		/// <summary>
+		/// Adds a new aggregate to the graph
+		/// </summary>
+		/// <param name="Name">Name of the agent</param>
+		/// <param name="Outputs">Initial set of tokens to include in the aggregate</param>
+		/// <param name="Label"></param>
+		/// <returns>Agent specification object</returns>
+		public BgAggregateSpec AddAggregate(BgString Name, BgList<BgFileSet> Outputs, BgString Label)
+		{
+			AddLabel(Config => Config.RequiredNodes = Config.RequiredNodes.Add(Outputs).Distinct());
+			return AddAggregate(Name, Outputs);
+		}
+
+		/// <summary>
+		/// Adds a new label to the graph
+		/// </summary>
+		/// <param name="Configure">Callback for configuring the label</param>
+		/// <returns>Label specification object</returns>
+		public BgLabelSpec AddLabel(Action<BgLabelConfig> Configure)
+		{
+			BgLabelConfig Config = new BgLabelConfig();
+			Configure(Config);
+
+			BgLabelSpec LabelSpec = new BgLabelSpec(Config);
+			LabelSpecs.Add(LabelSpec);
+
+			return LabelSpec;
+		}
+
+		/// <summary>
+		/// Adds a diagnostic to the graph
+		/// </summary>
+		/// <param name="Condition">Condition for adding the diagnostic to the graph</param>
+		/// <param name="Level"></param>
+		/// <param name="Message"></param>
+		public void AddDiagnostic(BgBool Condition, LogLevel Level, BgString Message)
+		{
+			BgDiagnosticSpec Diagnostic = new BgDiagnosticSpec(Level, Message);
+			DiagnosticSpecs = DiagnosticSpecs.AddIf(Condition, Diagnostic);
+		}
+
+		/// <summary>
+		/// Creates a concrete graph from the specifications
+		/// </summary>
+		/// <param name="Targets"></param>
+		/// <param name="Options"></param>
+		/// <returns>New graph instance</returns>
+		public BgGraph CreateGraph(HashSet<string> Targets, Dictionary<string, string> Options)
+		{
+			BgExprContext Context = new BgExprContext();
+			Context.Options = Options;
+
+			BgGraph Graph = new BgGraph();
+			foreach (BgAgentSpec AgentSpec in AgentSpecs)
+			{
+				AgentSpec.AddToGraph(Context, Graph);
+			}
+			foreach (BgAggregateSpec AggregateSpec in AggregateSpecs)
+			{
+				AggregateSpec.AddToGraph(Context, Graph);
+			}
+			foreach (BgDiagnosticSpec DiagnosticSpec in DiagnosticSpecs.GetEnumerable(Context))
+			{
+				DiagnosticSpec.AddToGraph(Context, Graph, null, null);
+			}
+			foreach (BgLabelSpec LabelSpec in LabelSpecs)
+			{
+				LabelSpec.AddToGraph(Context, Graph);
+			}
+
+			return Graph;
+		}
+	}
+}
diff --git a/Engine/Source/Programs/Shared/EpicGames.BuildGraph/BgLabelSpec.cs b/Engine/Source/Programs/Shared/EpicGames.BuildGraph/BgLabelSpec.cs
new file mode 100644
--- /dev/null
+++ b/Engine/Source/Programs/Shared/EpicGames.BuildGraph/BgLabelSpec.cs
@@ -0,0 +1,87 @@
+// Copyright Epic Games, Inc. All Rights Reserved.
+
+using EpicGames.BuildGraph.Expressions;
+using System;
+using System.Collections.Generic;
+using System.Linq;
+using System.Text;
+
+namespace EpicGames.BuildGraph
+{
+	/// <summary>
+	/// Configuration for a label
+	/// </summary>
+	public class BgLabelConfig
+	{
+		/// <summary>
+		/// Name of this badge
+		/// </summary>
+		public BgString? DashboardName { get; set; }
+
+		/// <summary>
+		/// Category for this label
+		/// </summary>
+		public BgString? DashboardCategory { get; }
+
+		/// <summary>
+		/// Name of the badge in UGS
+		/// </summary>
+		public BgString? UgsBadge { get; }
+
+		/// <summary>
+		/// Path to the project folder in UGS
+		/// </summary>
+		public BgString? UgsProject { get; }
+
+		/// <summary>
+		/// Which change to show the badge for
+		/// </summary>
+		public BgString? Change;
+
+		/// <summary>
+		/// Set of nodes that must be run for this label to be shown.
+		/// </summary>
+		public BgList<BgFileSet> RequiredNodes = BgList<BgFileSet>.Empty;
+
+		/// <summary>
+		/// Set of nodes that will be included in this label if present.
+		/// </summary>
+		public BgList<BgFileSet> IncludedNodes = BgList<BgFileSet>.Empty;
+	}
+
+	/// <summary>
+	/// Specification for a label
+	/// </summary>
+	public class BgLabelSpec
+	{
+		BgLabelConfig Config { get; }
+
+		internal BgLabelSpec(BgLabelConfig Config)
+		{
+			this.Config = Config;
+		}
+
+		internal void AddToGraph(BgExprContext Context, BgGraph Graph)
+		{
+			string? DashboardName = Config.DashboardName?.Compute(Context);
+			string? DashboardCategory = Config.DashboardCategory?.Compute(Context);
+			string? UgsBadge = Config.UgsBadge?.Compute(Context);
+			string? UgsProject = Config.UgsBadge?.Compute(Context);
+
+			BgLabelChange LabelChange;
+			if ((object?)Config.Change == null)
+			{
+				LabelChange = BgLabelChange.Current;
+			}
+			else
+			{
+				LabelChange = Enum.Parse<BgLabelChange>(Config.Change.Compute(Context));
+			}
+
+			BgLabel Label = new BgLabel(DashboardName, DashboardCategory, UgsBadge, UgsProject, LabelChange);
+			Label.RequiredNodes.UnionWith(Config.RequiredNodes.ComputeTags(Context).Select(x => Graph.TagNameToNodeOutput[x].ProducingNode));
+			Label.IncludedNodes.UnionWith(Config.IncludedNodes.ComputeTags(Context).Select(x => Graph.TagNameToNodeOutput[x].ProducingNode));
+			Graph.Labels.Add(Label);
+		}
+	}
+}
diff --git a/Engine/Source/Programs/Shared/EpicGames.BuildGraph/BgNodeSpec.cs b/Engine/Source/Programs/Shared/EpicGames.BuildGraph/BgNodeSpec.cs
new file mode 100644
--- /dev/null
+++ b/Engine/Source/Programs/Shared/EpicGames.BuildGraph/BgNodeSpec.cs
@@ -0,0 +1,520 @@
+// Copyright Epic Games, Inc. All Rights Reserved.
+
+using EpicGames.BuildGraph.Expressions;
+using EpicGames.Core;
+using Microsoft.Extensions.Logging;
+using System;
+using System.Collections.Generic;
+using System.Linq;
+using System.Linq.Expressions;
+using System.Reflection;
+using System.Text;
+using System.Text.Json;
+using System.Threading.Tasks;
+
+namespace EpicGames.BuildGraph
+{
+	/// <summary>
+	/// Exception for constructing nodes
+	/// </summary>
+	public sealed class BgNodeException : Exception
+	{
+		/// <summary>
+		/// Constructor
+		/// </summary>
+		/// <param name="Message"></param>
+		public BgNodeException(string Message) : base(Message)
+		{
+		}
+	}
+
+	/// <summary>
+	/// Speecifies the node name for a method. Parameters from the method may be embedded in the name using the {ParamName} syntax.
+	/// </summary>
+	[AttributeUsage(AttributeTargets.Method)]
+	public class BgNodeNameAttribute : Attribute
+	{
+		/// <summary>
+		/// The format string
+		/// </summary>
+		public string Template { get; }
+
+		/// <summary>
+		/// Constructor
+		/// </summary>
+		/// <param name="Template">Format string for the name</param>
+		public BgNodeNameAttribute(string Template)
+		{
+			this.Template = Template;
+		}
+	}
+
+	/// <summary>
+	/// Specification for a node to execute
+	/// </summary>
+	public class BgNodeSpec
+	{
+		internal string ClassName;
+		internal string MethodName;
+		internal IBgExpr?[] ArgumentExprs;
+		internal BgFileSet[] ResultExprs { get; }
+
+		/// <summary>
+		/// Name of the node
+		/// </summary>
+		public BgString Name { get; set; }
+
+		/// <summary>
+		/// Tokens for inputs of this node
+		/// </summary>
+		public BgList<BgFileSet> InputDependencies { get; set; } = BgList<BgFileSet>.Empty;
+
+		/// <summary>
+		/// Weak dependency on outputs that must be generated for the node to run, without making those dependencies inputs.
+		/// </summary>
+		public BgList<BgFileSet> AfterDependencies { get; set; } = BgList<BgFileSet>.Empty;
+
+		/// <summary>
+		/// Token for the output of this node
+		/// </summary>
+		BgFileSet Output { get; }
+
+		/// <summary>
+		/// All inputs and outputs of this node
+		/// </summary>
+		BgList<BgFileSet> InputsAndOutputs => InputDependencies.Add(Output);
+
+		/// <summary>
+		/// Whether this node should start running as soon as its dependencies are ready, even if other nodes in the same agent are not.
+		/// </summary>
+		public BgBool CanRunEarly { get; set; } = BgBool.False;
+
+		/// <summary>
+		/// Diagnostics for running this node
+		/// </summary>
+		public BgList<BgDiagnosticSpec> Diagnostics { get; set; } = BgList<BgDiagnosticSpec>.Empty;
+
+		/// <summary>
+		/// Constructor
+		/// </summary>
+		internal BgNodeSpec(MethodCallExpression Call)
+		{
+			string MethodSignature = $"{Call.Method.DeclaringType?.Name}.{Call.Method.Name}";
+			if (!Call.Method.IsStatic)
+			{
+				throw new BgNodeException($"Node {Call.Method.DeclaringType?.Name}.{Call.Method.Name} must be static");
+			}
+
+			this.ClassName = Call.Method.DeclaringType!.AssemblyQualifiedName!;
+			this.MethodName = Call.Method.Name;
+
+			try
+			{
+				ArgumentExprs = CreateArgumentExprs(Call);
+				Name = GetDefaultNodeName(Call.Method, ArgumentExprs);
+				ResultExprs = CreateReturnExprs(Call.Method, Name, Call.Method.ReturnType);
+			}
+			catch (Exception Ex)
+			{
+				ExceptionUtils.AddContext(Ex, $"while calling method {Call.Method.DeclaringType?.Name}.{Call.Method.Name}");
+				throw;
+			}
+
+			Output = new BgFileSetTagFromStringExpr(BgString.Format("#{0}", Name));
+		}
+
+		/// <summary>
+		/// Creates a node specification
+		/// </summary>
+		internal static BgNodeSpec Create(Expression<Func<BgContext, Task>> Action)
+		{
+			MethodCallExpression Call = (MethodCallExpression)Action.Body;
+			return new BgNodeSpec(Call);
+		}
+
+		/// <summary>
+		/// Creates a node specification
+		/// </summary>
+		internal static BgNodeSpec<T> Create<T>(Expression<Func<BgContext, Task<T>>> Function)
+		{
+			MethodCallExpression Call = (MethodCallExpression)Function.Body;
+			return new BgNodeSpec<T>(Call);
+		}
+
+		static IBgExpr?[] CreateArgumentExprs(MethodCallExpression Call)
+		{
+			IBgExpr?[] Args = new IBgExpr?[Call.Arguments.Count];
+			for (int Idx = 0; Idx < Call.Arguments.Count; Idx++)
+			{
+				Expression Expr = Call.Arguments[Idx];
+				if (Expr is ParameterExpression ParameterExpr)
+				{
+					if (ParameterExpr.Type != typeof(BgContext))
+					{
+						throw new BgNodeException($"Unable to determine type of parameter '{ParameterExpr.Name}'");
+					}
+				}
+				else
+				{
+					Delegate Compiled = Expression.Lambda(Expr).Compile();
+
+					object? Result = Compiled.DynamicInvoke();
+					if (Result is IBgExpr Computable)
+					{
+						Args[Idx] = Computable;
+					}
+					else
+					{
+						Args[Idx] = (BgString)(Result?.ToString() ?? String.Empty);
+					}
+				}
+			}
+			return Args;
+		}
+
+		static BgFileSet[] CreateReturnExprs(MethodInfo MethodInfo, BgString Name, Type Type)
+		{
+			if (Type == typeof(Task))
+			{
+				return Array.Empty<BgFileSet>();
+			}
+			else if (Type.IsGenericType && Type.GetGenericTypeDefinition() == typeof(Task<>))
+			{
+				return CreateReturnExprsInner(MethodInfo, Name, Type.GetGenericArguments()[0]);
+			}
+			else
+			{
+				throw new NotImplementedException();
+			}
+		}
+
+		static BgFileSet[] CreateReturnExprsInner(MethodInfo MethodInfo, BgString Name, Type Type)
+		{
+			Type[] OutputTypes;
+			if (IsValueTuple(Type))
+			{
+				OutputTypes = Type.GetGenericArguments();
+			}
+			else
+			{
+				OutputTypes = new[] { Type };
+			}
+
+			BgFileSet[] OutputExprs = new BgFileSet[OutputTypes.Length];
+			for (int Idx = 0; Idx < OutputTypes.Length; Idx++)
+			{
+				OutputExprs[Idx] = CreateOutputExpr(MethodInfo, Name, Idx, OutputTypes[Idx]);
+			}
+			return OutputExprs;
+		}
+
+		static BgFileSet CreateOutputExpr(MethodInfo MethodInfo, BgString Name, int Index, Type Type)
+		{
+			if (Type == typeof(BgFileSet))
+			{
+				return new BgFileSetTagFromStringExpr(BgString.Format("#{0}${1}", Name, Index));
+			}
+			else
+			{
+				throw new BgNodeException($"Unsupported return type for {MethodInfo.Name}: {Type.Name}");
+			}
+		}
+
+		internal static bool IsValueTuple(Type ReturnType)
+		{
+			if (ReturnType.IsGenericType)
+			{
+				Type GenericType = ReturnType.GetGenericTypeDefinition();
+				if(GenericType.FullName != null && GenericType.FullName.StartsWith("System.ValueTuple`", StringComparison.Ordinal))
+				{
+					return true;
+				}
+			}
+			return false;
+		}
+
+		static BgString GetDefaultNodeName(MethodInfo MethodInfo, IBgExpr?[] Args)
+		{
+			// Check if it's got an attribute override for the node name
+			BgNodeNameAttribute? NameAttr = MethodInfo.GetCustomAttribute<BgNodeNameAttribute>();
+			if (NameAttr != null)
+			{
+				return GetNodeNameFromTemplate(NameAttr.Template, MethodInfo.GetParameters(), Args);
+			}
+			else
+			{
+				return GetNodeNameFromMethodName(MethodInfo.Name);
+			}
+		}
+
+		static BgString GetNodeNameFromTemplate(string Template, ParameterInfo[] Parameters, IBgExpr?[] Args)
+		{
+			// Create a list of lazily computed string fragments which comprise the evaluated name
+			List<BgString> Fragments = new List<BgString>();
+
+			int LastIdx = 0;
+			for (int NextIdx = 0; NextIdx < Template.Length; NextIdx++)
+			{
+				if (Template[NextIdx] == '{')
+				{
+					if (NextIdx + 1 < Template.Length && Template[NextIdx + 1] == '{')
+					{
+						Fragments.Add(Template.Substring(LastIdx, NextIdx - LastIdx));
+						LastIdx = ++NextIdx;
+					}
+					else
+					{
+						Fragments.Add(Template.Substring(LastIdx, NextIdx - LastIdx));
+						NextIdx++;
+
+						int EndIdx = Template.IndexOf('}', NextIdx);
+						if (EndIdx == -1)
+						{
+							throw new BgNodeException($"Unterminated parameter expression for {nameof(BgNodeNameAttribute)} in {Template}");
+						}
+
+						StringView ParamName = new StringView(Template, NextIdx, EndIdx - NextIdx);
+
+						int ParamIdx = Array.FindIndex(Parameters, x => x.Name != null && ParamName.Equals(x.Name, StringComparison.Ordinal));
+						if (ParamIdx == -1)
+						{
+							throw new BgNodeException($"Unable to find parameter named {ParamName} in {Template}");
+						}
+
+						IBgExpr? Arg = Args[ParamIdx];
+						if(Arg != null)
+						{
+							Fragments.Add(Arg.ToBgString());
+						}
+
+						LastIdx = NextIdx = EndIdx + 1;
+					}
+				}
+				else if (Template[NextIdx] == '}')
+				{
+					if (NextIdx + 1 < Template.Length && Template[NextIdx + 1] == '{')
+					{
+						Fragments.Add(Template.Substring(LastIdx, NextIdx - LastIdx));
+						LastIdx = ++NextIdx;
+					}
+				}
+			}
+			Fragments.Add(Template.Substring(LastIdx, Template.Length - LastIdx));
+
+			return BgString.Join(BgString.Empty, Fragments);
+		}
+
+		/// <summary>
+		/// Inserts spaces into a PascalCase method name to create a node name
+		/// </summary>
+		public static string GetNodeNameFromMethodName(string MethodName)
+		{
+			StringBuilder Name = new StringBuilder();
+			Name.Append(MethodName[0]);
+
+			int Length = MethodName.Length;
+			if (Length > 5 && MethodName.EndsWith("Async", StringComparison.Ordinal))
+			{
+				Length -= 5;
+			}
+
+			bool bIsAcronym = false;
+			for (int Idx = 1; Idx < Length; Idx++)
+			{
+				bool bLastIsUpper = Char.IsUpper(MethodName[Idx - 1]);
+				bool bNextIsUpper = Char.IsUpper(MethodName[Idx]);
+				if (bLastIsUpper && bNextIsUpper)
+				{
+					bIsAcronym = true;
+				}
+				else if (bIsAcronym)
+				{
+					Name.Insert(Name.Length - 2, ' ');
+					bIsAcronym = false;
+				}
+				else if (!bLastIsUpper && bNextIsUpper)
+				{
+					Name.Append(' ');
+				}
+				Name.Append(MethodName[Idx]);
+			}
+
+			return Name.ToString();
+		}
+
+		/// <summary>
+		/// Gets the signature for a method
+		/// </summary>
+		public static string GetSignature(MethodInfo MethodInfo)
+		{
+			StringBuilder Arguments = new StringBuilder();
+			foreach (ParameterInfo ParameterInfo in MethodInfo.GetParameters())
+			{
+				Arguments.AppendLine(ParameterInfo.ParameterType.FullName);
+			}
+			return $"{MethodInfo.Name}:{Digest.Compute<Sha1>(Arguments.ToString())}";
+		}
+
+		/// <summary>
+		/// Creates a concrete node
+		/// </summary>
+		/// <param name="Context"></param>
+		/// <param name="Graph"></param>
+		/// <param name="Agent"></param>
+		/// <returns></returns>
+		internal void AddToGraph(BgExprContext Context, BgGraph Graph, BgAgent Agent)
+		{
+			HashSet<string> InputTags = new HashSet<string>();
+			InputTags.UnionWith(this.InputDependencies.ComputeTags(Context));
+			InputTags.UnionWith(ArgumentExprs.OfType<BgFileSet>().Select(x => x.ComputeTag(Context)));
+			InputTags.UnionWith(ArgumentExprs.OfType<BgList<BgFileSet>>().SelectMany(x => x.GetEnumerable(Context)).Select(x => x.ComputeTag(Context)));
+
+			HashSet<string> AfterTags = new HashSet<string>(InputTags);
+			AfterTags.UnionWith(this.AfterDependencies.ComputeTags(Context));
+
+			string Name = this.Name.Compute(Context);
+			BgNodeOutput[] InputDependencies = InputTags.Select(x => Graph.TagNameToNodeOutput[x]).Distinct().ToArray();
+			BgNodeOutput[] AfterDependencies = AfterTags.Select(x => Graph.TagNameToNodeOutput[x]).Distinct().ToArray();
+			string[] OutputNames = Array.ConvertAll(ResultExprs, x => x.ComputeTag(Context));
+			BgNode[] InputNodes = InputDependencies.Select(x => x.ProducingNode).Distinct().ToArray();
+			BgNode[] AfterNodes = AfterDependencies.Select(x => x.ProducingNode).Distinct().ToArray();
+			bool RunEarly = CanRunEarly.Compute(Context);
+
+			BgNode Node = new BgNode(Name, InputDependencies, OutputNames, InputNodes, AfterNodes, Array.Empty<FileReference>());
+			Node.bRunEarly = RunEarly;
+
+			BgScriptLocation Location = new BgScriptLocation("(unknown)", "(unknown)", 1);
+
+			BgTask Task = new BgTask(Location, "CallMethod");
+			Task.Arguments["Class"] = ClassName;
+			Task.Arguments["Method"] = MethodName;
+			for (int Idx = 0; Idx < ArgumentExprs.Length; Idx++)
+			{
+				IBgExpr? ArgumentExpr = ArgumentExprs[Idx];
+				if (ArgumentExpr != null)
+				{
+					Task.Arguments[$"Arg{Idx + 1}"] = BgType.Get(ArgumentExpr.GetType()).SerializeArgument(ArgumentExpr, Context);
+				}
+			}
+			if (OutputNames != null)
+			{
+				Task.Arguments["Tags"] = String.Join(";", OutputNames);
+			}
+			Node.Tasks.Add(Task);
+
+			Agent.Nodes.Add(Node);
+			Graph.NameToNode.Add(Name, Node);
+
+			foreach (BgDiagnosticSpec Precondition in Diagnostics.GetEnumerable(Context))
+			{
+				Precondition.AddToGraph(Context, Graph, Agent, Node);
+			}
+
+			foreach(BgNodeOutput Output in Node.Outputs)
+			{
+				Graph.TagNameToNodeOutput.Add(Output.TagName, Output);
+			}
+		}
+
+		/// <summary>
+		/// Allows using the 
+		/// </summary>
+		/// <param name="NodeSpec"></param>
+		public static implicit operator BgList<BgFileSet>(BgNodeSpec NodeSpec)
+		{
+			return NodeSpec.InputsAndOutputs.Add(NodeSpec.Output);
+		}
+	}
+
+	/// <summary>
+	/// Specification for a node to execute that returns one or more tagged outputs
+	/// </summary>
+	/// <typeparam name="T">Return type from the node. Must consist of latent Bg* types (eg. BgToken)</typeparam>
+	public class BgNodeSpec<T> : BgNodeSpec
+	{
+		/// <summary>
+		/// Output value from this node
+		/// </summary>
+		public T Output { get; }
+
+		/// <summary>
+		/// Constructor
+		/// </summary>
+		internal BgNodeSpec(MethodCallExpression Call)
+			: base(Call)
+		{
+			Type ReturnType = Call.Method.ReturnType;
+			if (IsValueTuple(ReturnType))
+			{
+				Output = (T)Activator.CreateInstance(ReturnType, ResultExprs)!;
+			}
+			else
+			{
+				Output = (T)(object)ResultExprs[0];
+			}
+		}
+	}
+
+	/// <summary>
+	/// Extension methods for BgNode types
+	/// </summary>
+	public static class BgNodeExtensions
+	{
+		/// <summary>
+		/// Add dependencies onto other nodes or outputs. Outputs from the given tokens will be copied to the current machine before execution of the node.
+		/// </summary>
+		/// <param name="NodeSpec">The node specification</param>
+		/// <param name="Tokens">Tokens to add dependencies on</param>
+		/// <returns>The current node spec, to allow chaining calls</returns>
+		public static T Requires<T>(this T NodeSpec, params BgList<BgFileSet>[] Tokens) where T : BgNodeSpec
+		{
+			foreach (BgList<BgFileSet> TokenSet in Tokens)
+			{
+				NodeSpec.InputDependencies = NodeSpec.InputDependencies.Add(TokenSet);
+			}
+			return NodeSpec;
+		}
+
+		/// <summary>
+		/// Add weak dependencies onto other nodes or outputs. The producing nodes must complete successfully if they are part of the graph, but outputs from them will not be 
+		/// transferred to the machine running this node.
+		/// </summary>
+		/// <param name="NodeSpec">The node specification</param>
+		/// <param name="Tokens">Tokens to add dependencies on</param>
+		/// <returns>The current node spec, to allow chaining calls</returns>
+		public static T After<T>(this T NodeSpec, params BgList<BgFileSet>[] Tokens) where T : BgNodeSpec
+		{
+			foreach (BgList<BgFileSet> TokenSet in Tokens)
+			{
+				NodeSpec.AfterDependencies = NodeSpec.AfterDependencies.Add(TokenSet);
+			}
+			return NodeSpec;
+		}
+
+		/// <summary>
+		/// Adds a warning when executing this node
+		/// </summary>
+		/// <param name="NodeSpec">The node specification</param>
+		/// <param name="Condition">Condition for writing the message</param>
+		/// <param name="Message">Message to output</param>
+		/// <returns>The current node spec, to allow chaining calls</returns>
+		public static BgNodeSpec WarnIf<T>(this T NodeSpec, BgBool Condition, BgString Message) where T : BgNodeSpec
+		{
+			NodeSpec.Diagnostics = NodeSpec.Diagnostics.AddIf(Condition, new BgDiagnosticSpec(LogLevel.Warning, Message));
+			return NodeSpec;
+		}
+
+		/// <summary>
+		/// Adds an error when executing this node
+		/// </summary>
+		/// <param name="NodeSpec">The node specification</param>
+		/// <param name="Condition">Condition for writing the message</param>
+		/// <param name="Message">Message to output</param>
+		/// <returns>The current node spec, to allow chaining calls</returns>
+		public static BgNodeSpec ErrorIf<T>(this T NodeSpec, BgBool Condition, BgString Message) where T : BgNodeSpec
+		{
+			NodeSpec.Diagnostics = NodeSpec.Diagnostics.AddIf(Condition, new BgDiagnosticSpec(LogLevel.Error, Message));
+			return NodeSpec;
+		}
+	}
+}
diff --git a/Engine/Source/Programs/Shared/EpicGames.BuildGraph/BgOption.cs b/Engine/Source/Programs/Shared/EpicGames.BuildGraph/BgOption.cs
new file mode 100644
--- /dev/null
+++ b/Engine/Source/Programs/Shared/EpicGames.BuildGraph/BgOption.cs
@@ -0,0 +1,51 @@
+// Copyright Epic Games, Inc. All Rights Reserved.
+
+using System;
+using System.Collections.Generic;
+using System.Text;
+
+namespace EpicGames.BuildGraph
+{
+	/// <summary>
+	/// Represents a graph option. These are expanded during preprocessing, but are retained in order to display help messages.
+	/// </summary>
+	public class BgOption
+	{
+		/// <summary>
+		/// Name of this option
+		/// </summary>
+		public string Name;
+
+		/// <summary>
+		/// Description for this option
+		/// </summary>
+		public string Description;
+
+		/// <summary>
+		/// Default value for this option
+		/// </summary>
+		public string DefaultValue;
+
+		/// <summary>
+		/// Constructor
+		/// </summary>
+		/// <param name="Name">The name of this option</param>
+		/// <param name="Description">Description of the option, for display on help pages</param>
+		/// <param name="DefaultValue">Default value for the option</param>
+		public BgOption(string Name, string Description, string DefaultValue)
+		{
+			this.Name = Name;
+			this.Description = Description;
+			this.DefaultValue = DefaultValue;
+		}
+
+		/// <summary>
+		/// Returns a name of this option for debugging
+		/// </summary>
+		/// <returns>Name of the option</returns>
+		public override string ToString()
+		{
+			return Name;
+		}
+	}
+}
diff --git a/Engine/Source/Programs/Shared/EpicGames.BuildGraph/BgOptionSpec.cs b/Engine/Source/Programs/Shared/EpicGames.BuildGraph/BgOptionSpec.cs
new file mode 100644
--- /dev/null
+++ b/Engine/Source/Programs/Shared/EpicGames.BuildGraph/BgOptionSpec.cs
@@ -0,0 +1,324 @@
+// Copyright Epic Games, Inc. All Rights Reserved.
+
+using EpicGames.BuildGraph.Expressions;
+using System;
+using System.Collections.Generic;
+using System.Linq;
+using System.Text;
+using System.Text.RegularExpressions;
+
+namespace EpicGames.BuildGraph
+{
+	/// <summary>
+	/// Exception thrown if an option fails validation
+	/// </summary>
+	sealed class BgOptionValidationException : Exception
+	{
+		/// <summary>
+		/// Constructor
+		/// </summary>
+		/// <param name="Message"></param>
+		public BgOptionValidationException(string Message)
+			: base(Message)
+		{
+		}
+	}
+
+	/// <summary>
+	/// Base class for option configuration
+	/// </summary>
+	public interface IBgOption : IBgExpr
+	{
+		/// <summary>
+		/// Name of the option
+		/// </summary>
+		string Name { get; }
+
+		/// <summary>
+		/// Label to show against the option in the UI
+		/// </summary>
+		BgString? Label { get; set; }
+
+		/// <summary>
+		/// Description for the option
+		/// </summary>
+		BgString Description { get; set; }
+	}
+
+	/// <summary>
+	/// A boolean option expression
+	/// </summary>
+	public class BgBoolOption : BgBool, IBgOption
+	{
+		/// <inheritdoc/>
+		public string Name { get; }
+
+		/// <inheritdoc/>
+		public BgString? Label { get; set; }
+
+		/// <inheritdoc/>
+		public BgString Description { get; set; }
+
+		/// <summary>
+		/// Default value for the option
+		/// </summary>
+		public BgBool DefaultValue { get; set; }
+
+		/// <summary>
+		/// Constructor
+		/// </summary>
+		internal BgBoolOption(string Name, BgString Description, BgBool DefaultValue)
+		{
+			this.Name = Name;
+			this.Description = Description;
+			this.DefaultValue = DefaultValue;
+		}
+
+		/// <inheritdoc/>
+		public override bool Compute(BgExprContext Context)
+		{
+			string? Value;
+			if (Context.Options.TryGetValue(Name, out Value))
+			{
+				bool BoolValue;
+				if (!bool.TryParse(Value, out BoolValue))
+				{
+					throw new BgOptionValidationException($"Argument for {Name} is not a valid bool ({Value})");
+				}
+				return BoolValue;
+			}
+			return DefaultValue.Compute(Context);
+		}
+	}
+
+	/// <summary>
+	/// An integer option expression
+	/// </summary>
+	public class BgIntOption : BgInt, IBgOption
+	{
+		/// <inheritdoc/>
+		public string Name { get; }
+
+		/// <inheritdoc/>
+		public BgString? Label { get; set; }
+
+		/// <inheritdoc/>
+		public BgString Description { get; set; }
+
+		/// <summary>
+		/// Default value for the option
+		/// </summary>
+		public BgInt DefaultValue { get; set; }
+
+		/// <summary>
+		/// Minimum allowed value
+		/// </summary>
+		public BgInt? MinValue { get; set; }
+
+		/// <summary>
+		/// Maximum allowed value
+		/// </summary>
+		public BgInt? MaxValue { get; set; }
+
+		/// <summary>
+		/// Constructor
+		/// </summary>
+		internal BgIntOption(string Name, BgString Description, BgInt DefaultValue)
+		{
+			this.Name = Name;
+			this.Description = Description;
+			this.DefaultValue = DefaultValue;
+		}
+
+		/// <inheritdoc/>
+		public override int Compute(BgExprContext Context)
+		{
+			string? Value;
+			if (Context.Options.TryGetValue(Name, out Value))
+			{
+				int IntValue;
+				if (!int.TryParse(Value, out IntValue))
+				{
+					throw new BgOptionValidationException($"Argument for '{Name}' is not a valid integer");
+				}
+				if (!ReferenceEquals(MinValue, null))
+				{
+					int IntMinValue = MinValue.Compute(Context);
+					if (IntValue < IntMinValue)
+					{
+						throw new BgOptionValidationException($"Argument for '{Name}' is less than the allowed minimum ({IntValue} < {IntMinValue})");
+					}
+				}
+				if (!ReferenceEquals(MaxValue, null))
+				{
+					int IntMaxValue = MaxValue.Compute(Context);
+					if (IntValue > IntMaxValue)
+					{
+						throw new BgOptionValidationException($"Argument for '{Name}' is greater than the allowed maximum ({IntValue} > {IntMaxValue})");
+					}
+				}
+			}
+			return DefaultValue.Compute(Context);
+		}
+	}
+
+	/// <summary>
+	/// A string option expression
+	/// </summary>
+	public class BgEnumOption<TEnum> : BgEnum<TEnum>, IBgOption where TEnum : struct
+	{
+		/// <inheritdoc/>
+		public string Name { get; }
+
+		/// <inheritdoc/>
+		public BgString? Label { get; set; }
+
+		/// <inheritdoc/>
+		public BgString Description { get; set; }
+
+		/// <summary>
+		/// Default value for the option
+		/// </summary>
+		public BgEnum<TEnum> DefaultValue { get; set; }
+
+		/// <summary>
+		/// Constructor
+		/// </summary>
+		internal BgEnumOption(string Name, BgString Description, BgEnum<TEnum> DefaultValue)
+		{
+			this.Name = Name;
+			this.Description = Description;
+			this.DefaultValue = DefaultValue;
+		}
+
+		/// <inheritdoc/>
+		public override TEnum Compute(BgExprContext Context)
+		{
+			string? Value;
+			if (Context.Options.TryGetValue(Name, out Value))
+			{
+				TEnum EnumValue;
+				if (!Enum.TryParse<TEnum>(Value, true, out EnumValue))
+				{
+					throw new BgOptionValidationException($"Argument '{Name}' is not a valid value for {typeof(TEnum).Name}");
+				}
+				return EnumValue;
+			}
+			return DefaultValue.Compute(Context);
+		}
+	}
+
+	/// <summary>
+	/// A string option expression
+	/// </summary>
+	public class BgStringOption : BgString, IBgOption
+	{
+		/// <inheritdoc/>
+		public string Name { get; }
+
+		/// <inheritdoc/>
+		public BgString? Label { get; set; }
+
+		/// <inheritdoc/>
+		public BgString Description { get; set; }
+
+		/// <summary>
+		/// Default value for the option
+		/// </summary>
+		public BgString DefaultValue { get; set; }
+
+		/// <summary>
+		/// Regex for validating values for the option
+		/// </summary>
+		public BgString? Pattern { get; set; }
+
+		/// <summary>
+		/// Message to display if validation fails
+		/// </summary>
+		public BgString? PatternFailed { get; set; }
+
+		/// <summary>
+		/// Allowed values of the option
+		/// </summary>
+		public BgList<BgString>? Enum { get; set; }
+
+		/// <summary>
+		/// Constructor
+		/// </summary>
+		internal BgStringOption(string Name, BgString Description, BgString DefaultValue)
+		{
+			this.Name = Name;
+			this.Description = Description;
+			this.DefaultValue = DefaultValue;
+		}
+
+		/// <inheritdoc/>
+		public override string Compute(BgExprContext Context)
+		{
+			string? Value;
+			if (Context.Options.TryGetValue(Name, out Value))
+			{
+				if (!Object.ReferenceEquals(Pattern, null))
+				{
+					string PatternValue = Pattern.Compute(Context);
+					if (!Regex.IsMatch(Value, PatternValue))
+					{
+						string PatternFailedValue = PatternFailed?.Compute(Context) ?? $"Argument '{Name}' does not match the required pattern: '{PatternValue}'";
+						throw new BgOptionValidationException(PatternFailedValue);
+					}
+				}
+				if (!Object.ReferenceEquals(Enum, null))
+				{
+					List<string> EnumValues = Enum.Compute(Context);
+					if (!EnumValues.Any(x => x.Equals(Value, StringComparison.OrdinalIgnoreCase)))
+					{
+						throw new BgOptionValidationException($"Argument '{Name}' is invalid");
+					}
+				}
+				return Value;
+			}
+			return DefaultValue.Compute(Context);
+		}
+	}
+
+	/// <summary>
+	/// A list option expression
+	/// </summary>
+	public class BgEnumListOption<TEnum> : BgList<BgEnum<TEnum>>, IBgOption where TEnum : struct
+	{
+		/// <inheritdoc/>
+		public string Name { get; }
+
+		/// <inheritdoc/>
+		public BgString? Label { get; set; }
+
+		/// <inheritdoc/>
+		public BgString Description { get; set; }
+
+		/// <summary>
+		/// Default value for the option
+		/// </summary>
+		public BgList<BgEnum<TEnum>> DefaultValue { get; set; }
+
+		/// <summary>
+		/// Constructor
+		/// </summary>
+		internal BgEnumListOption(string Name, BgString Description, BgList<BgEnum<TEnum>> DefaultValue)
+		{
+			this.Name = Name;
+			this.Description = Description;
+			this.DefaultValue = DefaultValue;
+		}
+
+		/// <inheritdoc/>
+		public override IEnumerable<BgEnum<TEnum>> GetEnumerable(BgExprContext Context)
+		{
+			BgList<BgEnum<TEnum>> Value = DefaultValue;
+			if (Context.Options.TryGetValue(Name, out string? ValueText))
+			{
+				Value = BgType.Get<BgList<BgEnum<TEnum>>>().DeserializeArgument(ValueText);
+			}
+			return Value.GetEnumerable(Context);
+		}
+	}
+}
diff --git a/Engine/Source/Programs/Shared/EpicGames.BuildGraph/BgSchema.cs b/Engine/Source/Programs/Shared/EpicGames.BuildGraph/BgSchema.cs
--- a/Engine/Source/Programs/Shared/EpicGames.BuildGraph/BgSchema.cs
+++ b/Engine/Source/Programs/Shared/EpicGames.BuildGraph/BgSchema.cs
@@ -129,6 +129,7 @@ namespace EpicGames.BuildGraph
 		}
 	}
 
+#pragma warning disable CS1591
 	/// <summary>
 	/// Enumeration of standard types used in the schema. Avoids hard-coding names.
 	/// </summary>
@@ -168,6 +169,7 @@ namespace EpicGames.BuildGraph
 		Integer,
 		LabelChange
 	}
+#pragma warning restore CS1591
 
 	/// <summary>
 	/// Schema for build graph definitions. Stores information about the supported tasks, and allows validating an XML document.
@@ -257,7 +259,8 @@ namespace EpicGames.BuildGraph
 		/// <summary>
 		/// Constructor
 		/// </summary>
-		/// <param name="InNameToTask">Mapping of task name to information about how to construct it</param>
+		/// <param name="Tasks">Set of known tasks</param>
+		/// <param name="PrimitiveTypes">Mapping of task name to information about how to construct it</param>
 		public BgScriptSchema(IEnumerable<BgScriptTask> Tasks, List<(Type, ScriptSchemaStandardType)> PrimitiveTypes)
 		{
 			// Create a lookup from standard types to their qualified names
diff --git a/Engine/Source/Programs/Shared/EpicGames.BuildGraph/BgScriptReader.cs b/Engine/Source/Programs/Shared/EpicGames.BuildGraph/BgScriptReader.cs
--- a/Engine/Source/Programs/Shared/EpicGames.BuildGraph/BgScriptReader.cs
+++ b/Engine/Source/Programs/Shared/EpicGames.BuildGraph/BgScriptReader.cs
@@ -264,7 +264,7 @@ namespace EpicGames.BuildGraph
 		/// <summary>
 		/// Tests whether the given file or directory exists
 		/// </summary>
-		/// <param name="File">Path to a file</param>
+		/// <param name="Path">Path to a file</param>
 		/// <returns>True if the file exists</returns>
 		Task<bool> ExistsAsync(string Path);
 
@@ -409,7 +409,6 @@ namespace EpicGames.BuildGraph
 		/// Reads the contents of a graph
 		/// </summary>
 		/// <param name="Element">The parent element to read from</param>
-		/// <param name="Arguments">Arguments passed in to the graph on the command line</param>
 		async Task ReadGraphBodyAsync(XmlElement Element)
 		{
 			foreach (BgScriptElement ChildElement in Element.ChildNodes.OfType<BgScriptElement>())
@@ -747,7 +746,6 @@ namespace EpicGames.BuildGraph
 		/// Read the contents of an agent definition
 		/// </summary>
 		/// <param name="Element">Xml element to read the definition from</param>
-		/// <param name="ParentAgent">The agent to contain the definition</param>
 		protected async Task ReadAgentBodyAsync(BgScriptElement Element)
 		{
 			EnterScope();
@@ -1420,7 +1418,7 @@ namespace EpicGames.BuildGraph
 						// Create a new option object to store the settings
 						string Description = ReadAttribute(Element, "Description");
 						string DefaultValue = ReadAttribute(Element, "DefaultValue");
-						BgScriptOption Option = new BgScriptOption(Name, Description, DefaultValue);
+						BgOption Option = new BgOption(Name, Description, DefaultValue);
 						Graph.Options.Add(Option);
 
 						// Get the value of this property
diff --git a/Engine/Source/Programs/Shared/EpicGames.BuildGraph/EpicGames.BuildGraph.csproj b/Engine/Source/Programs/Shared/EpicGames.BuildGraph/EpicGames.BuildGraph.csproj
--- a/Engine/Source/Programs/Shared/EpicGames.BuildGraph/EpicGames.BuildGraph.csproj
+++ b/Engine/Source/Programs/Shared/EpicGames.BuildGraph/EpicGames.BuildGraph.csproj
@@ -9,6 +9,7 @@
     <Configurations>Debug;Release;Development</Configurations>
     <Nullable>enable</Nullable>
     <DebugType>pdbonly</DebugType>
+    <GenerateDocumentationFile>True</GenerateDocumentationFile>
   </PropertyGroup>
   <PropertyGroup Condition=" '$(Configuration)|$(Platform)' == 'Debug|AnyCPU' ">
     <DocumentationFile></DocumentationFile>
diff --git a/Engine/Source/Programs/Shared/EpicGames.BuildGraph/Expressions/BgBool.cs b/Engine/Source/Programs/Shared/EpicGames.BuildGraph/Expressions/BgBool.cs
new file mode 100644
--- /dev/null
+++ b/Engine/Source/Programs/Shared/EpicGames.BuildGraph/Expressions/BgBool.cs
@@ -0,0 +1,198 @@
+// Copyright Epic Games, Inc. All Rights Reserved.
+
+using System;
+using System.Collections.Generic;
+using System.ComponentModel;
+using System.Diagnostics;
+using System.IO;
+using System.Linq;
+using System.Reflection;
+using System.Text;
+using System.Xml;
+using EpicGames.Core;
+using Microsoft.Extensions.Logging;
+using System.Threading.Tasks;
+
+namespace EpicGames.BuildGraph.Expressions
+{
+	/// <summary>
+	/// Abstract base class for expressions returning a boolean value 
+	/// </summary>
+	[BgType(typeof(BgBoolType))]
+	public abstract class BgBool : IBgExpr<BgBool>
+	{
+		/// <summary>
+		/// Constant value for false
+		/// </summary>
+		public static BgBool False { get; } = false;
+
+		/// <summary>
+		/// Constant value for true
+		/// </summary>
+		public static BgBool True { get; } = true;
+
+		/// <inheritdoc/>
+		public BgBool IfThen(BgBool Condition, BgBool ValueIfTrue) => new BgBoolChooseExpr(Condition, ValueIfTrue, this);
+
+		/// <summary>
+		/// Implict conversion operator from a boolean literal
+		/// </summary>
+		public static implicit operator BgBool(bool Value)
+		{
+			return new BgBoolConstantExpr(Value);
+		}
+
+		/// <inheritdoc/>
+		public static BgBool operator !(BgBool Inner) => new BgBoolNotExpr(Inner);
+
+		/// <inheritdoc/>
+		public static BgBool operator &(BgBool Lhs, BgBool Rhs) => new BgBoolAndExpr(Lhs, Rhs);
+
+		/// <inheritdoc/>
+		public static BgBool operator |(BgBool Lhs, BgBool Rhs) => new BgBoolOrExpr(Lhs, Rhs);
+
+		/// <inheritdoc/>
+		public static BgBool operator ==(BgBool Lhs, BgBool Rhs) => new BgBoolEqExpr(Lhs, Rhs);
+
+		/// <inheritdoc/>
+		public static BgBool operator !=(BgBool Lhs, BgBool Rhs) => new BgBoolNotExpr(new BgBoolEqExpr(Lhs, Rhs));
+
+		/// <inheritdoc/>
+		object IBgExpr.Compute(BgExprContext Context) => Compute(Context);
+
+		/// <inheritdoc/>
+		public abstract bool Compute(BgExprContext Context);
+
+		/// <inheritdoc/>
+		public sealed override bool Equals(object? obj) => throw new InvalidOperationException();
+
+		/// <inheritdoc/>
+		public sealed override int GetHashCode() => throw new InvalidOperationException();
+
+		/// <inheritdoc/>
+		public BgString ToBgString() => new BgBoolFormatExpr(this);
+	}
+
+	/// <summary>
+	/// Type traits for a <see cref="BgBool"/>
+	/// </summary>
+	class BgBoolType : BgTypeBase<BgBool>
+	{
+		/// <inheritdoc/>
+		public override BgBool DeserializeArgument(string Text) => bool.Parse(Text);
+
+		/// <inheritdoc/>
+		public override string SerializeArgument(BgBool Value, BgExprContext Context) => Value.Compute(Context).ToString();
+
+		/// <inheritdoc/>
+		public override BgBool CreateConstant(object Value) => new BgBoolConstantExpr((bool)Value);
+
+		/// <inheritdoc/>
+		public override IBgExprVariable<BgBool> CreateVariable() => new BgBoolVariableExpr();
+	}
+
+	#region Expression classes
+
+	class BgBoolNotExpr : BgBool
+	{
+		public BgBool Inner { get; }
+
+		public BgBoolNotExpr(BgBool Inner)
+		{
+			this.Inner = Inner;
+		}
+
+		public override bool Compute(BgExprContext Context) => !Inner.Compute(Context);
+	}
+
+	class BgBoolEqExpr : BgBool
+	{
+		public BgBool Lhs { get; }
+		public BgBool Rhs { get; }
+
+		public BgBoolEqExpr(BgBool Lhs, BgBool Rhs)
+		{
+			this.Lhs = Lhs;
+			this.Rhs = Rhs;
+		}
+
+		public override bool Compute(BgExprContext Context) => Lhs.Compute(Context) == Rhs.Compute(Context);
+	}
+
+	class BgBoolAndExpr : BgBool
+	{
+		public BgBool Lhs { get; }
+		public BgBool Rhs { get; }
+
+		public BgBoolAndExpr(BgBool Lhs, BgBool Rhs)
+		{
+			this.Lhs = Lhs;
+			this.Rhs = Rhs;
+		}
+
+		public override bool Compute(BgExprContext Context) => Lhs.Compute(Context) && Rhs.Compute(Context);
+	}
+
+	class BgBoolOrExpr : BgBool
+	{
+		public BgBool Lhs { get; }
+		public BgBool Rhs { get; }
+
+		public BgBoolOrExpr(BgBool Lhs, BgBool Rhs)
+		{
+			this.Lhs = Lhs;
+			this.Rhs = Rhs;
+		}
+
+		public override bool Compute(BgExprContext Context) => Lhs.Compute(Context) || Rhs.Compute(Context);
+	}
+
+	class BgBoolChooseExpr : BgBool
+	{
+		public BgBool Condition;
+		public BgBool ValueIfTrue;
+		public BgBool ValueIfFalse;
+
+		public BgBoolChooseExpr(BgBool Condition, BgBool ValueIfTrue, BgBool ValueIfFalse)
+		{
+			this.Condition = Condition;
+			this.ValueIfTrue = ValueIfTrue;
+			this.ValueIfFalse = ValueIfFalse;
+		}
+
+		public override bool Compute(BgExprContext Context) => Condition.Compute(Context) ? ValueIfTrue.Compute(Context) : ValueIfFalse.Compute(Context);
+	}
+
+	class BgBoolVariableExpr : BgBool, IBgExprVariable<BgBool>
+	{
+		public BgBool Value { get; set; } = BgBool.False;
+
+		public override bool Compute(BgExprContext Context) => Value.Compute(Context);
+	}
+
+	class BgBoolConstantExpr : BgBool
+	{
+		public bool Value { get; }
+
+		public BgBoolConstantExpr(bool Value)
+		{
+			this.Value = Value;
+		}
+
+		public override bool Compute(BgExprContext Context) => Value;
+	}
+
+	class BgBoolFormatExpr : BgString
+	{
+		BgBool Value;
+
+		public BgBoolFormatExpr(BgBool Value)
+		{
+			this.Value = Value;
+		}
+
+		public override string Compute(BgExprContext Context) => Value.Compute(Context).ToString();
+	}
+
+	#endregion
+}
diff --git a/Engine/Source/Programs/Shared/EpicGames.BuildGraph/Expressions/BgEnum.cs b/Engine/Source/Programs/Shared/EpicGames.BuildGraph/Expressions/BgEnum.cs
new file mode 100644
--- /dev/null
+++ b/Engine/Source/Programs/Shared/EpicGames.BuildGraph/Expressions/BgEnum.cs
@@ -0,0 +1,171 @@
+// Copyright Epic Games, Inc. All Rights Reserved.
+
+using EpicGames.BuildGraph;
+using System;
+using System.Collections.Generic;
+using System.ComponentModel;
+using System.Diagnostics;
+using System.IO;
+using System.Linq;
+using System.Reflection;
+using System.Text;
+using System.Xml;
+using EpicGames.Core;
+using Microsoft.Extensions.Logging;
+using OpenTracing;
+using OpenTracing.Util;
+using System.Threading.Tasks;
+using System.Text.RegularExpressions;
+
+namespace EpicGames.BuildGraph.Expressions
+{
+	/// <summary>
+	/// Abstract base class for expressions returning an enum. Use <see cref="BgEnum{TEnum}"/> for a strongly typed version.
+	/// </summary>
+	public interface IBgEnum : IBgExpr
+	{
+		/// <summary>
+		/// The enum type
+		/// </summary>
+		Type EnumType { get; }
+	}
+
+	/// <summary>
+	/// Abstract base class for expressions returning a string value 
+	/// </summary>
+	[BgType(typeof(BgEnumType<>))]
+	public abstract class BgEnum<TEnum> : IBgEnum, IBgExpr<BgEnum<TEnum>> where TEnum : struct
+	{
+		/// <inheritdoc/>
+		public Type EnumType => typeof(TEnum);
+
+		/// <summary>
+		/// Implicit conversion from a regular enum type
+		/// </summary>
+		public static implicit operator BgEnum<TEnum>(TEnum Value)
+		{
+			return new BgEnumConstantExpr<TEnum>(Value);
+		}
+
+		/// <summary>
+		/// Explicit conversion from a string value
+		/// </summary>
+		public static explicit operator BgEnum<TEnum>(BgString Value)
+		{
+			return new BgEnumParseExpr<TEnum>(Value);
+		}
+
+		/// <inheritdoc/>
+		public BgEnum<TEnum> IfThen(BgBool Condition, BgEnum<TEnum> ValueIfTrue) => new BgEnumChooseExpr<TEnum>(Condition, ValueIfTrue, this);
+
+		/// <inheritdoc/>
+		object IBgExpr.Compute(BgExprContext Context) => Compute(Context);
+
+		/// <inheritdoc/>
+		public abstract TEnum Compute(BgExprContext Context);
+
+		/// <inheritdoc/>
+		public BgString ToBgString() => new BgEnumFormatExpr<TEnum>(this);
+	}
+
+	/// <summary>
+	/// Type traits for a <see cref="BgEnum{TEnum}"/>
+	/// </summary>
+	class BgEnumType<TEnum> : BgTypeBase<BgEnum<TEnum>> where TEnum : struct
+	{
+		/// <inheritdoc/>
+		public override BgEnum<TEnum> DeserializeArgument(string Text) => Enum.Parse<TEnum>(Text);
+
+		/// <inheritdoc/>
+		public override string SerializeArgument(BgEnum<TEnum> Value, BgExprContext Context) => Value.Compute(Context).ToString() ?? String.Empty;
+
+		/// <inheritdoc/>
+		public override BgEnum<TEnum> CreateConstant(object Value) => new BgEnumConstantExpr<TEnum>((TEnum)Value);
+
+		/// <inheritdoc/>
+		public override IBgExprVariable<BgEnum<TEnum>> CreateVariable() => throw new NotImplementedException();
+	}
+
+	#region Expression classes
+
+	class BgEnumParseExpr<TEnum> : BgEnum<TEnum> where TEnum : struct
+	{
+		public BgString Value { get; }
+
+		public BgEnumParseExpr(BgString Value)
+		{
+			this.Value = Value;
+		}
+
+		public override TEnum Compute(BgExprContext Context) => Enum.Parse<TEnum>(Value.Compute(Context));
+
+		public override string ToString()
+		{
+			return $"Parse<{typeof(TEnum).Name}>({Value})";
+		}
+	}
+
+	class BgEnumConstantExpr<TEnum> : BgEnum<TEnum> where TEnum : struct
+	{
+		public TEnum Value { get; }
+
+		public BgEnumConstantExpr(TEnum Value)
+		{
+			this.Value = Value;
+		}
+
+		public override TEnum Compute(BgExprContext Context) => Value;
+
+		public override string ToString()
+		{
+			return $"\"{Value}\"";
+		}
+	}
+
+	class BgEnumChooseExpr<TEnum> : BgEnum<TEnum> where TEnum : struct
+	{
+		public BgBool Condition { get; }
+		public BgEnum<TEnum> ValueIfTrue { get; }
+		public BgEnum<TEnum> ValueIfFalse { get; }
+
+		public BgEnumChooseExpr(BgBool Condition, BgEnum<TEnum> ValueIfTrue, BgEnum<TEnum> ValueIfFalse)
+		{
+			this.Condition = Condition;
+			this.ValueIfTrue = ValueIfTrue;
+			this.ValueIfFalse = ValueIfFalse;
+		}
+
+		public override TEnum Compute(BgExprContext Context) => Condition.Compute(Context) ? ValueIfTrue.Compute(Context) : ValueIfFalse.Compute(Context);
+
+		public override string ToString()
+		{
+			return $"Choose({Condition}, {ValueIfTrue}, {ValueIfFalse});";
+		}
+	}
+
+	class BgEnumFormatExpr<TEnum> : BgString where TEnum : struct
+	{
+		public BgEnum<TEnum> Value { get; }
+
+		public BgEnumFormatExpr(BgEnum<TEnum> Value)
+		{
+			this.Value = Value;
+		}
+
+		public override string Compute(BgExprContext Context) => Value.Compute(Context).ToString() ?? String.Empty;
+
+		public override string ToString()
+		{
+			return $"Format({Value})";
+		}
+	}
+
+	class BgEnumVariableExpr<TEnum> : BgEnum<TEnum>, IBgExprVariable<BgEnum<TEnum>> where TEnum : struct
+	{
+		public BgEnum<TEnum> Value { get; set; } = null!;
+
+		public override TEnum Compute(BgExprContext Context) => Value.Compute(Context);
+	}
+
+	#endregion
+}
diff --git a/Engine/Source/Programs/Shared/EpicGames.BuildGraph/Expressions/BgFileSet.cs b/Engine/Source/Programs/Shared/EpicGames.BuildGraph/Expressions/BgFileSet.cs
new file mode 100644
--- /dev/null
+++ b/Engine/Source/Programs/Shared/EpicGames.BuildGraph/Expressions/BgFileSet.cs
@@ -0,0 +1,152 @@
+// Copyright Epic Games, Inc. All Rights Reserved.
+
+using EpicGames.Core;
+using System;
+using System.Collections.Generic;
+using System.Linq;
+using System.Text;
+
+namespace EpicGames.BuildGraph.Expressions
+{
+	/// <summary>
+	/// Represents a placeholder for the output from a node, which can be exchanged for the artifacts produced by a node at runtime
+	/// </summary>
+	[BgType(typeof(BgFileSetType))]
+	public abstract class BgFileSet : IBgExpr<BgFileSet>
+	{
+		/// <inheritdoc/>
+		public BgFileSet IfThen(BgBool Condition, BgFileSet ValueIfTrue) => new BgFileSetIfThenExpr(Condition, ValueIfTrue, this);
+
+		/// <inheritdoc/>
+		object IBgExpr.Compute(BgExprContext Context) => Compute(Context);
+
+		/// <inheritdoc/>
+		public abstract BgFileSet Compute(BgExprContext Context);
+
+		/// <summary>
+		/// Gets the tag name for this fileset
+		/// </summary>
+		/// <returns></returns>
+		public string ComputeTag(BgExprContext Context) => ((BgFileSetTagExpr)Compute(Context)).Name;
+
+		/// <summary>
+		/// Gets the tag name for this fileset
+		/// </summary>
+		/// <returns></returns>
+		public FileSet ComputeValue(BgExprContext Context) => ((BgFileSetValueExpr)Compute(Context)).FileSet;
+
+		/// <summary>
+		/// Implicit conversion from a file set to a functional file set
+		/// </summary>
+		/// <param name="FileSet"></param>
+		public static implicit operator BgFileSet(FileSet FileSet) => new BgFileSetValueExpr(FileSet);
+
+		/// <inheritdoc/>
+		public BgString ToBgString() => new BgFileSetToStringExpr(this);
+	}
+
+	/// <summary>
+	/// Traits for a <see cref="BgFileSet"/>
+	/// </summary>
+	class BgFileSetType : BgTypeBase<BgFileSet>
+	{
+		/// <inheritdoc/>
+		public override BgFileSet DeserializeArgument(string Text) => new BgFileSetTagExpr(Text);
+
+		/// <inheritdoc/>
+		public override string SerializeArgument(BgFileSet Value, BgExprContext Context) => ((BgFileSetTagExpr)Value.Compute(Context)).Name;
+
+		/// <inheritdoc/>
+		public override BgFileSet CreateConstant(object Value) => new BgFileSetTagExpr(((BgFileSetTagExpr)Value).Name);
+
+		/// <inheritdoc/>
+		public override IBgExprVariable<BgFileSet> CreateVariable() => throw new NotImplementedException();
+	}
+
+	#region Expression classes
+
+	class BgFileSetIfThenExpr : BgFileSet
+	{
+		public BgBool Condition;
+		public BgFileSet ValueIfTrue;
+		public BgFileSet ValueIfFalse;
+
+		public BgFileSetIfThenExpr(BgBool Condition, BgFileSet ValueIfTrue, BgFileSet ValueIfFalse)
+		{
+			this.Condition = Condition;
+			this.ValueIfTrue = ValueIfTrue;
+			this.ValueIfFalse = ValueIfFalse;
+		}
+
+		public override BgFileSet Compute(BgExprContext Context) => Condition.Compute(Context) ? ValueIfTrue.Compute(Context) : ValueIfFalse.Compute(Context);
+	}
+
+	class BgFileSetToStringExpr : BgString
+	{
+		BgFileSet FileSet;
+
+		internal BgFileSetToStringExpr(BgFileSet Token)
+		{
+			this.FileSet = Token;
+		}
+
+		public override string Compute(BgExprContext Context)
+		{
+			BgFileSet FileSetValue = FileSet.Compute(Context);
+			if (FileSetValue is BgFileSetTagExpr TagExpr)
+			{
+				return TagExpr.Name;
+			}
+			else
+			{
+				return "{...}";
+			}
+		}
+	}
+
+	class BgFileSetTagExpr : BgFileSet
+	{
+		public string Name { get; set; }
+
+		public BgFileSetTagExpr(string Name)
+		{
+			this.Name = Name;
+		}
+
+		public override BgFileSet Compute(BgExprContext Context) => this;
+	}
+
+	class BgFileSetTagFromStringExpr : BgFileSet
+	{
+		public BgString Name { get; }
+
+		public BgFileSetTagFromStringExpr(BgString Name)
+		{
+			this.Name = Name;
+		}
+
+		public override BgFileSet Compute(BgExprContext Context) => new BgFileSetTagExpr(Name.Compute(Context));
+	}
+
+	class BgFileSetVariableExpr : BgFileSet, IBgExprVariable<BgFileSet>
+	{
+		public BgFileSet Value { get; set; } = new BgFileSetValueExpr(FileSet.Empty);
+
+		public override BgFileSet Compute(BgExprContext Context) => Value.Compute(Context);
+	}
+
+	class BgFileSetValueExpr : BgFileSet
+	{
+		public FileSet FileSet { get; }
+
+		public BgFileSetValueExpr(FileSet FileSet)
+		{
+			this.FileSet = FileSet;
+		}
+
+		public override BgFileSet Compute(BgExprContext Context) => this;
+	}
+
+	#endregion
+}
+
diff --git a/Engine/Source/Programs/Shared/EpicGames.BuildGraph/Expressions/BgInt.cs b/Engine/Source/Programs/Shared/EpicGames.BuildGraph/Expressions/BgInt.cs
new file mode 100644
--- /dev/null
+++ b/Engine/Source/Programs/Shared/EpicGames.BuildGraph/Expressions/BgInt.cs
@@ -0,0 +1,214 @@
+// Copyright Epic Games, Inc. All Rights Reserved.
+
+using EpicGames.BuildGraph;
+using System;
+using System.Collections.Generic;
+using System.ComponentModel;
+using System.Diagnostics;
+using System.IO;
+using System.Linq;
+using System.Reflection;
+using System.Text;
+using System.Xml;
+using EpicGames.Core;
+using Microsoft.Extensions.Logging;
+using OpenTracing;
+using OpenTracing.Util;
+using System.Threading.Tasks;
+
+namespace EpicGames.BuildGraph.Expressions
+{
+	/// <summary>
+	/// Abstract base class for expressions returning a 32-bit integer value 
+	/// </summary>
+	[BgType(typeof(BgIntTraits))]
+	public abstract class BgInt : IBgExpr<BgInt>
+	{
+		/// <summary>
+		/// Implicit conversion from an integer value
+		/// </summary>
+		public static implicit operator BgInt(int Value)
+		{
+			return new BgIntConstantExpr(Value);
+		}
+
+		/// <inheritdoc/>
+		public static BgInt operator +(BgInt Lhs, BgInt Rhs) => new BgIntAddExpr(Lhs, Rhs);
+
+		/// <inheritdoc/>
+		public static BgInt operator -(BgInt Lhs, BgInt Rhs) => new BgIntAddExpr(Lhs, new BgIntNegateExpr(Rhs));
+
+		/// <inheritdoc/>
+		public static BgBool operator <(BgInt Lhs, BgInt Rhs) => new BgIntLtExpr(Lhs, Rhs);
+
+		/// <inheritdoc/>
+		public static BgBool operator >(BgInt Lhs, BgInt Rhs) => new BgIntGtExpr(Lhs, Rhs);
+
+		/// <inheritdoc/>
+		public static BgBool operator ==(BgInt Lhs, BgInt Rhs) => new BgIntEqExpr(Lhs, Rhs);
+
+		/// <inheritdoc/>
+		public static BgBool operator !=(BgInt Lhs, BgInt Rhs) => new BgBoolNotExpr(new BgIntEqExpr(Lhs, Rhs));
+
+		/// <inheritdoc/>
+		public static BgBool operator <=(BgInt Lhs, BgInt Rhs) => new BgBoolNotExpr(new BgIntGtExpr(Lhs, Rhs));
+
+		/// <inheritdoc/>
+		public static BgBool operator >=(BgInt Lhs, BgInt Rhs) => new BgBoolNotExpr(new BgIntLtExpr(Lhs, Rhs));
+
+		/// <inheritdoc/>
+		object IBgExpr.Compute(BgExprContext Context) => Compute(Context);
+
+		/// <inheritdoc/>
+		public abstract int Compute(BgExprContext Context);
+
+		/// <inheritdoc/>
+		public BgInt IfThen(BgBool Condition, BgInt ValueIfTrue) => new BgIntChooseExpr(Condition, ValueIfTrue, this);
+
+		/// <inheritdoc/>
+		public override bool Equals(object? obj) => throw new InvalidOperationException();
+
+		/// <inheritdoc/>
+		public override int GetHashCode() => throw new InvalidOperationException();
+
+		/// <inheritdoc/>
+		public BgString ToBgString() => new BgIntFormatExpr(this);
+	}
+
+	/// <summary>
+	/// Traits for a <see cref="BgInt"/>
+	/// </summary>
+	class BgIntTraits : BgTypeBase<BgInt>
+	{
+		/// <inheritdoc/>
+		public override BgInt DeserializeArgument(string Value) => (BgInt)int.Parse(Value);
+
+		/// <inheritdoc/>
+		public override string SerializeArgument(BgInt Value, BgExprContext Context) => Value.Compute(Context).ToString();
+
+		/// <inheritdoc/>
+		public override BgInt CreateConstant(object Value) => new BgIntConstantExpr((int)Value);
+
+		/// <inheritdoc/>
+		public override IBgExprVariable<BgInt> CreateVariable() => new BgIntVariableExpr();
+	}
+
+	#region Expression classes
+
+	class BgIntEqExpr : BgBool
+	{
+		public BgInt Lhs { get; }
+		public BgInt Rhs { get; }
+
+		public BgIntEqExpr(BgInt Lhs, BgInt Rhs)
+		{
+			this.Lhs = Lhs;
+			this.Rhs = Rhs;
+		}
+
+		public override bool Compute(BgExprContext Context) => Lhs.Compute(Context) == Rhs.Compute(Context);
+	}
+
+	class BgIntLtExpr : BgBool
+	{
+		public BgInt Lhs { get; }
+		public BgInt Rhs { get; }
+
+		public BgIntLtExpr(BgInt Lhs, BgInt Rhs)
+		{
+			this.Lhs = Lhs;
+			this.Rhs = Rhs;
+		}
+
+		public override bool Compute(BgExprContext Context) => Lhs.Compute(Context) < Rhs.Compute(Context);
+	}
+
+	class BgIntGtExpr : BgBool
+	{
+		public BgInt Lhs { get; }
+		public BgInt Rhs { get; }
+
+		public BgIntGtExpr(BgInt Lhs, BgInt Rhs)
+		{
+			this.Lhs = Lhs;
+			this.Rhs = Rhs;
+		}
+
+		public override bool Compute(BgExprContext Context) => Lhs.Compute(Context) > Rhs.Compute(Context);
+	}
+
+	class BgIntAddExpr : BgInt
+	{
+		public BgInt Lhs { get; }
+		public BgInt Rhs { get; }
+
+		public BgIntAddExpr(BgInt Lhs, BgInt Rhs)
+		{
+			this.Lhs = Lhs;
+			this.Rhs = Rhs;
+		}
+
+		public override int Compute(BgExprContext Context) => Lhs.Compute(Context) + Rhs.Compute(Context);
+	}
+
+	class BgIntNegateExpr : BgInt
+	{
+		public BgInt Inner { get; }
+
+		public BgIntNegateExpr(BgInt Inner)
+		{
+			this.Inner = Inner;
+		}
+
+		public override int Compute(BgExprContext Context) => -Inner.Compute(Context);
+	}
+
+	class BgIntChooseExpr : BgInt
+	{
+		public BgBool Condition;
+		public BgInt ValueIfTrue;
+		public BgInt ValueIfFalse;
+
+		public BgIntChooseExpr(BgBool Condition, BgInt ValueIfTrue, BgInt ValueIfFalse)
+		{
+			this.Condition = Condition;
+			this.ValueIfTrue = ValueIfTrue;
+			this.ValueIfFalse = ValueIfFalse;
+		}
+
+		public override int Compute(BgExprContext Context) => Condition.Compute(Context) ? ValueIfTrue.Compute(Context) : ValueIfFalse.Compute(Context);
+	}
+
+	class BgIntConstantExpr : BgInt
+	{
+		public int Value { get; }
+
+		public BgIntConstantExpr(int Value)
+		{
+			this.Value = Value;
+		}
+
+		public override int Compute(BgExprContext Context) => Value;
+	}
+
+	class BgIntVariableExpr : BgInt, IBgExprVariable<BgInt>
+	{
+		public BgInt Value { get; set; } = 0;
+
+		public override int Compute(BgExprContext Context) => Value.Compute(Context);
+	}
+
+	class BgIntFormatExpr : BgString
+	{
+		BgInt Value;
+
+		public BgIntFormatExpr(BgInt Value)
+		{
+			this.Value = Value;
+		}
+
+		public override string Compute(BgExprContext Context) => Value.Compute(Context).ToString();
+	}
+
+#endregion
+}
diff --git a/Engine/Source/Programs/Shared/EpicGames.BuildGraph/Expressions/BgList.cs b/Engine/Source/Programs/Shared/EpicGames.BuildGraph/Expressions/BgList.cs
new file mode 100644
--- /dev/null
+++ b/Engine/Source/Programs/Shared/EpicGames.BuildGraph/Expressions/BgList.cs
@@ -0,0 +1,401 @@
+// Copyright Epic Games, Inc. All Rights Reserved.
+
+using EpicGames.BuildGraph;
+using System;
+using System.Collections.Generic;
+using System.ComponentModel;
+using System.Diagnostics;
+using System.IO;
+using System.Linq;
+using System.Reflection;
+using System.Text;
+using System.Xml;
+using EpicGames.Core;
+using Microsoft.Extensions.Logging;
+using OpenTracing;
+using OpenTracing.Util;
+using System.Threading.Tasks;
+using System.Diagnostics.Contracts;
+
+namespace EpicGames.BuildGraph.Expressions
+{
+	/// <summary>
+	/// Interface for all list types
+	/// </summary>
+	public interface IBgList : IBgExpr
+	{
+		/// <summary>
+		/// Type of elements in this list
+		/// </summary>
+		Type ElementType { get; }
+	}
+
+	/// <summary>
+	/// Abstract base class for expressions returning an immutable list of values
+	/// </summary>
+	[BgType(typeof(BgListType<>))]
+	public abstract class BgList<T> : IBgExpr<BgList<T>>, IBgList where T : IBgExpr<T>
+	{
+		/// <inheritdoc/>
+		public Type ElementType => typeof(T);
+
+		/// <summary>
+		/// Constant representation of an empty list
+		/// </summary>
+		public static BgList<T> Empty { get; } = new BgListConstantExpr<T>(Array.Empty<T>());
+
+		/// <summary>
+		/// Implicit conversion operator from a single value
+		/// </summary>
+		public static implicit operator BgList<T>(T Value) => new BgListConstantExpr<T>(new[] { Value });
+
+		/// <summary>
+		/// Implicit conversion operator from an array of values
+		/// </summary>
+		public static implicit operator BgList<T>(T[] Value) => new BgListConstantExpr<T>(Value);
+
+		/// <summary>
+		/// Implicit conversion operator from a list of values
+		/// </summary>
+		public static implicit operator BgList<T>(List<T> Value) => new BgListConstantExpr<T>(Value);
+
+		/// <summary>
+		/// Crates a list from an array of values
+		/// </summary>
+		/// <param name="Items">Sequence to construct from</param>
+		/// <returns></returns>
+		public static BgList<T> Create(IEnumerable<T> Items) => new BgListConstantExpr<T>(Items);
+
+		/// <summary>
+		/// Crates a list from an array of values
+		/// </summary>
+		/// <param name="Items">Sequence to construct from</param>
+		/// <returns></returns>
+		public static BgList<T> Create(params T[] Items) => new BgListConstantExpr<T>(Items);
+
+		/// <summary>
+		/// Concatenates mutiple lists together
+		/// </summary>
+		/// <param name="Sources"></param>
+		/// <returns></returns>
+		public static BgList<T> Concat(params BgList<T>[] Sources) => new BgListConcatExpr<T>(Sources);
+
+		/// <summary>
+		/// Adds items to the end of the list, returning the new list
+		/// </summary>
+		/// <param name="Items">Items to add</param>
+		/// <returns>New list containing the given items</returns>
+		public BgList<T> Add(params T[] Items) => new BgListConcatExpr<T>(this, Items);
+
+		/// <inheritdoc cref="Add(T[])"/>
+		public BgList<T> Add(params BgList<T>[] Items) => new BgListConcatExpr<T>(this, Items);
+
+		/// <inheritdoc cref="Add(T[])"/>
+		public BgList<T> AddIf(BgBool Condition, params T[] Items) => IfThen(Condition, Add(Items));
+
+		/// <inheritdoc cref="Add(T[])"/>
+		public BgList<T> AddIf(BgBool Condition, params BgList<T>[] Items) => IfThen(Condition, Add(Items));
+
+		/// <summary>
+		/// Removes the given items from this list
+		/// </summary>
+		/// <param name="Items">Items to remove</param>
+		/// <returns>New list without the given items</returns>
+		public BgList<T> Except(BgList<T> Items) => new BgListExceptExpr<T>(this, Items);
+
+		/// <summary>
+		/// Removes the given items from this list
+		/// </summary>
+		/// <param name="Condition">Condition to remove the items</param>
+		/// <param name="Items">Items to remove</param>
+		/// <returns>New list without the given items</returns>
+		public BgList<T> ExceptIf(BgBool Condition, BgList<T> Items) => IfThen(Condition, Except(Items));
+
+		/// <summary>
+		/// Removes any duplicate items from the list. The first item in the list is retained in its original order.
+		/// </summary>
+		/// <returns>New list containing the distinct items.</returns>
+		public BgList<T> Distinct() => new BgListDistinctExpr<T>(this);
+
+		/// <inheritdoc cref="Enumerable.Select{TSource, TResult}(IEnumerable{TSource}, Func{TSource, TResult})"/>
+		public BgList<TResult> Select<TResult>(Func<T, TResult> Function) where TResult : IBgExpr<TResult> => new BgListSelectExpr<T, TResult>(this, Function);
+
+		/// <inheritdoc cref="Enumerable.Where{TSource}(IEnumerable{TSource}, Func{TSource, bool})"/>
+		public BgList<T> Where(Func<T, BgBool> Predicate) => new BgListWhereExpr<T>(this, Predicate);
+
+		/// <inheritdoc cref="Enumerable.Contains{TSource}(IEnumerable{TSource}, TSource)"/>
+		public BgBool Contains(T Item) => new BgListContainsExpr<T>(this, Item);
+
+		/// <inheritdoc/>
+		object IBgExpr.Compute(BgExprContext Context) => GetEnumerable(Context);
+
+		/// <inheritdoc/>
+		public abstract IEnumerable<T> GetEnumerable(BgExprContext Context);
+
+		/// <inheritdoc/>
+		public BgList<T> IfThen(BgBool Condition, BgList<T> ValueIfTrue) => new BgListChooseExpr<T>(Condition, ValueIfTrue, this);
+
+		/// <inheritdoc/>
+		public BgString ToBgString() => new BgListFormatExpr<T>(this);
+	}
+
+	/// <summary>
+	/// Traits for a <see cref="BgList{T}"/>
+	/// </summary>
+	class BgListType<T> : BgTypeBase<BgList<T>> where T : IBgExpr<T>
+	{
+		/// <inheritdoc/>
+		public override BgList<T> DeserializeArgument(string Text)
+		{
+			IBgType<T> Converter = BgType.Get<T>();
+			return BgList<T>.Create(Text.Split(';').Select(x => Converter.DeserializeArgument(x)));
+		}
+
+		/// <inheritdoc/>
+		public override string SerializeArgument(BgList<T> Value, BgExprContext Context)
+		{
+			IBgType<T> Converter = BgType.Get<T>();
+			return String.Join(";", Value.GetEnumerable(Context).Select(x => Converter.SerializeArgument(x, Context)));
+		}
+
+		/// <inheritdoc/>
+		public override BgList<T> CreateConstant(object Value)
+		{
+			return new BgListConstantExpr<T>((IEnumerable<T>)Value);
+		}
+
+		/// <inheritdoc/>
+		public override IBgExprVariable<BgList<T>> CreateVariable()
+		{
+			throw new NotImplementedException();
+		}
+	}
+
+	/// <summary>
+	/// Extension methods for lists
+	/// </summary>
+	public static class BgListExtensions
+	{
+		internal static List<string> Compute(this BgList<BgString> List, BgExprContext Context)
+		{
+			return List.GetEnumerable(Context).Select(x => x.Compute(Context)).ToList();
+		}
+
+		internal static List<BgFileSet> Compute(this BgList<BgFileSet> List, BgExprContext Context)
+		{
+			return List.GetEnumerable(Context).Select(x => x.Compute(Context)).ToList();
+		}
+
+		internal static List<string> ComputeTags(this BgList<BgFileSet> List, BgExprContext Context)
+		{
+			return List.GetEnumerable(Context).Select(x => x.ComputeTag(Context)).ToList();
+		}
+	}
+
+	#region Expression classes
+
+	class BgListConcatExpr<T> : BgList<T> where T : IBgExpr<T>
+	{
+		public BgList<T>[] Sources;
+
+		public BgListConcatExpr(params BgList<T>[] Sources)
+		{
+			this.Sources = Sources;
+		}
+
+		public BgListConcatExpr(BgList<T> FirstSource, BgList<T>[] OtherSources)
+		{
+			Sources = new BgList<T>[OtherSources.Length + 1];
+			Sources[0] = FirstSource;
+			Array.Copy(OtherSources, 0, Sources, 1, OtherSources.Length);
+		}
+
+		public override IEnumerable<T> GetEnumerable(BgExprContext Context)
+		{
+			return Sources.SelectMany(x => x.GetEnumerable(Context));
+		}
+	}
+
+	class BgListChooseExpr<T> : BgList<T> where T : IBgExpr<T>
+	{
+		public BgBool Condition;
+		public BgList<T> ValueIfTrue;
+		public BgList<T> ValueIfFalse;
+
+		public BgListChooseExpr(BgBool Condition, BgList<T> ValueIfTrue, BgList<T> ValueIfFalse)
+		{
+			this.Condition = Condition;
+			this.ValueIfTrue = ValueIfTrue;
+			this.ValueIfFalse = ValueIfFalse;
+		}
+
+		public override IEnumerable<T> GetEnumerable(BgExprContext Context) => Condition.Compute(Context) ? ValueIfTrue.GetEnumerable(Context) : ValueIfFalse.GetEnumerable(Context);
+	}
+
+	class BgListConstantExpr<T> : BgList<T> where T : IBgExpr<T>
+	{
+		public IEnumerable<T> Value { get; }
+
+		public BgListConstantExpr(IEnumerable<T> Value)
+		{
+			this.Value = Value;
+		}
+
+		public override IEnumerable<T> GetEnumerable(BgExprContext Context) => Value;
+	}
+
+	class BgListFormatExpr<T> : BgString where T : IBgExpr<T>
+	{
+		public BgList<T> Inner;
+
+		public BgListFormatExpr(BgList<T> Inner)
+		{
+			this.Inner = Inner;
+		}
+
+		public override string Compute(BgExprContext Context) => String.Join(";", Inner.GetEnumerable(Context).Select(x => x.ToBgString().Compute(Context)));
+	}
+
+	class BgListVariableExpr<T> : BgList<T>, IBgExprVariable<BgList<T>> where T : IBgExpr<T>
+	{
+		public BgList<T> Value { get; set; } = BgList<T>.Empty;
+
+		public override IEnumerable<T> GetEnumerable(BgExprContext Context) => Value.GetEnumerable(Context);
+	}
+
+	class BgListDistinctExpr<T> : BgList<T> where T : IBgExpr<T>
+	{
+		BgList<T> Source;
+
+		public BgListDistinctExpr(BgList<T> Source)
+		{
+			this.Source = Source;
+		}
+
+		public override IEnumerable<T> GetEnumerable(BgExprContext Context)
+		{
+			IBgType<T> Type = BgType.Get<T>();
+
+			HashSet<object> Values = new HashSet<object>();
+			foreach (T Item in Source.GetEnumerable(Context))
+			{
+				object ItemValue = Item.Compute(Context);
+				if (Values.Add(ItemValue))
+				{
+					yield return Item;
+				}
+			}
+		}
+	}
+
+	class BgFunc<TIn, TOut> where TIn : IBgExpr<TIn> where TOut : IBgExpr<TOut>
+	{
+		IBgExprVariable<TIn> ArgVar;
+		IBgExpr<TOut> FuncExpr;
+
+		public BgFunc(IBgExprVariable<TIn> ArgVar, IBgExpr<TOut> FuncExpr)
+		{
+			this.ArgVar = ArgVar;
+			this.FuncExpr = FuncExpr;
+		}
+
+		public static implicit operator BgFunc<TIn, TOut>(Func<TIn, TOut> Function)
+		{
+			IBgExprVariable<TIn> ArgVar = BgType.Get<TIn>().CreateVariable();
+			return new BgFunc<TIn, TOut>(ArgVar, Function((TIn)ArgVar));
+		}
+
+		public object Compute(BgExprContext Context, TIn Argument)
+		{
+			ArgVar.Value = Argument;
+			return FuncExpr.Compute(Context);
+		}
+
+		public BgString ToBgString()
+		{
+			throw new NotImplementedException();
+		}
+	}
+
+	class BgListSelectExpr<TIn, TOut> : BgList<TOut> where TIn : IBgExpr<TIn> where TOut : IBgExpr<TOut>
+	{
+		BgList<TIn> Source;
+		BgFunc<TIn, TOut> Function;
+
+		public BgListSelectExpr(BgList<TIn> Source, BgFunc<TIn, TOut> Function)
+		{
+			this.Source = Source;
+			this.Function = Function;
+		}
+
+		public override IEnumerable<TOut> GetEnumerable(BgExprContext Context)
+		{
+			return Source.GetEnumerable(Context).Select(x => (TOut)Function.Compute(Context, x));
+		}
+	}
+
+	class BgListWhereExpr<T> : BgList<T> where T : IBgExpr<T>
+	{
+		BgList<T> Source;
+		BgFunc<T, BgBool> Predicate;
+
+		public BgListWhereExpr(BgList<T> Source, BgFunc<T, BgBool> Predicate)
+		{
+			this.Source = Source;
+			this.Predicate = Predicate;
+		}
+
+		public override IEnumerable<T> GetEnumerable(BgExprContext Context)
+		{
+			return Source.GetEnumerable(Context).Where(x => (bool)Predicate.Compute(Context, x));
+		}
+	}
+
+	class BgListExceptExpr<T> : BgList<T> where T : IBgExpr<T>
+	{
+		BgList<T> SourceList;
+		BgList<T> ExceptList;
+
+		public BgListExceptExpr(BgList<T> Source, BgList<T> Except)
+		{
+			this.SourceList = Source;
+			this.ExceptList = Except;
+		}
+
+		public override IEnumerable<T> GetEnumerable(BgExprContext Context)
+		{
+			HashSet<object> ExceptValues = new HashSet<object>();
+			ExceptValues.UnionWith(ExceptList.GetEnumerable(Context).Select(x => x.Compute(Context)));
+
+			IBgType<T> Type = BgType.Get<T>();
+			foreach (T SourceItem in SourceList.GetEnumerable(Context))
+			{
+				object SourceValue = SourceItem.Compute(Context);
+				if (!ExceptValues.Contains(SourceValue))
+				{
+					yield return Type.CreateConstant(SourceValue);
+				}
+			}
+		}
+	}
+
+	class BgListContainsExpr<T> : BgBool where T : IBgExpr<T>
+	{
+		BgList<T> Source;
+		T Item;
+
+		public BgListContainsExpr(BgList<T> Source, T Item)
+		{
+			this.Source = Source;
+			this.Item = Item;
+		}
+
+		public override bool Compute(BgExprContext Context)
+		{
+			object Value = Item.Compute(Context);
+			return Source.GetEnumerable(Context).Any(x => x.Compute(Context).Equals(Value));
+		}
+	}
+
+	#endregion
+}
diff --git a/Engine/Source/Programs/Shared/EpicGames.BuildGraph/Expressions/BgString.cs b/Engine/Source/Programs/Shared/EpicGames.BuildGraph/Expressions/BgString.cs
new file mode 100644
--- /dev/null
+++ b/Engine/Source/Programs/Shared/EpicGames.BuildGraph/Expressions/BgString.cs
@@ -0,0 +1,282 @@
+// Copyright Epic Games, Inc. All Rights Reserved.
+
+using EpicGames.BuildGraph;
+using System;
+using System.Collections.Generic;
+using System.ComponentModel;
+using System.Diagnostics;
+using System.IO;
+using System.Linq;
+using System.Reflection;
+using System.Text;
+using System.Xml;
+using EpicGames.Core;
+using Microsoft.Extensions.Logging;
+using OpenTracing;
+using OpenTracing.Util;
+using System.Threading.Tasks;
+using System.Text.RegularExpressions;
+
+namespace EpicGames.BuildGraph.Expressions
+{
+	/// <summary>
+	/// Abstract base class for expressions returning a string value 
+	/// </summary>
+	[BgType(typeof(BgStringTraits))]
+	public abstract class BgString : IBgExpr<BgString>
+	{
+		/// <summary>
+		/// Constant value for an empty string
+		/// </summary>
+		public static BgString Empty { get; } = new BgStringConstantExpr(String.Empty);
+
+		/// <summary>
+		/// Implicit conversion from a regular string type
+		/// </summary>
+		public static implicit operator BgString(string Value)
+		{
+			return new BgStringConstantExpr(Value);
+		}
+
+		/// <inheritdoc cref="String.Equals(string?, string?, StringComparison)"/>
+		public static BgBool Equals(BgString Lhs, BgString Rhs, StringComparison Comparison = StringComparison.CurrentCulture) => Compare(Lhs, Rhs, Comparison) == 0;
+
+		/// <inheritdoc cref="String.Compare(string?, string?, StringComparison)"/>
+		public static BgInt Compare(BgString Lhs, BgString Rhs, StringComparison Comparison = StringComparison.CurrentCulture) => new BgStringCompareExpr(Lhs, Rhs, Comparison);
+
+		/// <inheritdoc cref="String.Join{T}(string?, IEnumerable{T})"/>
+		public static BgString Join(BgString Separator, BgList<BgString> Values) => new BgStringJoinExpr(Separator, Values);
+
+		/// <inheritdoc cref="String.Format(string, object?[])"/>
+		public static BgString Format(string Format, params object?[] Args) => new BgStringFormatExpr(Format, Args);
+
+		/// <inheritdoc/>
+		public static BgString operator +(BgString Lhs, BgString Rhs) => new BgStringConcatExpr(Lhs, Rhs);
+
+		/// <inheritdoc/>
+		public static BgBool operator ==(BgString Lhs, BgString Rhs) => Equals(Lhs, Rhs);
+
+		/// <inheritdoc/>
+		public static BgBool operator !=(BgString Lhs, BgString Rhs) => !Equals(Lhs, Rhs);
+
+		/// <summary>
+		/// Appens another string to this one
+		/// </summary>
+		public BgString Append(BgString Other) => this + Other;
+
+		/// <summary>
+		/// Appens another string to this one if a condition is true
+		/// </summary>
+		public BgString AppendIf(BgBool Condition, BgString Other) => IfThen(Condition, this + Other);
+
+		/// <inheritdoc/>
+		public BgString IfThen(BgBool Condition, BgString ValueIfTrue) => new BgStringChooseExpr(Condition, ValueIfTrue, this);
+
+		/// <inheritdoc/>
+		public BgBool IsMatch(BgString Input, string Pattern) => new BgStringIsMatchExpr(Input, Pattern);
+
+		/// <inheritdoc/>
+		public BgString Replace(BgString Input, string Pattern, string Replace) => new BgStringReplaceExpr(Input, Pattern, Replace);
+
+		/// <inheritdoc/>
+		object IBgExpr.Compute(BgExprContext Context) => Compute(Context);
+
+		/// <inheritdoc/>
+		public abstract string Compute(BgExprContext Context);
+
+		/// <inheritdoc/>
+		public override bool Equals(object? obj) => throw new InvalidOperationException();
+
+		/// <inheritdoc/>
+		public override int GetHashCode() => throw new InvalidOperationException();
+
+		/// <inheritdoc/>
+		public BgString ToBgString() => this;
+	}
+
+	/// <summary>
+	/// Traits implementation for <see cref="BgString"/>
+	/// </summary>
+	class BgStringTraits : BgTypeBase<BgString>
+	{
+		/// <inheritdoc/>
+		public override BgString DeserializeArgument(string Text) => Text;
+
+		/// <inheritdoc/>
+		public override string SerializeArgument(BgString Value, BgExprContext Context) => Value.Compute(Context);
+
+		/// <inheritdoc/>
+		public override BgString CreateConstant(object Value) => new BgStringConstantExpr((string)Value);
+
+		/// <inheritdoc/>
+		public override IBgExprVariable<BgString> CreateVariable() => new BgStringVariableExpr();
+	}
+
+	#region Expression classes
+
+	class BgStringCompareExpr : BgInt
+	{
+		public BgString Lhs { get; }
+		public BgString Rhs { get; }
+		public StringComparison Comparison { get; }
+
+		public BgStringCompareExpr(BgString Lhs, BgString Rhs, StringComparison Comparison)
+		{
+			this.Lhs = Lhs;
+			this.Rhs = Rhs;
+			this.Comparison = Comparison;
+		}
+
+		public override int Compute(BgExprContext Context) => String.Compare(Lhs.Compute(Context), Rhs.Compute(Context), Comparison);
+
+		public override string ToString() => $"Compare({Lhs}, {Rhs})";
+	}
+
+	class BgStringConcatExpr : BgString
+	{
+		public BgString Lhs { get; }
+		public BgString Rhs { get; }
+
+		public BgStringConcatExpr(BgString Lhs, BgString Rhs)
+		{
+			this.Lhs = Lhs;
+			this.Rhs = Rhs;
+		}
+
+		public override string Compute(BgExprContext Context) => Lhs.Compute(Context) + Rhs.Compute(Context);
+
+		public override string ToString() => $"Concat({Lhs}, {Rhs})";
+	}
+
+	class BgStringJoinExpr : BgString
+	{
+		public BgString Separator { get; }
+		public BgList<BgString> Values { get; }
+
+		public BgStringJoinExpr(BgString Separator, BgList<BgString> Values)
+		{
+			this.Separator = Separator;
+			this.Values = Values;
+		}
+
+		public override string Compute(BgExprContext Context) => String.Join(Separator.Compute(Context), Values.Compute(Context));
+
+		public override string ToString() => $"Join({Separator}, {Values})";
+	}
+
+	class BgStringIsMatchExpr : BgBool
+	{
+		public BgString Input { get; }
+		public string Pattern { get; }
+
+		public BgStringIsMatchExpr(BgString Input, string Pattern)
+		{
+			this.Input = Input;
+			this.Pattern = Pattern;
+		}
+
+		public override bool Compute(BgExprContext Context) => Regex.IsMatch(Input.Compute(Context), Pattern);
+	}
+
+	class BgStringReplaceExpr : BgString
+	{
+		public BgString Input { get; }
+		public string Pattern { get; }
+		public string Replacement { get; }
+
+		public BgStringReplaceExpr(BgString Input, string Pattern, string Replacement)
+		{
+			this.Input = Input;
+			this.Pattern = Pattern;
+			this.Replacement = Replacement;
+		}
+
+		public override string Compute(BgExprContext Context) => Regex.Replace(Input.Compute(Context), Pattern, Replacement);
+	}
+
+	class BgStringFormatExpr : BgString
+	{
+		new string Format { get; }
+		BgString[] Arguments { get; }
+
+		public BgStringFormatExpr(string Format, object?[] Args)
+		{
+			this.Format = Format;
+
+			Arguments = new BgString[Args.Length];
+			for (int Idx = 0; Idx < Args.Length; Idx++)
+			{
+				if (Args[Idx] is IBgExpr Expr)
+				{
+					Arguments[Idx] = Expr.ToBgString();
+				}
+				else if (Args[Idx] != null)
+				{
+					Arguments[Idx] = Args[Idx]?.ToString() ?? String.Empty;
+				}
+				else
+				{
+					Arguments[Idx] = String.Empty;
+				}
+			}
+		}
+
+		public override string Compute(BgExprContext Context) => String.Format(Format, Arguments.Select(x => x.Compute(Context)).ToArray());
+
+		public override string ToString() => $"Format({String.Join(", ", new[] { Format }.Concat(Arguments.Select(x => x.ToString())))})";
+	}
+
+	class BgStringChooseExpr : BgString
+	{
+		public BgBool Condition;
+		public BgString ValueIfTrue;
+		public BgString ValueIfFalse;
+
+		public BgStringChooseExpr(BgBool Condition, BgString ValueIfTrue, BgString ValueIfFalse)
+		{
+			this.Condition = Condition;
+			this.ValueIfTrue = ValueIfTrue;
+			this.ValueIfFalse = ValueIfFalse;
+		}
+
+		public override string Compute(BgExprContext Context) => Condition.Compute(Context) ? ValueIfTrue.Compute(Context) : ValueIfFalse.Compute(Context);
+	}
+
+	class BgStringVariableExpr : BgString, IBgExprVariable<BgString>
+	{
+		public BgString Value { get; set; } = BgString.Empty;
+
+		public BgString Variable => this;
+
+		public override string Compute(BgExprContext Context) => Value.Compute(Context);
+	}
+
+	class BgStringConstantExpr : BgString
+	{
+		public string Value { get; }
+
+		public BgStringConstantExpr(string Value)
+		{
+			this.Value = Value;
+		}
+
+		public override string Compute(BgExprContext Context) => Value;
+
+		public override bool Equals(object? Obj)
+		{
+			return Obj is string ObjStr && String.Equals(Value, ObjStr, StringComparison.OrdinalIgnoreCase);
+		}
+
+		public override int GetHashCode()
+		{
+			return String.GetHashCode(Value, StringComparison.OrdinalIgnoreCase);
+		}
+
+		public override string ToString()
+		{
+			return $"\"Value\"";
+		}
+	}
+
+	#endregion
+}
diff --git a/Engine/Source/Programs/Shared/EpicGames.BuildGraph/Expressions/IBgExpr.cs b/Engine/Source/Programs/Shared/EpicGames.BuildGraph/Expressions/IBgExpr.cs
new file mode 100644
--- /dev/null
+++ b/Engine/Source/Programs/Shared/EpicGames.BuildGraph/Expressions/IBgExpr.cs
@@ -0,0 +1,69 @@
+// Copyright Epic Games, Inc. All Rights Reserved.
+
+using EpicGames.BuildGraph.Expressions;
+using EpicGames.Core;
+using System;
+using System.Collections.Generic;
+using System.Reflection;
+using System.Text;
+
+namespace EpicGames.BuildGraph
+{
+	/// <summary>
+	/// Context object used for evaluating expressions
+	/// </summary>
+	public class BgExprContext
+	{
+		/// <summary>
+		/// Command line options for the graph
+		/// </summary>
+		public Dictionary<string, string> Options = new Dictionary<string, string>();
+
+		/// <summary>
+		/// Map from tag name to fileset
+		/// </summary>
+		public Dictionary<string, FileSet> TagNameToFileSet = new Dictionary<string, FileSet>();
+	}
+
+	/// <summary>
+	/// Base class for computable expressions
+	/// </summary>
+	public interface IBgExpr
+	{
+		/// <summary>
+		/// Compute the value of this expression
+		/// </summary>
+		/// <param name="Context"></param>
+		object Compute(BgExprContext Context);
+
+		/// <summary>
+		/// Helper method to create an expression which formats this expression as a string
+		/// </summary>
+		BgString ToBgString();
+	}
+
+	/// <summary>
+	/// Interface for a computable expression.
+	/// </summary>
+	public interface IBgExpr<TExpr> : IBgExpr
+	{
+		/// <summary>
+		/// Chooses between the current value and another value based on a condition
+		/// </summary>
+		/// <param name="Condition">The condition to evaluate</param>
+		/// <param name="ValueIfTrue">Value to return if the condition is true</param>
+		TExpr IfThen(BgBool Condition, TExpr ValueIfTrue);
+	}
+
+	/// <summary>
+	/// Sentinel expression node which can be assigned different values
+	/// </summary>
+	/// <typeparam name="TExpr">The expression type</typeparam>
+	public interface IBgExprVariable<TExpr> where TExpr : IBgExpr<TExpr>
+	{
+		/// <summary>
+		/// The current value of the variable
+		/// </summary>
+		TExpr Value { get; set; }
+	}
+}
diff --git a/Engine/Source/Programs/Shared/EpicGames.BuildGraph/Expressions/IBgType.cs b/Engine/Source/Programs/Shared/EpicGames.BuildGraph/Expressions/IBgType.cs
new file mode 100644
--- /dev/null
+++ b/Engine/Source/Programs/Shared/EpicGames.BuildGraph/Expressions/IBgType.cs
@@ -0,0 +1,196 @@
+// Copyright Epic Games, Inc. All Rights Reserved.
+
+using System;
+using System.Collections.Concurrent;
+using System.Collections.Generic;
+using System.Reflection;
+using System.Text;
+
+namespace EpicGames.BuildGraph.Expressions
+{
+	/// <summary>
+	/// Converts expressions to and from strings
+	/// </summary>
+	public interface IBgType
+	{
+		/// <summary>
+		/// Serializes the value of an expression to a string
+		/// </summary>
+		/// <param name="Value">The value to serialize</param>
+		/// <param name="Context">Context for evaluating the expression</param>
+		/// <returns></returns>
+		string SerializeArgument(object Value, BgExprContext Context);
+
+		/// <summary>
+		/// Constructs an expression from a string value
+		/// </summary>
+		/// <param name="Text">The serialized value</param>
+		/// <returns>Instance of the expression</returns>
+		object DeserializeArgument(string Text);
+	}
+
+	/// <summary>
+	/// Converts expressions to and from strings
+	/// </summary>
+	public interface IBgType<T> : IBgType where T : IBgExpr<T>
+	{
+		/// <summary>
+		/// Serializes the value of an expression to an argument string
+		/// </summary>
+		/// <param name="Value">The value to serialize</param>
+		/// <param name="Context">Context for evaluating the expression</param>
+		/// <returns></returns>
+		string SerializeArgument(T Value, BgExprContext Context);
+
+		/// <summary>
+		/// Constructs an expression from an argument value
+		/// </summary>
+		/// <param name="Text">The serialized value</param>
+		/// <returns>Instance of the expression</returns>
+		new T DeserializeArgument(string Text);
+
+		/// <summary>
+		/// Creates a value of the expression type
+		/// </summary>
+		/// <param name="Value">Value to wrap</param>
+		/// <returns></returns>
+		T CreateConstant(object Value);
+
+		/// <summary>
+		/// Creates a variable of the given type
+		/// </summary>
+		/// <returns></returns>
+		IBgExprVariable<T> CreateVariable();
+	}
+
+	/// <summary>
+	/// Base class for implementations of <see cref="IBgType{T}"/>
+	/// </summary>
+	public abstract class BgTypeBase<T> : IBgType<T> where T : IBgExpr<T>
+	{
+		/// <inheritdoc/>
+		object IBgType.DeserializeArgument(string Text) => DeserializeArgument(Text);
+
+		/// <inheritdoc/>
+		public abstract T DeserializeArgument(string Text);
+
+		/// <inheritdoc/>
+		string IBgType.SerializeArgument(object Value, BgExprContext Context) => SerializeArgument((T)Value, Context);
+
+		/// <inheritdoc/>
+		public abstract string SerializeArgument(T Value, BgExprContext Context);
+
+		/// <inheritdoc/>
+		public abstract T CreateConstant(object Value);
+
+		/// <inheritdoc/>
+		public abstract IBgExprVariable<T> CreateVariable();
+	}
+
+	/// <summary>
+	/// Attribute used to specify the converter class to use for a type
+	/// </summary>
+	class BgTypeAttribute : Attribute
+	{
+		/// <summary>
+		/// The converter type
+		/// </summary>
+		public Type Type { get; }
+
+		/// <summary>
+		/// Constructor
+		/// </summary>
+		/// <param name="Type"></param>
+		public BgTypeAttribute(Type Type)
+		{
+			this.Type = Type;
+		}
+	}
+
+	/// <summary>
+	/// Utility methods for type traits
+	/// </summary>
+	public static class BgType
+	{
+		/// <summary>
+		/// Cache of traits implementations for each type
+		/// </summary>
+		static ConcurrentDictionary<Type, IBgType> TypeToTraits = new ConcurrentDictionary<Type, IBgType>();
+
+		/// <summary>
+		/// Cache of traits for a static type
+		/// </summary>
+		class BgTypeTraitsCache<T> where T : IBgExpr<T>
+		{
+			public static IBgType<T> Traits { get; } = (IBgType<T>)BgType.Get(typeof(T));
+		}
+
+		/// <summary>
+		/// Create a converter instance for the given type
+		/// </summary>
+		/// <param name="Type"></param>
+		/// <returns></returns>
+		public static IBgType Get(Type Type)
+		{
+			IBgType? Traits = GetInner(Type);
+			if(Traits == null)
+			{
+				throw new ArgumentException($"Missing converter attribute on type {Type.Name}");
+			}
+			return Traits;
+		}
+
+		/// <summary>
+		/// Helper method for getting the type traits from a hierarchy of types
+		/// </summary>
+		/// <param name="Type"></param>
+		/// <returns></returns>
+		static IBgType? GetInner(Type Type)
+		{
+			IBgType? Traits;
+			if (!TypeToTraits.TryGetValue(Type, out Traits))
+			{
+				BgTypeAttribute? ConverterAttr = Type.GetCustomAttribute<BgTypeAttribute>(false);
+				if (ConverterAttr == null)
+				{
+					if (Type.BaseType == null)
+					{
+						Traits = null;
+					}
+					else
+					{
+						Traits = GetInner(Type.BaseType);
+					}
+				}
+				else
+				{
+					Type ConverterType = ConverterAttr.Type;
+					if (ConverterType.IsGenericType)
+					{
+						ConverterType = ConverterType.MakeGenericType(Type.GetGenericArguments());
+					}
+
+					IBgType NewTraits = (IBgType)Activator.CreateInstance(ConverterType)!;
+					if (TypeToTraits.TryAdd(Type, NewTraits))
+					{
+						Traits = NewTraits;
+					}
+					else
+					{
+						Traits = TypeToTraits[Type];
+					}
+				}
+			}
+			return Traits;
+		}
+
+		/// <summary>
+		/// Create a traits instance for the given type
+		/// </summary>
+		/// <returns></returns>
+		public static IBgType<T> Get<T>() where T : IBgExpr<T>
+		{
+			return BgTypeTraitsCache<T>.Traits;
+		}
+	}
+}
diff --git a/Engine/Source/Programs/Shared/EpicGames.Core/FileSet.cs b/Engine/Source/Programs/Shared/EpicGames.Core/FileSet.cs
new file mode 100644
--- /dev/null
+++ b/Engine/Source/Programs/Shared/EpicGames.Core/FileSet.cs
@@ -0,0 +1,475 @@
+// Copyright Epic Games, Inc. All Rights Reserved.
+
+using System;
+using System.Collections;
+using System.Collections.Generic;
+using System.IO;
+using System.Linq;
+using System.Text;
+
+namespace EpicGames.Core
+{
+	/// <summary>
+	/// Describes a tree of files represented by some arbitrary type. Allows manipulating files/directories in a functional manner; 
+	/// filtering a view of a certain directory, mapping files from one location to another, etc... before actually realizing those changes on disk.
+	/// </summary>
+	public abstract class FileSet : IEnumerable<FileReference>
+	{
+		/// <summary>
+		/// An empty fileset
+		/// </summary>
+		public static FileSet Empty { get; } = new FileSetFromFiles(Enumerable.Empty<(string, FileReference)>());
+		
+		/// <summary>
+		/// Path of this tree
+		/// </summary>
+		public string Path { get; }
+
+		/// <summary>
+		/// Constructor
+		/// </summary>
+		/// <param name="Path">Relative path within the tree</param>
+		public FileSet(string Path)
+		{
+			this.Path = Path;
+		}
+
+		/// <summary>
+		/// Enumerate files in the current tree
+		/// </summary>
+		/// <returns>Sequence consisting of names and file objects</returns>
+		public abstract IEnumerable<KeyValuePair<string, FileReference>> EnumerateFiles();
+
+		/// <summary>
+		/// Enumerate subtrees in the current tree
+		/// </summary>
+		/// <returns>Sequence consisting of names and subtree objects</returns>
+		public abstract IEnumerable<KeyValuePair<string, FileSet>> EnumerateDirectories();
+
+		/// <summary>
+		/// Creates a file tree from a given set of files
+		/// </summary>
+		/// <param name="Files"></param>
+		/// <returns>Tree containing the given files</returns>
+		public static FileSet FromFile(DirectoryReference Directory, string File)
+		{
+			return FromFiles(new[] { (File, FileReference.Combine(Directory, File)) });
+		}
+
+		/// <summary>
+		/// Creates a file tree from a given set of files
+		/// </summary>
+		/// <param name="Files"></param>
+		/// <returns>Tree containing the given files</returns>
+		public static FileSet FromFiles(IEnumerable<(string, FileReference)> Files)
+		{
+			return new FileSetFromFiles(Files);
+		}
+
+		/// <summary>
+		/// Creates a file tree from a given set of files
+		/// </summary>
+		/// <param name="Files"></param>
+		/// <returns>Tree containing the given files</returns>
+		public static FileSet FromFile(DirectoryReference Directory, FileReference File)
+		{
+			return FromFiles(Directory, new[] { File });
+		}
+
+		/// <summary>
+		/// Creates a file tree from a given set of files
+		/// </summary>
+		/// <param name="Files"></param>
+		/// <returns>Tree containing the given files</returns>
+		public static FileSet FromFiles(DirectoryReference Directory, IEnumerable<FileReference> Files)
+		{
+			return new FileSetFromFiles(Files.Select(x => (x.MakeRelativeTo(Directory), x)));
+		}
+
+		/// <summary>
+		/// Creates a file tree from a folder on disk
+		/// </summary>
+		/// <param name="Directory"></param>
+		/// <returns></returns>
+		public static FileSet FromDirectory(DirectoryReference Directory)
+		{
+			return new FileSetFromDirectory(new DirectoryInfo(Directory.FullName));
+		}
+
+		/// <summary>
+		/// Creates a file tree from a folder on disk
+		/// </summary>
+		/// <param name="DirectoryInfo"></param>
+		/// <returns></returns>
+		public static FileSet FromDirectory(DirectoryInfo DirectoryInfo)
+		{
+			return new FileSetFromDirectory(DirectoryInfo);
+		}
+
+		/// <summary>
+		/// Create a tree containing files filtered by any of the given wildcards
+		/// </summary>
+		/// <param name="Rules"></param>
+		/// <returns></returns>
+		public FileSet Filter(string Rules)
+		{
+			return Filter(Rules.Split(';'));
+		}
+
+		/// <summary>
+		/// Create a tree containing files filtered by any of the given wildcards
+		/// </summary>
+		/// <param name="Rules"></param>
+		/// <returns></returns>
+		public FileSet Filter(params string[] Rules)
+		{
+			return new FileSetFromFilter(this, new FileFilter(Rules));
+		}
+
+		/// <summary>
+		/// Create a tree containing files filtered by any of the given file filter objects
+		/// </summary>
+		/// <param name="Filters"></param>
+		/// <returns></returns>
+		public FileSet Filter(params FileFilter[] Filters)
+		{
+			return new FileSetFromFilter(this, Filters);
+		}
+
+		/// <summary>
+		/// Create a tree containing the exception of files with another tree
+		/// </summary>
+		/// <param name="Filter">Files to exclude from the filter</param>
+		/// <returns></returns>
+		public FileSet Except(string Filter)
+		{
+			return Except(Filter.Split(';'));
+		}
+
+		/// <summary>
+		/// Create a tree containing the exception of files with another tree
+		/// </summary>
+		/// <param name="Filter">Files to exclude from the filter</param>
+		/// <returns></returns>
+		public FileSet Except(params string[] Rules)
+		{
+			return new FileSetFromFilter(this, new FileFilter(Rules.Select(x => $"-{x}")));
+		}
+
+		/// <summary>
+		/// Create a tree containing the union of files with another tree
+		/// </summary>
+		/// <param name="Lhs"></param>
+		/// <param name="Rhs"></param>
+		/// <returns></returns>
+		public static FileSet Union(FileSet Lhs, FileSet Rhs)
+		{
+			return new FileSetFromUnion(Lhs, Rhs);
+		}
+
+		/// <summary>
+		/// Create a tree containing the exception of files with another tree
+		/// </summary>
+		/// <param name="Lhs"></param>
+		/// <param name="Rhs"></param>
+		/// <returns></returns>
+		public static FileSet Except(FileSet Lhs, FileSet Rhs)
+		{
+			return new FileSetFromExcept(Lhs, Rhs);
+		}
+
+		/// <inheritdoc cref="Union(FileSet, FileSet)"/>
+		public static FileSet operator +(FileSet Lhs, FileSet Rhs)
+		{
+			return Union(Lhs, Rhs);
+		}
+
+		/// <inheritdoc cref="Except(FileSet, FileSet)"/>
+		public static FileSet operator -(FileSet Lhs, FileSet Rhs)
+		{
+			return Except(Lhs, Rhs);
+		}
+
+		/// <summary>
+		/// Flatten to a map of files in a target directory
+		/// </summary>
+		/// <returns></returns>
+		public Dictionary<string, FileReference> Flatten()
+		{
+			Dictionary<string, FileReference> PathToSourceFile = new Dictionary<string, FileReference>(StringComparer.OrdinalIgnoreCase);
+			FlattenInternal(String.Empty, PathToSourceFile);
+			return PathToSourceFile;
+		}
+
+		private void FlattenInternal(string PathPrefix, Dictionary<string, FileReference> PathToSourceFile)
+		{
+			foreach ((string Path, FileReference File) in EnumerateFiles())
+			{
+				PathToSourceFile[PathPrefix + Path] = File;
+			}
+			foreach((string Path, FileSet FileSet) in EnumerateDirectories())
+			{
+				FileSet.FlattenInternal(PathPrefix + Path + "/", PathToSourceFile);
+			}
+		}
+
+		/// <summary>
+		/// Flatten to a map of files in a target directory
+		/// </summary>
+		/// <returns></returns>
+		public Dictionary<FileReference, FileReference> Flatten(DirectoryReference OutputDir)
+		{
+			Dictionary<FileReference, FileReference> TargetToSourceFile = new Dictionary<FileReference, FileReference>();
+			foreach ((string Path, FileReference SourceFile) in Flatten())
+			{
+				FileReference TargetFile = FileReference.Combine(OutputDir, Path);
+				TargetToSourceFile[TargetFile] = SourceFile;
+			}
+			return TargetToSourceFile;
+		}
+
+		/// <inheritdoc/>
+		IEnumerator IEnumerable.GetEnumerator() => GetEnumerator();
+
+		/// <inheritdoc/>
+		public IEnumerator<FileReference> GetEnumerator() => Flatten().Values.GetEnumerator();
+	}
+
+	/// <summary>
+	/// File tree from a known set of files
+	/// </summary>
+	class FileSetFromFiles : FileSet
+	{
+		Dictionary<string, FileReference> Files = new Dictionary<string, FileReference>();
+		Dictionary<string, FileSetFromFiles> SubTrees = new Dictionary<string, FileSetFromFiles>();
+
+		/// <summary>
+		/// Private constructor
+		/// </summary>
+		/// <param name="Path"></param>
+		private FileSetFromFiles(string Path)
+			: base(Path)
+		{
+		}
+
+		/// <summary>
+		/// Creates a tree from a given set of files
+		/// </summary>
+		/// <param name="InputFiles"></param>
+		public FileSetFromFiles(IEnumerable<(string, FileReference)> InputFiles)
+			: this(String.Empty)
+		{
+			foreach ((string Path, FileReference File) in InputFiles)
+			{
+				string[] Fragments = Path.Split(new[] { '/', '\\' }, StringSplitOptions.RemoveEmptyEntries);
+
+				FileSetFromFiles Current = this;
+				for (int Idx = 0; Idx < Fragments.Length - 1; Idx++)
+				{
+					FileSetFromFiles? Next;
+					if (!Current.SubTrees.TryGetValue(Fragments[Idx], out Next))
+					{
+						Next = new FileSetFromFiles(Current.Path + Fragments[Idx] + "/");
+						Current.SubTrees.Add(Fragments[Idx], Next);
+					}
+					Current = Next;
+				}
+
+				Current.Files.Add(Fragments[^1], File);
+			}
+		}
+
+		/// <inheritdoc/>
+		public override IEnumerable<KeyValuePair<string, FileReference>> EnumerateFiles() => Files;
+
+		/// <inheritdoc/>
+		public override IEnumerable<KeyValuePair<string, FileSet>> EnumerateDirectories() => SubTrees.Select(x => new KeyValuePair<string, FileSet>(x.Key, x.Value));
+	}
+
+	/// <summary>
+	/// File tree enumerated from the contents of an existing directory
+	/// </summary>
+	sealed class FileSetFromDirectory : FileSet
+	{
+		DirectoryInfo DirectoryInfo;
+
+		/// <summary>
+		/// Constructor
+		/// </summary>
+		public FileSetFromDirectory(DirectoryInfo DirectoryInfo)
+			: this(DirectoryInfo, "/")
+		{
+		}
+
+		/// <summary>
+		/// Constructor
+		/// </summary>
+		public FileSetFromDirectory(DirectoryInfo DirectoryInfo, string Path)
+			: base(Path)
+		{
+			this.DirectoryInfo = DirectoryInfo;
+		}
+
+		/// <inheritdoc/>
+		public override IEnumerable<KeyValuePair<string, FileReference>> EnumerateFiles() => DirectoryInfo.EnumerateFiles().Select(x => new KeyValuePair<string, FileReference>(x.Name, new FileReference(x)));
+
+		/// <inheritdoc/>
+		public override IEnumerable<KeyValuePair<string, FileSet>> EnumerateDirectories() => DirectoryInfo.EnumerateDirectories().Select(x => KeyValuePair.Create<string, FileSet>(x.Name, new FileSetFromDirectory(x)));
+	}
+
+	/// <summary>
+	/// File tree enumerated from the combination of two separate trees
+	/// </summary>
+	class FileSetFromUnion : FileSet
+	{
+		FileSet Lhs;
+		FileSet Rhs;
+
+		/// <summary>
+		/// Constructor
+		/// </summary>
+		/// <param name="Lhs">First file tree for the union</param>
+		/// <param name="Rhs">Other file tree for the union</param>
+		public FileSetFromUnion(FileSet Lhs, FileSet Rhs)
+			: base(Lhs.Path)
+		{
+			this.Lhs = Lhs;
+			this.Rhs = Rhs;
+		}
+
+		/// <inheritdoc/>
+		public override IEnumerable<KeyValuePair<string, FileReference>> EnumerateFiles()
+		{
+			Dictionary<string, FileReference> Files = new Dictionary<string, FileReference>(Lhs.EnumerateFiles(), StringComparer.OrdinalIgnoreCase);
+			foreach ((string Name, FileReference File) in Rhs.EnumerateFiles())
+			{
+				FileReference? ExistingFile;
+				if (!Files.TryGetValue(Name, out ExistingFile))
+				{
+					Files.Add(Name, File);
+				}
+				else if (ExistingFile == null || !ExistingFile.Equals(File))
+				{
+					throw new InvalidOperationException($"Conflict for contents of {Path}{Name} - could be {ExistingFile} or {File}");
+				}
+			}
+			return Files;
+		}
+
+		/// <inheritdoc/>
+		public override IEnumerable<KeyValuePair<string, FileSet>> EnumerateDirectories()
+		{
+			Dictionary<string, FileSet> NameToSubTree = new Dictionary<string, FileSet>(Lhs.EnumerateDirectories(), StringComparer.OrdinalIgnoreCase);
+			foreach ((string Name, FileSet SubTree) in Rhs.EnumerateDirectories())
+			{
+				FileSet? ExistingSubTree;
+				if (NameToSubTree.TryGetValue(Name, out ExistingSubTree))
+				{
+					NameToSubTree[Name] = new FileSetFromUnion(ExistingSubTree, SubTree);
+				}
+				else
+				{
+					NameToSubTree[Name] = SubTree;
+				}
+			}
+			return NameToSubTree;
+		}
+	}
+
+	/// <summary>
+	/// File tree enumerated from the combination of two separate trees
+	/// </summary>
+	class FileSetFromExcept : FileSet
+	{
+		FileSet Lhs;
+		FileSet Rhs;
+
+		/// <summary>
+		/// Constructor
+		/// </summary>
+		/// <param name="Lhs">First file tree for the union</param>
+		/// <param name="Rhs">Other file tree for the union</param>
+		public FileSetFromExcept(FileSet Lhs, FileSet Rhs)
+			: base(Lhs.Path)
+		{
+			this.Lhs = Lhs;
+			this.Rhs = Rhs;
+		}
+
+		/// <inheritdoc/>
+		public override IEnumerable<KeyValuePair<string, FileReference>> EnumerateFiles()
+		{
+			HashSet<string> RhsFiles = new HashSet<string>(Rhs.EnumerateFiles().Select(x => x.Key), StringComparer.OrdinalIgnoreCase);
+			return Lhs.EnumerateFiles().Where(x => !RhsFiles.Contains(x.Key));
+		}
+
+		/// <inheritdoc/>
+		public override IEnumerable<KeyValuePair<string, FileSet>> EnumerateDirectories()
+		{
+			Dictionary<string, FileSet> RhsDirs = new Dictionary<string, FileSet>(Rhs.EnumerateDirectories(), StringComparer.OrdinalIgnoreCase);
+			foreach ((string Name, FileSet LhsSet) in Lhs.EnumerateDirectories())
+			{
+				FileSet? RhsSet;
+				if (RhsDirs.TryGetValue(Name, out RhsSet))
+				{
+					yield return KeyValuePair.Create<string, FileSet>(Name, new FileSetFromExcept(LhsSet, RhsSet));
+				}
+				else
+				{
+					yield return KeyValuePair.Create(Name, LhsSet);
+				}
+			}
+		}
+	}
+
+	/// <summary>
+	/// File tree which includes only those files which match any given filter
+	/// </summary>
+	/// <typeparam name="T">Class containing information about a file</typeparam>
+	class FileSetFromFilter : FileSet
+	{
+		FileSet Inner;
+		FileFilter[] Filters;
+
+		/// <summary>
+		/// Constructor
+		/// </summary>
+		/// <param name="Inner">The tree to filter</param>
+		/// <param name="Filters"></param>
+		public FileSetFromFilter(FileSet Inner, params FileFilter[] Filters)
+			: base(Inner.Path)
+		{
+			this.Inner = Inner;
+			this.Filters = Filters;
+		}
+
+		/// <inheritdoc/>
+		public override IEnumerable<KeyValuePair<string, FileReference>> EnumerateFiles()
+		{
+			foreach (KeyValuePair<string, FileReference> Item in Inner.EnumerateFiles())
+			{
+				string FilterName = Inner.Path + Item.Key;
+				if (Filters.Any(x => x.Matches(FilterName)))
+				{
+					yield return Item;
+				}
+			}
+		}
+
+		/// <inheritdoc/>
+		public override IEnumerable<KeyValuePair<string, FileSet>> EnumerateDirectories()
+		{
+			foreach (KeyValuePair<string, FileSet> Item in Inner.EnumerateDirectories())
+			{
+				string FilterName = Inner.Path + Item.Key;
+
+				FileFilter[] PossibleFilters = Filters.Where(x => x.PossiblyMatches(FilterName)).ToArray();
+				if (PossibleFilters.Length > 0)
+				{
+					FileSetFromFilter SubTreeFilter = new FileSetFromFilter(Item.Value, PossibleFilters);
+					yield return new KeyValuePair<string, FileSet>(Item.Key, SubTreeFilter);
+				}
+			}
+		}
+	}
+}
diff --git a/Engine/Source/Programs/Shared/EpicGames.Core/FileUtils.cs b/Engine/Source/Programs/Shared/EpicGames.Core/FileUtils.cs
--- a/Engine/Source/Programs/Shared/EpicGames.Core/FileUtils.cs
+++ b/Engine/Source/Programs/Shared/EpicGames.Core/FileUtils.cs
@@ -435,7 +435,9 @@ namespace EpicGames.Core
 
 		static readonly IntPtr INVALID_HANDLE_VALUE = new IntPtr(-1);
 
+		const int ERROR_FILE_NOT_FOUND = 2;
 		const int ERROR_PATH_NOT_FOUND = 3;
+		const int ERROR_ACCESS_DENIED = 5;
 
 		private static void ForceDeleteLongDirectoryContentsWin32(string DirName)
 		{
@@ -485,16 +487,13 @@ namespace EpicGames.Core
 			if (!RemoveDirectory(DirName))
 			{
 				int ErrorCode = Marshal.GetLastWin32Error();
-				if (ErrorCode != ERROR_PATH_NOT_FOUND)
+				if (ErrorCode != ERROR_FILE_NOT_FOUND && ErrorCode != ERROR_PATH_NOT_FOUND)
 				{
 					throw new WrappedFileOrDirectoryException(new Win32Exception(ErrorCode), "Unable to delete " + DirName);
 				}
 			}
 		}
 
-		const int ERROR_FILE_NOT_FOUND = 2;
-		const int ERROR_ACCESS_DENIED = 5;
-
 		[DllImport("kernel32.dll", CharSet = CharSet.Unicode, SetLastError = true)]
 		internal static extern int GetFileAttributesW(string lpFileName);
 
