Fix an issue where editor heartbeat was not written to the log like the client heartbeat.

diff --git a/Engine/Source/Programs/AutomationTool/Gauntlet/Unreal/Base/Gauntlet.UnrealTestNode.cs b/Engine/Source/Programs/AutomationTool/Gauntlet/Unreal/Base/Gauntlet.UnrealTestNode.cs
--- a/Engine/Source/Programs/AutomationTool/Gauntlet/Unreal/Base/Gauntlet.UnrealTestNode.cs
+++ b/Engine/Source/Programs/AutomationTool/Gauntlet/Unreal/Base/Gauntlet.UnrealTestNode.cs
@@ -233,9 +233,14 @@ namespace Gauntlet
 		protected UnrealSession UnrealApp;
 
 		/// <summary>
-		/// Used to track how much of our log has been written out
+		/// Used to track how much of our app log has been written out
 		/// </summary>
-		private int LastLogCount;
+		private int LastAppLogCount;
+
+		/// <summary>
+		/// Used to track how much of our editor log has been written out
+		/// </summary>
+		private int LastEditorLogCount;
 
 		private int CurrentPass;
 
@@ -303,7 +308,8 @@ namespace Gauntlet
 
 			UnrealTestResult = TestResult.Invalid;
 			TimeToWaitForProcesses = 5;
-			LastLogCount = 0;
+			LastAppLogCount = 0;
+			LastEditorLogCount = 0;
 			CurrentPass = 0;
 			NumPasses = 0;
 			TestVersion = new Version("1.0.0");
@@ -678,7 +684,8 @@ namespace Gauntlet
 			UnrealTestResult = TestResult.Invalid;
 			CurrentPass = Pass;
 			NumPasses = InNumPasses;
-			LastLogCount = 0;
+			LastAppLogCount = 0;
+			LastEditorLogCount = 0;
 			LastHeartbeatTime = DateTime.MinValue;
 			LastActiveHeartbeatTime = DateTime.MinValue;			
 
@@ -800,7 +807,8 @@ namespace Gauntlet
 			SessionArtifacts = Enumerable.Empty<UnrealRoleArtifacts>();
 			RoleResults = Enumerable.Empty<UnrealRoleResult>();
 
-			LastLogCount = 0;
+			LastAppLogCount = 0;
+			LastEditorLogCount = 0;
 			LastHeartbeatTime = DateTime.MinValue;
 			LastActiveHeartbeatTime = DateTime.MinValue;
 
@@ -829,38 +837,41 @@ namespace Gauntlet
 		public override void TickTest()
 		{
 			IAppInstance App = null;
-
+			string AppInfoPrefix = "App";
 			if (TestInstance.ClientApps == null)
 			{
 				App = TestInstance.ServerApp;
+				AppInfoPrefix = "Server";
 			}
 			else
 			{
 				if (TestInstance.ClientApps.Length > 0)
 				{
 					App = TestInstance.ClientApps.First();
+					AppInfoPrefix = "Client";
 				}
 			}
 
-			if (App != null)
+			List<string> LogCategories = new List<string>();
+			LogCategories.Add("Gauntlet");
 			{
-				List<string> LogCategories = new List<string>();
-				LogCategories.Add("Gauntlet");
-				{
-					// get the categories used to monitor process (this needs rethought).
-					IEnumerable<string> HeartbeatCategories = GetHeartbeatLogCategories().Union(GetCachedConfiguration().LogCategoriesForEvents);
-					LogCategories.AddRange(HeartbeatCategories);
-				}
-				LogCategories = LogCategories.Distinct().ToList();
+				// get the categories used to monitor process (this needs rethought).
+				IEnumerable<string> HeartbeatCategories = GetHeartbeatLogCategories().Union(GetCachedConfiguration().LogCategoriesForEvents);
+				LogCategories.AddRange(HeartbeatCategories);
+			}
+			LogCategories = LogCategories.Distinct().ToList();
+
 
+			if (App != null)
+			{
 				UnrealLogParser Parser = new UnrealLogParser(App.StdOut);
 				List<string> TestLines = new List<string>();
 				// ONLY ADD RANGE ONCE. Ordering is important and will be skewed if multiple ranges are added which can skew how logs are pulled out.
 				TestLines.AddRange(Parser.GetLogChannels(LogCategories, true));
 								
-				for (int i = LastLogCount; i < TestLines.Count(); i++)
+				for (int i = LastAppLogCount; i < TestLines.Count(); i++)
 				{
-					Log.Info(TestLines[i]);
+					Log.Info(string.Format("{0}: {1}", AppInfoPrefix, TestLines[i]));
 
 					if (Regex.IsMatch(TestLines[i], @".*GauntletHeartbeat\: Active.*"))
 					{
@@ -873,12 +884,28 @@ namespace Gauntlet
 					}
 				}
 
-				LastLogCount = TestLines.Count();
+				LastAppLogCount = TestLines.Count();
 
 				// Detect missed heartbeats and fail the test
 				CheckHeartbeat();
 			}
 
+			IAppInstance EditorApp = TestInstance.EditorApp;
+			if (EditorApp != null)
+			{
+				UnrealLogParser Parser = new UnrealLogParser(EditorApp.StdOut);
+				List<string> TestLines = new List<string>();
+				// ONLY ADD RANGE ONCE. Ordering is important and will be skewed if multiple ranges are added which can skew how logs are pulled out.
+				TestLines.AddRange(Parser.GetLogChannels(LogCategories, true));
+
+				for (int i = LastEditorLogCount; i < TestLines.Count(); i++)
+				{
+					Log.Info(string.Format("Editor: {0}", TestLines[i]));
+				}
+
+				LastEditorLogCount = TestLines.Count();
+			}
+
 
 			// Check status and health after updating logs
 			if (GetTestStatus() == TestStatus.InProgress && IsTestRunning() == false)
