Turnkey.Automation.csproj: Ensure referenced dlls are copied to binary directory

diff --git a/Engine/Source/Programs/AutomationTool/Turnkey/Turnkey.Automation.csproj b/Engine/Source/Programs/AutomationTool/Turnkey/Turnkey.Automation.csproj
--- a/Engine/Source/Programs/AutomationTool/Turnkey/Turnkey.Automation.csproj
+++ b/Engine/Source/Programs/AutomationTool/Turnkey/Turnkey.Automation.csproj
@@ -13,7 +13,9 @@
     <OutputPath>..\..\..\..\Binaries\DotNET\AutomationTool\AutomationScripts\Turnkey\</OutputPath>
     <TreatWarningsAsErrors>true</TreatWarningsAsErrors>
     <DebugType>pdbonly</DebugType>
-  </PropertyGroup>
+    <CopyLocalLockFileAssemblies>true</CopyLocalLockFileAssemblies>
+    <SatelliteResourceLanguages>en</SatelliteResourceLanguages>    <!-- remove non english resource languages -->
+</PropertyGroup>
   <PropertyGroup Condition="'$(Configuration)|$(Platform)' == 'Development|AnyCPU' ">
     <DefineConstants>$(DefineConstants);TRACE</DefineConstants>
     <Optimize>true</Optimize>
