UNSAFE_TYPECAST_WARNING macros allow force enabling of warnings.

diff --git a/Engine/Source/Programs/AutomationTool/Scripts/CheckBalancedMacros.cs b/Engine/Source/Programs/AutomationTool/Scripts/CheckBalancedMacros.cs
--- a/Engine/Source/Programs/AutomationTool/Scripts/CheckBalancedMacros.cs
+++ b/Engine/Source/Programs/AutomationTool/Scripts/CheckBalancedMacros.cs
@@ -41,8 +41,12 @@ namespace AutomationTool
 			},
 			{
 				"PRAGMA_DISABLE_UNSAFE_TYPECAST_WARNINGS",
-				"PRAGMA_ENABLE_UNSAFE_TYPECAST_WARNINGS"
+				"PRAGMA_RESTORE_UNSAFE_TYPECAST_WARNINGS"
 			},
+			{
+				"PRAGMA_FORCE_UNSAFE_TYPECAST_WARNINGS",
+				"PRAGMA_RESTORE_UNSAFE_TYPECAST_WARNINGS"
+			},			
 			{
 				"PRAGMA_DISABLE_UNDEFINED_IDENTIFIER_WARNINGS",
 				"PRAGMA_ENABLE_UNDEFINED_IDENTIFIER_WARNINGS"
@@ -55,10 +59,6 @@ namespace AutomationTool
 				"BEGIN_FUNCTION_BUILD_OPTIMIZATION",
 				"END_FUNCTION_BUILD_OPTIMIZATION"
 			},
-			{
-				"BEGIN_FUNCTION_BUILD_OPTIMIZATION",
-				"END_FUNCTION_BUILD_OPTIMIZATION"
-			},
 			{
 				"BEGIN_SLATE_FUNCTION_BUILD_OPTIMIZATION",
 				"END_SLATE_FUNCTION_BUILD_OPTIMIZATION"
@@ -82,11 +82,19 @@ namespace AutomationTool
 		public override void ExecuteBuild()
 		{
 			// Build a lookup of flags to set and clear for each identifier
-			Dictionary<string, int> IdentifierToIndex = new Dictionary<string, int>();
+			Dictionary<string, List<int>> IdentifierToIndex = new Dictionary<string, List<int>>();
 			for(int Idx = 0; Idx < MacroPairs.GetLength(0); Idx++)
 			{
-				IdentifierToIndex[MacroPairs[Idx, 0]] = Idx;
-				IdentifierToIndex[MacroPairs[Idx, 1]] = ~Idx;
+				for (int SubIdx = 0; SubIdx < 2; SubIdx++)
+				{
+					ref string Key = ref MacroPairs[Idx, SubIdx];
+					if (!IdentifierToIndex.ContainsKey(Key))
+					{
+						IdentifierToIndex[Key] = new List<int>();
+					}
+
+					IdentifierToIndex[Key].Add(SubIdx == 0 ? Idx : ~Idx);
+				}
 			}
 
 			// Check if we want to just parse a single file
@@ -211,7 +219,7 @@ namespace AutomationTool
 		/// <param name="SourceFile"></param>
 		/// <param name="IdentifierToIndex">Map of macro identifier to bit index. The complement of an index is used to indicate the end of the pair.</param>
 		/// <param name="LogLock">Object used to marshal access to the global log instance</param>
-		void CheckSourceFile(FileReference SourceFile, Dictionary<string, int> IdentifierToIndex, object LogLock)
+		void CheckSourceFile(FileReference SourceFile, Dictionary<string, List<int>> IdentifierToIndex, object LogLock)
 		{
 			// Read the text
 			string Text = FileReference.ReadAllText(SourceFile);
@@ -233,28 +241,44 @@ namespace AutomationTool
 					string Identifier = Text.Substring(StartIdx, Idx - StartIdx);
 
 					// Find the matching flag
-					int Index;
+					List<int> Index;
 					if(IdentifierToIndex.TryGetValue(Identifier, out Index))
 					{
-						if(Index >= 0)
+						if(Index[0] >= 0)
 						{
 							// Set the flag (should not already be set)
-							int Flag = 1 << Index;
+							int Flag = 1 << Index[0];
 							if((Flags & Flag) != 0)
 							{
-								EpicGames.Core.Log.TraceWarningTask(SourceFile, GetLineNumber(Text, StartIdx), "{0} macro appears a second time without matching {1} macro", Identifier, MacroPairs[Index, 1]);
+								EpicGames.Core.Log.TraceWarningTask(SourceFile, GetLineNumber(Text, StartIdx), "{0} macro appears a second time without matching {1} macro", Identifier, MacroPairs[Index[0], 1]);
 							}
 							Flags |= Flag;
 						}
 						else
 						{
-							// Clear the flag (should already be set)
-							int Flag = 1 << ~Index;
-							if((Flags & Flag) == 0)
+							bool bMatched = false;
+							// Check for any flag. We clear the first we find when validating, even if that's not technically the correct match. TODO: This means we may report the wrong tag as left over where tags are nested and there's a missing end tag.
+							foreach (int SubIndex in Index)
+							{
+								int Flag = 1 << ~SubIndex;
+								// Clear the flag (should already be set)
+								if((Flags & Flag) != 0)
+								{
+									Flags &= ~Flag;
+									bMatched = true;
+									break;
+								}
+							}
+							
+							if (!bMatched)
 							{
-								EpicGames.Core.Log.TraceWarningTask(SourceFile, GetLineNumber(Text, StartIdx), "{0} macro appears without matching {1} macro", Identifier, MacroPairs[~Index, 0]);
+								string MissingMatching = "";
+								foreach (int SubIndex in Index)
+								{
+									MissingMatching += (MissingMatching.Length > 0 ? ", " : "") +  MacroPairs[~SubIndex, 0];
+								}
+								EpicGames.Core.Log.TraceWarningTask(SourceFile, GetLineNumber(Text, StartIdx), "{0} macro appears without matching {1} macro", Identifier, MissingMatching);
 							}
-							Flags &= ~Flag;
 						}
 					}
 				}
diff --git a/Engine/Source/Runtime/Core/Private/FileCache/FileCache.cpp b/Engine/Source/Runtime/Core/Private/FileCache/FileCache.cpp
--- a/Engine/Source/Runtime/Core/Private/FileCache/FileCache.cpp
+++ b/Engine/Source/Runtime/Core/Private/FileCache/FileCache.cpp
@@ -227,7 +227,7 @@ public:
 	{
 		return BlockIDType(FMath::DivideAndRoundDown(Offset, BlockIDType::GetSize()));
 	}
-	PRAGMA_ENABLE_UNSAFE_TYPECAST_WARNINGS
+	PRAGMA_RESTORE_UNSAFE_TYPECAST_WARNINGS
 
 	template<typename BlockIDType> inline int32 GetNumBlocks(int64 Offset, int64 Size)
 	{
diff --git a/Engine/Source/Runtime/Core/Private/HAL/MallocBinned.cpp b/Engine/Source/Runtime/Core/Private/HAL/MallocBinned.cpp
--- a/Engine/Source/Runtime/Core/Private/HAL/MallocBinned.cpp
+++ b/Engine/Source/Runtime/Core/Private/HAL/MallocBinned.cpp
@@ -1810,4 +1810,4 @@ const TCHAR* FMallocBinned::GetDescriptiveName()
 	return TEXT("binned");
 }
 
-PRAGMA_ENABLE_UNSAFE_TYPECAST_WARNINGS
+PRAGMA_RESTORE_UNSAFE_TYPECAST_WARNINGS
diff --git a/Engine/Source/Runtime/Core/Private/HAL/MallocBinned2.cpp b/Engine/Source/Runtime/Core/Private/HAL/MallocBinned2.cpp
--- a/Engine/Source/Runtime/Core/Private/HAL/MallocBinned2.cpp
+++ b/Engine/Source/Runtime/Core/Private/HAL/MallocBinned2.cpp
@@ -1352,4 +1352,4 @@ void FMallocBinned2::UpdateStats()
 }
 
 
-PRAGMA_ENABLE_UNSAFE_TYPECAST_WARNINGS
+PRAGMA_RESTORE_UNSAFE_TYPECAST_WARNINGS
diff --git a/Engine/Source/Runtime/Core/Private/HAL/MallocBinned3.cpp b/Engine/Source/Runtime/Core/Private/HAL/MallocBinned3.cpp
--- a/Engine/Source/Runtime/Core/Private/HAL/MallocBinned3.cpp
+++ b/Engine/Source/Runtime/Core/Private/HAL/MallocBinned3.cpp
@@ -1613,4 +1613,4 @@ void FMallocBinned3::DumpAllocatorStats(class FOutputDevice& Ar)
 #endif
 #endif
 
-PRAGMA_ENABLE_UNSAFE_TYPECAST_WARNINGS
+PRAGMA_RESTORE_UNSAFE_TYPECAST_WARNINGS
diff --git a/Engine/Source/Runtime/Core/Private/HAL/MallocBinnedArena.cpp b/Engine/Source/Runtime/Core/Private/HAL/MallocBinnedArena.cpp
--- a/Engine/Source/Runtime/Core/Private/HAL/MallocBinnedArena.cpp
+++ b/Engine/Source/Runtime/Core/Private/HAL/MallocBinnedArena.cpp
@@ -1353,4 +1353,4 @@ void FMallocBinnedArena::DumpAllocatorStats(class FOutputDevice& Ar)
 }
 #endif
 
-PRAGMA_ENABLE_UNSAFE_TYPECAST_WARNINGS
+PRAGMA_RESTORE_UNSAFE_TYPECAST_WARNINGS
diff --git a/Engine/Source/Runtime/Core/Private/HAL/MallocBinnedCommon.cpp b/Engine/Source/Runtime/Core/Private/HAL/MallocBinnedCommon.cpp
--- a/Engine/Source/Runtime/Core/Private/HAL/MallocBinnedCommon.cpp
+++ b/Engine/Source/Runtime/Core/Private/HAL/MallocBinnedCommon.cpp
@@ -473,4 +473,4 @@ uint32 FBitTree::CountOnes(uint32 UpTo) const
 
 #endif
 
-PRAGMA_ENABLE_UNSAFE_TYPECAST_WARNINGS
+PRAGMA_RESTORE_UNSAFE_TYPECAST_WARNINGS
diff --git a/Engine/Source/Runtime/Core/Private/HAL/MallocBinnedGPU.cpp b/Engine/Source/Runtime/Core/Private/HAL/MallocBinnedGPU.cpp
--- a/Engine/Source/Runtime/Core/Private/HAL/MallocBinnedGPU.cpp
+++ b/Engine/Source/Runtime/Core/Private/HAL/MallocBinnedGPU.cpp
@@ -1295,4 +1295,4 @@ void FMallocBinnedGPU::DumpAllocatorStats(class FOutputDevice& Ar)
 }
 #endif
 
-PRAGMA_ENABLE_UNSAFE_TYPECAST_WARNINGS
+PRAGMA_RESTORE_UNSAFE_TYPECAST_WARNINGS
diff --git a/Engine/Source/Runtime/Core/Private/Math/UnitConversion.cpp b/Engine/Source/Runtime/Core/Private/Math/UnitConversion.cpp
--- a/Engine/Source/Runtime/Core/Private/Math/UnitConversion.cpp
+++ b/Engine/Source/Runtime/Core/Private/Math/UnitConversion.cpp
@@ -778,4 +778,4 @@ namespace UnitConversion
 
 #undef LOCTEXT_NAMESPACE
 
-PRAGMA_ENABLE_UNSAFE_TYPECAST_WARNINGS
+PRAGMA_RESTORE_UNSAFE_TYPECAST_WARNINGS
diff --git a/Engine/Source/Runtime/Core/Private/Memory/MemoryArena.cpp b/Engine/Source/Runtime/Core/Private/Memory/MemoryArena.cpp
--- a/Engine/Source/Runtime/Core/Private/Memory/MemoryArena.cpp
+++ b/Engine/Source/Runtime/Core/Private/Memory/MemoryArena.cpp
@@ -477,4 +477,4 @@ FMemoryArena* FArenaMap::MapPtrToArena(const void* VaBase)
 
 //////////////////////////////////////////////////////////////////////////
 
-PRAGMA_ENABLE_UNSAFE_TYPECAST_WARNINGS
+PRAGMA_RESTORE_UNSAFE_TYPECAST_WARNINGS
diff --git a/Engine/Source/Runtime/Core/Private/Microsoft/MicrosoftAsyncIO.h b/Engine/Source/Runtime/Core/Private/Microsoft/MicrosoftAsyncIO.h
--- a/Engine/Source/Runtime/Core/Private/Microsoft/MicrosoftAsyncIO.h
+++ b/Engine/Source/Runtime/Core/Private/Microsoft/MicrosoftAsyncIO.h
@@ -517,4 +517,4 @@ const TCHAR* FMicrosoftReadRequest::GetFileNameForErrorMessagesAndPanicRetry()
 	return *Owner->FileNameForErrorMessagesAndPanicRetry;
 }
 
-PRAGMA_ENABLE_UNSAFE_TYPECAST_WARNINGS
+PRAGMA_RESTORE_UNSAFE_TYPECAST_WARNINGS
diff --git a/Engine/Source/Runtime/Core/Private/Misc/Compression.cpp b/Engine/Source/Runtime/Core/Private/Misc/Compression.cpp
--- a/Engine/Source/Runtime/Core/Private/Misc/Compression.cpp
+++ b/Engine/Source/Runtime/Core/Private/Misc/Compression.cpp
@@ -1015,4 +1015,4 @@ bool FCompression::VerifyCompressionFlagsValid(int32 InCompressionFlags)
 /***********************
   Deprecated functions
 ***********************/
-PRAGMA_ENABLE_UNSAFE_TYPECAST_WARNINGS
+PRAGMA_RESTORE_UNSAFE_TYPECAST_WARNINGS
diff --git a/Engine/Source/Runtime/Core/Private/Serialization/Archive.cpp b/Engine/Source/Runtime/Core/Private/Serialization/Archive.cpp
--- a/Engine/Source/Runtime/Core/Private/Serialization/Archive.cpp
+++ b/Engine/Source/Runtime/Core/Private/Serialization/Archive.cpp
@@ -1409,4 +1409,4 @@ void FArchiveState::SetIsPersistent(bool bInIsPersistent)
 
 static_assert(sizeof(FArchive) == sizeof(FArchiveState), "New FArchive members should be added to FArchiveState instead");
 
-PRAGMA_ENABLE_UNSAFE_TYPECAST_WARNINGS
+PRAGMA_RESTORE_UNSAFE_TYPECAST_WARNINGS
diff --git a/Engine/Source/Runtime/Core/Private/Serialization/ArchiveLoadCompressedProxy.cpp b/Engine/Source/Runtime/Core/Private/Serialization/ArchiveLoadCompressedProxy.cpp
--- a/Engine/Source/Runtime/Core/Private/Serialization/ArchiveLoadCompressedProxy.cpp
+++ b/Engine/Source/Runtime/Core/Private/Serialization/ArchiveLoadCompressedProxy.cpp
@@ -142,4 +142,4 @@ int64 FArchiveLoadCompressedProxy::Tell()
 	return RawBytesSerialized;
 }
 
-PRAGMA_ENABLE_UNSAFE_TYPECAST_WARNINGS
+PRAGMA_RESTORE_UNSAFE_TYPECAST_WARNINGS
diff --git a/Engine/Source/Runtime/Core/Private/Serialization/ArchiveSaveCompressedProxy.cpp b/Engine/Source/Runtime/Core/Private/Serialization/ArchiveSaveCompressedProxy.cpp
--- a/Engine/Source/Runtime/Core/Private/Serialization/ArchiveSaveCompressedProxy.cpp
+++ b/Engine/Source/Runtime/Core/Private/Serialization/ArchiveSaveCompressedProxy.cpp
@@ -146,4 +146,4 @@ int64 FArchiveSaveCompressedProxy::Tell()
 	}
 }
 
-PRAGMA_ENABLE_UNSAFE_TYPECAST_WARNINGS
+PRAGMA_RESTORE_UNSAFE_TYPECAST_WARNINGS
diff --git a/Engine/Source/Runtime/Core/Private/Serialization/BitReader.cpp b/Engine/Source/Runtime/Core/Private/Serialization/BitReader.cpp
--- a/Engine/Source/Runtime/Core/Private/Serialization/BitReader.cpp
+++ b/Engine/Source/Runtime/Core/Private/Serialization/BitReader.cpp
@@ -311,4 +311,4 @@ void FBitReaderMark::Copy( FBitReader& Reader, TArray<uint8> &Buffer )
 	}
 }
 
-PRAGMA_ENABLE_UNSAFE_TYPECAST_WARNINGS
+PRAGMA_RESTORE_UNSAFE_TYPECAST_WARNINGS
diff --git a/Engine/Source/Runtime/Core/Private/Serialization/BitWriter.cpp b/Engine/Source/Runtime/Core/Private/Serialization/BitWriter.cpp
--- a/Engine/Source/Runtime/Core/Private/Serialization/BitWriter.cpp
+++ b/Engine/Source/Runtime/Core/Private/Serialization/BitWriter.cpp
@@ -307,4 +307,4 @@ void FBitWriterMark::Copy( FBitWriter& Writer, TArray<uint8> &Buffer )
 	}
 }
 
-PRAGMA_ENABLE_UNSAFE_TYPECAST_WARNINGS
+PRAGMA_RESTORE_UNSAFE_TYPECAST_WARNINGS
diff --git a/Engine/Source/Runtime/Core/Private/Tests/Math/UnitConversionTests.cpp b/Engine/Source/Runtime/Core/Private/Tests/Math/UnitConversionTests.cpp
--- a/Engine/Source/Runtime/Core/Private/Tests/Math/UnitConversionTests.cpp
+++ b/Engine/Source/Runtime/Core/Private/Tests/Math/UnitConversionTests.cpp
@@ -166,4 +166,4 @@ bool FParsingUnitTests::RunTest(const FString& Parameters)
 
 #endif //WITH_DEV_AUTOMATION_TESTS
 
-PRAGMA_ENABLE_UNSAFE_TYPECAST_WARNINGS
+PRAGMA_RESTORE_UNSAFE_TYPECAST_WARNINGS
diff --git a/Engine/Source/Runtime/Core/Private/UObject/UnrealNames.cpp b/Engine/Source/Runtime/Core/Private/UObject/UnrealNames.cpp
--- a/Engine/Source/Runtime/Core/Private/UObject/UnrealNames.cpp
+++ b/Engine/Source/Runtime/Core/Private/UObject/UnrealNames.cpp
@@ -4004,4 +4004,4 @@ bool ShouldReplicateAsInteger(EName Ename, const FName& Name)
 	return Ename <= EName(MAX_NETWORKED_HARDCODED_NAME) && Name.GetNumber() == NAME_NO_NUMBER_INTERNAL;
 }
 
-PRAGMA_ENABLE_UNSAFE_TYPECAST_WARNINGS
+PRAGMA_RESTORE_UNSAFE_TYPECAST_WARNINGS
diff --git a/Engine/Source/Runtime/Core/Private/Windows/WindowsPlatformProcess.cpp b/Engine/Source/Runtime/Core/Private/Windows/WindowsPlatformProcess.cpp
--- a/Engine/Source/Runtime/Core/Private/Windows/WindowsPlatformProcess.cpp
+++ b/Engine/Source/Runtime/Core/Private/Windows/WindowsPlatformProcess.cpp
@@ -2114,6 +2114,6 @@ void FWindowsPlatformProcess::SetThreadName( const TCHAR* ThreadName )
 	WindowsPlatformProcessImpl::SetThreadName(TCHAR_TO_ANSI(ThreadName));
 }
 
-PRAGMA_ENABLE_UNSAFE_TYPECAST_WARNINGS
+PRAGMA_RESTORE_UNSAFE_TYPECAST_WARNINGS
 
 #include "Windows/HideWindowsPlatformTypes.h"
diff --git a/Engine/Source/Runtime/Core/Public/Clang/ClangPlatformCompilerPreSetup.h b/Engine/Source/Runtime/Core/Public/Clang/ClangPlatformCompilerPreSetup.h
--- a/Engine/Source/Runtime/Core/Public/Clang/ClangPlatformCompilerPreSetup.h
+++ b/Engine/Source/Runtime/Core/Public/Clang/ClangPlatformCompilerPreSetup.h
@@ -79,12 +79,24 @@
 #define DISABLE_IMPLICIT_FLOAT_CONVERSION_FRAGMENT
 #endif
 
+#if __has_warning("-Wimplicit-float-conversion")
+#define FORCE_IMPLICIT_FLOAT_CONVERSION_FRAGMENT _Pragma("clang diagnostic warning \"-Wimplicit-float-conversion\"")
+#else
+#define FORCE_IMPLICIT_FLOAT_CONVERSION_FRAGMENT
+#endif
+
 #if __has_warning("-Wimplicit-int-conversion")
 #define DISABLE_IMPLICIT_INT_CONVERSION_FRAGMENT _Pragma("clang diagnostic ignored \"-Wimplicit-int-conversion\"")
 #else
 #define DISABLE_IMPLICIT_INT_CONVERSION_FRAGMENT
 #endif
 
+#if __has_warning("-Wimplicit-int-conversion")
+#define FORCE_IMPLICIT_INT_CONVERSION_FRAGMENT _Pragma("clang diagnostic warning \"-Wimplicit-int-conversion\"")
+#else
+#define FORCE_IMPLICIT_INT_CONVERSION_FRAGMENT
+#endif
+
 #ifndef PRAGMA_DISABLE_UNSAFE_TYPECAST_WARNINGS
 	#define PRAGMA_DISABLE_UNSAFE_TYPECAST_WARNINGS \
 		_Pragma("clang diagnostic push") \
@@ -96,7 +108,21 @@
 
 #ifndef PRAGMA_ENABLE_UNSAFE_TYPECAST_WARNINGS
 	#define PRAGMA_ENABLE_UNSAFE_TYPECAST_WARNINGS \
-		_Pragma("clang diagnostic pop")
+		DEPRECATED_MACRO(5.0, "The PRAGMA_ENABLE_UNSAFE_TYPECAST_WARNINGS macro has been deprecated in favor of PRAGMA_RESTORE_UNSAFE_TYPECAST_WARNINGS. To force enable warnings use PRAGMA_FORCE_UNSAFE_TYPECAST_WARNINGS.")
+#endif // PRAGMA_ENABLE_UNSAFE_TYPECAST_WARNINGS
+
+#ifndef PRAGMA_FORCE_UNSAFE_TYPECAST_WARNINGS
+	#define PRAGMA_FORCE_UNSAFE_TYPECAST_WARNINGS \
+		_Pragma("clang diagnostic push") \
+		_Pragma("clang diagnostic warning \"-Wfloat-conversion\"") \
+		FORCE_IMPLICIT_FLOAT_CONVERSION_FRAGMENT \
+		FORCE_IMPLICIT_INT_CONVERSION_FRAGMENT \
+		_Pragma("clang diagnostic warning \"-Wc++11-narrowing\"")
+#endif // PRAGMA_DISABLE_UNSAFE_TYPECAST_WARNINGS
+
+#ifndef PRAGMA_RESTORE_UNSAFE_TYPECAST_WARNINGS
+	#define PRAGMA_RESTORE_UNSAFE_TYPECAST_WARNINGS \
+	_Pragma("clang diagnostic pop")
 #endif // PRAGMA_ENABLE_UNSAFE_TYPECAST_WARNINGS
 
 #ifndef PRAGMA_DISABLE_UNDEFINED_IDENTIFIER_WARNINGS
diff --git a/Engine/Source/Runtime/Core/Public/HAL/MallocBinned3.h b/Engine/Source/Runtime/Core/Public/HAL/MallocBinned3.h
--- a/Engine/Source/Runtime/Core/Public/HAL/MallocBinned3.h
+++ b/Engine/Source/Runtime/Core/Public/HAL/MallocBinned3.h
@@ -723,7 +723,7 @@ public:
 	static void* AllocateMetaDataMemory(SIZE_T Size);
 };
 
-PRAGMA_ENABLE_UNSAFE_TYPECAST_WARNINGS
+PRAGMA_RESTORE_UNSAFE_TYPECAST_WARNINGS
 
 #define BINNED3_INLINE (1)
 #if BINNED3_INLINE // during development, it helps with iteration time to not include these here, but rather in the .cpp
diff --git a/Engine/Source/Runtime/Core/Public/HAL/MallocBinnedArena.h b/Engine/Source/Runtime/Core/Public/HAL/MallocBinnedArena.h
--- a/Engine/Source/Runtime/Core/Public/HAL/MallocBinnedArena.h
+++ b/Engine/Source/Runtime/Core/Public/HAL/MallocBinnedArena.h
@@ -690,6 +690,6 @@ public:
 	TArray<void*> MallocedPointers;
 };
 
-PRAGMA_ENABLE_UNSAFE_TYPECAST_WARNINGS
+PRAGMA_RESTORE_UNSAFE_TYPECAST_WARNINGS
 
 #endif
diff --git a/Engine/Source/Runtime/Core/Public/HAL/MallocBinnedGPU.h b/Engine/Source/Runtime/Core/Public/HAL/MallocBinnedGPU.h
--- a/Engine/Source/Runtime/Core/Public/HAL/MallocBinnedGPU.h
+++ b/Engine/Source/Runtime/Core/Public/HAL/MallocBinnedGPU.h
@@ -662,6 +662,6 @@ public:
 	TArray<void*> MallocedPointers;
 };
 
-PRAGMA_ENABLE_UNSAFE_TYPECAST_WARNINGS
+PRAGMA_RESTORE_UNSAFE_TYPECAST_WARNINGS
 
 #endif
diff --git a/Engine/Source/Runtime/Core/Public/HoloLens/HoloLensPlatformCompilerPreSetup.h b/Engine/Source/Runtime/Core/Public/HoloLens/HoloLensPlatformCompilerPreSetup.h
--- a/Engine/Source/Runtime/Core/Public/HoloLens/HoloLensPlatformCompilerPreSetup.h
+++ b/Engine/Source/Runtime/Core/Public/HoloLens/HoloLensPlatformCompilerPreSetup.h
@@ -32,14 +32,27 @@
 #ifndef PRAGMA_DISABLE_UNSAFE_TYPECAST_WARNINGS
 	#define PRAGMA_DISABLE_UNSAFE_TYPECAST_WARNINGS \
 		__pragma (warning(push)) \
-		__pragma (warning(disable: 4244)) /* 'argument': conversion from 'type1' to 'type2', possible loss of data */
+		__pragma (warning(disable: 4244)) /* 'argument': conversion from 'type1' to 'type2', possible loss of data */	\
+		__pragma (warning(disable: 4838)) /* 'argument': conversion from 'type1' to 'type2' requires a narrowing conversion */
 #endif // PRAGMA_DISABLE_UNSAFE_TYPECAST_WARNINGS
 
 #ifndef PRAGMA_ENABLE_UNSAFE_TYPECAST_WARNINGS
 	#define PRAGMA_ENABLE_UNSAFE_TYPECAST_WARNINGS \
-		__pragma(warning(pop))
+		DEPRECATED_MACRO(5.0, "The PRAGMA_ENABLE_UNSAFE_TYPECAST_WARNINGS macro has been deprecated in favor of PRAGMA_RESTORE_UNSAFE_TYPECAST_WARNINGS. To force enable warnings use PRAGMA_FORCE_UNSAFE_TYPECAST_WARNINGS.")
 #endif // PRAGMA_ENABLE_UNSAFE_TYPECAST_WARNINGS
 
+#ifndef PRAGMA_FORCE_UNSAFE_TYPECAST_WARNINGS
+	#define PRAGMA_FORCE_UNSAFE_TYPECAST_WARNINGS \
+		__pragma (warning(push)) \
+		__pragma (warning(error: 4244)) /* 'argument': conversion from 'type1' to 'type2', possible loss of data */	\
+		__pragma (warning(error: 4838)) /* 'argument': conversion from 'type1' to 'type2' requires a narrowing conversion */
+#endif // PRAGMA_FORCE_UNSAFE_TYPECAST_WARNINGS
+
+#ifndef PRAGMA_RESTORE_UNSAFE_TYPECAST_WARNINGS
+	#define PRAGMA_RESTORE_UNSAFE_TYPECAST_WARNINGS \
+		__pragma(warning(pop))
+#endif // PRAGMA_RESTORE_UNSAFE_TYPECAST_WARNINGS
+
 #ifndef PRAGMA_DISABLE_UNDEFINED_IDENTIFIER_WARNINGS
 	#define PRAGMA_DISABLE_UNDEFINED_IDENTIFIER_WARNINGS \
 		__pragma(warning(push)) \
diff --git a/Engine/Source/Runtime/Core/Public/Windows/WindowsPlatformCompilerPreSetup.h b/Engine/Source/Runtime/Core/Public/Windows/WindowsPlatformCompilerPreSetup.h
--- a/Engine/Source/Runtime/Core/Public/Windows/WindowsPlatformCompilerPreSetup.h
+++ b/Engine/Source/Runtime/Core/Public/Windows/WindowsPlatformCompilerPreSetup.h
@@ -79,14 +79,27 @@
 	#ifndef PRAGMA_DISABLE_UNSAFE_TYPECAST_WARNINGS
 		#define PRAGMA_DISABLE_UNSAFE_TYPECAST_WARNINGS \
 			__pragma (warning(push)) \
-			__pragma (warning(disable: 4244)) /* 'argument': conversion from 'type1' to 'type2', possible loss of data */
+			__pragma (warning(disable: 4244)) /* 'argument': conversion from 'type1' to 'type2', possible loss of data */	\
+			__pragma (warning(disable: 4838)) /* 'argument': conversion from 'type1' to 'type2' requires a narrowing conversion */
 	#endif // PRAGMA_DISABLE_UNSAFE_TYPECAST_WARNINGS
 
 	#ifndef PRAGMA_ENABLE_UNSAFE_TYPECAST_WARNINGS
 		#define PRAGMA_ENABLE_UNSAFE_TYPECAST_WARNINGS \
-			__pragma(warning(pop))
+			DEPRECATED_MACRO(5.0, "The PRAGMA_ENABLE_UNSAFE_TYPECAST_WARNINGS macro has been deprecated in favor of PRAGMA_RESTORE_UNSAFE_TYPECAST_WARNINGS. To force enable warnings use PRAGMA_FORCE_UNSAFE_TYPECAST_WARNINGS.")
 	#endif // PRAGMA_ENABLE_UNSAFE_TYPECAST_WARNINGS
 
+	#ifndef PRAGMA_FORCE_UNSAFE_TYPECAST_WARNINGS
+		#define PRAGMA_FORCE_UNSAFE_TYPECAST_WARNINGS \
+			__pragma (warning(push)) \
+			__pragma (warning(error: 4244)) /* 'argument': conversion from 'type1' to 'type2', possible loss of data */	\
+			__pragma (warning(error: 4838)) /* 'argument': conversion from 'type1' to 'type2' requires a narrowing conversion */
+	#endif // PRAGMA_FORCE_UNSAFE_TYPECAST_WARNINGS
+
+	#ifndef PRAGMA_RESTORE_UNSAFE_TYPECAST_WARNINGS
+		#define PRAGMA_RESTORE_UNSAFE_TYPECAST_WARNINGS \
+			__pragma(warning(pop))
+	#endif // PRAGMA_RESTORE_UNSAFE_TYPECAST_WARNINGS
+
 	#ifndef PRAGMA_DISABLE_UNDEFINED_IDENTIFIER_WARNINGS
 		#define PRAGMA_DISABLE_UNDEFINED_IDENTIFIER_WARNINGS \
 			__pragma(warning(push)) \
