Lyra - Adding file lock checking for the deployment script, adding a temporary method of using BPT tool 1.5 until it's available online, will have to correct the script then.

diff --git a/Engine/Source/Programs/AutomationTool/Gauntlet/Framework/FileLockInfo.cs b/Engine/Source/Programs/AutomationTool/Gauntlet/Framework/FileLockInfo.cs
--- a/Engine/Source/Programs/AutomationTool/Gauntlet/Framework/FileLockInfo.cs
+++ b/Engine/Source/Programs/AutomationTool/Gauntlet/Framework/FileLockInfo.cs
@@ -37,6 +37,12 @@ namespace FileLockInfo
 		/// http://wyupdate.googlecode.com/svn-history/r401/trunk/frmFilesInUse.cs (no copyright in code at time of viewing)
 		/// </remarks>
 		public static List<Process> GetProcessesLockingFile(string path)
+		{
+			string[] resources = { path }; // Just checking on one resource.
+			return GetProcessesLockingFiles(resources);
+		}
+
+		public static List<Process> GetProcessesLockingFiles(string[] FilePaths)
 		{
 			uint handle;
 			string key = Guid.NewGuid().ToString();
@@ -49,9 +55,7 @@ namespace FileLockInfo
 				const int MORE_DATA = 234;
 				uint pnProcInfoNeeded, pnProcInfo = 0, lpdwRebootReasons = RmRebootReasonNone;
 
-				string[] resources = { path }; // Just checking on one resource.
-
-				res = RmRegisterResources(handle, (uint)resources.Length, resources, 0, null, 0, null);
+				res = RmRegisterResources(handle, (uint)FilePaths.Length, FilePaths, 0, null, 0, null);
 
 				if (res != 0) throw new Exception("Could not register resource.");
 
