AutomationTool: Remove UE4 references

#trivial

diff --git a/Engine/Source/Programs/AutomationTool/Android/AndroidPlatform.Automation.cs b/Engine/Source/Programs/AutomationTool/Android/AndroidPlatform.Automation.cs
--- a/Engine/Source/Programs/AutomationTool/Android/AndroidPlatform.Automation.cs
+++ b/Engine/Source/Programs/AutomationTool/Android/AndroidPlatform.Automation.cs
@@ -893,12 +893,12 @@ public class AndroidPlatform : Platform
 			bool bHaveAPK = !bEnableBundle;     // do not have a standard APK if bundle enabled
 			if (!SC.IsCodeBasedProject)
 			{
-				string UE4SOName = GetFinalApkName(Params, SC.StageExecutables[0], false, bMakeSeparateApks ? Architecture : "");
-                UE4SOName = UE4SOName.Replace(".apk", ".so");
-                if (FileExists_NoExceptions(UE4SOName) == false)
+				string UnrealSOName = GetFinalApkName(Params, SC.StageExecutables[0], false, bMakeSeparateApks ? Architecture : "");
+                UnrealSOName = UnrealSOName.Replace(".apk", ".so");
+                if (FileExists_NoExceptions(UnrealSOName) == false)
 				{
-					LogInformation("Failed to find game .so " + UE4SOName);
-                    throw new AutomationException(ExitCode.Error_MissingExecutable, "Stage Failed. Could not find .so {0}. You may need to build the UE project with your target configuration and platform.", UE4SOName);
+					LogInformation("Failed to find game .so " + UnrealSOName);
+                    throw new AutomationException(ExitCode.Error_MissingExecutable, "Stage Failed. Could not find .so {0}. You may need to build the UE project with your target configuration and platform.", UnrealSOName);
 				}
 			}
 				
diff --git a/Engine/Source/Programs/AutomationTool/AutomationUtils/BuildCommand.cs b/Engine/Source/Programs/AutomationTool/AutomationUtils/BuildCommand.cs
--- a/Engine/Source/Programs/AutomationTool/AutomationUtils/BuildCommand.cs
+++ b/Engine/Source/Programs/AutomationTool/AutomationUtils/BuildCommand.cs
@@ -256,13 +256,13 @@ namespace AutomationTool
 			if (bForeign)
 			{
 				var DestSample = ParseParamValue("DestSample", "CopiedHoverShip");
-				var Dest = ParseParamValue("ForeignDest", CombinePaths(@"C:\testue4\foreign\", DestSample + "_ _Dir"));
+				var Dest = ParseParamValue("ForeignDest", CombinePaths(@"C:\testue\foreign\", DestSample + "_ _Dir"));
 				ProjectFullPath = new FileReference(CombinePaths(Dest, DestSample + ".uproject"));
 			}
 			else if (bForeignCode)
 			{
 				var DestSample = ParseParamValue("DestSample", "PlatformerGame");
-				var Dest = ParseParamValue("ForeignDest", CombinePaths(@"C:\testue4\foreign\", DestSample + "_ _Dir"));
+				var Dest = ParseParamValue("ForeignDest", CombinePaths(@"C:\testue\foreign\", DestSample + "_ _Dir"));
 				ProjectFullPath = new FileReference(CombinePaths(Dest, DestSample + ".uproject"));
 			}
 			else
diff --git a/Engine/Source/Programs/AutomationTool/AutomationUtils/CommandEnvironment.cs b/Engine/Source/Programs/AutomationTool/AutomationUtils/CommandEnvironment.cs
--- a/Engine/Source/Programs/AutomationTool/AutomationUtils/CommandEnvironment.cs
+++ b/Engine/Source/Programs/AutomationTool/AutomationUtils/CommandEnvironment.cs
@@ -46,7 +46,7 @@ namespace AutomationTool
 	public class CommandEnvironment
 	{
 		/// <summary>
-		/// Path to a file we know to always exist under the UE4 root directory.
+		/// Path to a file we know to always exist under the Unreal root directory.
 		/// </summary>
 		public static readonly string KnownFileRelativeToRoot = @"Engine/Config/BaseEngine.ini";
 
diff --git a/Engine/Source/Programs/AutomationTool/AutomationUtils/DeploymentContext.cs b/Engine/Source/Programs/AutomationTool/AutomationUtils/DeploymentContext.cs
--- a/Engine/Source/Programs/AutomationTool/AutomationUtils/DeploymentContext.cs
+++ b/Engine/Source/Programs/AutomationTool/AutomationUtils/DeploymentContext.cs
@@ -162,7 +162,7 @@ public class DeploymentContext //: ProjectParams
 	public string CustomConfig;
 
 	/// <summary>
-	/// This is the root directory that contains the engine: d:\a\UE4\
+	/// This is the root directory that contains the engine: d:\a\UE\
 	/// </summary>
 	public DirectoryReference LocalRoot;
 
@@ -172,7 +172,7 @@ public class DeploymentContext //: ProjectParams
 	public DirectoryReference EngineRoot;
 
 	/// <summary>
-	/// The directory that contains the project: d:\a\UE4\ShooterGame
+	/// The directory that contains the project: d:\a\UE\ShooterGame
 	/// </summary>
 	public DirectoryReference ProjectRoot;
 
diff --git a/Engine/Source/Programs/AutomationTool/AutomationUtils/HostPlatform.cs b/Engine/Source/Programs/AutomationTool/AutomationUtils/HostPlatform.cs
--- a/Engine/Source/Programs/AutomationTool/AutomationUtils/HostPlatform.cs
+++ b/Engine/Source/Programs/AutomationTool/AutomationUtils/HostPlatform.cs
@@ -53,7 +53,7 @@ namespace AutomationTool
 		abstract public string GetDotnetMsbuildExe();
 
 		/// <summary>
-		/// Folder under UE4/ to the platform's binaries.
+		/// Folder under UE/ to the platform's binaries.
 		/// </summary>
 		abstract public string RelativeBinariesFolder { get; }
 
diff --git a/Engine/Source/Programs/AutomationTool/AutomationUtils/P4Environment.cs b/Engine/Source/Programs/AutomationTool/AutomationUtils/P4Environment.cs
--- a/Engine/Source/Programs/AutomationTool/AutomationUtils/P4Environment.cs
+++ b/Engine/Source/Programs/AutomationTool/AutomationUtils/P4Environment.cs
@@ -42,7 +42,7 @@ namespace AutomationTool
 		}
 
 		/// <summary>
-		/// The current branch containing UAT. This may be a depot path (eg. //depot/UE4) or the name of a stream (eg. //UE4/Main).
+		/// The current branch containing UAT. This may be a depot path (e.g. //depot/UE4) or the name of a stream (e.g. //UE5/Main).
 		/// </summary>
 		public string Branch
 		{
@@ -51,7 +51,7 @@ namespace AutomationTool
 		}
 
 		/// <summary>
-		/// Path to the root of the branch as a client path (//MyClient/UE4).
+		/// Path to the root of the branch as a client path (//MyClient/UE).
 		/// </summary>
 		public string ClientRoot
 		{
diff --git a/Engine/Source/Programs/AutomationTool/AutomationUtils/P4Utils.cs b/Engine/Source/Programs/AutomationTool/AutomationUtils/P4Utils.cs
--- a/Engine/Source/Programs/AutomationTool/AutomationUtils/P4Utils.cs
+++ b/Engine/Source/Programs/AutomationTool/AutomationUtils/P4Utils.cs
@@ -3930,7 +3930,7 @@ namespace AutomationTool
 		/// <summary>
 		/// Enumerates all streams in a depot
 		/// </summary>
-		/// <param name="StreamPath">The path for streams to enumerate (eg. "//UE4/...")</param>
+		/// <param name="StreamPath">The path for streams to enumerate (eg. "//UE5/...")</param>
 		/// <returns>List of streams matching the given criteria</returns>
 		public List<P4StreamRecord> Streams(string StreamPath)
 		{
@@ -3940,7 +3940,7 @@ namespace AutomationTool
 		/// <summary>
 		/// Enumerates all streams in a depot
 		/// </summary>
-		/// <param name="StreamPath">The path for streams to enumerate (eg. "//UE4/...")</param>
+		/// <param name="StreamPath">The path for streams to enumerate (eg. "//UE5/...")</param>
 		/// <param name="MaxResults">Maximum number of results to return</param>
 		/// <param name="Filter">Additional filter to be applied to the results</param>
 		/// <param name="bUnloaded">Whether to enumerate unloaded workspaces</param>
diff --git a/Engine/Source/Programs/AutomationTool/AutomationUtils/Utils.cs b/Engine/Source/Programs/AutomationTool/AutomationUtils/Utils.cs
--- a/Engine/Source/Programs/AutomationTool/AutomationUtils/Utils.cs
+++ b/Engine/Source/Programs/AutomationTool/AutomationUtils/Utils.cs
@@ -825,7 +825,7 @@ namespace AutomationTool
 	            {
 	                int Retry = 0;
 	                int NumRetries = 60;
-	                if(!Directoryname.Contains("UE4"))
+	                if(!Directoryname.Contains("UE"))
 	                {
 	                    NumRetries = 2;
 	                }
diff --git a/Engine/Source/Programs/AutomationTool/BuildGraph/Tasks/TagReceiptTask.cs b/Engine/Source/Programs/AutomationTool/BuildGraph/Tasks/TagReceiptTask.cs
--- a/Engine/Source/Programs/AutomationTool/BuildGraph/Tasks/TagReceiptTask.cs
+++ b/Engine/Source/Programs/AutomationTool/BuildGraph/Tasks/TagReceiptTask.cs
@@ -174,7 +174,7 @@ namespace AutomationTool.Tasks
 							continue;
 						}
 
-						// Check which files exist, and warn about any that don't. Ignore debug files, as they are frequently excluded for size (eg. UE4 on GitHub). This matches logic during staging.
+						// Check which files exist, and warn about any that don't. Ignore debug files, as they are frequently excluded for size (eg. UE on GitHub). This matches logic during staging.
 						FileReference DependencyPath = RuntimeDependency.Path;
 						if (FileReference.Exists(DependencyPath))
 						{
diff --git a/Engine/Source/Programs/AutomationTool/CookedEditor/CookedEditor.Automation.cs b/Engine/Source/Programs/AutomationTool/CookedEditor/CookedEditor.Automation.cs
--- a/Engine/Source/Programs/AutomationTool/CookedEditor/CookedEditor.Automation.cs
+++ b/Engine/Source/Programs/AutomationTool/CookedEditor/CookedEditor.Automation.cs
@@ -1025,7 +1025,7 @@ public class MakeCookedEditor : BuildCommand
 		foreach (string AProp in Props)
 		{
 			string Prop = AProp.Trim(" \t".ToCharArray());
-			// find the first = (UE4 properties can't have an equal sign, so it's valid to do)
+			// find the first = (UE properties can't have an equal sign, so it's valid to do)
 			int Equals = Prop.IndexOf('=');
 			// we must have one
 			if (Equals == -1)
diff --git a/Engine/Source/Programs/AutomationTool/IOS/IOSPlatform.Automation.cs b/Engine/Source/Programs/AutomationTool/IOS/IOSPlatform.Automation.cs
--- a/Engine/Source/Programs/AutomationTool/IOS/IOSPlatform.Automation.cs
+++ b/Engine/Source/Programs/AutomationTool/IOS/IOSPlatform.Automation.cs
@@ -595,19 +595,19 @@ public class IOSPlatform : Platform
 		return PakParams;
 	}
 
-	public virtual bool PrepForUATPackageOrDeploy(UnrealTargetConfiguration Config, FileReference ProjectFile, string InProjectName, DirectoryReference InProjectDirectory, string InExecutablePath, DirectoryReference InEngineDir, bool bForDistribution, string CookFlavor, bool bIsDataDeploy, bool bCreateStubIPA, bool bIsUE4Game)
+	public virtual bool PrepForUATPackageOrDeploy(UnrealTargetConfiguration Config, FileReference ProjectFile, string InProjectName, DirectoryReference InProjectDirectory, string InExecutablePath, DirectoryReference InEngineDir, bool bForDistribution, string CookFlavor, bool bIsDataDeploy, bool bCreateStubIPA, bool bIsUEGame)
 	{
-		FileReference TargetReceiptFileName = GetTargetReceiptFileName(Config, InExecutablePath, InEngineDir, InProjectDirectory, bIsUE4Game);
+		FileReference TargetReceiptFileName = GetTargetReceiptFileName(Config, InExecutablePath, InEngineDir, InProjectDirectory, bIsUEGame);
 
 		return IOSExports.PrepForUATPackageOrDeploy(Config, ProjectFile, InProjectName, InProjectDirectory, InExecutablePath, InEngineDir, bForDistribution, CookFlavor, bIsDataDeploy, bCreateStubIPA, TargetReceiptFileName);
 	}
 
 
-	private FileReference GetTargetReceiptFileName(UnrealTargetConfiguration Config, string InExecutablePath, DirectoryReference InEngineDir, DirectoryReference InProjectDirectory, bool bIsUE4Game)
+	private FileReference GetTargetReceiptFileName(UnrealTargetConfiguration Config, string InExecutablePath, DirectoryReference InEngineDir, DirectoryReference InProjectDirectory, bool bIsUEGame)
 	{
 		string TargetName = Path.GetFileNameWithoutExtension(InExecutablePath).Split("-".ToCharArray())[0];
 		FileReference TargetReceiptFileName;
-		if (bIsUE4Game)
+		if (bIsUEGame)
 		{
 			TargetReceiptFileName = TargetReceipt.GetDefaultPath(InEngineDir, "UnrealGame", UnrealTargetPlatform.IOS, Config, "");
 		}
@@ -623,10 +623,10 @@ public class IOSPlatform : Platform
 		IOSExports.GetProvisioningData(InProject, bDistribution, out MobileProvision, out SigningCertificate, out TeamUUID, out bAutomaticSigning);
 	}
 
-	public virtual bool DeployGeneratePList(FileReference ProjectFile, UnrealTargetConfiguration Config, DirectoryReference ProjectDirectory, bool bIsUE4Game, string GameName, bool bIsClient, string ProjectName, DirectoryReference InEngineDir, DirectoryReference AppDirectory, string InExecutablePath, out bool bSupportsPortrait, out bool bSupportsLandscape)
+	public virtual bool DeployGeneratePList(FileReference ProjectFile, UnrealTargetConfiguration Config, DirectoryReference ProjectDirectory, bool bIsUEGame, string GameName, bool bIsClient, string ProjectName, DirectoryReference InEngineDir, DirectoryReference AppDirectory, string InExecutablePath, out bool bSupportsPortrait, out bool bSupportsLandscape)
 	{
-		FileReference TargetReceiptFileName = GetTargetReceiptFileName(Config, InExecutablePath, InEngineDir, ProjectDirectory, bIsUE4Game);
-		return IOSExports.GeneratePList(ProjectFile, Config, ProjectDirectory, bIsUE4Game, GameName, bIsClient, ProjectName, InEngineDir, AppDirectory, TargetReceiptFileName, out bSupportsPortrait, out bSupportsLandscape);
+		FileReference TargetReceiptFileName = GetTargetReceiptFileName(Config, InExecutablePath, InEngineDir, ProjectDirectory, bIsUEGame);
+		return IOSExports.GeneratePList(ProjectFile, Config, ProjectDirectory, bIsUEGame, GameName, bIsClient, ProjectName, InEngineDir, AppDirectory, TargetReceiptFileName, out bSupportsPortrait, out bSupportsLandscape);
 	}
 
 	protected string MakeIPAFileName(UnrealTargetConfiguration TargetConfiguration, ProjectParams Params, DeploymentContext SC, bool bAllowDistroPrefix)
@@ -676,9 +676,9 @@ public class IOSPlatform : Platform
 		string InExecutablePath = CombinePaths(Path.GetDirectoryName(Params.GetProjectExeForPlatform(TargetPlatformType).ToString()), SC.StageExecutables[0]);
 		DirectoryReference InEngineDir = DirectoryReference.Combine(SC.LocalRoot, "Engine");
 		DirectoryReference InProjectDirectory = Params.RawProjectPath.Directory;
-		bool bIsUE4Game = !SC.IsCodeBasedProject;
+		bool bIsUEGame = !SC.IsCodeBasedProject;
 
-		FileReference ReceiptFileName = GetTargetReceiptFileName(Config, InExecutablePath, InEngineDir, InProjectDirectory, bIsUE4Game);
+		FileReference ReceiptFileName = GetTargetReceiptFileName(Config, InExecutablePath, InEngineDir, InProjectDirectory, bIsUEGame);
 		TargetReceipt Receipt;
 		bool bIsReadSuccessful = TargetReceipt.TryRead(ReceiptFileName, out Receipt);
 
@@ -770,14 +770,14 @@ public class IOSPlatform : Platform
 
 		bool bIsBuiltAsFramework = IsBuiltAsFramework(Params, SC);
 
-		// ensure the ue4game binary exists, if applicable
+		// ensure the UnrealGame binary exists, if applicable
 #if !PLATFORM_MAC
 		string ProjectGameExeFilename = Params.GetProjectExeForPlatform(TargetPlatformType).ToString();
 		string FullExePath = CombinePaths(Path.GetDirectoryName(ProjectGameExeFilename), SC.StageExecutables[0] + (UnrealBuildTool.BuildHostPlatform.Current.Platform != UnrealTargetPlatform.Mac ? ".stub" : ""));
 		if (!SC.IsCodeBasedProject && !FileExists_NoExceptions(FullExePath) && !bIsBuiltAsFramework)
 		{
 			LogError("Failed to find game binary " + FullExePath);
-			throw new AutomationException(ExitCode.Error_MissingExecutable, "Stage Failed. Could not find binary {0}. You may need to build the UE4 project with your target configuration and platform.", FullExePath);
+			throw new AutomationException(ExitCode.Error_MissingExecutable, "Stage Failed. Could not find binary {0}. You may need to build the Unreal Engine project with your target configuration and platform.", FullExePath);
 		}
 #endif // PLATFORM_MAC
 
@@ -1045,7 +1045,7 @@ public class IOSPlatform : Platform
 
 	private string EnsureXcodeProjectExists(FileReference RawProjectPath, string LocalRoot, string ShortProjectName, string ProjectRoot, bool IsCodeBasedProject, out bool bWasGenerated)
 	{
-		// first check for ue4.xcodeproj
+		// first check for the .xcodeproj
 		bWasGenerated = false;
 		string RawProjectDir = RawProjectPath.Directory.FullName;
 		string XcodeProj = RawProjectPath.FullName.Replace(".uproject", "_" + PlatformName + ".xcworkspace");
diff --git a/Engine/Source/Programs/AutomationTool/Linux/LinuxPlatform.Automation.cs b/Engine/Source/Programs/AutomationTool/Linux/LinuxPlatform.Automation.cs
--- a/Engine/Source/Programs/AutomationTool/Linux/LinuxPlatform.Automation.cs
+++ b/Engine/Source/Programs/AutomationTool/Linux/LinuxPlatform.Automation.cs
@@ -126,10 +126,10 @@ public abstract class BaseLinuxPlatform : Platform
 		string EOL = "\n";
 		Script.Append("#!/bin/sh" + EOL);
 		// allow running from symlinks
-		Script.AppendFormat("UE4_TRUE_SCRIPT_NAME=$(echo \\\"$0\\\" | xargs readlink -f)" + EOL);
-		Script.AppendFormat("UE4_PROJECT_ROOT=$(dirname \"$UE4_TRUE_SCRIPT_NAME\")" + EOL);
-		Script.AppendFormat("chmod +x \"$UE4_PROJECT_ROOT/{0}\"" + EOL, StagedRelativeTargetPath);
-		Script.AppendFormat("\"$UE4_PROJECT_ROOT/{0}\" {1} \"$@\" " + EOL, StagedRelativeTargetPath, StagedArguments);
+		Script.AppendFormat("UE_TRUE_SCRIPT_NAME=$(echo \\\"$0\\\" | xargs readlink -f)" + EOL);
+		Script.AppendFormat("UE_PROJECT_ROOT=$(dirname \"$UE_TRUE_SCRIPT_NAME\")" + EOL);
+		Script.AppendFormat("chmod +x \"$UE_PROJECT_ROOT/{0}\"" + EOL, StagedRelativeTargetPath);
+		Script.AppendFormat("\"$UE_PROJECT_ROOT/{0}\" {1} \"$@\" " + EOL, StagedRelativeTargetPath, StagedArguments);
 
 		// write out the 
 		FileReference.WriteAllText(IntermediateFile, Script.ToString());
diff --git a/Engine/Source/Programs/AutomationTool/Program.cs b/Engine/Source/Programs/AutomationTool/Program.cs
--- a/Engine/Source/Programs/AutomationTool/Program.cs
+++ b/Engine/Source/Programs/AutomationTool/Program.cs
@@ -364,7 +364,7 @@ namespace AutomationToolDriver
             ExitCode ReturnCode = ExitCode.Success;
 			try
 			{
-				// Set the working directory to the UE4 root
+				// Set the working directory to the Unreal root directory
 				Environment.CurrentDirectory = Path.GetFullPath(Path.Combine(Path.GetDirectoryName(Assembly.GetExecutingAssembly().GetOriginalLocation()), "..", "..", "..", ".."));
 
 				// Ensure we can resolve any external assemblies as necessary.
diff --git a/Engine/Source/Programs/AutomationTool/Scripts/BuildCookRun.Automation.cs b/Engine/Source/Programs/AutomationTool/Scripts/BuildCookRun.Automation.cs
--- a/Engine/Source/Programs/AutomationTool/Scripts/BuildCookRun.Automation.cs
+++ b/Engine/Source/Programs/AutomationTool/Scripts/BuildCookRun.Automation.cs
@@ -229,7 +229,7 @@ public class BuildCookRun : BuildCommand
 			throw new AutomationException("Can't find source directory to make foreign sample {0}.", Src);
 		}
 
-		var Dest = ParseParamValue("ForeignDest", CombinePaths(@"C:\testue4\foreign\", DestSample + "_ _Dir"));
+		var Dest = ParseParamValue("ForeignDest", CombinePaths(@"C:\testue\foreign\", DestSample + "_ _Dir"));
 		LogInformation("Make a foreign sample {0} -> {1}", Src, Dest);
 
 		CloneDirectory(Src, Dest);
@@ -255,7 +255,7 @@ public class BuildCookRun : BuildCommand
 			throw new AutomationException("Can't find source directory to make foreign sample {0}.", Src);
 		}
 
-		var Dest = ParseParamValue("ForeignDest", CombinePaths(@"C:\testue4\foreign\", DestSample + "_ _Dir"));
+		var Dest = ParseParamValue("ForeignDest", CombinePaths(@"C:\testue\foreign\", DestSample + "_ _Dir"));
 		LogInformation("Make a foreign sample {0} -> {1}", Src, Dest);
 
 		CloneDirectory(Src, Dest);
diff --git a/Engine/Source/Programs/AutomationTool/Scripts/CheckForHacks.cs b/Engine/Source/Programs/AutomationTool/Scripts/CheckForHacks.cs
--- a/Engine/Source/Programs/AutomationTool/Scripts/CheckForHacks.cs
+++ b/Engine/Source/Programs/AutomationTool/Scripts/CheckForHacks.cs
@@ -14,7 +14,7 @@ namespace AutomationTool
 {
 	[Help("Audits the current branch for comments denoting a hack that was not meant to leave another branch, following a given format (\"BEGIN XXXX HACK\", where XXXX is one or more tags separated by spaces).")]
 	[Help("Allowed tags may be specified manually on the command line. At least one must match, otherwise it will print a warning.")]
-	[Help("The current branch name and fragments of the branch path will also be added by default, so running from //UE4/Main will add \"//UE4/Main\", \"UE4\", and \"Main\".")]
+	[Help("The current branch name and fragments of the branch path will also be added by default, so running from //UE5/Main will add \"//UE5/Main\", \"UE5\", and \"Main\".")]
 	[Help("-Allow", "Specifies additional tags which are allowed in the BEGIN ... HACK tag list")]
 	class CheckForHacks : BuildCommand
 	{
diff --git a/Engine/Source/Programs/AutomationTool/Scripts/ListMobileDevices.Automation.cs b/Engine/Source/Programs/AutomationTool/Scripts/ListMobileDevices.Automation.cs
--- a/Engine/Source/Programs/AutomationTool/Scripts/ListMobileDevices.Automation.cs
+++ b/Engine/Source/Programs/AutomationTool/Scripts/ListMobileDevices.Automation.cs
@@ -1,4 +1,4 @@
-﻿// Copyright Epic Games, Inc. All Rights Reserved.
+// Copyright Epic Games, Inc. All Rights Reserved.
 
 using System.Collections.Generic;
 using AutomationTool;
@@ -15,7 +15,7 @@ namespace AutomationScripts.Automation
 
 			var GlobalParams = new ProjectParams(
 				Command: this,
-				RawProjectPath: new FileReference(@"D:\UE-Main\UE4\Samples\Games\TappyChicken\TappyChicken.uproject")
+				RawProjectPath: new FileReference(@"D:\UE\Samples\Games\TappyChicken\TappyChicken.uproject")
 				);
 
 			if (ParseParam("android"))
diff --git a/Engine/Source/Programs/AutomationTool/TVOS/TVOSPlatform.Automation.cs b/Engine/Source/Programs/AutomationTool/TVOS/TVOSPlatform.Automation.cs
--- a/Engine/Source/Programs/AutomationTool/TVOS/TVOSPlatform.Automation.cs
+++ b/Engine/Source/Programs/AutomationTool/TVOS/TVOSPlatform.Automation.cs
@@ -17,11 +17,11 @@ public class TVOSPlatform : IOSPlatform
 		TargetIniPlatformType = UnrealTargetPlatform.IOS;
 	}
 
-	public override bool PrepForUATPackageOrDeploy(UnrealTargetConfiguration Config, FileReference ProjectFile, string InProjectName, DirectoryReference InProjectDirectory, string InExecutablePath, DirectoryReference InEngineDir, bool bForDistribution, string CookFlavor, bool bIsDataDeploy, bool bCreateStubIPA, bool bIsUE4Game)
+	public override bool PrepForUATPackageOrDeploy(UnrealTargetConfiguration Config, FileReference ProjectFile, string InProjectName, DirectoryReference InProjectDirectory, string InExecutablePath, DirectoryReference InEngineDir, bool bForDistribution, string CookFlavor, bool bIsDataDeploy, bool bCreateStubIPA, bool bIsUEGame)
 	{
 		string TargetName = Path.GetFileNameWithoutExtension(InExecutablePath).Split("-".ToCharArray())[0];
 		FileReference TargetReceiptFileName;
-		if (bIsUE4Game)
+		if (bIsUEGame)
 		{
 			TargetReceiptFileName = TargetReceipt.GetDefaultPath(InEngineDir, "UnrealGame", UnrealTargetPlatform.TVOS, Config, "");
 		}
@@ -37,11 +37,11 @@ public class TVOSPlatform : IOSPlatform
 		TVOSExports.GetProvisioningData(InProject, bDistribution, out MobileProvision, out SigningCertificate, out Team, out bAutomaticSigning);
     }
 
-	public override bool DeployGeneratePList(FileReference ProjectFile, UnrealTargetConfiguration Config, DirectoryReference ProjectDirectory, bool bIsUE4Game, string GameName, bool bIsClient, string ProjectName, DirectoryReference InEngineDir, DirectoryReference AppDirectory, string InExecutablePath, out bool bSupportsPortrait, out bool bSupportsLandscape)
+	public override bool DeployGeneratePList(FileReference ProjectFile, UnrealTargetConfiguration Config, DirectoryReference ProjectDirectory, bool bIsUEGame, string GameName, bool bIsClient, string ProjectName, DirectoryReference InEngineDir, DirectoryReference AppDirectory, string InExecutablePath, out bool bSupportsPortrait, out bool bSupportsLandscape)
 	{
 		string TargetName = Path.GetFileNameWithoutExtension(InExecutablePath).Split("-".ToCharArray())[0];
 		FileReference TargetReceiptFileName;
-		if (bIsUE4Game)
+		if (bIsUEGame)
 		{
 			TargetReceiptFileName = TargetReceipt.GetDefaultPath(InEngineDir, "UnrealGame", UnrealTargetPlatform.TVOS, Config, "");
 		}
@@ -49,7 +49,7 @@ public class TVOSPlatform : IOSPlatform
 		{
 			TargetReceiptFileName = TargetReceipt.GetDefaultPath(ProjectDirectory, TargetName, UnrealTargetPlatform.TVOS, Config, "");
 		}
-		return TVOSExports.GeneratePList(ProjectFile, Config, ProjectDirectory, bIsUE4Game, GameName, bIsClient, ProjectName, InEngineDir, AppDirectory, TargetReceiptFileName, out bSupportsPortrait, out bSupportsLandscape);
+		return TVOSExports.GeneratePList(ProjectFile, Config, ProjectDirectory, bIsUEGame, GameName, bIsClient, ProjectName, InEngineDir, AppDirectory, TargetReceiptFileName, out bSupportsPortrait, out bSupportsLandscape);
 	}
 
     public override string GetCookPlatform(bool bDedicatedServer, bool bIsClientOnly)
diff --git a/Engine/Source/Programs/AutomationTool/Turnkey/Commands/Test.cs b/Engine/Source/Programs/AutomationTool/Turnkey/Commands/Test.cs
--- a/Engine/Source/Programs/AutomationTool/Turnkey/Commands/Test.cs
+++ b/Engine/Source/Programs/AutomationTool/Turnkey/Commands/Test.cs
@@ -25,7 +25,7 @@ namespace Turnkey.Commands
 				@"googledrive:'1O6V1JO-LNZYxcB6wIqCiHjQ44mz4tv7S'/*",
 				@"file:D:\Engine\*\Build\*.xml",
 				@"file:D:\Engine\*\*",
-				@"perforce://UE4/Main/*/Build/*.xml",
+				@"perforce://UE5/Main/*/Build/*.xml",
 			};
 
 			foreach (string Test in TestEnumerates)
diff --git a/Engine/Source/Programs/AutomationTool/Turnkey/CopyProvider.cs b/Engine/Source/Programs/AutomationTool/Turnkey/CopyProvider.cs
--- a/Engine/Source/Programs/AutomationTool/Turnkey/CopyProvider.cs
+++ b/Engine/Source/Programs/AutomationTool/Turnkey/CopyProvider.cs
@@ -24,7 +24,7 @@ namespace Turnkey
 	abstract class CopyProvider
 	{
 		/// <summary>
-		/// Unique provider token to direct a copy operation ("perforce://UE4/Main/...)
+		/// Unique provider token to direct a copy operation ("perforce://UE5/Main/...)
 		/// </summary>
 		public abstract string ProviderToken { get; }
 
@@ -32,7 +32,7 @@ namespace Turnkey
 		/// Perform the copy operation
 		/// </summary>
 		/// <param name="Operation">Description for the operation</param>
-		/// <returns>The output path of the copied file, or a directory that contains all of the wildcards in the operation ("perforce://UE4/SDKs/.../Windows/*" would return something like "d:\UE4\Sdks")</returns>
+		/// <returns>The output path of the copied file, or a directory that contains all of the wildcards in the operation ("perforce://depot/CarefullyRedist/.../Windows/*" would return something like "d:\\sdks")</returns>
 		public abstract string Execute(string Operation, CopyExecuteSpecialMode SpecialMode, string SpecialModeHint);
 
 		/// <summary>
