BuildGraph: Take the digest of any *.version files when copying files to temp storage, so their timestamps can be updated without making any content changes.

diff --git a/Engine/Source/Programs/AutomationTool/BuildGraph/TempStorage.cs b/Engine/Source/Programs/AutomationTool/BuildGraph/TempStorage.cs
--- a/Engine/Source/Programs/AutomationTool/BuildGraph/TempStorage.cs
+++ b/Engine/Source/Programs/AutomationTool/BuildGraph/TempStorage.cs
@@ -262,7 +262,7 @@ namespace AutomationTool
 		/// <returns>True to generate a digest for this file, rather than relying on timestamps</returns>
 		bool GenerateDigest()
 		{
-			return RelativePath.EndsWith(".modules", StringComparison.OrdinalIgnoreCase);
+			return RelativePath.EndsWith(".version", StringComparison.OrdinalIgnoreCase) || RelativePath.EndsWith(".modules", StringComparison.OrdinalIgnoreCase);
 		}
 
 		/// <summary>
