- Fixed up various files included with new MakeCookedEditor
- Added "MapMode" which can allow a cooked editor/cooker to stage cooked, uncooked, or no maps

diff --git a/Engine/Config/BaseGame.ini b/Engine/Config/BaseGame.ini
--- a/Engine/Config/BaseGame.ini
+++ b/Engine/Config/BaseGame.ini
@@ -265,6 +265,8 @@ ReleaseTargetType=Game
 
 
 [CookedEditorSettings_CookedEditor]
+; cooked maps don't work correctly in the editor now, so stage them uncooked
+MapMode=Uncooked
 
 ; currently don't work well with cooked content
 +DisabledPlugins=GeometryMode
@@ -281,3 +283,9 @@ ReleaseTargetType=Game
 +DisabledPluginsInShipping=TraceSourceFilters
 
 [CookedEditorSettings_CookedCooker]
+; cooked cookers don't edit maps, so we can use cooked ones
+MapMode=Cooked
+
+; stage all non-NDAd platgform configs, and editor configs
++EngineExtraStageFiles=(Path="Config")
++ProjectExtraStageFiles=(Path="Config")
diff --git a/Engine/Source/Programs/AutomationTool/CookedEditor/CookedEditor.Automation.cs b/Engine/Source/Programs/AutomationTool/CookedEditor/CookedEditor.Automation.cs
--- a/Engine/Source/Programs/AutomationTool/CookedEditor/CookedEditor.Automation.cs
+++ b/Engine/Source/Programs/AutomationTool/CookedEditor/CookedEditor.Automation.cs
@@ -286,8 +286,8 @@ public class ModifyStageContext
 	{
 		if (bIsForExternalDistribution)
 		{
-			// remove entries where any restricted folder names are in the name
-			Files = Files.Where(x => !SC.RestrictedFolderNames.Any(y => x.Key.ContainsName(y))).ToDictionary(x => x.Key, x => x.Value);
+			// remove entries where any restricted folder names are in the name remapped path (if we remap from NFL to non-NFL, then we don't remove it)
+			Files = Files.Where(x => !SC.RestrictedFolderNames.Any(y => Project.ApplyDirectoryRemap(SC, x.Key).ContainsName(y))).ToDictionary(x => x.Key, x => x.Value);
 		}
 		else
 		{
@@ -304,8 +304,28 @@ public class ModifyStageContext
 
 	private void UncookMaps(DeploymentContext SC)
 	{
-		// remove maps from SC and Context (SC has path to the cooked map, so we have to come back from Staged refernece that doesn't have the Cooked dir in it)
-		AddUFSFilesToList(FilesToUncook, ".umap", SC);
+		string CookedMapMode = ConfigHelper.GetString("MapMode").ToLower();
+		if (CookedMapMode == "cooked")
+		{
+			// nothing to do, they are already staged as cooked as normal
+		}
+		else if (CookedMapMode == "uncooked")
+		{
+			// remove maps from SC and Context (SC has path to the cooked map, so we have to come back from Staged refernece that doesn't have the Cooked dir in it)
+			AddUFSFilesToList(FilesToUncook, ".umap", SC);
+		}
+		else if (CookedMapMode == "none")
+		{
+			// remove umaps and their sidecar files so they won't be staged (remove extension so that we remove Foo.umap and Foo.uexp)
+			// also remove Foo_BuiltData.*
+			HashSet<string> Maps = UFSFilesToStage.Where(x => x.HasExtension("umap")).Select(x => Path.ChangeExtension(x.FullName, null)).ToHashSet();
+			UFSFilesToStage = UFSFilesToStage.Where(x => !Maps.Contains(Path.ChangeExtension(x.FullName.Replace("_BuiltData", ""), null))).ToList();
+
+			Maps = SC.FilesToStage.UFSFiles.Keys.Where(x => x.HasExtension("umap")).Select(x => Path.ChangeExtension(x.Name, null)).ToHashSet();
+			Console.WriteLine($"Found {Maps.Count()} maps");
+			SC.FilesToStage.UFSFiles = SC.FilesToStage.UFSFiles.Where(x => !Maps.Contains(Path.ChangeExtension(x.Key.Name.Replace("_BuiltData", ""), null))).ToDictionary(x => x.Key, x => x.Value);
+
+		}
 	}
 
 	private void UnUFSFiles(DeploymentContext SC)
@@ -315,6 +335,7 @@ public class ModifyStageContext
 			// UAT needs uplugin and ini files, so make sure they are not in the .pak
 			AddUFSFilesToList(NonUFSFilesToStage, ".uplugin", SC);
 			AddUFSFilesToList(NonUFSFilesToStage, ".ini", SC);
+			NonUFSFilesToStage = NonUFSFilesToStage.Where(x => x.GetFileName() != "BinaryConfig.ini").ToList();
 		}
 	}
 
@@ -473,7 +494,6 @@ public class MakeCookedEditor : BuildCommand
 				// get location of the project where the other paths are relative to
 				DirectoryReference RecordRoot = FileReference.Combine(JsonFile.Directory, BuildRecord.ProjectPath).Directory;
 				string FullPath = RecordRoot.FullName;
-				Console.WriteLine(FullPath);
 
 				// stage the output and everything it pulled in next to it
 				foreach (FileReference TargetDirFile in DirectoryReference.EnumerateFiles(FileReference.Combine(RecordRoot, BuildRecord.TargetPath).Directory, "*", SearchOption.AllDirectories))
@@ -636,36 +656,42 @@ public class MakeCookedEditor : BuildCommand
 
 
 		// plugins are already handled in the Plugins staging code
-		List<string> RootFoldersToStrip = new List<string> { "Source", "Plugins" };//, "Binaries" };
-		List<string> SubFoldersToStrip = new List<string> { "Source", "Intermediate", "Tests", "Binaries" + Path.DirectorySeparatorChar + HostPlatform.Current.HostEditorPlatform.ToString() };
+		List<string> RootFoldersToStrip = new List<string> { "source", "plugins" };//, "binaries" };
+		List<string> SubFoldersToStrip = new List<string> { "source", "intermediate", "tests", "binaries" + Path.DirectorySeparatorChar + HostPlatform.Current.HostEditorPlatform.ToString().ToLower() };
+		List<string> RootNonUFSFolders = new List<string> { "shaders", "binaries", "build", "extras" };
+
+
 		if (!Context.bStageShaderDirs)
 		{
-			RootFoldersToStrip.Add("Shaders");
+			RootFoldersToStrip.Add("shaders");
 		}
 		if (!Context.bStageBuildDirs)
 		{
-			RootFoldersToStrip.Add("Build");
+			RootFoldersToStrip.Add("build");
 		}
 		if (!Context.bStageExtrasDirs)
 		{
-			RootFoldersToStrip.Add("Extras");
+			RootFoldersToStrip.Add("extras");
 		}
 
 		foreach (DirectoryReference PlatformDir in Unreal.GetExtensionDirs(RootDir, true, false, false))
 		{
 			foreach (DirectoryReference Subdir in DirectoryReference.EnumerateDirectories(PlatformDir, "*", SearchOption.TopDirectoryOnly))
 			{
+				string SubdirName = Subdir.GetDirectoryName().ToLower();
+
 				// Remvoe some unnecessary folders that can be large
 				List<FileReference> ContextFileList = Context.UFSFilesToStage;
 
-				if (Subdir.GetDirectoryName() == "Shaders" || Subdir.GetDirectoryName() == "Binaries")
+				// some files need to be NonUFS for C# etc to access
+				if (RootNonUFSFolders.Contains(SubdirName))
 				{
 					ContextFileList = Context.NonUFSFilesToStage;
 				}
 
 				List<FileReference> FilesToStage = new List<FileReference>();
 				// if we aren't in a bad subdir, add files
-				if (!RootFoldersToStrip.Contains(Subdir.GetDirectoryName(), StringComparer.InvariantCultureIgnoreCase))
+				if (!RootFoldersToStrip.Contains(SubdirName))
 				{
 					FilesToStage.AddRange(DirectoryReference.EnumerateFiles(Subdir, "*", SearchOption.AllDirectories));
 
