AutomationTool: More vigorously validate build records before attempting to load script modules

diff --git a/Engine/Source/Programs/AutomationTool/CompileScriptModules.cs b/Engine/Source/Programs/AutomationTool/CompileScriptModules.cs
--- a/Engine/Source/Programs/AutomationTool/CompileScriptModules.cs
+++ b/Engine/Source/Programs/AutomationTool/CompileScriptModules.cs
@@ -91,21 +91,32 @@ namespace AutomationToolDriver
 				GameFolders: Unreal.IsEngineInstalled() ? GameDirectories : new List<DirectoryReference>(), 
 				ForeignPlugins: null, AdditionalSearchPaths: AdditionalDirectories.Concat(GameBuildDirectories).ToList()));
 
-			bool bUseBuildRecordsOnlyForProjectDiscovery = bNoCompile || Unreal.IsEngineInstalled() || FoundAutomationProjects.Count() == 0;
+			bool bUseBuildRecordsOnlyForProjectDiscovery = bNoCompile || Unreal.IsEngineInstalled();
 
 			// Load existing build records, validating them only if (re)compiling script projects is an option
-			Dictionary<CsProjBuildRecord, FileReference> ExistingBuildRecords = LoadExistingBuildRecords(BaseDirectories, !bUseBuildRecordsOnlyForProjectDiscovery, ref WriteTimeCache);
+			Dictionary<FileReference, (CsProjBuildRecord, FileReference)> ExistingBuildRecords = LoadExistingBuildRecords(BaseDirectories);
+
+			Dictionary<FileReference, CsProjBuildRecord> ValidBuildRecords = new Dictionary<FileReference, CsProjBuildRecord>(ExistingBuildRecords.Count);
+			Dictionary<FileReference, (CsProjBuildRecord BuildRecord, FileReference BuildRecordPath)> InvalidBuildRecords = new Dictionary<FileReference, (CsProjBuildRecord, FileReference)>(ExistingBuildRecords.Count);
+
+			if (bUseBuildRecords)
+			{
+				foreach (FileReference AutomationProject in FoundAutomationProjects)
+				{
+					ValidateBuildRecordRecursively(ValidBuildRecords, InvalidBuildRecords, ExistingBuildRecords, AutomationProject, ref WriteTimeCache);
+				}
+			}
 
 			if (bUseBuildRecordsOnlyForProjectDiscovery)
             {
-				// when the engine is installed, we expect to find at least one script module (AutomationUtils is a necessity)
+				// when the engine is installed, or UAT is invoked with -NoCompile, we expect to find at least one script module (AutomationUtils is a necessity)
 				if (ExistingBuildRecords.Count == 0)
                 {
 					throw new Exception("Found no script module records.");
                 }
 
 				HashSet<FileReference> BuiltTargets = new HashSet<FileReference>(ExistingBuildRecords.Count);
-				foreach ((CsProjBuildRecord BuildRecord, FileReference BuildRecordPath) in ExistingBuildRecords)
+				foreach ((CsProjBuildRecord BuildRecord, FileReference BuildRecordPath) in ExistingBuildRecords.Values)
                 {
 					FileReference ProjectPath = FileReference.Combine(BuildRecordPath.Directory, BuildRecord.ProjectPath);
 					FileReference TargetPath = FileReference.Combine(ProjectPath.Directory, BuildRecord.TargetPath);
@@ -134,34 +145,27 @@ namespace AutomationToolDriver
             }
 			else
             {
-				// when the engine is not installed, delete any .json file that does not have a corresponding .csproj file
-				foreach ((CsProjBuildRecord BuildRecord, FileReference BuildRecordPath) in ExistingBuildRecords)
+				// when the engine is not installed, delete any build record .json file that is not valid
+				foreach ((CsProjBuildRecord _, FileReference BuildRecordPath) in InvalidBuildRecords.Values)
 				{
-					FileReference ProjectPath = FileReference.Combine(BuildRecordPath.Directory, BuildRecord.ProjectPath);
-
-					if (!FileReference.Exists(ProjectPath))
-                    {
-						Log.TraceInformation($"Deleting \"{BuildRecordPath}\" because referenced project file \"{ProjectPath}\" was not found");
-						FileReference.Delete(BuildRecordPath);
-                    }
+					Log.TraceLog($"Deleting invalid build record \"{BuildRecordPath}\"");
+					FileReference.Delete(BuildRecordPath);
                 }
             }
 
 			if (!bForceCompile && bUseBuildRecords)
 			{
-				// fastest path: if we have an up-to-date record of a previous build, we should be able to start faster
-				HashSet<FileReference> ScriptModules = TryGetAllUpToDateScriptModules(ExistingBuildRecords, FoundAutomationProjects, ref WriteTimeCache);
-
-				if (ScriptModules != null)
-				{
+				// If all found records are valid, we can return their targets directly
+				if (ValidBuildRecords.Count == FoundAutomationProjects.Count)
+                {
 					bBuildSuccess = true;
-					return ScriptModules;
-				}
+					return new HashSet<FileReference>(ValidBuildRecords.Select(x => FileReference.Combine(x.Key.Directory, x.Value.TargetPath)));
+                }
 			}
 
 			// Fall back to the slower approach: use msbuild to load csproj files & build as necessary
 			RegisterMsBuildPath();
-			return BuildAllScriptPlugins(FoundAutomationProjects, bForceCompile, bNoCompile, out bBuildSuccess, ref WriteTimeCache, BaseDirectories);
+			return BuildAllScriptPlugins(FoundAutomationProjects, bForceCompile || !bUseBuildRecords, out bBuildSuccess, ref WriteTimeCache, BaseDirectories);
 		}
 
 		/// <summary>
@@ -169,9 +173,9 @@ namespace AutomationToolDriver
 		/// </summary>
 		/// <param name="BaseDirectories"></param>
 		/// <returns></returns>
-		static Dictionary<CsProjBuildRecord, FileReference> LoadExistingBuildRecords(List<DirectoryReference> BaseDirectories, bool bValidateBuildRecordsAreUpToDate, ref WriteTimeCache Cache)
+		static Dictionary<FileReference, (CsProjBuildRecord, FileReference)> LoadExistingBuildRecords(List<DirectoryReference> BaseDirectories)
         {
-			Dictionary<CsProjBuildRecord, FileReference> LoadedBuildRecords = new Dictionary<CsProjBuildRecord, FileReference>();
+			Dictionary<FileReference, (CsProjBuildRecord, FileReference)> LoadedBuildRecords = new Dictionary<FileReference, (CsProjBuildRecord, FileReference)>();
 
 			foreach (DirectoryReference Directory in BaseDirectories)
 			{
@@ -183,32 +187,29 @@ namespace AutomationToolDriver
 
 				foreach (FileReference JsonFile in DirectoryReference.EnumerateFiles(IntermediateDirectory, "*.json"))
                 {
+					CsProjBuildRecord BuildRecord = default;
+
 					// filesystem errors or json parsing might result in an exception. If that happens, we fall back to the
-					// slower path - buildrecord files will be re-generated, other filesystem errors may persist
+					// slower path - if compiling, buildrecord files will be re-generated; other filesystem errors may persist
 					try
 					{
-						CsProjBuildRecord BuildRecord =
-							JsonSerializer.Deserialize<CsProjBuildRecord>(FileReference.ReadAllText(JsonFile));
-
-						if (bValidateBuildRecordsAreUpToDate)
-						{
-							DirectoryReference ProjectDirectory = FileReference.Combine(IntermediateDirectory, BuildRecord.ProjectPath).Directory;
-
-							if (!ValidateBuildRecord(BuildRecord, ProjectDirectory, out string ValidationFailureMessage, ref Cache))
-							{
-								Log.TraceLog($"[{JsonFile}] {ValidationFailureMessage}");
-								continue;
-							}
-						}
-
+						BuildRecord = JsonSerializer.Deserialize<CsProjBuildRecord>(FileReference.ReadAllText(JsonFile));
 						Log.TraceLog($"Loaded script module build record {JsonFile}");
-						LoadedBuildRecords.Add(BuildRecord, JsonFile);
 					}
 					catch(Exception Ex)
 					{
 						Log.TraceWarning($"[{JsonFile}] Failed to load build record: {Ex.Message}");
 					}
 
+					if (BuildRecord.ProjectPath != null)
+					{
+						LoadedBuildRecords.Add(FileReference.FromString(Path.GetFullPath(BuildRecord.ProjectPath, JsonFile.Directory.FullName)), (BuildRecord, JsonFile));
+					}
+					else
+                    {
+						// Delete the invalid build record
+						Log.TraceWarning($"Deleting invalid build record {JsonFile}");
+                    }
                 }
 			}
 
@@ -336,87 +337,64 @@ namespace AutomationToolDriver
 			return true;
 		}
 
-		// Loads build records for each project, if they exist, and then checks all recorded build dependencies to ensure
-		// that nothing has changed since the last build.
-		// This function is (currently?) all-or-nothing: either all projects are up-to-date, or none are.
-		private static HashSet<FileReference> TryGetAllUpToDateScriptModules(Dictionary<CsProjBuildRecord, FileReference> ExistingBuildRecords, HashSet<FileReference> FoundAutomationProjects, ref WriteTimeCache Cache)
+		static void ValidateBuildRecordRecursively(
+			Dictionary<FileReference, CsProjBuildRecord> ValidBuildRecords,
+			Dictionary<FileReference, (CsProjBuildRecord, FileReference)> InvalidBuildRecords,
+			Dictionary<FileReference, (CsProjBuildRecord, FileReference)> BuildRecords,
+			FileReference ProjectPath, ref WriteTimeCache Cache)
 		{
-			Dictionary<FileReference, CsProjBuildRecord> ExistingBuildRecordLookup = new Dictionary<FileReference, CsProjBuildRecord>(ExistingBuildRecords.Count);
-
-			Dictionary<FileReference, CsProjBuildRecord> ValidatedBuildRecords = new Dictionary<FileReference, CsProjBuildRecord>(ExistingBuildRecords.Count);
+			if (ValidBuildRecords.ContainsKey(ProjectPath) || InvalidBuildRecords.ContainsKey(ProjectPath))
+			{
+				// Project validity has already been determined
+				return;
+			}
 
-			foreach((CsProjBuildRecord BuildRecord, FileReference BuildRecordPath) in ExistingBuildRecords)
-            {
-				FileReference ProjectPath = FileReference.FromString(
-					Path.GetFullPath(BuildRecord.ProjectPath, BuildRecordPath.Directory.FullName));
-				ExistingBuildRecordLookup.Add(ProjectPath, BuildRecord);
-            }
+			// Was a build record loaded for this project path? (relevant when considering referenced projects)
+			if (!BuildRecords.TryGetValue(ProjectPath, out (CsProjBuildRecord BuildRecord, FileReference BuildRecordPath) Entry))
+			{
+				Log.TraceLog($"Found project {ProjectPath} with no existing build record");
+				return;
+			}
 
-			bool ValidateProjectBuildRecords(FileReference ProjectPath, ref WriteTimeCache Cache)
+			// Is this particular build record valid?
+			if (!ValidateBuildRecord(Entry.BuildRecord, ProjectPath.Directory, out string ValidationFailureMessage, ref Cache))
 			{
-				if (ValidatedBuildRecords.ContainsKey(ProjectPath))
-				{
-					return true;
-				}
+				string AutomationProjectRelativePath = Path.GetRelativePath(Unreal.EngineDirectory.FullName, ProjectPath.FullName);
+				Log.TraceLog($"[{AutomationProjectRelativePath}] {ValidationFailureMessage}");
 
-				if (!ExistingBuildRecordLookup.TryGetValue(ProjectPath, out CsProjBuildRecord BuildRecord)) // LoadUpToDateBuildRecord(ProjectPath, ref Cache, AllScriptFolders);
-				{
-					Log.TraceLog($"Found project {ProjectPath} with no existing build record");
-					return false;
-				}
+				InvalidBuildRecords.Add(ProjectPath, Entry);
+				return;
+			}
+
+			// Are all referenced build records valid?
+			foreach (string ReferencedProjectPath in Entry.BuildRecord.ProjectReferences)
+			{
+				FileReference FullProjectPath = FileReference.FromString(Path.GetFullPath(ReferencedProjectPath, ProjectPath.Directory.FullName));
+				ValidateBuildRecordRecursively(ValidBuildRecords, InvalidBuildRecords, BuildRecords, FullProjectPath, ref Cache);
 
-				if (!ValidateBuildRecord(BuildRecord, ProjectPath.Directory,
-					out string ValidationFailureMessage, ref Cache))
+				if (!ValidBuildRecords.ContainsKey(FullProjectPath))
 				{
-					string AutomationProjectRelativePath =
-						Path.GetRelativePath(Unreal.EngineDirectory.FullName, ProjectPath.FullName);
-				
-					Log.TraceLog($"[{AutomationProjectRelativePath}] {ValidationFailureMessage}");
-					return false;
+					string AutomationProjectRelativePath = Path.GetRelativePath(Unreal.EngineDirectory.FullName, ProjectPath.FullName);
+					string DependencyRelativePath = Path.GetRelativePath(Unreal.EngineDirectory.FullName, FullProjectPath.FullName);
+					Log.TraceLog($"[{AutomationProjectRelativePath}] Existing output is not valid because dependency {DependencyRelativePath} is not valid");
+					InvalidBuildRecords.Add(ProjectPath, Entry);
+					return;
 				}
 
-				ValidatedBuildRecords.Add(ProjectPath, BuildRecord);
-
-				foreach (string ReferencedProjectPath in BuildRecord.ProjectReferences)
+				// Ensure that the dependency was not built more recently than the project
+				if (Entry.BuildRecord.TargetBuildTime < ValidBuildRecords[FullProjectPath].TargetBuildTime)
 				{
-					FileReference FullProjectPath = FileReference.FromString(Path.GetFullPath(ReferencedProjectPath, ProjectPath.Directory.FullName));
-					if (!ValidateProjectBuildRecords(FullProjectPath, ref Cache))
-					{
-						return false;
-					}
+					string AutomationProjectRelativePath = Path.GetRelativePath(Unreal.EngineDirectory.FullName, ProjectPath.FullName);
+					string DependencyRelativePath = Path.GetRelativePath(Unreal.EngineDirectory.FullName, FullProjectPath.FullName);
+					Log.TraceLog($"[{AutomationProjectRelativePath}] Existing output is not valid because dependency {DependencyRelativePath} is newer");
+					InvalidBuildRecords.Add(ProjectPath, Entry);
+					return;
 				}
-
-				return true;
 			}
 
-			HashSet<FileReference> FoundAssemblies = new HashSet<FileReference>();
-			foreach (FileReference AutomationProject in FoundAutomationProjects)
-			{
-				if (!ValidateProjectBuildRecords(AutomationProject, ref Cache))
-                {
-					return null;
-                }
-				FoundAssemblies.Add(FileReference.Combine(AutomationProject.Directory, ValidatedBuildRecords[AutomationProject].TargetPath));
-			}
-
-			// it is possible that a referenced project has been rebuilt separately, that it is up to date, but that
-			// its build time is newer than a project that references it. Check for that.
-			foreach (KeyValuePair<FileReference, CsProjBuildRecord> Entry in ValidatedBuildRecords)
-            {
-				foreach(string ReferencedProjectPath in Entry.Value.ProjectReferences)
-                {
-					FileReference FullReferencedProjectPath = FileReference.FromString(Path.GetFullPath(ReferencedProjectPath, Entry.Key.Directory.FullName));
-					if (ValidatedBuildRecords[FullReferencedProjectPath].TargetBuildTime > Entry.Value.TargetBuildTime)
-                    {
-						Log.TraceLog($"[{Entry.Key.MakeRelativeTo(Unreal.EngineDirectory)}] referenced project target {ValidatedBuildRecords[FullReferencedProjectPath].TargetPath} build time ({ValidatedBuildRecords[FullReferencedProjectPath].TargetBuildTime}) is more recent than this project's build time ({Entry.Value.TargetBuildTime})");
-						return null;
-                    }
-                }
-            }
-
-			return FoundAssemblies;
+			ValidBuildRecords.Add(ProjectPath, Entry.BuildRecord);
 		}
-		
+
 		/// <summary>
 		/// Register our bundled dotnet installation to be used by Microsoft.Build
 		/// This needs to happen in a function called before the first use of any Microsoft.Build types
@@ -591,11 +569,10 @@ namespace AutomationToolDriver
 #endif
             };
 
-		private static HashSet<FileReference> BuildAllScriptPlugins(HashSet<FileReference> FoundAutomationProjects, bool bForceCompile, bool bNoCompile, out bool bBuildSuccess, 
-			ref WriteTimeCache Cache, List<DirectoryReference> BaseDirectories)
+		private static HashSet<FileReference> BuildAllScriptPlugins(HashSet<FileReference> FoundAutomationProjects,
+			bool bForceCompile, out bool bBuildSuccess, ref WriteTimeCache Cache, List<DirectoryReference> BaseDirectories)
 		{
-			// The -IgnoreBuildRecords prevents the loading & parsing of build record .json files from Intermediate/ScriptModules - but CsProjBuildRecord objects will be used in this function regardless
-			Dictionary<FileReference, CsProjBuildRecord> BuildRecords = new Dictionary<FileReference, CsProjBuildRecord>();
+			Dictionary<FileReference, (CsProjBuildRecord BuildRecord, FileReference _)> BuildRecords = new Dictionary<FileReference, (CsProjBuildRecord, FileReference)>();
 			
 			Dictionary<string, Project> Projects = new Dictionary<string, Project>();
 			
@@ -791,7 +768,7 @@ namespace AutomationToolDriver
 						Include = Include, Exclude = Exclude, Remove = Remove });
 				}
 
-				BuildRecords.Add(FileReference.FromString(Project.FullPath), BuildRecord);
+				BuildRecords.Add(FileReference.FromString(Project.FullPath), (BuildRecord, default));
 			}
 
 			// Potential optimization: Contructing the ProjectGraph here gives the full graph of dependencies - which is nice,
@@ -809,57 +786,16 @@ namespace AutomationToolDriver
 			}
 			else
 			{
-				HashSet<ProjectGraphNode> OutOfDateProjects = new HashSet<ProjectGraphNode>(FoundAutomationProjects.Count);
-
-				foreach (ProjectGraphNode Project in InputProjectGraph.ProjectNodesTopologicallySorted)
-				{
-					CsProjBuildRecord BuildRecord = BuildRecords[FileReference.FromString(Project.ProjectInstance.FullPath)];
-
-					string ValidationFailureMessage;
-					if (!ValidateBuildRecord(BuildRecord, DirectoryReference.FromString(Project.ProjectInstance.Directory),
-						out ValidationFailureMessage, ref Cache))
-					{
-						Log.TraceLog($"[{Path.GetFileName(Project.ProjectInstance.FullPath)}] is out of date:\n{ValidationFailureMessage}");
-						OutOfDateProjects.Add(Project);
-					}
-				}
+				Dictionary<FileReference, CsProjBuildRecord> ValidBuildRecords = new Dictionary<FileReference, CsProjBuildRecord>(BuildRecords.Count);
+				Dictionary<FileReference, (CsProjBuildRecord, FileReference _)> InvalidBuildRecords = new Dictionary<FileReference, (CsProjBuildRecord, FileReference)>(BuildRecords.Count);
 
-				// it is possible that a referenced project has been rebuilt separately, that it is up to date, but that
-				// its build time is newer than a project that references it. Check for that.
 				foreach (ProjectGraphNode Project in InputProjectGraph.ProjectNodesTopologicallySorted)
 				{
-					CsProjBuildRecord BuildRecord = BuildRecords[FileReference.FromString(Project.ProjectInstance.FullPath)];
-
-					foreach(string ReferencedProjectPath in BuildRecord.ProjectReferences)
-					{
-						FileReference FullReferencedProjectPath = FileReference.FromString(Path.GetFullPath(ReferencedProjectPath, Project.ProjectInstance.Directory));
-						if (BuildRecords[FullReferencedProjectPath].TargetBuildTime > BuildRecord.TargetBuildTime)
-						{
-							Log.TraceLog($"[{Path.GetFileName(Project.ProjectInstance.FullPath)}] referenced project target {BuildRecords[FullReferencedProjectPath].TargetPath} build time ({BuildRecords[FullReferencedProjectPath].TargetBuildTime}) is more recent than this project's build time ({BuildRecord.TargetBuildTime})");
-							OutOfDateProjects.Add(Project);
-						}
-					}
+					ValidateBuildRecordRecursively(ValidBuildRecords, InvalidBuildRecords, BuildRecords, FileReference.FromString(Project.ProjectInstance.FullPath), ref Cache);
 				}
 
-				// for any out of date project, mark everything that references it as out of date
-				Queue<ProjectGraphNode> OutOfDateQueue = new Queue<ProjectGraphNode>(OutOfDateProjects);
-				while (OutOfDateQueue.TryDequeue(out ProjectGraphNode OutOfDateProject))
-				{
-					foreach (ProjectGraphNode Referee in OutOfDateProject.ReferencingProjects)
-					{
-						if (OutOfDateProjects.Add(Referee))
-						{
-							OutOfDateQueue.Enqueue(Referee);
-						}
-					}
-				}
-
-				if (bNoCompile)
-				{
-					bBuildSuccess = true;
-					// return a list of all script modules that are up to date (without touching any CsProjBuildRecords)
-					return new HashSet<FileReference>(InputProjectGraph.EntryPointNodes.Where(P => !OutOfDateProjects.Contains(P)).Select(P => FileReference.FromString(P.ProjectInstance.GetPropertyValue("TargetPath"))));
-				}
+				// Select the projects that have been found to be out of date
+				HashSet<ProjectGraphNode> OutOfDateProjects = new HashSet<ProjectGraphNode>(InputProjectGraph.ProjectNodes.Where(x => InvalidBuildRecords.Keys.Contains(FileReference.FromString(x.ProjectInstance.FullPath))));
 
 				if (OutOfDateProjects.Count > 0)
 				{
@@ -883,14 +819,14 @@ namespace AutomationToolDriver
 
 				FileReference BuildRecordPath = ConstructBuildRecordPath(ProjectPath, BaseDirectories);
 
-				CsProjBuildRecord BuildRecord = BuildRecords[ProjectPath];
+				CsProjBuildRecord BuildRecord = BuildRecords[ProjectPath].BuildRecord;
 
 				// update target build times into build records to ensure everything is up-to-date
 				FileReference FullPath = FileReference.Combine(ProjectPath.Directory, BuildRecord.TargetPath);
 				BuildRecord.TargetBuildTime = FileReference.GetLastWriteTime(FullPath);
 
 				if (FileReference.WriteAllTextIfDifferent(BuildRecordPath, 
-					JsonSerializer.Serialize(BuildRecords[ProjectPath], new JsonSerializerOptions { WriteIndented = true })))
+					JsonSerializer.Serialize<CsProjBuildRecord>(BuildRecord, new JsonSerializerOptions { WriteIndented = true })))
                 {
 					Log.TraceLog($"Wrote script module build record to {BuildRecordPath}");
                 }
