Fixed the initial map setting of the project launcher being ignored when also specifying maps to cook (from jamie.dale)

diff --git a/Engine/Source/Developer/LauncherServices/Private/Launcher/LauncherWorker.cpp b/Engine/Source/Developer/LauncherServices/Private/Launcher/LauncherWorker.cpp
--- a/Engine/Source/Developer/LauncherServices/Private/Launcher/LauncherWorker.cpp
+++ b/Engine/Source/Developer/LauncherServices/Private/Launcher/LauncherWorker.cpp
@@ -509,23 +509,17 @@ FString FLauncherWorker::CreateUATCommand( const ILauncherProfileRef& InProfile,
 		*InProfile->GetAdditionalCommandLineParameters());
 
 	// map list
-	FString MapList = TEXT("");
+	FString MapList = TEXT(" -map=");
 	const TArray<FString>& CookedMaps = InProfile->GetCookedMaps();
 	if (CookedMaps.Num() > 0 && (InProfile->GetCookMode() == ELauncherProfileCookModes::ByTheBook || InProfile->GetCookMode() == ELauncherProfileCookModes::ByTheBookInEditor))
 	{
-		MapList += TEXT(" -map=");
-		for (int32 MapIndex = 0; MapIndex < CookedMaps.Num(); ++MapIndex)
-		{
-			MapList += CookedMaps[MapIndex];
-			if (MapIndex+1 < CookedMaps.Num())
-			{
-				MapList += "+";
-			}
-		}
+		MapList += InitialMap;
+		MapList += TEXT("+");
+		MapList += FString::Join(CookedMaps, TEXT("+"));
 	}
 	else
 	{
-		MapList = TEXT(" -map=") + InitialMap;
+		MapList += InitialMap;
 	}
 
 	// culture list
diff --git a/Engine/Source/Programs/AutomationTool/AutomationUtils/ProjectParams.cs b/Engine/Source/Programs/AutomationTool/AutomationUtils/ProjectParams.cs
--- a/Engine/Source/Programs/AutomationTool/AutomationUtils/ProjectParams.cs
+++ b/Engine/Source/Programs/AutomationTool/AutomationUtils/ProjectParams.cs
@@ -831,7 +831,6 @@ namespace AutomationTool
 			this.SeparateDebugInfo = GetParamValueIfNotSpecified(Command, SeparateDebugInfo, this.SeparateDebugInfo, "separatedebuginfo");
 			this.MapFile = GetParamValueIfNotSpecified(Command, MapFile, this.MapFile, "mapfile");
 			this.NoCleanStage = GetParamValueIfNotSpecified(Command, NoCleanStage, this.NoCleanStage, "nocleanstage");
-			this.MapToRun = ParseParamValueIfNotSpecified(Command, MapToRun, "map", String.Empty);
 			this.AdditionalServerMapParams = ParseParamValueIfNotSpecified(Command, AdditionalServerMapParams, "AdditionalServerMapParams", String.Empty);
 			this.Foreign = GetParamValueIfNotSpecified(Command, Foreign, this.Foreign, "foreign");
 			this.ForeignCode = GetParamValueIfNotSpecified(Command, ForeignCode, this.ForeignCode, "foreigncode");
@@ -850,6 +849,17 @@ namespace AutomationTool
 			this.Deploy = GetParamValueIfNotSpecified(Command, Deploy, this.Deploy, "deploy");
 			this.DeployFolder = ParseParamValueIfNotSpecified(Command, DeployFolder, "deploy", null);
 
+			// if the user specified multiple -map arguments, just use the first one
+			this.MapToRun = ParseParamValueIfNotSpecified(Command, MapToRun, "map", String.Empty);
+			if (!String.IsNullOrEmpty(this.MapToRun))
+			{
+				int DelimiterIndex = this.MapToRun.IndexOf('+');
+				if (DelimiterIndex != -1)
+				{
+					this.MapToRun = this.MapToRun.Remove(DelimiterIndex);
+				}
+			}
+
 			// if the user specified -deploy but no folder, set the default
 			if (this.Deploy && string.IsNullOrEmpty(this.DeployFolder))
 			{
