Fix for non-threadsafe logging that can result in occasional errors when running parallel tests

diff --git a/Engine/Source/Programs/AutomationTool/Gauntlet/Framework/Gauntlet.Utils.cs b/Engine/Source/Programs/AutomationTool/Gauntlet/Framework/Gauntlet.Utils.cs
--- a/Engine/Source/Programs/AutomationTool/Gauntlet/Framework/Gauntlet.Utils.cs
+++ b/Engine/Source/Programs/AutomationTool/Gauntlet/Framework/Gauntlet.Utils.cs
@@ -194,7 +194,7 @@ namespace Gauntlet
 			}
 		}
 
-		static StreamWriter LogFile = null;
+		static TextWriter LogFile = null;
 
 		static List<Action<string>> Callbacks;
 
@@ -237,7 +237,7 @@ namespace Gauntlet
 
 				try
 				{
-					LogFile = new StreamWriter(Outpath);
+					LogFile = TextWriter.Synchronized(new StreamWriter(Outpath));
 				}
 				catch (UnauthorizedAccessException Ex)
 				{
