Fix MapsToCook when multiple +-separated maps are provided to the cooker

diff --git a/Engine/Source/Programs/AutomationTool/AutomationUtils/ProjectParams.cs b/Engine/Source/Programs/AutomationTool/AutomationUtils/ProjectParams.cs
--- a/Engine/Source/Programs/AutomationTool/AutomationUtils/ProjectParams.cs
+++ b/Engine/Source/Programs/AutomationTool/AutomationUtils/ProjectParams.cs
@@ -831,6 +831,7 @@ namespace AutomationTool
 			this.SeparateDebugInfo = GetParamValueIfNotSpecified(Command, SeparateDebugInfo, this.SeparateDebugInfo, "separatedebuginfo");
 			this.MapFile = GetParamValueIfNotSpecified(Command, MapFile, this.MapFile, "mapfile");
 			this.NoCleanStage = GetParamValueIfNotSpecified(Command, NoCleanStage, this.NoCleanStage, "nocleanstage");
+			this.MapToRun = ParseParamValueIfNotSpecified(Command, MapToRun, "map", String.Empty);
 			this.AdditionalServerMapParams = ParseParamValueIfNotSpecified(Command, AdditionalServerMapParams, "AdditionalServerMapParams", String.Empty);
 			this.Foreign = GetParamValueIfNotSpecified(Command, Foreign, this.Foreign, "foreign");
 			this.ForeignCode = GetParamValueIfNotSpecified(Command, ForeignCode, this.ForeignCode, "foreigncode");
@@ -849,17 +850,6 @@ namespace AutomationTool
 			this.Deploy = GetParamValueIfNotSpecified(Command, Deploy, this.Deploy, "deploy");
 			this.DeployFolder = ParseParamValueIfNotSpecified(Command, DeployFolder, "deploy", null);
 
-			// if the user specified multiple -map arguments, just use the first one
-			this.MapToRun = ParseParamValueIfNotSpecified(Command, MapToRun, "map", String.Empty);
-			if (!String.IsNullOrEmpty(this.MapToRun))
-			{
-				int DelimiterIndex = this.MapToRun.IndexOf('+');
-				if (DelimiterIndex != -1)
-				{
-					this.MapToRun = this.MapToRun.Remove(DelimiterIndex);
-				}
-			}
-
 			// if the user specified -deploy but no folder, set the default
 			if (this.Deploy && string.IsNullOrEmpty(this.DeployFolder))
 			{
@@ -1030,6 +1020,16 @@ namespace AutomationTool
                 this.MapsToCook.Add(this.MapToRun);
             }
 
+			// if the user specified multiple -map arguments, just use the first one
+			if (!String.IsNullOrEmpty(this.MapToRun))
+			{
+				int DelimiterIndex = this.MapToRun.IndexOf('+');
+				if (DelimiterIndex != -1)
+				{
+					this.MapToRun = this.MapToRun.Remove(DelimiterIndex);
+				}
+			}
+
 			if (InMapsToRebuildLightMaps == null)
 			{
 				if (Command != null)
