UAT:  Fix skipClient support

- skipClient option allows for client steps to be run in UAT with a null role which causes them to be skipped
  - the test still runs and is considered successful and all dependencies are still executed
  - this results in a "Null" client device with a null platform
  - IDeviceUsageReporter causes exceptions with this configuration
- change IDeviceUsageReporter methods to take a platform reference and handle a null platform internally
  - no events are logged for the null device and platform

- the null client device associated with teh skip tests will create a NullAppInstance
  - NullAppInstance only completes upon WaitForExit() but many tests don't do this as they run until they see a specific line in the output log or something similar (ex:  boot tests)
  - this causes some NullAppInstance tests to never complete and to run until a timeout is encountered
- NullAppInstance is now considered completed and HasExited upon creation to resolve this
  - the base test node then identifies the null app instance as completed

diff --git a/Engine/Source/Programs/AutomationTool/Gauntlet/Framework/Devices/Gauntlet.UnrealDeviceReservation.cs b/Engine/Source/Programs/AutomationTool/Gauntlet/Framework/Devices/Gauntlet.UnrealDeviceReservation.cs
--- a/Engine/Source/Programs/AutomationTool/Gauntlet/Framework/Devices/Gauntlet.UnrealDeviceReservation.cs
+++ b/Engine/Source/Programs/AutomationTool/Gauntlet/Framework/Devices/Gauntlet.UnrealDeviceReservation.cs
@@ -145,7 +145,7 @@ namespace Gauntlet
 			{
 				foreach (ITargetDevice device in ReservedDevices)
 				{
-					IDeviceUsageReporter.RecordEnd(device.Name, (UnrealTargetPlatform)device.Platform, IDeviceUsageReporter.EventType.Device);
+					IDeviceUsageReporter.RecordEnd(device.Name, device.Platform, IDeviceUsageReporter.EventType.Device);
 				}
 				DevicePool.Instance.ReleaseDevices(ReservedDevices);
 				ReservedDevices.Clear();
diff --git a/Engine/Source/Programs/AutomationTool/Gauntlet/Platform/Null/Gauntlet.TargetDeviceNull.cs b/Engine/Source/Programs/AutomationTool/Gauntlet/Platform/Null/Gauntlet.TargetDeviceNull.cs
--- a/Engine/Source/Programs/AutomationTool/Gauntlet/Platform/Null/Gauntlet.TargetDeviceNull.cs
+++ b/Engine/Source/Programs/AutomationTool/Gauntlet/Platform/Null/Gauntlet.TargetDeviceNull.cs
@@ -16,6 +16,7 @@ namespace Gauntlet
 		public NullAppInstance(ITargetDevice InDevice)
 		{
 			Device = InDevice;
+			HasExited = true;
 		}
 
 		public string ArtifactPath
diff --git a/Engine/Source/Programs/AutomationTool/Gauntlet/Unreal/Base/Gauntlet.UnrealSession.cs b/Engine/Source/Programs/AutomationTool/Gauntlet/Unreal/Base/Gauntlet.UnrealSession.cs
--- a/Engine/Source/Programs/AutomationTool/Gauntlet/Unreal/Base/Gauntlet.UnrealSession.cs
+++ b/Engine/Source/Programs/AutomationTool/Gauntlet/Unreal/Base/Gauntlet.UnrealSession.cs
@@ -379,7 +379,7 @@ namespace Gauntlet
 					{
 						App.Kill();
 						// Apps that are still running have timed out => fail
-						IDeviceUsageReporter.RecordEnd(App.Device.Name, (UnrealTargetPlatform)App.Device.Platform, IDeviceUsageReporter.EventType.Test, IDeviceUsageReporter.EventState.Success);
+						IDeviceUsageReporter.RecordEnd(App.Device.Name, App.Device.Platform, IDeviceUsageReporter.EventType.Test, IDeviceUsageReporter.EventState.Success);
 					});
 				}
 
@@ -389,7 +389,7 @@ namespace Gauntlet
 					ClosedApps.ForEach(App =>
 					{
 						// Apps that have already exited have 'succeeded'
-						IDeviceUsageReporter.RecordEnd(App.Device.Name, (UnrealTargetPlatform)App.Device.Platform, IDeviceUsageReporter.EventType.Test, IDeviceUsageReporter.EventState.Failure);
+						IDeviceUsageReporter.RecordEnd(App.Device.Name, App.Device.Platform, IDeviceUsageReporter.EventType.Test, IDeviceUsageReporter.EventState.Failure);
 					});
 				}
 			}
@@ -734,12 +734,12 @@ namespace Gauntlet
 
 					IAppInstall Install = null;
 
-					IDeviceUsageReporter.RecordStart(Device.Name, (UnrealTargetPlatform)Device.Platform, IDeviceUsageReporter.EventType.Device, IDeviceUsageReporter.EventState.Success);
-					IDeviceUsageReporter.RecordStart(Device.Name, (UnrealTargetPlatform)Device.Platform, IDeviceUsageReporter.EventType.Install, IDeviceUsageReporter.EventState.Success, BuildSource.BuildName);
+					IDeviceUsageReporter.RecordStart(Device.Name, Device.Platform, IDeviceUsageReporter.EventType.Device, IDeviceUsageReporter.EventState.Success);
+					IDeviceUsageReporter.RecordStart(Device.Name, Device.Platform, IDeviceUsageReporter.EventType.Install, IDeviceUsageReporter.EventState.Success, BuildSource.BuildName);
 					try
 					{
 						Install = Device.InstallApplication(AppConfig);
-						IDeviceUsageReporter.RecordEnd(Device.Name, (UnrealTargetPlatform)Device.Platform, IDeviceUsageReporter.EventType.Install, IDeviceUsageReporter.EventState.Success);
+						IDeviceUsageReporter.RecordEnd(Device.Name, Device.Platform, IDeviceUsageReporter.EventType.Install, IDeviceUsageReporter.EventState.Success);
 					}
 					catch (System.Exception Ex)
 					{
@@ -747,7 +747,7 @@ namespace Gauntlet
 						Log.Info("Failed to install app onto device {0} for role {1}. {2}. Will retry with new device", Device, Role, Ex);
 						UnrealDeviceReservation.MarkProblemDevice(Device);
 						InstallSuccess = false;
-						IDeviceUsageReporter.RecordEnd(Device.Name, (UnrealTargetPlatform)Device.Platform, IDeviceUsageReporter.EventType.Install, IDeviceUsageReporter.EventState.Failure);
+						IDeviceUsageReporter.RecordEnd(Device.Name, Device.Platform, IDeviceUsageReporter.EventType.Install, IDeviceUsageReporter.EventState.Failure);
 						break;
 					}
 					
@@ -794,7 +794,7 @@ namespace Gauntlet
 						try
 						{
 							IAppInstance Instance = CurrentInstall.Run();
-							IDeviceUsageReporter.RecordStart(Instance.Device.Name, (UnrealTargetPlatform)Instance.Device.Platform, IDeviceUsageReporter.EventType.Test);
+							IDeviceUsageReporter.RecordStart(Instance.Device.Name, Instance.Device.Platform, IDeviceUsageReporter.EventType.Test);
 
 							if (Instance != null || Globals.CancelSignalled)
 							{
@@ -1165,7 +1165,7 @@ namespace Gauntlet
 					Log.VeryVerbose("Calling SaveRoleArtifacts, Role: {0}  Artifact Path: {1}", App.ToString(), App.AppInstance.ArtifactPath);
 					UnrealRoleArtifacts Artifacts = null;
 					ITargetDevice device = App.AppInstance.Device;
-					IDeviceUsageReporter.RecordStart(device.Name, (UnrealTargetPlatform)device.Platform, IDeviceUsageReporter.EventType.SavingArtifacts);
+					IDeviceUsageReporter.RecordStart(device.Name, device.Platform, IDeviceUsageReporter.EventType.SavingArtifacts);
 					try
 					{
 						Artifacts = SaveRoleArtifacts(Context, App, DestPath);
@@ -1173,12 +1173,12 @@ namespace Gauntlet
 					catch (Exception)
 					{
 						// Caught an exception -> report failure
-						IDeviceUsageReporter.RecordEnd(device.Name, (UnrealTargetPlatform)device.Platform, IDeviceUsageReporter.EventType.SavingArtifacts, IDeviceUsageReporter.EventState.Failure);
+						IDeviceUsageReporter.RecordEnd(device.Name, device.Platform, IDeviceUsageReporter.EventType.SavingArtifacts, IDeviceUsageReporter.EventState.Failure);
 						// Pass exception to the surrounding try/catch in UnrealTestNode while preserving the original callstack
 						throw;
 					}
 					// Did not catch -> successful reporting
-					IDeviceUsageReporter.RecordEnd(device.Name, (UnrealTargetPlatform)device.Platform, IDeviceUsageReporter.EventType.SavingArtifacts, IDeviceUsageReporter.EventState.Success);
+					IDeviceUsageReporter.RecordEnd(device.Name, device.Platform, IDeviceUsageReporter.EventType.SavingArtifacts, IDeviceUsageReporter.EventState.Success);
 
 					if (Artifacts != null)
 					{
diff --git a/Engine/Source/Programs/AutomationTool/Gauntlet/Unreal/Base/Gauntlet.UnrealTestNode.cs b/Engine/Source/Programs/AutomationTool/Gauntlet/Unreal/Base/Gauntlet.UnrealTestNode.cs
--- a/Engine/Source/Programs/AutomationTool/Gauntlet/Unreal/Base/Gauntlet.UnrealTestNode.cs
+++ b/Engine/Source/Programs/AutomationTool/Gauntlet/Unreal/Base/Gauntlet.UnrealTestNode.cs
@@ -705,8 +705,8 @@ namespace Gauntlet
 			{
 				if (AppInstance != null)
 				{
-					IDeviceUsageReporter.RecordComment(AppInstance.Device.Name, (UnrealTargetPlatform)AppInstance.Device.Platform, IDeviceUsageReporter.EventType.Device, Context.Options.JobDetails);
-					IDeviceUsageReporter.RecordComment(AppInstance.Device.Name, (UnrealTargetPlatform)AppInstance.Device.Platform, IDeviceUsageReporter.EventType.Test, this.GetType().Name);
+					IDeviceUsageReporter.RecordComment(AppInstance.Device.Name, AppInstance.Device.Platform, IDeviceUsageReporter.EventType.Device, Context.Options.JobDetails);
+					IDeviceUsageReporter.RecordComment(AppInstance.Device.Name, AppInstance.Device.Platform, IDeviceUsageReporter.EventType.Test, this.GetType().Name);
 				}
 			}
 			
diff --git a/Engine/Source/Programs/AutomationTool/Gauntlet/Unreal/Utils/Gauntlet.DeviceUsageReporter.cs b/Engine/Source/Programs/AutomationTool/Gauntlet/Unreal/Utils/Gauntlet.DeviceUsageReporter.cs
--- a/Engine/Source/Programs/AutomationTool/Gauntlet/Unreal/Utils/Gauntlet.DeviceUsageReporter.cs
+++ b/Engine/Source/Programs/AutomationTool/Gauntlet/Unreal/Utils/Gauntlet.DeviceUsageReporter.cs
@@ -22,19 +22,28 @@ namespace Gauntlet
 			, Success = 1
 		};
 
-		public static void RecordStart(string deviceName, UnrealTargetPlatform platform, EventType et, EventState state = EventState.Success, string comment = "")
+		public static void RecordStart(string deviceName, UnrealTargetPlatform? platform, EventType et, EventState state = EventState.Success, string comment = "")
 		{
-			RecordToAll(deviceName, platform, et, true, state, comment);
+			if (platform.HasValue)
+			{
+				RecordToAll(deviceName, platform.Value, et, true, state, comment);
+			}
 		}
 
-		public static void RecordEnd(string deviceName, UnrealTargetPlatform platform, EventType et, EventState state = EventState.Success, string comment = "")
+		public static void RecordEnd(string deviceName, UnrealTargetPlatform? platform, EventType et, EventState state = EventState.Success, string comment = "")
 		{
-			RecordToAll(deviceName, platform, et, false, state, comment);
+			if (platform.HasValue)
+			{
+				RecordToAll(deviceName, platform.Value, et, false, state, comment);
+			}
 		}
 
-		public static void RecordComment(string deviceName, UnrealTargetPlatform platform, EventType et, string comment)
+		public static void RecordComment(string deviceName, UnrealTargetPlatform? platform, EventType et, string comment)
 		{
-			CommentToAll(deviceName, platform, et, comment);
+			if (platform.HasValue)
+            {
+				CommentToAll(deviceName, platform.Value, et, comment);
+			}
 		}
 
 		private static List<IDeviceUsageReporter> GetEnabledReporters()
