AutomationTool: Remove error string that served only to misdirect

#trivial

diff --git a/Engine/Source/Programs/AutomationTool/Program.cs b/Engine/Source/Programs/AutomationTool/Program.cs
--- a/Engine/Source/Programs/AutomationTool/Program.cs
+++ b/Engine/Source/Programs/AutomationTool/Program.cs
@@ -395,7 +395,7 @@ namespace AutomationToolDriver
 			}
 			catch (Exception Ex)
             {
-				Log.TraceError("Exception occurred between AutomationToolDriver.Main() and Automation.Process()" + ExceptionUtils.FormatException(Ex));
+				Log.TraceError(ExceptionUtils.FormatException(Ex));
             }
             finally
             {
