AutomationTool: For installed tools, check for shared dlls in the AutomationTool binary directory

diff --git a/Engine/Source/Programs/AutomationTool/CompileScriptModules.cs b/Engine/Source/Programs/AutomationTool/CompileScriptModules.cs
--- a/Engine/Source/Programs/AutomationTool/CompileScriptModules.cs
+++ b/Engine/Source/Programs/AutomationTool/CompileScriptModules.cs
@@ -134,8 +134,20 @@ namespace AutomationToolDriver
 						}
 						else
 						{
-							// when the engine is installed, expect to find a built target assembly for every record that was found
-							throw new Exception($"Script module \"{TargetPath}\" not found for record \"{BuildRecordPath}\"");
+							// Workaround to address UE-138692:
+							// Shared/EpicGames.* assemblies build to their own local project directories, and those directories are not included in packages
+							// but those dlls, referenced by AutomationTool or its dependencies, are copied into AutomationTool's output directory.
+							// So, if the file is not found at TargetPath, look for the same file in the output directory.
+							DirectoryReference AutomationToolBinaryDirectory = DirectoryReference.Combine(Unreal.EngineDirectory, "Binaries", "DotNET", "AutomationTool");
+							if (FileReference.Exists(FileReference.Combine(AutomationToolBinaryDirectory, TargetPath.GetFileName())))
+							{
+								Log.TraceLog($"Using {TargetPath.GetFileName()} from {AutomationToolBinaryDirectory.FullName}");
+							}
+							else
+							{
+								// when the engine is installed, expect to find a built target assembly for every record that was found
+								throw new Exception($"Script module \"{TargetPath}\" not found for record \"{BuildRecordPath}\"");
+							}
 						}
 						
                     }
