BuildCookRun: Fix -SeparateDebugInfo not placing debug files in separate directory with the manifests

Looks like this broke in a refactor from 15771665

diff --git a/Engine/Source/Programs/AutomationTool/Scripts/CopyBuildToStagingDirectory.Automation.cs b/Engine/Source/Programs/AutomationTool/Scripts/CopyBuildToStagingDirectory.Automation.cs
--- a/Engine/Source/Programs/AutomationTool/Scripts/CopyBuildToStagingDirectory.Automation.cs
+++ b/Engine/Source/Programs/AutomationTool/Scripts/CopyBuildToStagingDirectory.Automation.cs
@@ -1707,11 +1707,12 @@ namespace AutomationScripts
 		{
 			if (bPerformCopy)
 			{
-				LogInformation("Copying {0} to staging directory: {1}", ManifestName, SC.StageDirectory);
+				var StageDirectory = ManifestName == "DebugFiles" ? SC.DebugStageDirectory : SC.StageDirectory;
+				LogInformation("Copying {0} to staging directory: {1}", ManifestName, StageDirectory);
 				foreach (KeyValuePair<StagedFileReference, FileReference> Pair in Mapping)
 				{
 					FileReference Src = Pair.Value;
-					FileReference Dest = FileReference.Combine(SC.StageDirectory, Pair.Key.Name);
+					FileReference Dest = FileReference.Combine(StageDirectory, Pair.Key.Name);
 					if (Src != Dest)  // special case for things created in the staging directory, like the pak file
 					{
 						CopyFileIncremental(Src, Dest, IniKeyDenyList: SC.IniKeyDenyList, IniSectionDenyList: SC.IniSectionDenyList);
