fix for turnkey install sdk crash when there are no valid sdks available.

#rnx

diff --git a/Engine/Source/Programs/AutomationTool/Turnkey/Commands/InstallSdk.cs b/Engine/Source/Programs/AutomationTool/Turnkey/Commands/InstallSdk.cs
--- a/Engine/Source/Programs/AutomationTool/Turnkey/Commands/InstallSdk.cs
+++ b/Engine/Source/Programs/AutomationTool/Turnkey/Commands/InstallSdk.cs
@@ -232,7 +232,7 @@ namespace Turnkey.Commands
 						Options.Add("Show Invalid Sdks");
 					}
 
-					int BestVersionIndex = Sdks.IndexOf(BestByType.First().Value);
+					int BestVersionIndex = Sdks.Any() ? Sdks.IndexOf(BestByType.First().Value) : -1;
 
 					string Current, MinVersion, MaxVersion;
 					if (DesiredType == FileSource.SourceType.Flash)
