Unshelved from pending changelist '18515746':

PR #8762: Iterative Deploy: Change the filename of the retrieved manifest so they don't override the already generated manifest. (Contributed by [REDACTED_PERSON])

diff --git a/Engine/Source/Programs/AutomationTool/Android/AndroidPlatform.Automation.cs b/Engine/Source/Programs/AutomationTool/Android/AndroidPlatform.Automation.cs
--- a/Engine/Source/Programs/AutomationTool/Android/AndroidPlatform.Automation.cs
+++ b/Engine/Source/Programs/AutomationTool/Android/AndroidPlatform.Automation.cs
@@ -1706,28 +1706,30 @@ public class AndroidPlatform : Platform
 		string RemoteDir = StorageLocation + "/UnrealGame/" + Params.ShortProjectName;
 
 		// Try retrieving the UFS files manifest files from the device
-		string UFSManifestFileName = CombinePaths(SC.StageDirectory.FullName, SC.GetUFSDeployedManifestFileName(DeviceName));
-		IProcessResult UFSResult = RunAdbCommand(Params, DeviceName, " pull " + RemoteDir + "/" + SC.GetUFSDeployedManifestFileName(null) + " \"" + UFSManifestFileName + "\"", null, ERunOptions.AppMustExist);
+		string RetrievedUFSManifestFileName = CombinePaths(SC.StageDirectory.FullName, $"Retrieved_{SC.GetUFSDeployedManifestFileName(DeviceName)}");
+		IProcessResult UFSResult = RunAdbCommand(Params, DeviceName, " pull " + RemoteDir + "/" + SC.GetUFSDeployedManifestFileName(null) + " \"" + RetrievedUFSManifestFileName + "\"", null, ERunOptions.AppMustExist);
 		if (!(UFSResult.Output.Contains("bytes") || UFSResult.Output.Contains("[100%]")))
 		{
+			LogWarning("Failed retrieving UFS Manifest: {0}", UFSResult.Output);
 			return false;
 		}
 
 		// Try retrieving the non UFS files manifest files from the device
-		string NonUFSManifestFileName = CombinePaths(SC.StageDirectory.FullName, SC.GetNonUFSDeployedManifestFileName(DeviceName));
-		IProcessResult NonUFSResult = RunAdbCommand(Params, DeviceName, " pull " + RemoteDir + "/" + SC.GetNonUFSDeployedManifestFileName(null) + " \"" + NonUFSManifestFileName + "\"", null, ERunOptions.AppMustExist);
+		string RetrievedNonUFSManifestFileName = CombinePaths(SC.StageDirectory.FullName, $"Retrieved_{SC.GetNonUFSDeployedManifestFileName(DeviceName)}");
+		IProcessResult NonUFSResult = RunAdbCommand(Params, DeviceName, " pull " + RemoteDir + "/" + SC.GetNonUFSDeployedManifestFileName(null) + " \"" + RetrievedNonUFSManifestFileName + "\"", null, ERunOptions.AppMustExist);
 		if (!(NonUFSResult.Output.Contains("bytes") || NonUFSResult.Output.Contains("[100%]")))
 		{
+			LogWarning("Failed retrieving NonUFS Manifest: {0}", NonUFSResult.Output);
 			// Did not retrieve both so delete one we did retrieve
-			File.Delete(UFSManifestFileName);
+			File.Delete(RetrievedUFSManifestFileName);
 			return false;
 		}
 
 		// Return the manifest files
 		UFSManifests = new List<string>();
-		UFSManifests.Add(UFSManifestFileName);
+		UFSManifests.Add(RetrievedUFSManifestFileName);
 		NonUFSManifests = new List<string>();
-		NonUFSManifests.Add(NonUFSManifestFileName);
+		NonUFSManifests.Add(RetrievedNonUFSManifestFileName);
 
 		return true;
 	}
