Turnkey will show the current device software version when updating.

#rnx

diff --git a/Engine/Source/Programs/AutomationTool/Turnkey/Commands/VerifySdk.cs b/Engine/Source/Programs/AutomationTool/Turnkey/Commands/VerifySdk.cs
--- a/Engine/Source/Programs/AutomationTool/Turnkey/Commands/VerifySdk.cs
+++ b/Engine/Source/Programs/AutomationTool/Turnkey/Commands/VerifySdk.cs
@@ -209,7 +209,7 @@ namespace Turnkey.Commands
 							{
 								if (Device.bCanConnect)
 								{
-									FileSource MatchingInstallableSdk = FileSource.FindMatchingSdk(AutomationPlatform, new FileSource.SourceType[] { FileSource.SourceType.Flash }, bSelectBest: bUnattended, DeviceType: Device.Type);
+									FileSource MatchingInstallableSdk = FileSource.FindMatchingSdk(AutomationPlatform, new FileSource.SourceType[] { FileSource.SourceType.Flash }, bSelectBest: bUnattended, DeviceType: Device.Type, CurrentSdk: Device.SoftwareVersion);
 
 									if (MatchingInstallableSdk == null)
 									{
diff --git a/Engine/Source/Programs/AutomationTool/Turnkey/FileSource.cs b/Engine/Source/Programs/AutomationTool/Turnkey/FileSource.cs
--- a/Engine/Source/Programs/AutomationTool/Turnkey/FileSource.cs
+++ b/Engine/Source/Programs/AutomationTool/Turnkey/FileSource.cs
@@ -292,7 +292,7 @@ namespace Turnkey
 			return Best;
 		}
 
-		static public FileSource FindMatchingSdk(AutomationTool.Platform Platform, SourceType[] TypePriority, bool bSelectBest, string DeviceType = null)
+		static public FileSource FindMatchingSdk(AutomationTool.Platform Platform, SourceType[] TypePriority, bool bSelectBest, string DeviceType = null, string CurrentSdk = null)
 		{
 			UEBuildPlatformSDK SDK = UEBuildPlatformSDK.GetSDKForPlatform(Platform.PlatformType.ToString());
 
@@ -334,9 +334,32 @@ namespace Turnkey
 					return Best;
 				}
 
+				// find the current sdk, if any
+				FileSource Current = null;
+				if (CurrentSdk != null)
+				{
+					Current = Sdks.Find( x => x.Version == CurrentSdk);
+				}
+
+				// helper for getting the display name for the sdks as they are presented to the user
+				string GetDisplayName( FileSource Sdk )
+				{
+					string Result = Sdk.Name;
+					if (Current != null && Sdk == Current) 
+					{
+						Result += " (current)";
+					}
+					if (Sdk == Best)
+					{
+						Result += " [Best Choice]";
+					}
+					return Result;
+				}
+
+
 				// if unable to pick one automatically, ask the user
 				int BestIndex = Sdks.IndexOf(Best);
-				int Choice = TurnkeyUtils.ReadInputInt("Multiple Sdks found that could be installed. Please select one:", Sdks.Select(x => x.Name + (x == Best ? " [Best Choice]" : "")).ToList(), true, BestIndex >= 0 ? BestIndex + 1 : -1);
+				int Choice = TurnkeyUtils.ReadInputInt("Multiple Sdks found that could be installed. Please select one:", Sdks.Select(x => GetDisplayName(x)).ToList(), true, BestIndex >= 0 ? BestIndex + 1 : -1);
 
 				// we take canceling to mean to try the next type
 				if (Choice == 0)
