- Removing platform Build dirs which were bloating the CW size with NewSchool mode

diff --git a/Engine/Config/BaseGame.ini b/Engine/Config/BaseGame.ini
--- a/Engine/Config/BaseGame.ini
+++ b/Engine/Config/BaseGame.ini
@@ -247,7 +247,7 @@ FrameInterpolationProcessor=Class'/Script/LiveLink.LiveLinkBasicFrameInterpolati
 
 ; default directories to stage
 bStageShaderDirs=true
-bStageBuildDirs=true
+bStagePlatformBuildDirs=false
 bStageExtrasDirs=false
 bStagePlatformDirs=true
 bStageUAT=true
diff --git a/Engine/Source/Programs/AutomationTool/CookedEditor/CookedEditor.Automation.cs b/Engine/Source/Programs/AutomationTool/CookedEditor/CookedEditor.Automation.cs
--- a/Engine/Source/Programs/AutomationTool/CookedEditor/CookedEditor.Automation.cs
+++ b/Engine/Source/Programs/AutomationTool/CookedEditor/CookedEditor.Automation.cs
@@ -79,7 +79,7 @@ public class ModifyStageContext
 	public List<FileReference> NonUFSFilesToStage = new List<FileReference>();
 
 	public bool bStageShaderDirs = true;
-	public bool bStageBuildDirs = true;
+	public bool bStagePlatformBuildDirs = true;
 	public bool bStageExtrasDirs = false;
 	public bool bStagePlatformDirs = true;
 	public bool bStageUAT = false;
@@ -109,7 +109,7 @@ public class ModifyStageContext
 		this.ConfigHelper = ConfigHelper;
 
 		bStageShaderDirs = ConfigHelper.GetBool("bStageShaderDirs");
-		bStageBuildDirs = ConfigHelper.GetBool("bStageBuildDirs");
+		bStagePlatformBuildDirs = ConfigHelper.GetBool("bStagePlatformBuildDirs");
 		bStageExtrasDirs = ConfigHelper.GetBool("bStageExtrasDirs");
 		bStagePlatformDirs = ConfigHelper.GetBool("bStagePlatformDirs");
 		bStageUAT = ConfigHelper.GetBool("bStageUAT");
@@ -665,7 +665,7 @@ public class MakeCookedEditor : BuildCommand
 		{
 			RootFoldersToStrip.Add("shaders");
 		}
-		if (!Context.bStageBuildDirs)
+		if (!Context.bStagePlatformBuildDirs)
 		{
 			RootFoldersToStrip.Add("build");
 		}
@@ -891,7 +891,6 @@ public class MakeCookedEditor : BuildCommand
 		// get cooked. This is okay, because the editor stuff should exist. We may want to revist this, and not cook anything that would
 		// cause the issues
 		Params.AdditionalCookerOptions += " -AllowUnsafeBlueprintCalls";
-		Params.AdditionalCookerOptions += " -dpcvars=cook.displaymode=2,r.ForceDebugViewModes=1";
 
 		// Params.AdditionalCookerOptions += " -NoFilterAssetRegistry";
 
