AutomationTool: Validate the existence of Automation.{csproj,json} targets only

The default build behavior for script assemblies is that their dependencies are copied to the same output dir - assume dependencies are present (if not, it's a packaging problem, and will error later when the assembly load fails)

diff --git a/Engine/Source/Programs/AutomationTool/CompileScriptModules.cs b/Engine/Source/Programs/AutomationTool/CompileScriptModules.cs
--- a/Engine/Source/Programs/AutomationTool/CompileScriptModules.cs
+++ b/Engine/Source/Programs/AutomationTool/CompileScriptModules.cs
@@ -116,7 +116,7 @@ namespace AutomationToolDriver
                 }
 
 				HashSet<FileReference> BuiltTargets = new HashSet<FileReference>(ExistingBuildRecords.Count);
-				foreach ((CsProjBuildRecord BuildRecord, FileReference BuildRecordPath) in ExistingBuildRecords.Values)
+				foreach ((CsProjBuildRecord BuildRecord, FileReference BuildRecordPath) in ExistingBuildRecords.Values.Where(x => x.Item2.HasExtension(".Automation.json")))
                 {
 					FileReference ProjectPath = FileReference.Combine(BuildRecordPath.Directory, BuildRecord.ProjectPath);
 					FileReference TargetPath = FileReference.Combine(ProjectPath.Directory, BuildRecord.TargetPath);
@@ -134,20 +134,8 @@ namespace AutomationToolDriver
 						}
 						else
 						{
-							// Workaround to address UE-138692:
-							// Shared/EpicGames.* assemblies build to their own local project directories, and those directories are not included in packages
-							// but those dlls, referenced by AutomationTool or its dependencies, are copied into AutomationTool's output directory.
-							// So, if the file is not found at TargetPath, look for the same file in the output directory.
-							DirectoryReference AutomationToolBinaryDirectory = DirectoryReference.Combine(Unreal.EngineDirectory, "Binaries", "DotNET", "AutomationTool");
-							if (FileReference.Exists(FileReference.Combine(AutomationToolBinaryDirectory, TargetPath.GetFileName())))
-							{
-								Log.TraceLog($"Using {TargetPath.GetFileName()} from {AutomationToolBinaryDirectory.FullName}");
-							}
-							else
-							{
-								// when the engine is installed, expect to find a built target assembly for every record that was found
-								throw new Exception($"Script module \"{TargetPath}\" not found for record \"{BuildRecordPath}\"");
-							}
+							// when the engine is installed, expect to find a built target assembly for every record that was found
+							throw new Exception($"Script module \"{TargetPath}\" not found for record \"{BuildRecordPath}\"");
 						}
 						
                     }
