AutomationTool: More idiomatic perhaps

#trivial

diff --git a/Engine/Source/Programs/AutomationTool/Program.cs b/Engine/Source/Programs/AutomationTool/Program.cs
--- a/Engine/Source/Programs/AutomationTool/Program.cs
+++ b/Engine/Source/Programs/AutomationTool/Program.cs
@@ -4,7 +4,9 @@
 // This software will not be supported.
 // Use at your own risk.
 using System;
+using System.Linq;
 using System.Threading;
+using System.Threading.Tasks;
 using System.Diagnostics;
 using System.Reflection;
 using EpicGames.Core;
@@ -14,7 +16,6 @@ using System.Text;
 using Microsoft.Build.Execution;
 using Microsoft.Build.Framework.Profiler;
 using UnrealBuildBase;
-using System.Threading.Tasks;
 
 namespace AutomationToolDriver
 {
@@ -428,17 +429,9 @@ namespace AutomationToolDriver
 				return ExitCode.Error_Unknown;
             }
 
-			Assembly AutomationUtilsAssembly = null;
-			
 			// Load AutomationUtils.Automation.dll
-			foreach (FileReference AssemblyPath in ScriptModuleAssemblyPaths)
-			{
-				if (AssemblyPath.GetFileNameWithoutExtension().Contains("AutomationUtils.Automation"))
-				{
-					AutomationUtilsAssembly = Assembly.LoadFrom(AssemblyPath.FullName);
-					break;
-				}
-			}
+			FileReference AssemblyPath = ScriptModuleAssemblyPaths.FirstOrDefault(x => x.GetFileNameWithoutExtension().Contains("AutomationUtils.Automation"));
+			Assembly AutomationUtilsAssembly = AssemblyPath != null ? Assembly.LoadFrom(AssemblyPath.FullName) : null;
 
 			if (AutomationUtilsAssembly == null)
             {
