Modify CreatePlatformExtension UAT command so that it can be used programatically.

 - also includes a fix so that a parent module is set to the same as the source module by default.

#rnx

diff --git a/Engine/Source/Programs/AutomationTool/Scripts/CreatePlatformExtension.cs b/Engine/Source/Programs/AutomationTool/Scripts/CreatePlatformExtension.cs
--- a/Engine/Source/Programs/AutomationTool/Scripts/CreatePlatformExtension.cs
+++ b/Engine/Source/Programs/AutomationTool/Scripts/CreatePlatformExtension.cs
@@ -16,9 +16,10 @@ using System.Linq;
 [Help("SkipPluginModules", "Do not generate platform extension module files when generating a platform extension plugin")]
 [Help("AllowOverwrite", "If target files already exist they'll be overwritten rather than skipped")]
 [Help("AllowUnknownPlatforms", "Allow platform & platform groups that are not known, for example when generating code for extensions we do not have access to")]
+[Help("AllowPlatformExtensionsAsParents", "When creating a platform extension from another platform extension, use the source platform as the parent")]
 [Help("P4", "Create a changelist for the new files")]
 [Help("CL", "Override the changelist #")]
-class CreatePlatformExtension : BuildCommand
+public class CreatePlatformExtension : BuildCommand
 {
 	readonly List<ModuleHostType> ModuleTypeDenyList = new List<ModuleHostType>
 	{
@@ -31,23 +32,24 @@ class CreatePlatformExtension : BuildCommand
 
 	ConfigHierarchy GameIni;
 	DirectoryReference ProjectDir;
-	List<string> NewFiles = new List<string>();
-	List<string> ModifiedFiles = new List<string>();
-	List<string> WritableFiles = new List<string>(); // for testing only
+	DirectoryReference SourceDir;
+	public List<string> NewFiles = new List<string>();
+	public List<string> ModifiedFiles = new List<string>();
+	public List<string> WritableFiles = new List<string>(); // for testing only
 	
 	string[] Platforms;
-	string Source;
-	string Project;
-	bool bSkipPluginModules;
-	bool bOverwriteExistingFile;
-	int CL = -1;
+	public bool bSkipPluginModules = false;
+	public bool bOverwriteExistingFile = false;
+	public bool bIsTest = false;
+	public bool bAllowPlatformExtensionsAsParents = false;
+	public int CL = -1;
 
 	public override void ExecuteBuild()
 	{
 		// Parse the parameters
-		Platforms = ParseParamValue("Platform", "").Split('+', StringSplitOptions.RemoveEmptyEntries);
-		Source = ParseParamValue("Source", "");
-		Project = ParseParamValue("Project", "");
+		string[] SrcPlatforms = ParseParamValue("Platform", "").Split('+', StringSplitOptions.RemoveEmptyEntries);
+		string Source = ParseParamValue("Source", "");
+		string Project = ParseParamValue("Project", "");
 		bSkipPluginModules = ParseParam("SkipPluginModules");
 		bOverwriteExistingFile = ParseParam("AllowOverwrite");
 		CL = ParseParamInt("CL", -1 );
@@ -60,18 +62,18 @@ class CreatePlatformExtension : BuildCommand
 		}
 
 		// Sanity check platforms list
-		Platforms = VerifyPlatforms(Platforms);
-		if (Platforms.Length == 0)
+		SrcPlatforms = VerifyPlatforms(SrcPlatforms);
+		if (SrcPlatforms.Length == 0)
 		{
 			Log.TraceError("Please specify at least one platform or platform group");
 			return;
 		}
 
-		// cannot have both -P4 and -Test
-		bool bIsTest = ParseParam("Test");
+		// cannot have both -P4 and -DebugTest
+		bIsTest = ParseParam("DebugTest") && System.Diagnostics.Debugger.IsAttached; //NB. -DebugTest is for debugging this program only
 		if (CommandUtils.P4Enabled && bIsTest)
 		{
-			Log.TraceError("Cannot specify both -P4 and -Test");
+			Log.TraceError("Cannot specify both -P4 and -DebugTest");
 			return;
 		}
 
@@ -93,50 +95,10 @@ class CreatePlatformExtension : BuildCommand
 			}
 		}
 
-		// Prepare values
-		ProjectDir = string.IsNullOrEmpty(Project) ? null : new FileReference(Project).Directory;
-		GameIni = ConfigCache.ReadHierarchy(ConfigHierarchyType.Game, ProjectDir, BuildHostPlatform.Current.Platform );
-
 		// Generate the code
 		try
 		{
-			if (Directory.Exists(Source))
-			{
-				// check the directory for plugins first, because the plugins will automatically generate the modules too
-				List<string> Plugins = Directory.EnumerateFiles(Source, "*.uplugin", SearchOption.AllDirectories ).ToList();
-				if (Plugins.Count > 0)
-				{
-					foreach (string Plugin in Plugins)
-					{
-						GeneratePluginPlatformExtension( new FileReference(Plugin) );
-					}
-				}
-				else
-				{
-					// there were no plugins found, so search for module & target rules instead
-					List<string> ModuleRules = Directory.EnumerateFiles(Source,"*.build.cs", SearchOption.AllDirectories).ToList();
-					ModuleRules.AddRange(Directory.EnumerateFiles(Source,"*.target.cs", SearchOption.AllDirectories));
-					if (ModuleRules.Count > 0)
-					{
-						foreach (string ModuleRule in ModuleRules)
-						{
-							GenerateModulePlatformExtension( new FileReference(ModuleRule), Platforms );
-						}
-					}
-					else
-					{
-						Log.TraceError($"Cannot find any supported files in {Source}");
-					}
-				}
-			}
-			else if (File.Exists(Source))
-			{
-				GeneratePlatformExtensionFromFile(new FileReference(Source));
-			}
-			else
-			{
-				Log.TraceError($"Invalid path or file name {Source}");
-			}
+			GeneratePlatformExtension(Project, Source, SrcPlatforms);
 
 			// check the generated files
 			if (NewFiles.Count > 0 || ModifiedFiles.Count > 0)
@@ -216,6 +178,59 @@ class CreatePlatformExtension : BuildCommand
 	}
 
 
+	/// <summary>
+	/// Generates platform extensions from the given source
+	/// </summary>
+	public void GeneratePlatformExtension( string InProject, string InSource, string[] InPlatforms )
+	{
+		// Prepare values
+		ProjectDir = string.IsNullOrEmpty(InProject) ? null : new FileReference(InProject).Directory;
+		GameIni = ConfigCache.ReadHierarchy(ConfigHierarchyType.Game, ProjectDir, BuildHostPlatform.Current.Platform);
+		Platforms = InPlatforms;
+
+		if (Directory.Exists(InSource))
+		{
+			SourceDir = new DirectoryReference(InSource);
+
+			// check the directory for plugins first, because the plugins will automatically generate the modules too
+			List<string> Plugins = Directory.EnumerateFiles(InSource, "*.uplugin", SearchOption.AllDirectories).ToList();
+			if (Plugins.Count > 0)
+			{
+				foreach (string Plugin in Plugins)
+				{
+					GeneratePluginPlatformExtension(new FileReference(Plugin));
+				}
+			}
+			else
+			{
+				// there were no plugins found, so search for module & target rules instead
+				List<string> ModuleRules = Directory.EnumerateFiles(InSource, "*.build.cs", SearchOption.AllDirectories).ToList();
+				ModuleRules.AddRange(Directory.EnumerateFiles(InSource, "*.target.cs", SearchOption.AllDirectories));
+				if (ModuleRules.Count > 0)
+				{
+					foreach (string ModuleRule in ModuleRules)
+					{
+						GenerateModulePlatformExtension(new FileReference(ModuleRule), Platforms);
+					}
+				}
+				else
+				{
+					Log.TraceError($"Cannot find any supported files in {InSource}");
+				}
+			}
+		}
+		else if (File.Exists(InSource))
+		{
+			FileReference Source = new FileReference(InSource);
+			SourceDir = Source.Directory;
+			GeneratePlatformExtensionFromFile(Source);
+		}
+		else
+		{
+			Log.TraceError($"Invalid path or file name {InSource}");
+		}
+	}
+
 
 
 	/// <summary>
@@ -446,7 +461,11 @@ class CreatePlatformExtension : BuildCommand
 			Log.TraceError($"Cannot find class declaration in ${ModulePath}");
 			return;
 		}
-		string ParentModuleName = ModuleClassDeclaration.Trim().Remove(0, ClassDeclaration.Length ).Split(' ', StringSplitOptions.None ).First();
+		string ParentModuleName = ModuleClassDeclaration.Trim().Remove(0, ClassDeclaration.Length).Split(' ', StringSplitOptions.None).Last();
+		if (bAllowPlatformExtensionsAsParents || ParentModuleName.Equals("ModuleRules") || ParentModuleName.Equals("TargetRules") )
+		{
+			ParentModuleName = ModuleClassDeclaration.Trim().Remove(0, ClassDeclaration.Length ).Split(' ', StringSplitOptions.None ).First();
+		}
 		if (string.IsNullOrEmpty(ParentModuleName))
 		{
 			Log.TraceError($"Cannot parse class declaration in ${ModulePath}");
@@ -691,7 +710,6 @@ class CreatePlatformExtension : BuildCommand
 		if (CL == -1)
 		{
 			// add the files to perforce if that is available
-			DirectoryReference SourceDir = new DirectoryReference(Source);
 			string Description = $"[AUTO-GENERATED] {string.Join('+', Platforms)} platform extension files from {SourceDir.MakeRelativeTo(Unreal.RootDirectory)}\n\n#nocheckin verify the code has been generated successfully before checking in!";
 
 			CL = P4.CreateChange(P4Env.Client, Description );
@@ -739,7 +757,6 @@ class CreatePlatformExtension : BuildCommand
 
 		if ((File.GetAttributes(ExistingFile) & FileAttributes.ReadOnly) == FileAttributes.ReadOnly)
 		{
-			bool bIsTest = ParseParam("Test");
 			if (bIsTest)
 			{
 				File.Copy(ExistingFile, ExistingFile + ".tmp.bak");
