Force disable bPakUsesSecondaryOrder for all projects when staging with iostore, i.e. don't use the non-deterministic CookerOpenOrder.log in the UAT stage and package step.
This improves the determinism in iostore packaged builds with no regression in load times.
Game specific ordering .log files are still enabled, so any files located in [GameName]\Build\[CookPlatform]\FileOpenOrder\ will still be used.

#rnx

diff --git a/Engine/Source/Developer/IoStoreUtilities/Private/IoStoreUtilities.cpp b/Engine/Source/Developer/IoStoreUtilities/Private/IoStoreUtilities.cpp
--- a/Engine/Source/Developer/IoStoreUtilities/Private/IoStoreUtilities.cpp
+++ b/Engine/Source/Developer/IoStoreUtilities/Private/IoStoreUtilities.cpp
@@ -5150,10 +5150,7 @@ static bool ParseOrderFileArguments(FIoStoreArguments& Arguments)
 		}
 		if (FParse::Value(FCommandLine::Get(), TEXT("CookerOrder="), LegacyParam, false))
 		{
-			UE_LOG(LogIoStore, Warning, TEXT("-GameOrder= and -CookerOrder= are deprecated in favor of -Order"));
-			TArray<FString> LegacyPaths;
-			LegacyParam.ParseIntoArray(LegacyPaths, TEXT(","), true);
-			OrderFilePaths.Append(LegacyPaths);
+			UE_LOG(LogIoStore, Warning, TEXT("-CookerOrder is ignored by IoStore. -GameOrder= and -CookerOrder= are deprecated in favor of -Order."));
 		}
 
 		FString OrderPriorityString;
diff --git a/Engine/Source/Programs/AutomationTool/Scripts/CopyBuildToStagingDirectory.Automation.cs b/Engine/Source/Programs/AutomationTool/Scripts/CopyBuildToStagingDirectory.Automation.cs
--- a/Engine/Source/Programs/AutomationTool/Scripts/CopyBuildToStagingDirectory.Automation.cs
+++ b/Engine/Source/Programs/AutomationTool/Scripts/CopyBuildToStagingDirectory.Automation.cs
@@ -2323,6 +2323,7 @@ namespace AutomationScripts
 		/// <param name="PakName"></param>
 		private static void CreatePaks(ProjectParams Params, DeploymentContext SC, List<CreatePakParams> PakParamsList, EncryptionAndSigning.CryptoSettings CryptoSettings, FileReference CryptoKeysCacheFilename)
 		{
+			bool bShouldCreateIoStoreContainerFiles = ShouldCreateIoStoreContainerFiles(Params, SC);
 			bool bShouldGeneratePatch = Params.IsGeneratingPatch && SC.StageTargetPlatform.GetPlatformPatchesWithDiffPak(Params, SC);
 
 			if (bShouldGeneratePatch && !Params.HasBasedOnReleaseVersion)
@@ -2669,7 +2670,12 @@ namespace AutomationScripts
 			bool bUseSecondaryOrder = false;
 			PlatformGameConfig.GetBool("/Script/UnrealEd.ProjectPackagingSettings", "bPakUsesSecondaryOrder", out bUseSecondaryOrder);
 
-			if(!bUseSecondaryOrder)
+			if (bShouldCreateIoStoreContainerFiles)
+			{
+				bUseSecondaryOrder = false;
+			}
+
+			if (!bUseSecondaryOrder)
 			{
 				OrderFiles.RemoveAll(x => (x.OrderType == OrderFile.OrderFileType.Cooker && x.AppendOrder == Int32.MaxValue));
 
@@ -2869,7 +2875,7 @@ namespace AutomationScripts
 						AdditionalArgs += " -platform=" + ConfigHierarchy.GetIniPlatformName(SC.StageTargetPlatform.IniPlatformType);
 
 						Dictionary<string, string> UnrealPakResponseFile = PakParams.UnrealPakResponseFile;
-						if (ShouldCreateIoStoreContainerFiles(Params, SC))
+						if (bShouldCreateIoStoreContainerFiles)
 						{
 							bool bAllowShadersInIoStore;
 							if (!PlatformEngineConfig.GetBool("Core.System", "AllowShadersInIoStore", out bAllowShadersInIoStore))
@@ -3072,7 +3078,7 @@ namespace AutomationScripts
 
 					InternalUtils.SafeCreateDirectory(Path.GetDirectoryName(ReleaseVersionPath));
 					InternalUtils.SafeCopyFile(OutputLocation.FullName, ReleaseVersionPath);
-					if (ShouldCreateIoStoreContainerFiles(Params, SC))
+					if (bShouldCreateIoStoreContainerFiles)
 					{
 						InternalUtils.SafeCopyFile(Path.ChangeExtension(OutputLocation.FullName, ".utoc"), Path.ChangeExtension(ReleaseVersionPath, ".utoc"));
 						InternalUtils.SafeCopyFile(Path.ChangeExtension(OutputLocation.FullName, ".ucas"), Path.ChangeExtension(ReleaseVersionPath, ".ucas"));
@@ -3182,7 +3188,7 @@ namespace AutomationScripts
 						{
 							HashSet<string> IncludedExtensions = new HashSet<string>();
 							IncludedExtensions.Add(OutputFilenameExtension);
-							if (ShouldCreateIoStoreContainerFiles(Params, SC))
+							if (bShouldCreateIoStoreContainerFiles)
 							{
 								IncludedExtensions.Add(".ucas");
 								IncludedExtensions.Add(".utoc");
