- Removed maps from Uncook list in MapMode=None, since the Uncook list will be added back in after the MapMode logic

diff --git a/Engine/Source/Programs/AutomationTool/CookedEditor/CookedEditor.Automation.cs b/Engine/Source/Programs/AutomationTool/CookedEditor/CookedEditor.Automation.cs
--- a/Engine/Source/Programs/AutomationTool/CookedEditor/CookedEditor.Automation.cs
+++ b/Engine/Source/Programs/AutomationTool/CookedEditor/CookedEditor.Automation.cs
@@ -321,6 +321,9 @@ public class ModifyStageContext
 			HashSet<string> Maps = UFSFilesToStage.Where(x => x.HasExtension("umap")).Select(x => Path.ChangeExtension(x.FullName, null)).ToHashSet();
 			UFSFilesToStage = UFSFilesToStage.Where(x => !Maps.Contains(Path.ChangeExtension(x.FullName.Replace("_BuiltData", ""), null))).ToList();
 
+			Maps = FilesToUncook.Where(x => x.HasExtension("umap")).Select(x => Path.ChangeExtension(x.FullName, null)).ToHashSet();
+			FilesToUncook = FilesToUncook.Where(x => !Maps.Contains(Path.ChangeExtension(x.FullName.Replace("_BuiltData", ""), null))).ToList();
+
 			Maps = SC.FilesToStage.UFSFiles.Keys.Where(x => x.HasExtension("umap")).Select(x => Path.ChangeExtension(x.Name, null)).ToHashSet();
 			Console.WriteLine($"Found {Maps.Count()} maps");
 			SC.FilesToStage.UFSFiles = SC.FilesToStage.UFSFiles.Where(x => !Maps.Contains(Path.ChangeExtension(x.Key.Name.Replace("_BuiltData", ""), null))).ToDictionary(x => x.Key, x => x.Value);
