Catch and report Google Drive verification errors for Turnkey.

#rnx

diff --git a/Engine/Source/Programs/AutomationTool/Turnkey/Providers/GoogleDriveCopyProvider.cs b/Engine/Source/Programs/AutomationTool/Turnkey/Providers/GoogleDriveCopyProvider.cs
--- a/Engine/Source/Programs/AutomationTool/Turnkey/Providers/GoogleDriveCopyProvider.cs
+++ b/Engine/Source/Programs/AutomationTool/Turnkey/Providers/GoogleDriveCopyProvider.cs
@@ -241,9 +241,16 @@ namespace Turnkey
 				}
 
 				TurnkeyUtils.Log("Connecting to GoogleDrive app '{0}'. If this pauses here, check your web browser for an authentication tab. This is required to be able to connect to Google Drive", GoogleDriveAppName);
-				ServiceHelper = new DriveServiceHelper(GoogleDriveAppName, SecretsFile, Path.GetDirectoryName(SecretsFile));
-
-				PathToFileCache = new Dictionary<string, Google.Apis.Drive.v3.Data.File>();
+				try
+				{
+					ServiceHelper = new DriveServiceHelper(GoogleDriveAppName, SecretsFile, Path.GetDirectoryName(SecretsFile));
+					PathToFileCache = new Dictionary<string, Google.Apis.Drive.v3.Data.File>();
+				}
+				catch(Exception ex)
+				{
+					TurnkeyUtils.Log("ERROR: Unable to access GoogleDrive: {0}", ex.Message);
+					return false;
+				}
 			}
 
 			return true;
