Prevent exceptions from propagating when verifying sdks

An exception was getting thrown on the Mac when checking for hololens when looping through all possible platforms.  I felt the safest option was to prevent any exceptions when checking for SDK versions, that way we can simply treat the SDK as not present if an exception occurs.

diff --git a/Engine/Source/Programs/AutomationTool/Turnkey/Commands/VerifySdk.cs b/Engine/Source/Programs/AutomationTool/Turnkey/Commands/VerifySdk.cs
--- a/Engine/Source/Programs/AutomationTool/Turnkey/Commands/VerifySdk.cs
+++ b/Engine/Source/Programs/AutomationTool/Turnkey/Commands/VerifySdk.cs
@@ -67,11 +67,18 @@ namespace Turnkey.Commands
 				string ManualSDKVersion = "", AutoSDKVersion = "";
 				string MinAllowedVersion = "", MaxAllowedVersion = "";
 				string MinSoftwareAllowedVersion = "", MaxSoftwareAllowedVersion = "";
-				if (PlatformSDK != null)
+
+				try
+				{
+					if (PlatformSDK != null)
+					{
+						PlatformSDK.GetInstalledVersions(out ManualSDKVersion, out AutoSDKVersion);
+						PlatformSDK.GetValidVersionRange(out MinAllowedVersion, out MaxAllowedVersion);
+						PlatformSDK.GetValidSoftwareVersionRange(out MinSoftwareAllowedVersion, out MaxSoftwareAllowedVersion);
+					}
+				}
+				catch (Exception)
 				{
-					PlatformSDK.GetInstalledVersions(out ManualSDKVersion, out AutoSDKVersion);
-					PlatformSDK.GetValidVersionRange(out MinAllowedVersion, out MaxAllowedVersion);
-					PlatformSDK.GetValidSoftwareVersionRange(out MinSoftwareAllowedVersion, out MaxSoftwareAllowedVersion);
 				}
 
 				SdkUtils.LocalAvailability ReportedState = LocalState;
