Replace/remove uses of whitelist/blacklist in MCP

diff --git a/Engine/Source/Programs/AutomationTool/AutomationUtils/MCPPublic.cs b/Engine/Source/Programs/AutomationTool/AutomationUtils/MCPPublic.cs
--- a/Engine/Source/Programs/AutomationTool/AutomationUtils/MCPPublic.cs
+++ b/Engine/Source/Programs/AutomationTool/AutomationUtils/MCPPublic.cs
@@ -855,7 +855,7 @@ namespace EpicGames.MCP.Automation
 			/// <summary>
 			/// List of artifact ids which will be used to generate patches, if this is left blank all artifact ids will be used
 			/// </summary>
-			public string WhitelistedArtifactIds;
+			public string AllowedlistArtifactIds;
 			/// <summary>
 			/// Specifies the client id allocated to you by Epic for uploading binaries to Epic's services.
 			/// </summary>
