- Allow staging empty files in CookedEditor via ini files
- Allow for CookedEditor's ini settings to remap a file (added a Dest param to ExtraStageFiles, must not be recursive or wildcard)

diff --git a/Engine/Source/Programs/AutomationTool/CookedEditor/CookedEditor.Automation.cs b/Engine/Source/Programs/AutomationTool/CookedEditor/CookedEditor.Automation.cs
--- a/Engine/Source/Programs/AutomationTool/CookedEditor/CookedEditor.Automation.cs
+++ b/Engine/Source/Programs/AutomationTool/CookedEditor/CookedEditor.Automation.cs
@@ -765,6 +765,29 @@ public class MakeCookedEditor : BuildCommand
 			{
 				SearchMode = SearchOption.TopDirectoryOnly;
 			}
+			// settings with a Dest are special, and have to go right to SC
+			if (Props.ContainsKey("Dest"))
+			{
+				FileReference SourceFile = FileReference.Combine(BaseDirectory, SubPath);
+
+				// allow a blank Path to mean empty file
+				if (SubPath == "")
+				{
+					SourceFile = FileReference.Combine(Context.EngineDirectory, "Intermediate/blankfile.txt");
+					FileReference.WriteAllText(SourceFile, "");
+					FileWildcard = "";
+				}
+
+				if (SearchMode == SearchOption.AllDirectories || FileWildcard.Contains("*"))
+				{
+					throw new AutomationException($"Unable to stage directories with \"Dest\" setting for CookedEditor: '{Entry}'");
+				}
+
+				// now stage it to a different location a specified in the params
+				StagedFileType FileType = (FileList == Context.NonUFSFilesToStage) ? StagedFileType.NonUFS : StagedFileType.UFS;
+				SC.StageFile(FileType, SourceFile, new StagedFileReference(Props["Dest"]));
+				continue;
+			}
 
 			// now enumerate files based on the settings
 			DirectoryReference Dir = DirectoryReference.Combine(BaseDirectory, SubPath);
