- Added NeverPackage flag to BuildCookRun which can skip steps in Staging knowing that Packaging will never happen (different than SkipPackage which assumes Packaging happened at some other point)

diff --git a/Engine/Source/Programs/AutomationTool/AutomationUtils/ProjectParams.cs b/Engine/Source/Programs/AutomationTool/AutomationUtils/ProjectParams.cs
--- a/Engine/Source/Programs/AutomationTool/AutomationUtils/ProjectParams.cs
+++ b/Engine/Source/Programs/AutomationTool/AutomationUtils/ProjectParams.cs
@@ -360,6 +360,7 @@ namespace AutomationTool
             this.ClientCommandline = InParams.ClientCommandline;
             this.Package = InParams.Package;
 			this.SkipPackage = InParams.SkipPackage;
+			this.NeverPackage = InParams.NeverPackage;
 			this.ForcePackageData = InParams.ForcePackageData;
 			this.Deploy = InParams.Deploy;
 			this.DeployFolder = InParams.DeployFolder;
@@ -505,6 +506,7 @@ namespace AutomationTool
 			bool? MapFile = null,
 			bool? NoXGE = null,
 			bool? SkipPackage = null,
+			bool? NeverPackage = null,
 			bool? Package = null,
 			bool? Pak = null,
 			bool? IgnorePaksFromDifferentCookSource = null,
@@ -845,6 +847,7 @@ namespace AutomationTool
             this.ClientCommandline = this.ClientCommandline.Replace('\'', '\"'); // replace any single quotes with double quotes
             this.Package = GetParamValueIfNotSpecified(Command, Package, this.Package, "package");
 			this.SkipPackage = GetParamValueIfNotSpecified(Command, SkipPackage, this.SkipPackage, "skippackage");
+			this.NeverPackage = GetParamValueIfNotSpecified(Command, NeverPackage, this.NeverPackage, "neverpackage");
 			this.ForcePackageData = GetParamValueIfNotSpecified(Command, Package, this.ForcePackageData, "forcepackagedata");
 
 			this.Deploy = GetParamValueIfNotSpecified(Command, Deploy, this.Deploy, "deploy");
@@ -2092,7 +2095,10 @@ namespace AutomationTool
 		
 		[Help("skippackage", "Skips packaging the project for the target platform")]
 		public bool SkipPackage { get; set; }
-		
+
+		[Help("neverpackage", "Skips preparing data that would be used during packaging, in earlier stages. Different from skippackage which is used to optimize later stages like archive, which still was packaged at some point")]
+		public bool NeverPackage { get; set; }
+
 		[Help("package", "Determine whether data is packaged. This can be an iteration optimization for platforms that require packages for deployment")]
 		public bool ForcePackageData { get; set; }
 
@@ -3026,6 +3032,7 @@ namespace AutomationTool
 				CommandUtils.LogLog("SkipEncryption={0}", SkipEncryption);
 				CommandUtils.LogLog("GenerateOptimizationData={0}", GenerateOptimizationData);
 				CommandUtils.LogLog("SkipPackage={0}", SkipPackage);
+				CommandUtils.LogLog("NeverPackage={0}", NeverPackage);
 				CommandUtils.LogLog("Package={0}", Package);
 				CommandUtils.LogLog("ForcePackageData={0}", ForcePackageData);
 				CommandUtils.LogLog("NullRHI={0}", NullRHI);
