- Refactored RunCommandlet and Project.Cook() so that a script can start cooking or other commandlet while doing something else (like monitoring the output)
#p4v-preflight-copy 18808074

diff --git a/Engine/Source/Programs/AutomationTool/AutomationUtils/CommandletUtils.cs b/Engine/Source/Programs/AutomationTool/AutomationUtils/CommandletUtils.cs
--- a/Engine/Source/Programs/AutomationTool/AutomationUtils/CommandletUtils.cs
+++ b/Engine/Source/Programs/AutomationTool/AutomationUtils/CommandletUtils.cs
@@ -249,24 +249,33 @@ namespace AutomationTool
 		/// <param name="DestLogFile">Log file after completion</param>
 		/// <param name="ErrorLevel">The minimum exit code, which is treated as an error.</param>
 		public static void RunCommandlet(FileReference ProjectName, string UnrealExe, string Commandlet, string Parameters, out string DestLogFile, int ErrorLevel = 1)
+		{
+			string LocalLogFile;
+			IProcessResult RunResult;
+
+			DateTime StartTime = DateTime.UtcNow;
+
+			StartRunCommandlet(ProjectName, UnrealExe, Commandlet, Parameters, ERunOptions.Default, out LocalLogFile, out RunResult);
+			FinishRunCommandlet(ProjectName, Commandlet, StartTime, RunResult, LocalLogFile, out DestLogFile);
+		}
+
+		public static void StartRunCommandlet(FileReference ProjectName, string UnrealExe, string Commandlet, string Parameters, ERunOptions RunOptions, out string LocalLogFile, out IProcessResult RunResult, ProcessResult.SpewFilterCallbackType SpewFilterCallback=null)
 		{
 			LogInformation("Running UnrealEditor {0} for project {1}", Commandlet, ProjectName);
 
-            var CWD = Path.GetDirectoryName(UnrealExe);
+			var CWD = Path.GetDirectoryName(UnrealExe);
 
-            string EditorExe = UnrealExe;
+			string EditorExe = UnrealExe;
 
-            if (String.IsNullOrEmpty(CWD))
-            {
-                EditorExe = HostPlatform.Current.GetUnrealExePath(UnrealExe);
-                CWD = CombinePaths(CmdEnv.LocalRoot, HostPlatform.Current.RelativeBinariesFolder);
-            }
-			
-			PushDir(CWD);
+			if (String.IsNullOrEmpty(CWD))
+			{
+				EditorExe = HostPlatform.Current.GetUnrealExePath(UnrealExe);
+				CWD = CombinePaths(CmdEnv.LocalRoot, HostPlatform.Current.RelativeBinariesFolder);
+			}
 
-			DateTime StartTime = DateTime.UtcNow;
+			PushDir(CWD);
 
-			string LocalLogFile = LogUtils.GetUniqueLogName(CombinePaths(CmdEnv.EngineSavedFolder, Commandlet));
+			LocalLogFile = LogUtils.GetUniqueLogName(CombinePaths(CmdEnv.EngineSavedFolder, Commandlet));
 			LogInformation("Commandlet log file is {0}", LocalLogFile);
 			string Args = String.Format(
 				"{0} -run={1} {2} -abslog={3} -stdout -CrashForUAT -unattended -NoLogTimes {5}{4}",
@@ -275,17 +284,19 @@ namespace AutomationTool
 				String.IsNullOrEmpty(Parameters) ? "" : Parameters,
 				CommandUtils.MakePathSafeToUseWithCommandLine(LocalLogFile),
 				IsBuildMachine ? "-buildmachine" : "",
-                (GlobalCommandLine.Verbose || GlobalCommandLine.AllowStdOutLogVerbosity) ? "-AllowStdOutLogVerbosity " : ""
+				(GlobalCommandLine.Verbose || GlobalCommandLine.AllowStdOutLogVerbosity) ? "-AllowStdOutLogVerbosity " : ""
 			);
-			ERunOptions Opts = ERunOptions.Default;
 			if (GlobalCommandLine.UTF8Output)
 			{
 				Args += " -UTF8Output";
-				Opts |= ERunOptions.UTF8Output;
+				RunOptions |= ERunOptions.UTF8Output;
 			}
-			IProcessResult RunResult = Run(EditorExe, Args, Options: Opts, Identifier: Commandlet);
+			RunResult = Run(EditorExe, Args, Options: RunOptions, Identifier: Commandlet, SpewFilterCallback: SpewFilterCallback);
 			PopDir();
+		}
 
+		public static void FinishRunCommandlet(FileReference ProjectName, string Commandlet, DateTime StartTime, IProcessResult RunResult, string LocalLogFile, out string DestLogFile, int ErrorLevel = 1)
+		{ 
 			// If we're running on a Windows build machine, copy any crash dumps into the log folder
 			if(HostPlatform.Current.HostEditorPlatform == UnrealTargetPlatform.Win64 && IsBuildMachine)
 			{
diff --git a/Engine/Source/Programs/AutomationTool/Scripts/CookCommand.Automation.cs b/Engine/Source/Programs/AutomationTool/Scripts/CookCommand.Automation.cs
--- a/Engine/Source/Programs/AutomationTool/Scripts/CookCommand.Automation.cs
+++ b/Engine/Source/Programs/AutomationTool/Scripts/CookCommand.Automation.cs
@@ -24,6 +24,190 @@ namespace AutomationScripts
 	/// </remarks>
 	public partial class Project : CommandUtils
 	{
+		public static void GetCookByTheBookCommandletParams(ProjectParams Params, out string CommandletParams, out string[] Maps, out string[] DirectoriesToCook, out string InternationalizationPreset, out string[] CulturesToCook, out string PlatformsToCook)
+		{
+			var PlatformsToCookSet = new HashSet<string>();
+			if (!Params.NoClient)
+			{
+				foreach (var ClientPlatform in Params.ClientTargetPlatforms)
+				{
+					// Use the data platform, sometimes we will copy another platform's data
+					var DataPlatformDesc = Params.GetCookedDataPlatformForClientTarget(ClientPlatform);
+					string PlatformToCook = Platform.Platforms[DataPlatformDesc].GetCookPlatform(false, Params.Client);
+					PlatformsToCookSet.Add(PlatformToCook);
+				}
+			}
+			if (Params.DedicatedServer)
+			{
+				foreach (var ServerPlatform in Params.ServerTargetPlatforms)
+				{
+					// Use the data platform, sometimes we will copy another platform's data
+					var DataPlatformDesc = Params.GetCookedDataPlatformForServerTarget(ServerPlatform);
+					string PlatformToCook = Platform.Platforms[DataPlatformDesc].GetCookPlatform(true, false);
+					PlatformsToCookSet.Add(PlatformToCook);
+				}
+			}
+
+			PlatformsToCook = CombineCommandletParams(PlatformsToCookSet);
+
+			if (Params.Clean.HasValue && Params.Clean.Value && !Params.IterativeCooking)
+			{
+				LogInformation("Cleaning cooked data.");
+				CleanupCookedData(PlatformsToCookSet.ToList(), Params);
+			}
+
+			// cook the set of maps, or the run map, or nothing
+			Maps = null;
+			if (Params.HasMapsToCook)
+			{
+				Maps = Params.MapsToCook.ToArray();
+				foreach (var M in Maps)
+				{
+					LogInformation("HasMapsToCook " + M.ToString());
+				}
+				foreach (var M in Params.MapsToCook)
+				{
+					LogInformation("Params.HasMapsToCook " + M.ToString());
+				}
+			}
+
+			DirectoriesToCook = null;
+			if (Params.HasDirectoriesToCook)
+			{
+				DirectoriesToCook = Params.DirectoriesToCook.ToArray();
+			}
+
+			InternationalizationPreset = null;
+			if (Params.HasInternationalizationPreset)
+			{
+				InternationalizationPreset = Params.InternationalizationPreset;
+			}
+
+			CulturesToCook = null;
+			if (Params.HasCulturesToCook)
+			{
+				CulturesToCook = Params.CulturesToCook.ToArray();
+			}
+
+			CommandletParams = IsBuildMachine ? "-buildmachine -fileopenlog" : "-fileopenlog";
+
+			if (Params.HasDDCGraph)
+			{
+				CommandletParams += " -ddc=" + Params.DDCGraph;
+			}
+			if (Params.UnversionedCookedContent)
+			{
+				CommandletParams += " -unversioned";
+			}
+			if (Params.FastCook)
+			{
+				CommandletParams += " -FastCook";
+			}
+			if (Params.Manifests)
+			{
+				CommandletParams += " -manifests";
+			}
+			if (Params.IterativeCooking)
+			{
+				CommandletParams += " -iterate";
+			}
+			if (Params.HasIterateSharedCookedBuild)
+			{
+				new SharedCookedBuild(Params).CopySharedCookedBuilds();
+				CommandletParams += " -iteratesharedcookedbuild";
+			}
+
+			if (Params.CookMapsOnly)
+			{
+				CommandletParams += " -mapsonly";
+			}
+			if (Params.CookAll)
+			{
+				CommandletParams += " -cookall";
+			}
+			if (Params.HasCreateReleaseVersion)
+			{
+				CommandletParams += " -createreleaseversion=" + Params.CreateReleaseVersion;
+			}
+			if (Params.SkipCookingEditorContent)
+			{
+				CommandletParams += " -skipeditorcontent";
+			}
+			if (Params.NumCookersToSpawn != 0)
+			{
+				CommandletParams += " -numcookerstospawn=" + Params.NumCookersToSpawn;
+			}
+			if (Params.CookPartialGC)
+			{
+				CommandletParams += " -partialgc";
+			}
+			if (Params.HasMapIniSectionsToCook)
+			{
+				string MapIniSections = CombineCommandletParams(Params.MapIniSectionsToCook.ToArray());
+
+				CommandletParams += " -MapIniSection=" + MapIniSections;
+			}
+			if (Params.HasDLCName)
+			{
+				CommandletParams += " -dlcname=" + Params.DLCFile.GetFileNameWithoutExtension();
+				if (!Params.DLCIncludeEngineContent)
+				{
+					CommandletParams += " -errorOnEngineContentUse";
+				}
+			}
+			if (!String.IsNullOrEmpty(Params.CookOutputDir))
+			{
+				CommandletParams += " -outputdir=" + CommandUtils.MakePathSafeToUseWithCommandLine(Params.CookOutputDir);
+			}
+			// don't include the based on release version unless we are cooking dlc or creating a new release version
+			// in this case the based on release version is used in packaging
+			if (Params.HasBasedOnReleaseVersion && (Params.HasDLCName || Params.HasCreateReleaseVersion))
+			{
+				CommandletParams += " -basedonreleaseversion=" + Params.BasedOnReleaseVersion;
+			}
+			if (!String.IsNullOrEmpty(Params.CreateReleaseVersionBasePath))
+			{
+				CommandletParams += " -createreleaseversionroot=" + Params.CreateReleaseVersionBasePath;
+			}
+			if (!String.IsNullOrEmpty(Params.BasedOnReleaseVersionBasePath))
+			{
+				CommandletParams += " -basedonreleaseversionroot=" + Params.BasedOnReleaseVersionBasePath;
+			}
+
+			// if we are not going to pak but we specified compressed then compress in the cooker ;)
+			// otherwise compress the pak files
+			if (!Params.Pak && !Params.SkipPak && Params.Compressed && !Params.ForceUncompressed)
+			{
+				CommandletParams += " -compressed";
+			}
+
+			if (Params.HasAdditionalCookerOptions)
+			{
+				string FormatedAdditionalCookerParams = Params.AdditionalCookerOptions.TrimStart(new char[] { '\"', ' ' }).TrimEnd(new char[] { '\"', ' ' });
+				CommandletParams += " ";
+				CommandletParams += FormatedAdditionalCookerParams;
+			}
+
+			if (!Params.NoClient)
+			{
+				var MapsList = Maps == null ? new List<string>() : Maps.ToList();
+				foreach (var ClientPlatform in Params.ClientTargetPlatforms)
+				{
+					var DataPlatformDesc = Params.GetCookedDataPlatformForClientTarget(ClientPlatform);
+					CommandletParams += (Platform.Platforms[DataPlatformDesc].GetCookExtraCommandLine(Params));
+					MapsList.AddRange((Platform.Platforms[ClientPlatform].GetCookExtraMaps()));
+				}
+				Maps = MapsList.ToArray();
+			}
+
+			// Config overrides (-ini)
+			foreach (string ConfigOverrideParam in Params.ConfigOverrideParams)
+			{
+				CommandletParams += " -";
+				CommandletParams += ConfigOverrideParam;
+			}
+		}
+
 		public static void Cook(ProjectParams Params)
 		{
 			if ((!Params.Cook && !(Params.CookOnTheFly && !Params.SkipServer)) || Params.SkipCook)
@@ -90,188 +274,20 @@ namespace AutomationScripts
 			}
 			else
 			{
-				var PlatformsToCook = new HashSet<string>();
-				if (!Params.NoClient)
-				{
-					foreach (var ClientPlatform in Params.ClientTargetPlatforms)
-					{
-						// Use the data platform, sometimes we will copy another platform's data
-						var DataPlatformDesc = Params.GetCookedDataPlatformForClientTarget(ClientPlatform);
-						string PlatformToCook = Platform.Platforms[DataPlatformDesc].GetCookPlatform(false, Params.Client);
-						PlatformsToCook.Add(PlatformToCook);
-					}
-				}
-				if (Params.DedicatedServer)
-				{
-					foreach (var ServerPlatform in Params.ServerTargetPlatforms)
-					{
-						// Use the data platform, sometimes we will copy another platform's data
-						var DataPlatformDesc = Params.GetCookedDataPlatformForServerTarget(ServerPlatform);
-						string PlatformToCook = Platform.Platforms[DataPlatformDesc].GetCookPlatform(true, false);
-						PlatformsToCook.Add(PlatformToCook);
-					}
-				}
-
-				if (Params.Clean.HasValue && Params.Clean.Value && !Params.IterativeCooking)
-				{
-					LogInformation("Cleaning cooked data.");
-					CleanupCookedData(PlatformsToCook.ToList(), Params);
-				}
-
-				// cook the set of maps, or the run map, or nothing
-				string[] Maps = null;
-				if (Params.HasMapsToCook)
-				{
-					Maps = Params.MapsToCook.ToArray();
-					foreach (var M in Maps)
-					{
-						LogInformation("HasMapsToCook " + M.ToString());
-					}
-					foreach (var M in Params.MapsToCook)
-					{
-						LogInformation("Params.HasMapsToCook " + M.ToString());
-					}
-				}
-
-				string[] Dirs = null;
-				if (Params.HasDirectoriesToCook)
-				{
-					Dirs = Params.DirectoriesToCook.ToArray();
-				}
-
-				string InternationalizationPreset = null;
-				if (Params.HasInternationalizationPreset)
-				{
-					InternationalizationPreset = Params.InternationalizationPreset;
-				}
-
-				string[] CulturesToCook = null;
-				if (Params.HasCulturesToCook)
-				{
-					CulturesToCook = Params.CulturesToCook.ToArray();
-				}
-
 				try
 				{
-					var CommandletParams = IsBuildMachine ? "-buildmachine -fileopenlog" : "-fileopenlog";
-
-					if (Params.HasDDCGraph)
-					{
-						CommandletParams += " -ddc=" + Params.DDCGraph;
-					}
-					if (Params.UnversionedCookedContent)
-					{
-						CommandletParams += " -unversioned";
-					}
-					if (Params.FastCook)
-					{
-						CommandletParams += " -FastCook";
-					}
-					if (Params.Manifests)
-					{
-						CommandletParams += " -manifests";
-					}
-					if (Params.IterativeCooking)
-					{
-						CommandletParams += " -iterate";
-					}
-					if (Params.HasIterateSharedCookedBuild)
-					{
-						new SharedCookedBuild(Params).CopySharedCookedBuilds();
-						CommandletParams += " -iteratesharedcookedbuild";
-					}
-
-					if (Params.CookMapsOnly)
-					{
-						CommandletParams += " -mapsonly";
-					}
-					if (Params.CookAll)
-					{
-						CommandletParams += " -cookall";
-					}
-					if (Params.HasCreateReleaseVersion)
-					{
-						CommandletParams += " -createreleaseversion=" + Params.CreateReleaseVersion;
-					}
-					if (Params.SkipCookingEditorContent)
-					{
-						CommandletParams += " -skipeditorcontent";
-					}
-					if (Params.NumCookersToSpawn != 0)
-					{
-						CommandletParams += " -numcookerstospawn=" + Params.NumCookersToSpawn;
-					}
-					if (Params.CookPartialGC)
-					{
-						CommandletParams += " -partialgc";
-					}
-					if (Params.HasMapIniSectionsToCook)
-					{
-						string MapIniSections = CombineCommandletParams(Params.MapIniSectionsToCook.ToArray());
-
-						CommandletParams += " -MapIniSection=" + MapIniSections;
-					}
-					if (Params.HasDLCName)
-					{
-						CommandletParams += " -dlcname=" + Params.DLCFile.GetFileNameWithoutExtension();
-						if (!Params.DLCIncludeEngineContent)
-						{
-							CommandletParams += " -errorOnEngineContentUse";
-						}
-					}
-					if (!String.IsNullOrEmpty(Params.CookOutputDir))
-					{
-						CommandletParams += " -outputdir=" + CommandUtils.MakePathSafeToUseWithCommandLine(Params.CookOutputDir);
-					}
-					// don't include the based on release version unless we are cooking dlc or creating a new release version
-					// in this case the based on release version is used in packaging
-					if (Params.HasBasedOnReleaseVersion && (Params.HasDLCName || Params.HasCreateReleaseVersion))
-					{
-						CommandletParams += " -basedonreleaseversion=" + Params.BasedOnReleaseVersion;
-					}
-					if (!String.IsNullOrEmpty(Params.CreateReleaseVersionBasePath))
-					{
-						CommandletParams += " -createreleaseversionroot=" + Params.CreateReleaseVersionBasePath;
-					}
-					if (!String.IsNullOrEmpty(Params.BasedOnReleaseVersionBasePath))
-					{
-						CommandletParams += " -basedonreleaseversionroot=" + Params.BasedOnReleaseVersionBasePath;
-					}
-
-					// if we are not going to pak but we specified compressed then compress in the cooker ;)
-					// otherwise compress the pak files
-					if (!Params.Pak && !Params.SkipPak && Params.Compressed && !Params.ForceUncompressed)
-					{
-						CommandletParams += " -compressed";
-					}
-
-					if (Params.HasAdditionalCookerOptions)
-					{
-						string FormatedAdditionalCookerParams = Params.AdditionalCookerOptions.TrimStart(new char[] { '\"', ' ' }).TrimEnd(new char[] { '\"', ' ' });
-						CommandletParams += " ";
-						CommandletParams += FormatedAdditionalCookerParams;
-					}
-
-					if (!Params.NoClient)
-					{
-						var MapsList = Maps == null ? new List<string>() : Maps.ToList();
-						foreach (var ClientPlatform in Params.ClientTargetPlatforms)
-						{
-							var DataPlatformDesc = Params.GetCookedDataPlatformForClientTarget(ClientPlatform);
-							CommandletParams += (Platform.Platforms[DataPlatformDesc].GetCookExtraCommandLine(Params));
-							MapsList.AddRange((Platform.Platforms[ClientPlatform].GetCookExtraMaps()));
-						}
-						Maps = MapsList.ToArray();
-					}
-
-					// Config overrides (-ini)
-					foreach (string ConfigOverrideParam in Params.ConfigOverrideParams)
-					{
-						CommandletParams += " -";
-						CommandletParams += ConfigOverrideParam;
-					}
-
-					CookCommandlet(Params.RawProjectPath, Params.UnrealExe, Maps, Dirs, InternationalizationPreset, CulturesToCook, CombineCommandletParams(PlatformsToCook.ToArray()), CommandletParams);
+					string CommandletParams;
+					string[] DirectoriesToCook;
+					string[] MapsToCook;
+					string InternationalizationPreset;
+					string[] CulturesToCook;
+					string PlatformsToCook;
+
+					// get the parameters to CBTB
+					GetCookByTheBookCommandletParams(Params, out CommandletParams, out MapsToCook, out DirectoriesToCook, out InternationalizationPreset, out CulturesToCook, out PlatformsToCook);
+
+					// run a blocking cook
+					CookCommandlet(Params.RawProjectPath, Params.UnrealExe, MapsToCook, DirectoriesToCook, InternationalizationPreset, CulturesToCook, PlatformsToCook, CommandletParams);
 				}
 				catch (Exception Ex)
 				{
