Gauntlet Android: Making sure the Local Keys directory exists before we try to see if files exist within it

diff --git a/Engine/Source/Programs/AutomationTool/Gauntlet/Platform/Android/Gauntlet.TargetDeviceAndroid.cs b/Engine/Source/Programs/AutomationTool/Gauntlet/Platform/Android/Gauntlet.TargetDeviceAndroid.cs
--- a/Engine/Source/Programs/AutomationTool/Gauntlet/Platform/Android/Gauntlet.TargetDeviceAndroid.cs
+++ b/Engine/Source/Programs/AutomationTool/Gauntlet/Platform/Android/Gauntlet.TargetDeviceAndroid.cs
@@ -1294,6 +1294,10 @@ namespace Gauntlet
 					Log.Info("Using adb keys at {0}", KeyPath);
 
 					string LocalKeyPath = Path.Combine(Environment.GetEnvironmentVariable("USERPROFILE"), ".android");
+					if(!Directory.Exists(LocalKeyPath))
+					{
+						Directory.CreateDirectory(LocalKeyPath);
+					}
 
 					string RemoteKeyFile = Path.Combine(KeyPath, "adbkey");
 					string RemotePubKeyFile = Path.Combine(KeyPath, "adbkey.pub");
