Allow for retrevial of log path from RebuildHLODCommandlet

diff --git a/Engine/Source/Programs/AutomationTool/AutomationUtils/CommandletUtils.cs b/Engine/Source/Programs/AutomationTool/AutomationUtils/CommandletUtils.cs
--- a/Engine/Source/Programs/AutomationTool/AutomationUtils/CommandletUtils.cs
+++ b/Engine/Source/Programs/AutomationTool/AutomationUtils/CommandletUtils.cs
@@ -127,7 +127,12 @@ namespace AutomationTool
 			RunCommandlet(ProjectName, UnrealExe, "ResavePackages", String.Format("-buildtexturestreaming -buildlighting -MapsOnly -ProjectOnly -AllowCommandletRendering -SkipSkinVerify {0} {1}", MapsToRebuildLighting, Parameters));
 		}
 
-        public static void RebuildHLODCommandlet(FileReference ProjectName, string UnrealExe = "UnrealEditor-Cmd.exe", string[] Maps = null, string Parameters = "")
+		public static void RebuildHLODCommandlet(FileReference ProjectName, string UnrealExe = "UnrealEditor-Cmd.exe", string[] Maps = null, string Parameters = "")
+		{
+			RebuildHLODCommandlet(ProjectName, out string LogFile, UnrealExe, Maps, Parameters);
+		}
+
+		public static void RebuildHLODCommandlet(FileReference ProjectName, out string DestLogFile, string UnrealExe = "UnrealEditor-Cmd.exe", string[] Maps = null, string Parameters = "")
         {
             string MapsToRebuildHLODs = "";
             if (!IsNullOrEmpty(Maps))
@@ -135,7 +140,7 @@ namespace AutomationTool
                 MapsToRebuildHLODs = "-Map=" + CombineCommandletParams(Maps).Trim();
             }
 
-            RunCommandlet(ProjectName, UnrealExe, "ResavePackages", String.Format("-BuildHLOD -ProjectOnly -AllowCommandletRendering -SkipSkinVerify {0} {1}", MapsToRebuildHLODs, Parameters));
+            RunCommandlet(ProjectName, UnrealExe, "ResavePackages", String.Format("-BuildHLOD -ProjectOnly -AllowCommandletRendering -SkipSkinVerify {0} {1}", MapsToRebuildHLODs, Parameters), out DestLogFile);
         }
 
         /// <summary>
