Add Linux support for Gauntlet

diff --git a/Engine/Source/Programs/AutomationTool/Gauntlet/Platform/Linux/Gauntlet.LinuxBuildSources.cs b/Engine/Source/Programs/AutomationTool/Gauntlet/Platform/Linux/Gauntlet.LinuxBuildSources.cs
new file mode 100644
--- /dev/null
+++ b/Engine/Source/Programs/AutomationTool/Gauntlet/Platform/Linux/Gauntlet.LinuxBuildSources.cs
@@ -0,0 +1,16 @@
+﻿// Copyright Epic Games, Inc. All Rights Reserved.
+
+using System;
+using UnrealBuildTool;
+
+namespace Gauntlet
+{
+	public class LinuxBuildSource : StagedBuildSource<StagedBuild>
+	{
+		public override string BuildName { get { return "LinuxStagedBuild"; } }
+
+		public override UnrealTargetPlatform Platform { get { return UnrealTargetPlatform.Linux; } }
+
+		public override string PlatformFolderPrefix { get { return "Linux"; } }
+	}
+}
diff --git a/Engine/Source/Programs/AutomationTool/Gauntlet/Platform/Linux/Gauntlet.LinuxPGO.cs b/Engine/Source/Programs/AutomationTool/Gauntlet/Platform/Linux/Gauntlet.LinuxPGO.cs
new file mode 100644
--- /dev/null
+++ b/Engine/Source/Programs/AutomationTool/Gauntlet/Platform/Linux/Gauntlet.LinuxPGO.cs
@@ -0,0 +1,37 @@
+// Copyright Epic Games, Inc. All Rights Reserved.
+
+using UnrealBuildTool;
+
+namespace Gauntlet
+{
+	/// <summary>
+	/// Skeleton Linux PGO platform implementation (primarily used to test PGO locally with editor builds)
+	/// </summary>
+	internal class LinuxPGOPlatform : IPGOPlatform
+	{
+
+		public UnrealTargetPlatform GetPlatform()
+		{
+			return UnrealTargetPlatform.Linux;
+		}
+
+		public void GatherResults(string ArtifactPath)
+		{
+
+		}
+
+		public void ApplyConfiguration(PGOConfig Config)
+		{
+
+		}
+
+		public bool TakeScreenshot(ITargetDevice Device, string ScreenshotDirectory, out string ImageFilename)
+		{			
+			ImageFilename = string.Empty;
+			return false;
+		}
+	}
+}
+
+
+
diff --git a/Engine/Source/Programs/AutomationTool/Gauntlet/Platform/Linux/Gauntlet.SelfTest.TestTargetDeviceLinux.cs b/Engine/Source/Programs/AutomationTool/Gauntlet/Platform/Linux/Gauntlet.SelfTest.TestTargetDeviceLinux.cs
new file mode 100644
--- /dev/null
+++ b/Engine/Source/Programs/AutomationTool/Gauntlet/Platform/Linux/Gauntlet.SelfTest.TestTargetDeviceLinux.cs
@@ -0,0 +1,29 @@
+﻿// Copyright Epic Games, Inc. All Rights Reserved.
+
+using Gauntlet;
+using System;
+using System.Collections.Generic;
+using System.Diagnostics;
+using System.IO;
+using System.Linq;
+using System.Text;
+using System.Threading.Tasks;
+using UnrealBuildTool;
+
+namespace Gauntlet.SelfTest
+{
+	[TestGroup("Devices")]
+	class TestTargetDeviceLinux : TestTargetDevice
+	{
+		public override void TickTest()
+		{
+			string TempDir = Path.Combine(Environment.CurrentDirectory, "GauntletTemp");
+
+			TargetDeviceLinux Device = new TargetDeviceLinux("Local Linux", Globals.TempDir);
+
+			CheckEssentialFunctions(Device);
+
+			MarkComplete();
+		}
+	}
+}
diff --git a/Engine/Source/Programs/AutomationTool/Gauntlet/Platform/Linux/Gauntlet.TargetDeviceLinux.cs b/Engine/Source/Programs/AutomationTool/Gauntlet/Platform/Linux/Gauntlet.TargetDeviceLinux.cs
new file mode 100644
--- /dev/null
+++ b/Engine/Source/Programs/AutomationTool/Gauntlet/Platform/Linux/Gauntlet.TargetDeviceLinux.cs
@@ -0,0 +1,444 @@
+// Copyright Epic Games, Inc. All Rights Reserved.
+
+using System;
+using System.Collections.Generic;
+using System.IO;
+using System.Linq;
+using System.Threading;
+using AutomationTool;
+using UnrealBuildTool;
+using System.Text.RegularExpressions;
+
+namespace Gauntlet
+{
+	class LinuxAppInstance : LocalAppProcess
+	{
+		protected LinuxAppInstall Install;
+
+		public LinuxAppInstance(LinuxAppInstall InInstall, IProcessResult InProcess, string ProcessLogFile = null)
+			: base(InProcess, InInstall.CommandArguments, ProcessLogFile)
+		{
+			Install = InInstall;
+		}
+
+		public override string ArtifactPath
+		{
+			get
+			{
+				return Install.ArtifactPath;
+			}
+		}
+
+		public override ITargetDevice Device
+		{
+			get
+			{
+				return Install.Device;
+			}
+		}
+	}
+
+
+	class LinuxAppInstall : IAppInstall
+	{
+		public string Name { get; private set; }
+
+		public string WorkingDirectory;
+
+		public string ExecutablePath;
+
+		public string CommandArguments;
+
+		public string ArtifactPath;
+
+		public string ProjectName;
+
+		public TargetDeviceLinux LinuxDevice { get; private set; }
+
+		public ITargetDevice Device { get { return LinuxDevice; } }
+
+		public CommandUtils.ERunOptions RunOptions { get; set; }
+
+		public LinuxAppInstall(string InName, string InProjectName, TargetDeviceLinux InDevice)
+		{
+			Name = InName;
+			ProjectName = InProjectName;
+			LinuxDevice = InDevice;
+			CommandArguments = "";
+			this.RunOptions = CommandUtils.ERunOptions.NoWaitForExit;
+		}
+
+		public IAppInstance Run()
+		{
+			return Device.Run(this);
+		}
+
+		public virtual void CleanDeviceArtifacts()
+		{
+			if (!string.IsNullOrEmpty(ArtifactPath) && Directory.Exists(ArtifactPath))
+			{
+				try
+				{
+					Log.Info("Clearing actifact path {0} for {1}", ArtifactPath, Device.Name);
+					Directory.Delete(ArtifactPath, true);
+				}
+				catch (Exception Ex)
+				{
+					Log.Warning("Failed to delete {0}. {1}", ArtifactPath, Ex.Message);
+				}
+			}
+		}
+	}
+
+	public class LinuxDeviceFactory : IDeviceFactory
+	{
+		public bool CanSupportPlatform(UnrealTargetPlatform? Platform)
+		{
+			return Platform == UnrealTargetPlatform.Linux;
+		}
+
+		public ITargetDevice CreateDevice(string InRef, string InCachePath, string InParam = null)
+		{
+			return new TargetDeviceLinux(InRef, InCachePath);
+		}
+	}
+
+	/// <summary>
+	/// Linux implementation of a device to run applications
+	/// </summary>
+	public class TargetDeviceLinux : ITargetDevice
+	{
+		public string Name { get; protected set; }
+
+		protected string UserDir { get; set; }
+
+		/// <summary>
+		/// Our mappings of Intended directories to where they actually represent on this platform.
+		/// </summary>
+		protected Dictionary<EIntendedBaseCopyDirectory, string> LocalDirectoryMappings { get; set; }
+
+		public TargetDeviceLinux(string InName, string InCacheDir)
+		{
+			Name = InName;
+			LocalCachePath = InCacheDir;
+			RunOptions = CommandUtils.ERunOptions.NoWaitForExit | CommandUtils.ERunOptions.NoLoggingOfRunCommand;
+
+			UserDir = Path.Combine(LocalCachePath, "UserDir");
+            LocalDirectoryMappings = new Dictionary<EIntendedBaseCopyDirectory, string>();
+		}
+
+		#region IDisposable Support
+		private bool disposedValue = false; // To detect redundant calls
+
+		protected virtual void Dispose(bool disposing)
+		{
+			if (!disposedValue)
+			{
+				if (disposing)
+				{
+					// TODO: dispose managed state (managed objects).
+				}
+
+				// TODO: free unmanaged resources (unmanaged objects) and override a finalizer below.
+				// TODO: set large fields to null.
+
+				disposedValue = true;
+			}
+		}
+
+		// This code added to correctly implement the disposable pattern.
+		public void Dispose()
+		{
+			// Do not change this code. Put cleanup code in Dispose(bool disposing) above.
+			Dispose(true);
+			// TODO: uncomment the following line if the finalizer is overridden above.
+			// GC.SuppressFinalize(this);
+		}
+		#endregion
+
+		public CommandUtils.ERunOptions RunOptions { get; set; }
+
+        // We care about UserDir in Linux as some of the roles may require files going into user instead of build dir.
+        public void PopulateDirectoryMappings(string BasePath, string UserDir)
+		{
+			LocalDirectoryMappings.Add(EIntendedBaseCopyDirectory.Build, Path.Combine(BasePath, "Build"));
+			LocalDirectoryMappings.Add(EIntendedBaseCopyDirectory.Binaries, Path.Combine(BasePath, "Binaries"));
+			LocalDirectoryMappings.Add(EIntendedBaseCopyDirectory.Config, Path.Combine(BasePath, "Saved", "Config"));
+            LocalDirectoryMappings.Add(EIntendedBaseCopyDirectory.Content, Path.Combine(BasePath, "Content"));
+            LocalDirectoryMappings.Add(EIntendedBaseCopyDirectory.Demos, Path.Combine(UserDir, "Saved", "Demos"));
+			LocalDirectoryMappings.Add(EIntendedBaseCopyDirectory.PersistentDownloadDir, Path.Combine(BasePath, "Saved", "PersistentDownloadDir"));
+			LocalDirectoryMappings.Add(EIntendedBaseCopyDirectory.Profiling, Path.Combine(BasePath, "Saved", "Profiling"));
+            LocalDirectoryMappings.Add(EIntendedBaseCopyDirectory.Saved, Path.Combine(BasePath, "Saved"));
+		}
+
+        public IAppInstance Run(IAppInstall App)
+		{
+			LinuxAppInstall LinuxApp = App as LinuxAppInstall;
+
+			if (LinuxApp == null)
+			{
+				throw new DeviceException("AppInstance is of incorrect type!");
+			}
+
+			if (File.Exists(LinuxApp.ExecutablePath) == false)
+			{
+				throw new DeviceException("Specified path {0} not found!", LinuxApp.ExecutablePath);
+			}
+
+			IProcessResult Result = null;
+			string ProcessLogFile = null;
+
+			lock (Globals.MainLock)
+			{
+				string ExePath = Path.GetDirectoryName(LinuxApp.ExecutablePath);
+				string NewWorkingDir = string.IsNullOrEmpty(LinuxApp.WorkingDirectory) ? ExePath : LinuxApp.WorkingDirectory;
+				string OldWD = Environment.CurrentDirectory;
+				Environment.CurrentDirectory = NewWorkingDir;
+
+				Log.Info("Launching {0} on {1}", App.Name, ToString());
+
+				string CmdLine = LinuxApp.CommandArguments;
+
+				// Look in app parameters if abslog is specified, if so use it
+				/* TODO for linux
+				Regex CLRegex = new Regex(@"(--?[a-zA-Z]+)[:\s=]?([A-Z]:(?:\\[\w\s-]+)+\\?(?=\s-)|\""[^\""]*\""|[^-][^\s]*)?");
+				foreach (Match M in CLRegex.Matches(CmdLine))
+				{
+					if (M.Groups.Count == 3 && M.Groups[1].Value == "-abslog")
+					{
+						ProcessLogFile = M.Groups[2].Value;
+					}
+				}
+				*/
+
+				// explicitly set log file when not already defined
+				if (string.IsNullOrEmpty(ProcessLogFile))
+				{
+					string LogFolder = string.Format(@"{0}/Logs", LinuxApp.ArtifactPath);
+
+					if (!Directory.Exists(LogFolder))
+					{
+						Directory.CreateDirectory(LogFolder);
+					}
+
+					ProcessLogFile = string.Format("{0}/{1}.log", LogFolder, LinuxApp.ProjectName);
+					CmdLine = string.Format("{0} -abslog=\"{1}\"", CmdLine, ProcessLogFile);
+				}
+
+				// cleanup any existing log file
+				try
+				{
+					if (File.Exists(ProcessLogFile))
+					{
+						File.Delete(ProcessLogFile);
+					}
+				}
+				catch (Exception Ex)
+				{
+					throw new AutomationException("Unable to delete existing log file {0} {1}", ProcessLogFile, Ex.Message);
+				}
+
+				Log.Verbose("\t{0}", CmdLine);
+
+				Result = CommandUtils.Run(LinuxApp.ExecutablePath, CmdLine, Options: LinuxApp.RunOptions | (ProcessLogFile != null ? CommandUtils.ERunOptions.NoStdOutRedirect : 0 ));
+
+				if (Result.HasExited && Result.ExitCode != 0)
+				{
+					throw new AutomationException("Failed to launch {0}. Error {1}", LinuxApp.ExecutablePath, Result.ExitCode);
+				}
+
+				Environment.CurrentDirectory = OldWD;
+			}
+
+			return new LinuxAppInstance(LinuxApp, Result, ProcessLogFile);
+		}
+
+		private void CopyAdditionalFiles(UnrealAppConfig AppConfig)
+		{
+			if (AppConfig.FilesToCopy != null)
+			{
+				foreach (UnrealFileToCopy FileToCopy in AppConfig.FilesToCopy)
+				{
+					string PathToCopyTo = Path.Combine(LocalDirectoryMappings[FileToCopy.TargetBaseDirectory], FileToCopy.TargetRelativeLocation);
+					if (File.Exists(FileToCopy.SourceFileLocation))
+					{
+						FileInfo SrcInfo = new FileInfo(FileToCopy.SourceFileLocation);
+						SrcInfo.IsReadOnly = false;
+						string DirectoryToCopyTo = Path.GetDirectoryName(PathToCopyTo);
+						if (!Directory.Exists(DirectoryToCopyTo))
+						{
+							Directory.CreateDirectory(DirectoryToCopyTo);
+						}
+						if (File.Exists(PathToCopyTo))
+						{
+							FileInfo ExistingFile = new FileInfo(PathToCopyTo);
+							ExistingFile.IsReadOnly = false;
+						}
+						SrcInfo.CopyTo(PathToCopyTo, true);
+						Log.Info("Copying {0} to {1}", FileToCopy.SourceFileLocation, PathToCopyTo);
+					}
+					else
+					{
+						Log.Warning("File to copy {0} not found", FileToCopy);
+					}
+				}
+			}
+		}
+
+		protected IAppInstall InstallStagedBuild(UnrealAppConfig AppConfig, StagedBuild InBuild)
+		{
+			bool SkipDeploy = Globals.Params.ParseParam("SkipDeploy");
+
+			string BuildPath = InBuild.BuildPath;
+
+			if (CanRunFromPath(BuildPath) == false)
+			{
+				string SubDir = string.IsNullOrEmpty(AppConfig.Sandbox) ? AppConfig.ProjectName : AppConfig.Sandbox;
+				string DestPath = Path.Combine(this.LocalCachePath, SubDir, AppConfig.ProcessType.ToString());
+
+				if (!SkipDeploy)
+				{
+					DestPath = StagedBuild.InstallBuildParallel(AppConfig, InBuild, BuildPath, DestPath, ToString());
+				}
+				else
+				{
+					Log.Info("Skipping install of {0} (-skipdeploy)", BuildPath);
+				}
+
+				Utils.SystemHelpers.MarkDirectoryForCleanup(DestPath);
+
+				BuildPath = DestPath;
+			}
+
+			LinuxAppInstall LinuxApp = new LinuxAppInstall(AppConfig.Name, AppConfig.ProjectName, this);
+			LinuxApp.RunOptions = RunOptions;
+
+			// Set commandline replace any InstallPath arguments with the path we use
+			LinuxApp.CommandArguments = Regex.Replace(AppConfig.CommandLine, @"\$\(InstallPath\)", BuildPath, RegexOptions.IgnoreCase);
+
+			if (string.IsNullOrEmpty(UserDir) == false)
+			{
+				LinuxApp.CommandArguments += string.Format(" -userdir=\"{0}\"", UserDir);
+				LinuxApp.ArtifactPath = Path.Combine(UserDir, @"Saved");
+
+				Utils.SystemHelpers.MarkDirectoryForCleanup(UserDir);
+			}
+			else
+			{
+				// e.g d:\Unreal\GameName\Saved
+				LinuxApp.ArtifactPath = Path.Combine(BuildPath, AppConfig.ProjectName, @"Saved");
+
+			}
+
+			// clear artifact path
+			LinuxApp.CleanDeviceArtifacts();
+
+            if (LocalDirectoryMappings.Count == 0)
+            {
+                PopulateDirectoryMappings(Path.Combine(BuildPath, AppConfig.ProjectName), UserDir);
+            }
+
+			CopyAdditionalFiles(AppConfig);
+
+            if (Path.IsPathRooted(InBuild.ExecutablePath))
+			{
+				LinuxApp.ExecutablePath = InBuild.ExecutablePath;
+			}
+			else
+			{
+				// TODO - this check should be at a higher level....
+				string BinaryPath = Path.Combine(BuildPath, InBuild.ExecutablePath);
+
+				// check for a local newer executable
+				if (Globals.Params.ParseParam("dev") && AppConfig.ProcessType.UsesEditor() == false)
+				{
+					string LocalBinary = Path.Combine(Environment.CurrentDirectory, InBuild.ExecutablePath);
+
+					bool LocalFileExists = File.Exists(LocalBinary);
+					bool LocalFileNewer = LocalFileExists && File.GetLastWriteTime(LocalBinary) > File.GetLastWriteTime(BinaryPath);
+
+					Log.Verbose("Checking for newer binary at {0}", LocalBinary);
+					Log.Verbose("LocalFile exists: {0}. Newer: {1}", LocalFileExists, LocalFileNewer);
+
+					if (LocalFileExists && LocalFileNewer)
+					{
+						// need to -basedir to have our exe load content from the path
+						LinuxApp.CommandArguments += string.Format(" -basedir={0}", Path.GetDirectoryName(BinaryPath));
+
+						BinaryPath = LocalBinary;
+					}
+				}
+
+				LinuxApp.ExecutablePath = BinaryPath;
+			}
+
+			return LinuxApp;
+		}
+
+
+		public IAppInstall InstallApplication(UnrealAppConfig AppConfig)
+		{
+			if (AppConfig.Build is StagedBuild)
+			{
+				return InstallStagedBuild(AppConfig, AppConfig.Build as StagedBuild);
+			}
+
+			EditorBuild EditorBuild = AppConfig.Build as EditorBuild;
+
+			if (EditorBuild == null)
+			{
+				throw new AutomationException("Invalid build type!");
+			}
+
+			LinuxAppInstall LinuxApp = new LinuxAppInstall(AppConfig.Name, AppConfig.ProjectName, this);
+
+			LinuxApp.WorkingDirectory = Path.GetDirectoryName(EditorBuild.ExecutablePath);
+			LinuxApp.RunOptions = RunOptions;
+	
+			// Force this to stop logs and other artifacts going to different places
+			LinuxApp.CommandArguments = AppConfig.CommandLine + string.Format(" -userdir=\"{0}\"", UserDir);
+			LinuxApp.ArtifactPath = Path.Combine(UserDir, @"Saved");
+			LinuxApp.ExecutablePath = EditorBuild.ExecutablePath;
+
+			if (LocalDirectoryMappings.Count == 0)
+			{
+				PopulateDirectoryMappings(AppConfig.ProjectFile.Directory.FullName, AppConfig.ProjectFile.Directory.FullName);
+			}
+
+			CopyAdditionalFiles(AppConfig);
+
+			return LinuxApp;
+		}
+		
+		public bool CanRunFromPath(string InPath)
+		{
+			return !Utils.SystemHelpers.IsNetworkPath(InPath);
+		}
+
+		public UnrealTargetPlatform? Platform { get { return UnrealTargetPlatform.Linux; } }
+
+		public string LocalCachePath { get; private set; }
+		public bool IsAvailable { get { return true; } }
+		public bool IsConnected { get { return true; } }
+		public bool IsOn { get { return true; } }
+		public bool PowerOn() { return true; }
+		public bool PowerOff() { return true; }
+		public bool Reboot() { return true; }
+		public bool Connect() { return true; }
+		public bool Disconnect(bool bForce = false) { return true; }
+
+		public override string ToString()
+		{
+			return Name;
+		}
+
+		public Dictionary<EIntendedBaseCopyDirectory, string> GetPlatformDirectoryMappings()
+		{
+			if (LocalDirectoryMappings.Count == 0)
+			{
+				Log.Warning("Platform directory mappings have not been populated for this platform! This should be done within InstallApplication()");
+			}
+			return LocalDirectoryMappings;
+		}
+	}
+}
diff --git a/Engine/Source/Programs/AutomationTool/Gauntlet/Unreal/Base/Gauntlet.UnrealTestConfiguration.cs b/Engine/Source/Programs/AutomationTool/Gauntlet/Unreal/Base/Gauntlet.UnrealTestConfiguration.cs
--- a/Engine/Source/Programs/AutomationTool/Gauntlet/Unreal/Base/Gauntlet.UnrealTestConfiguration.cs
+++ b/Engine/Source/Programs/AutomationTool/Gauntlet/Unreal/Base/Gauntlet.UnrealTestConfiguration.cs
@@ -526,7 +526,7 @@ namespace Gauntlet
 		protected bool Windowed { get; set; }
 
 		/// <summary>
-		/// Which window mode to use for the PC or Mac client. Only Windowed and Fullscreen are fully supported.
+		/// Which window mode to use for the PC or Mac or Linux client. Only Windowed and Fullscreen are fully supported.
 		/// </summary>
 		/// 
 		[AutoParam(EWindowMode.Windowed)]
@@ -822,7 +822,7 @@ namespace Gauntlet
 			}
 			else if (AppConfig.ProcessType.IsClient())
 			{
-				if (AppConfig.Platform == UnrealTargetPlatform.Win64 || AppConfig.Platform == UnrealTargetPlatform.Mac)
+				if (AppConfig.Platform == UnrealTargetPlatform.Win64 || AppConfig.Platform == UnrealTargetPlatform.Mac || AppConfig.Platform == UnrealTargetPlatform.Linux)
 				{
 					if (!IgnoreDefaultResolutionAndWindowMode)
 					{
@@ -855,6 +855,13 @@ namespace Gauntlet
 				}
 			}
 
+			// due to an issue with dotnet being extremely pedantic we have to drop our locks on files so we can read from the log file
+			// [REDACTED_PROFILE_URL]/runtime/issues/34126
+			if (AppConfig.Platform == UnrealTargetPlatform.Linux)
+			{
+				AppConfig.CommandLine += " -noexclusivelockonwrite";
+			}
+
 			// use -log on servers so we get a window..
 			if (AppConfig.ProcessType.IsServer())
 			{
@@ -864,6 +871,13 @@ namespace Gauntlet
 			if (Attended == false)
 			{
 				AppConfig.CommandLine += " -unattended";
+
+				// if we are unattended but still may need access to Vulkan passing renderoffscreen to allow not depending on
+				// the X11/Wayland display server to be around and use a dummy/offscreen rendering mode
+				if (AppConfig.Platform == UnrealTargetPlatform.Linux)
+				{
+					AppConfig.CommandLine += " -renderoffscreen";
+				}
 			}
 
 			AppConfig.CommandLine += " -stdout -AllowStdOutLogVerbosity";
@@ -893,4 +907,4 @@ namespace Gauntlet
 		}			
 	}
 
-}
\ No newline at end of file
+}
