Don't create devices for CookedCooker targetplatform (CookedEditor TP was already correct)
- Harden Turnkey device querying to print a message on device exceptions, but not fail out

#rnx

diff --git a/Engine/Source/Developer/CookedEditor/Public/CookedEditorTargetPlatform.h b/Engine/Source/Developer/CookedEditor/Public/CookedEditorTargetPlatform.h
--- a/Engine/Source/Developer/CookedEditor/Public/CookedEditorTargetPlatform.h
+++ b/Engine/Source/Developer/CookedEditor/Public/CookedEditorTargetPlatform.h
@@ -226,6 +226,24 @@ public:
 		return false;
 	}
 
+	virtual float GetVariantPriority() const override
+	{
+		// by returning -1, we will never use this variant when targeting host platform this class implements
+		// (without this, cooking for Windows in the editor may choose this variant to cook for, which we never want)
+		return -1.0f;
+	}
+
+	virtual void GetAllDevices(TArray<ITargetDevicePtr>& OutDevices) const override
+	{
+	}
+
+	virtual ITargetDevicePtr GetDefaultDevice() const override
+	{
+		return nullptr;
+	}
+
+
+
 	TUniquePtr<ICookedEditorPackageManager> PackageManager;
 
 };
diff --git a/Engine/Source/Programs/AutomationTool/Turnkey/TurnkeyUtils.cs b/Engine/Source/Programs/AutomationTool/Turnkey/TurnkeyUtils.cs
--- a/Engine/Source/Programs/AutomationTool/Turnkey/TurnkeyUtils.cs
+++ b/Engine/Source/Programs/AutomationTool/Turnkey/TurnkeyUtils.cs
@@ -382,9 +382,38 @@ namespace Turnkey
 
 		private static DeviceInfo GetDeviceByPlatformAndName(UnrealTargetPlatform Platform, string DeviceName)
 		{
-			return AutomationTool.Platform.GetPlatform(Platform).GetDeviceByName(DeviceName);
+			try
+			{
+				return AutomationTool.Platform.GetPlatform(Platform).GetDeviceByName(DeviceName);
+			}
+			catch (Exception Ex)
+			{
+				TurnkeyUtils.Log($"An error occurred trying to access device {DeviceName} for platform {Platform}: {Ex.Message}");
+
+				// swallow errors and just return no device
+				return null;
+			}
+		}
+
+		private static DeviceInfo[] GetDevicesForPlatform(UnrealTargetPlatform Platform)
+		{
+			try
+			{
+				DeviceInfo[] Devices = AutomationTool.Platform.GetPlatform(Platform).GetDevices();
+				
+				// in the null case, return an empty array, for cleaner code for callers 
+				return Devices != null ? Devices : new DeviceInfo[] { };
+			}
+			catch (Exception Ex)
+			{
+				TurnkeyUtils.Log($"An error occurred trying to access all devices for platform {Platform}: {Ex.Message}");
+
+				// swallow errors and just return no devices
+				return new DeviceInfo[] { } ;
+			}
 		}
 
+
 		public static List<DeviceInfo> GetDevicesFromCommandLineOrUser(string[] CommandOptions, UnrealTargetPlatform Platform)
 		{
 			return GetDevicesFromCommandLineOrUser(CommandOptions, new List<UnrealTargetPlatform>() { Platform });
@@ -439,7 +468,7 @@ namespace Turnkey
 						
 						if (DeviceName.ToLower() == "all")
 						{
-							ChosenDevices.AddRange(AutomationTool.Platform.GetPlatform(Platform).GetDevices());
+							ChosenDevices.AddRange(GetDevicesForPlatform(Platform));
 						}
 						else
 						{
@@ -491,11 +520,7 @@ namespace Turnkey
 				{
 					foreach (UnrealTargetPlatform Platform in ChosenPlatforms)
 					{
-						DeviceInfo[] Devices = AutomationTool.Platform.GetPlatform(Platform).GetDevices();
-						if (Devices != null)
-						{
-							ChosenDevices.AddRange(Devices);
-						}
+						ChosenDevices.AddRange(GetDevicesForPlatform(Platform));
 					}
 				}
 				// now if the list of devices was given, then attempt to find them in the platform
@@ -521,15 +546,10 @@ namespace Turnkey
 
 					foreach (UnrealTargetPlatform Platform in ChosenPlatforms)
 					{
-						Platform AutomationPlatform = AutomationTool.Platform.GetPlatform(Platform);
-						DeviceInfo[] Devices = AutomationPlatform.GetDevices();
-						if (Devices != null)
+						foreach (DeviceInfo Device in GetDevicesForPlatform(Platform))
 						{
-							foreach (DeviceInfo Device in Devices)
-							{
-								PossibleDevices.Add(Device);
-								Options.Add(string.Format("[{0} {1}] {2}", Platform, Device.Type, Device.Name));
-							}
+							PossibleDevices.Add(Device);
+							Options.Add(string.Format("[{0} {1}] {2}", Platform, Device.Type, Device.Name));
 						}
 					}
 
