#UE Config files in platform extension folders for plugins will now be staged

diff --git a/Engine/Source/Programs/AutomationTool/Scripts/CopyBuildToStagingDirectory.Automation.cs b/Engine/Source/Programs/AutomationTool/Scripts/CopyBuildToStagingDirectory.Automation.cs
--- a/Engine/Source/Programs/AutomationTool/Scripts/CopyBuildToStagingDirectory.Automation.cs
+++ b/Engine/Source/Programs/AutomationTool/Scripts/CopyBuildToStagingDirectory.Automation.cs
@@ -939,6 +939,11 @@ namespace AutomationScripts
 
 					// Stage plugin config files
 					List<KeyValuePair<StagedFileReference, FileReference>> StagedPlugins = SC.FilesToStage.UFSFiles.Where(x => x.Value.HasExtension(".uplugin")).ToList();
+					List<PluginInfo> AvailablePlugins = new List<PluginInfo>();
+					if (StagedPlugins.Count > 0 && PlatformExtensionsToStage.Count > 0)
+					{
+						AvailablePlugins = UnrealBuildTool.Plugins.ReadAvailablePlugins(Unreal.EngineDirectory, Params.RawProjectPath.Directory, null);
+					}
 					foreach (KeyValuePair<StagedFileReference, FileReference> StagedPlugin in StagedPlugins)
 					{
 						//PluginDescriptor Descriptor = PluginDescriptor.FromFile(StagedPlugin.Value);
@@ -947,6 +952,29 @@ namespace AutomationScripts
 						{
 							SC.StageFiles(StagedFileType.UFS, PluginConfigDirectory, "*.ini", StageFilesSearch.AllDirectories);
 						}
+
+						// look in platform extension locations of this plugin for config files
+						PluginInfo Info = AvailablePlugins.Find(x => x.File == StagedPlugin.Value);
+						if (Info != null)
+						{
+							foreach (string PlatformExtensionToStage in PlatformExtensionsToStage)
+							{
+								foreach (FileReference ChildPlugin in Info.ChildFiles)
+								{
+									// only look at a child plugin for a platform we are staging (ie, it is in a /PlatformName/ subdir)
+									if (ChildPlugin.ContainsName(PlatformExtensionToStage, 0))
+									{
+										// stage any config files that exist, into the destination plugin Config dir, so that Runtime will find them
+										DirectoryReference PlatformConfigDirectory = DirectoryReference.Combine(ChildPlugin.Directory, "Config");
+										if (DirectoryReference.Exists(PlatformConfigDirectory))
+										{
+											// stage info Plugin/Config/PlatformName/
+											SC.StageFiles(StagedFileType.UFS, PlatformConfigDirectory, "*.ini", StageFilesSearch.AllDirectories, StagedDirectoryReference.Combine(StagedPlugin.Key.Directory, "Config", PlatformExtensionToStage));
+										}
+									}
+								}
+							}
+						}
 					}
 
 					// Stage all project localization targets
