Fixed: Missing GPUDumpViewer files from a packaged builds on Android and iOS

diff --git a/Engine/Build/Rsync/RsyncEngine.txt b/Engine/Build/Rsync/RsyncEngine.txt
--- a/Engine/Build/Rsync/RsyncEngine.txt
+++ b/Engine/Build/Rsync/RsyncEngine.txt
@@ -129,6 +129,11 @@ protect /Programs/UnrealBuildTool/
 protect /Programs/UnrealHeaderTool/Intermediate/
 protect /Programs/UnrealHeaderTool/Saved/
 
+# Extras: Allow specific folders
++ /Extras/
++ /Extras/GPUDumpViewer/
++ /Extras/GPUDumpViewer/**
+
 # Saved
 protect /Saved/
 
diff --git a/Engine/Source/Programs/AutomationTool/Android/AndroidPlatform.Automation.cs b/Engine/Source/Programs/AutomationTool/Android/AndroidPlatform.Automation.cs
--- a/Engine/Source/Programs/AutomationTool/Android/AndroidPlatform.Automation.cs
+++ b/Engine/Source/Programs/AutomationTool/Android/AndroidPlatform.Automation.cs
@@ -601,7 +601,7 @@ public class AndroidPlatform : Platform
 		return bAllowOverflowOBBFiles;
 	}
 
-	private bool CreateOBBFile(DeploymentContext SC, string StageDirectoryPath, string OutputFilename, List<FileReference> FilesForObb)
+	private bool CreateOBBFile(DeploymentContext SC, string OutputFilename, List<FileReference> FilesForObb)
 	{
 		LogInformation("Creating {0} from {1}", OutputFilename, SC.StageDirectory);
 		using (ZipFile ObbFile = new ZipFile(OutputFilename))
@@ -626,8 +626,8 @@ public class AndroidPlatform : Platform
 
 			foreach (FileReference FileRef in FilesForObb)
 			{
-				string DestinationPath = Path.GetDirectoryName(FileRef.FullName).Replace(StageDirectoryPath, SC.ShortProjectName);
-				ObbFile.AddFile(FileRef.FullName, DestinationPath);
+				string DestinationDirectoryPath = Path.GetRelativePath(SC.StageDirectory.FullName, Path.GetDirectoryName(FileRef.FullName));
+				ObbFile.AddFile(FileRef.FullName, DestinationDirectoryPath);
 			}
 
 			// ObbFile.AddDirectory(SC.StageDirectory+"/"+SC.ShortProjectName, SC.ShortProjectName);
@@ -702,9 +702,20 @@ public class AndroidPlatform : Platform
 		{
 			ObbFileFilter.AddRules(ObbFilters);
 		}
+		// Filter out dynamic libraries from obb
+		ObbFileFilter.Exclude("*.so");
 
-		string StageDirectoryPath = Path.Combine(SC.StageDirectory.FullName, SC.ShortProjectName);
-		List<FileReference> FilesForObb = ObbFileFilter.ApplyToDirectory(new DirectoryReference(StageDirectoryPath), true);
+		List<FileReference> FilesForObb = new List<FileReference>();
+		// Add staged Engine files
+		{
+			string EngineStageDirectoryPath = Path.Combine(SC.StageDirectory.FullName, "Engine");
+			FilesForObb.AddRange(ObbFileFilter.ApplyToDirectory(new DirectoryReference(EngineStageDirectoryPath), true));
+		}
+		// Add staged project files
+		{
+			string ProjectStageDirectoryPath = Path.Combine(SC.StageDirectory.FullName, SC.ShortProjectName);
+			FilesForObb.AddRange(ObbFileFilter.ApplyToDirectory(new DirectoryReference(ProjectStageDirectoryPath), true));
+		}
 
 		bool OBBNeedsUpdate = false;
 
@@ -770,6 +781,7 @@ public class AndroidPlatform : Platform
 				FilesToObb = new List<FileReference>();
 
 				// Collect the filesize and place into Obb or Patch list
+				Int64 StagingDirLength = SC.StageDirectory.FullName.Length;
 				Int64 MinimumObbSize = 22 + 10;		// EOCD with comment (store version)
 				Int64 MainObbSize = MinimumObbSize;
 				Int64 PatchObbSize = MinimumObbSize;
@@ -779,9 +791,8 @@ public class AndroidPlatform : Platform
 				{
 					FileInfo LocalFileInfo = new FileInfo(FileRef.FullName);
 					Int64 LocalFileLength = LocalFileInfo.Length;
+					Int64 FilenameLength = FileRef.FullName.Length - StagingDirLength - 1;
 
-					string DestinationPath = Path.GetDirectoryName(FileRef.FullName).Replace(StageDirectoryPath, SC.ShortProjectName);
-					Int64 FilenameLength = DestinationPath.Length + Path.GetFileName(FileRef.FullName).Length + 1;
 					Int64 LocalOverhead = (30 + FilenameLength + 36);		// local file descriptor
 					Int64 GlobalOverhead = (46 + FilenameLength + 36);		// central directory cost
 					Int64 FileRequirements = LocalFileLength + LocalOverhead + GlobalOverhead;
@@ -825,7 +836,7 @@ public class AndroidPlatform : Platform
 			}
 
 			// Now create the main OBB as a ZIP archive.
-			if (!CreateOBBFile(SC, StageDirectoryPath, LocalObbName, FilesToObb))
+			if (!CreateOBBFile(SC, LocalObbName, FilesToObb))
 			{
 				LogInformation("Failed to build OBB: " + LocalObbName);
 				throw new AutomationException(ExitCode.Error_AndroidOBBError, "Stage Failed. Could not build OBB {0}. The file may be too big to fit in an OBB ({1} limit)", LocalObbName, LimitString);
@@ -834,7 +845,7 @@ public class AndroidPlatform : Platform
 			// Now create the patch OBB as a ZIP archive if required.
 			if (FilesToPatch.Count() > 0)
 			{
-				if (!CreateOBBFile(SC, StageDirectoryPath, LocalPatchName, FilesToPatch))
+				if (!CreateOBBFile(SC, LocalPatchName, FilesToPatch))
 				{
 					LogInformation("Failed to build OBB: " + LocalPatchName);
 					throw new AutomationException(ExitCode.Error_AndroidOBBError, "Stage Failed. Could not build OBB {0}. The file may be too big to fit in an OBB ({1} limit)", LocalPatchName, LimitString);
@@ -844,7 +855,7 @@ public class AndroidPlatform : Platform
 			// Now create the overflow1 OBB as a ZIP archive if required.
 			if (FilesToOverflow1.Count() > 0)
 			{
-				if (!CreateOBBFile(SC, StageDirectoryPath, LocalOverflow1Name, FilesToOverflow1))
+				if (!CreateOBBFile(SC, LocalOverflow1Name, FilesToOverflow1))
 				{
 					LogInformation("Failed to build OBB: " + LocalOverflow1Name);
 					throw new AutomationException(ExitCode.Error_AndroidOBBError, "Stage Failed. Could not build OBB {0}. The file may be too big to fit in an OBB ({1} limit)", LocalPatchName, LimitString);
@@ -854,7 +865,7 @@ public class AndroidPlatform : Platform
 			// Now create the overflow2 OBB as a ZIP archive if required.
 			if (FilesToOverflow2.Count() > 0)
 			{
-				if (!CreateOBBFile(SC, StageDirectoryPath, LocalOverflow2Name, FilesToOverflow2))
+				if (!CreateOBBFile(SC, LocalOverflow2Name, FilesToOverflow2))
 				{
 					LogInformation("Failed to build OBB: " + LocalOverflow2Name);
 					throw new AutomationException(ExitCode.Error_AndroidOBBError, "Stage Failed. Could not build OBB {0}. The file may be too big to fit in an OBB ({1} limit)", LocalPatchName, LimitString);
diff --git a/Engine/Source/Programs/AutomationTool/AutomationUtils/Platform.cs b/Engine/Source/Programs/AutomationTool/AutomationUtils/Platform.cs
--- a/Engine/Source/Programs/AutomationTool/AutomationUtils/Platform.cs
+++ b/Engine/Source/Programs/AutomationTool/AutomationUtils/Platform.cs
@@ -622,6 +622,16 @@ namespace AutomationTool
 			get { return BuildHostPlatform.Current.Platform == PlatformType; }
 		}
 
+		/// <summary>
+		/// return true if we need to call Remap of a specific file type
+		/// </summary>
+		/// <param name="FileType" The staged file type to check (UFS vs SsytemNonUFS, etc)
+		/// <returns>true if files should be remaped, for the given filetype</returns>
+		public virtual bool RemapFileType(StagedFileType FileType)
+		{
+			return (FileType == StagedFileType.UFS || FileType == StagedFileType.NonUFS);
+		}
+
 		/// <summary>
 		/// Remaps the given content directory to its final location
 		/// </summary>
diff --git a/Engine/Source/Programs/AutomationTool/IOS/IOSPlatform.Automation.cs b/Engine/Source/Programs/AutomationTool/IOS/IOSPlatform.Automation.cs
--- a/Engine/Source/Programs/AutomationTool/IOS/IOSPlatform.Automation.cs
+++ b/Engine/Source/Programs/AutomationTool/IOS/IOSPlatform.Automation.cs
@@ -1910,6 +1910,14 @@ public class IOSPlatform : Platform
 		}
 	}
 
+	public override bool RemapFileType(StagedFileType FileType)
+	{
+		return (
+			FileType == StagedFileType.UFS || 
+			FileType == StagedFileType.NonUFS ||
+			FileType == StagedFileType.DebugNonUFS);
+	}
+
 	public override StagedFileReference Remap(StagedFileReference Dest)
 	{
 		return new StagedFileReference("cookeddata/" + Dest.Name);
diff --git a/Engine/Source/Programs/AutomationTool/Scripts/CopyBuildToStagingDirectory.Automation.cs b/Engine/Source/Programs/AutomationTool/Scripts/CopyBuildToStagingDirectory.Automation.cs
--- a/Engine/Source/Programs/AutomationTool/Scripts/CopyBuildToStagingDirectory.Automation.cs
+++ b/Engine/Source/Programs/AutomationTool/Scripts/CopyBuildToStagingDirectory.Automation.cs
@@ -1319,22 +1319,18 @@ namespace AutomationScripts
 			}
 
 			// Remap all the non-ufs files if not using a PAK file
-			if (Params.HasIterateSharedCookedBuild)
+			// Shared NonUFS files are staged in their remapped location, and may be duplicated in the to-stage list.
+			if (SC.StageTargetPlatform.RemapFileType(StagedFileType.NonUFS))
 			{
-				// Shared NonUFS files are staged in their remapped location, and may be duplicated in the to-stage list.
-				Dictionary<StagedFileReference, FileReference> NonUFSToStage = new Dictionary<StagedFileReference, FileReference>();
-				foreach (KeyValuePair<StagedFileReference, FileReference> StagedFilePair in SC.FilesToStage.NonUFSFiles)
-				{
-					NonUFSToStage[SC.StageTargetPlatform.Remap(StagedFilePair.Key)] = StagedFilePair.Value;
-				}
-				SC.FilesToStage.NonUFSFiles = NonUFSToStage;
+				SC.FilesToStage.NonUFSFiles = SC.FilesToStage.NonUFSFiles.ToDictionary(x => SC.StageTargetPlatform.Remap(x.Key), x => x.Value);
 			}
-			else
+				
+			if (SC.StageTargetPlatform.RemapFileType(StagedFileType.DebugNonUFS))
 			{
-				SC.FilesToStage.NonUFSFiles = SC.FilesToStage.NonUFSFiles.ToDictionary(x => SC.StageTargetPlatform.Remap(x.Key), x => x.Value);
+				SC.FilesToStage.NonUFSDebugFiles = SC.FilesToStage.NonUFSDebugFiles.ToDictionary(x => SC.StageTargetPlatform.Remap(x.Key), x => x.Value);
 			}
 
-			if (!Params.UsePak(SC.StageTargetPlatform))
+			if (SC.StageTargetPlatform.RemapFileType(StagedFileType.UFS) && !Params.UsePak(SC.StageTargetPlatform))
 			{
 				SC.FilesToStage.UFSFiles = SC.FilesToStage.UFSFiles.ToDictionary(x => SC.StageTargetPlatform.Remap(x.Key), x => x.Value);
 			}
