Turnkey Google Drive filename matches are not case sensitive.

#rnx

diff --git a/Engine/Source/Programs/AutomationTool/Turnkey/Providers/GoogleDriveCopyProvider.cs b/Engine/Source/Programs/AutomationTool/Turnkey/Providers/GoogleDriveCopyProvider.cs
--- a/Engine/Source/Programs/AutomationTool/Turnkey/Providers/GoogleDriveCopyProvider.cs
+++ b/Engine/Source/Programs/AutomationTool/Turnkey/Providers/GoogleDriveCopyProvider.cs
@@ -457,7 +457,7 @@ namespace Turnkey
 				string Wildcard = (NextSlash == -1) ? PathString.Substring(PrevSlash + 1) : PathString.Substring(PrevSlash + 1, (NextSlash - PrevSlash) - 1);
 
 				// convert to a wildcard with capturing to get what the matches are
-				Regex Regex = new Regex(string.Format("^{0}$", Wildcard.Replace("*", "(.+?)")));
+				Regex Regex = new Regex(string.Format("^{0}$", Wildcard.Replace("*", "(.+?)")), RegexOptions.IgnoreCase);
 
 				// get files that could match the wildcard and are not trashed
 				string Query = string.Format("'{0}' in parents and trashed = false", FolderBeforeWildcard.Id);
