Add support for ASan automation in Gauntlet and EngineTest

diff --git a/Engine/Source/Programs/AutomationTool/Gauntlet/Unreal/BuildSource/Gauntlet.UnrealBuildSource.cs b/Engine/Source/Programs/AutomationTool/Gauntlet/Unreal/BuildSource/Gauntlet.UnrealBuildSource.cs
--- a/Engine/Source/Programs/AutomationTool/Gauntlet/Unreal/BuildSource/Gauntlet.UnrealBuildSource.cs
+++ b/Engine/Source/Programs/AutomationTool/Gauntlet/Unreal/BuildSource/Gauntlet.UnrealBuildSource.cs
@@ -622,6 +622,27 @@ namespace Gauntlet
 			return CommandlineToReturn;
 		}
 
+		private string GetSanitizerSuffix()
+		{
+			if (Globals.Params.ParseParam("asan"))
+			{
+				return "ASan";
+			}
+			else if (Globals.Params.ParseParam("msan"))
+			{
+				return "MSan";
+			}
+			else if (Globals.Params.ParseParam("ubsan"))
+			{
+				return "UBSan";
+			}
+			else if (Globals.Params.ParseParam("tsan"))
+			{
+				return "TSan";
+			}
+
+			return "";
+		}
 
 		/// <summary>
 		/// Given a role, platform, and config, returns the path to the binary for that config. E.g. Binaries\Win64\FooServer-Win64-Shipping.exe
@@ -632,6 +653,8 @@ namespace Gauntlet
 		/// <returns></returns>
 		virtual public string GetRelativeExecutablePath(UnrealTargetRole TargetRole, UnrealTargetPlatform TargetPlatform, UnrealTargetConfiguration TargetConfiguration)
 		{
+			string SanitizerSuffix = GetSanitizerSuffix();
+			bool HasSanitizerSuffix = !string.IsNullOrEmpty(SanitizerSuffix);
 			string ExePath;
 
 			if (TargetRole.UsesEditor())
@@ -646,13 +669,29 @@ namespace Gauntlet
 				if (EditorExe != null)
 				{
 					ExePath = EditorExe.FullName;
+
+					if (HasSanitizerSuffix)
+					{
+						ExePath += "-" + SanitizerSuffix;
+					}
 				}
 				else
 				{
 					string ExeFileName = "UnrealEditor";
 					if (TargetConfiguration != UnrealTargetConfiguration.Development)
 					{
-						ExeFileName += string.Format("-{0}-{1}", TargetPlatform.ToString(), TargetConfiguration.ToString());
+						if (HasSanitizerSuffix)
+						{
+							ExeFileName += string.Format("-{0}-{1}-{2}", SanitizerSuffix, TargetPlatform.ToString(), TargetConfiguration.ToString());
+						}
+						else
+						{
+							ExeFileName += string.Format("-{0}-{1}", TargetPlatform.ToString(), TargetConfiguration.ToString());
+						}
+					}
+					else if (HasSanitizerSuffix)
+					{
+						ExeFileName += "-" + SanitizerSuffix;
 					}
 
 					ExeFileName += Platform.GetExeExtension(TargetPlatform);
@@ -724,7 +763,18 @@ namespace Gauntlet
 
 					if (TargetConfiguration != UnrealTargetConfiguration.Development)
 					{
-						ExeFileName += string.Format("-{0}-{1}", TargetPlatform.ToString(), TargetConfiguration.ToString());
+						if (HasSanitizerSuffix)
+						{
+							ExeFileName += string.Format("-{0}-{1}-{2}", SanitizerSuffix, TargetPlatform.ToString(), TargetConfiguration.ToString());
+						}
+						else
+						{
+							ExeFileName += string.Format("-{0}-{1}", TargetPlatform.ToString(), TargetConfiguration.ToString());
+						}
+					}
+					else if (HasSanitizerSuffix)
+					{
+						ExeFileName += "-" + SanitizerSuffix;
 					}
 
 					ExeFileName += Platform.GetExeExtension(TargetPlatform);
