implement Turnkey GetDevices() for Linux.

diff --git a/Engine/Source/Programs/AutomationTool/Linux/LinuxPlatform.Automation.cs b/Engine/Source/Programs/AutomationTool/Linux/LinuxPlatform.Automation.cs
--- a/Engine/Source/Programs/AutomationTool/Linux/LinuxPlatform.Automation.cs
+++ b/Engine/Source/Programs/AutomationTool/Linux/LinuxPlatform.Automation.cs
@@ -6,6 +6,7 @@ using System.Text;
 using System.Text.RegularExpressions;
 using System.IO;
 using System.Diagnostics;
+using System.Runtime.InteropServices;
 using AutomationTool;
 using UnrealBuildTool;
 using EpicGames.Core;
@@ -22,6 +23,21 @@ public abstract class BaseLinuxPlatform : Platform
 	{
 	}
 
+	public override DeviceInfo[] GetDevices()
+	{
+		List<DeviceInfo> Devices = new List<DeviceInfo>();
+
+		if (HostPlatform.Current.HostEditorPlatform == TargetPlatformType)
+		{
+			DeviceInfo LocalMachine = new DeviceInfo(TargetPlatformType, Environment.MachineName, Environment.MachineName,
+				RuntimeInformation.OSDescription, "Computer", true, true);
+
+			Devices.Add(LocalMachine);
+		}
+
+		return Devices.ToArray();
+	}
+
 	public override void GetFilesToDeployOrStage(ProjectParams Params, DeploymentContext SC)
 	{
 		if (SC.bStageCrashReporter)
diff --git a/Engine/Source/Programs/UnrealBuildTool/Platform/Linux/LinuxPlatformSDK.cs b/Engine/Source/Programs/UnrealBuildTool/Platform/Linux/LinuxPlatformSDK.cs
--- a/Engine/Source/Programs/UnrealBuildTool/Platform/Linux/LinuxPlatformSDK.cs
+++ b/Engine/Source/Programs/UnrealBuildTool/Platform/Linux/LinuxPlatformSDK.cs
@@ -64,6 +64,15 @@ namespace UnrealBuildTool
 		{
 			if (StringValue != null)
 			{
+				// if it doesnt start with a v (for an SDK version), assume its a valid version
+				// which will be used for devices. If a messed up toolchain ends up not having a v<num>
+				// 1 will be an invalid range number to pass which will fail the SDK range check
+				if (StringValue[0] != 'v')
+				{
+					OutValue = 1;
+					return true;
+				}
+
 				// Example: v11_clang-5.0.0-centos7
 				string FullVersionPattern = @"^v([0-9]+)_.*$";
 				Match Result = Regex.Match(StringValue, FullVersionPattern);
@@ -77,10 +86,6 @@ namespace UnrealBuildTool
 			return false;
 		}
 
-
-
-
-
 		protected override bool PlatformSupportsAutoSDKs()
 		{
 			return true;
