Added missing change that tags digests as part of a cook

diff --git a/Engine/Source/Programs/AutomationTool/BuildGraph/Tasks/CookTask.cs b/Engine/Source/Programs/AutomationTool/BuildGraph/Tasks/CookTask.cs
--- a/Engine/Source/Programs/AutomationTool/BuildGraph/Tasks/CookTask.cs
+++ b/Engine/Source/Programs/AutomationTool/BuildGraph/Tasks/CookTask.cs
@@ -137,6 +137,13 @@ namespace AutomationTool.Tasks
 					List<FileReference> PackagingFiles = DirectoryReference.EnumerateFiles(PackagingFilesDirectory, "*", System.IO.SearchOption.AllDirectories).ToList();
 					CookedFiles.AddRange(PackagingFiles);
 				}
+
+				DirectoryReference VerseFilesDirectory = DirectoryReference.Combine(ProjectFile.Directory, "Saved", "Verse", Platform);
+				if (DirectoryReference.Exists(VerseFilesDirectory))
+				{
+					List<FileReference> VerseFiles = DirectoryReference.EnumerateFiles(VerseFilesDirectory, "*", System.IO.SearchOption.AllDirectories).ToList();
+					CookedFiles.AddRange(VerseFiles);
+				}
 			}
 			// Apply the optional tag to the build products
 			foreach(string TagName in FindTagNamesFromList(Parameters.Tag))
