Disable cef as it seems to want to talk to an x11 server even though there is none around

diff --git a/Engine/Source/Programs/AutomationTool/Gauntlet/Unreal/Base/Gauntlet.UnrealTestConfiguration.cs b/Engine/Source/Programs/AutomationTool/Gauntlet/Unreal/Base/Gauntlet.UnrealTestConfiguration.cs
--- a/Engine/Source/Programs/AutomationTool/Gauntlet/Unreal/Base/Gauntlet.UnrealTestConfiguration.cs
+++ b/Engine/Source/Programs/AutomationTool/Gauntlet/Unreal/Base/Gauntlet.UnrealTestConfiguration.cs
@@ -876,9 +876,11 @@ namespace Gauntlet
 				// the X11/Wayland display server to be around and use a dummy/offscreen rendering mode
 				//
 				// As well as disable sound as there are no audio devices when running through horde
+				//
+				// Disable cef as it seems to want to talk to an X11 server so unlikely its even working
 				if (AppConfig.Platform == UnrealTargetPlatform.Linux)
 				{
-					AppConfig.CommandLine += " -renderoffscreen -nosound";
+					AppConfig.CommandLine += " -renderoffscreen -nosound -nocef";
 				}
 			}
 
