- Added support in MakeCookedEditor code to kick off a side BuildCookRun to create a base Release when making a CookedEditor as DLC. This will make it easier to test DLC mode because this one script can make the Release and also the DLC in one run.
- Run with -MakeRelease=build,cook,stage (or a subset to skip the step)
- Can also optionally combine the release and DLC into a target directory with -CombineBuilds=<path>
#jria none

diff --git a/Engine/Source/Programs/AutomationTool/CookedEditor/CookedEditor.Automation.cs b/Engine/Source/Programs/AutomationTool/CookedEditor/CookedEditor.Automation.cs
--- a/Engine/Source/Programs/AutomationTool/CookedEditor/CookedEditor.Automation.cs
+++ b/Engine/Source/Programs/AutomationTool/CookedEditor/CookedEditor.Automation.cs
@@ -118,12 +118,16 @@ public class ModifyStageContext
 		// cache some useful properties
 		ProjectDirectory = Params.RawProjectPath.Directory;
 		ProjectName = Params.RawProjectPath.GetFileNameWithoutAnyExtensions();
-		IniPlatformName = ConfigHierarchy.GetIniPlatformName(SC.StageTargetPlatform.PlatformType);
-		bIsDLC = Params.DLCFile != null;
+		IniPlatformName = ConfigHierarchy.GetIniPlatformName(SC.StageTargetPlatform.IniPlatformType);
+		bIsDLC = Params.DLCFile != null && SC.MetadataDir != null; // MetadataDir needs to be set for DLC
+
+		Console.WriteLine("---> eleaeOverrideDir = {0}, MetadataDrir = {1}", Params.BasedOnReleaseVersionPathOverride, SC.MetadataDir);
+
 
 		// cache info for DLC against a release
 		if (Params.BasedOnReleaseVersionPathOverride != null)
 		{
+			// look in Metadata and the override locations
 			ReleaseMetadataLocation = DirectoryReference.Combine(new DirectoryReference(Params.BasedOnReleaseVersionPathOverride), "Metadata");
 		}
 		else
@@ -353,6 +357,10 @@ public class MakeCookedEditor : BuildCommand
 
 	protected ConfigHelper ConfigHelper;
 
+	// used to remember the locations of stating output exactly as calculated by the staging code
+	protected DirectoryReference CookedEditorStageDirectory = null;
+	protected DirectoryReference ReleaseStageDirectory = null;
+
 	public override void ExecuteBuild()
 	{
 		LogInformation("************************* MakeCookedEditor");
@@ -363,12 +371,30 @@ public class MakeCookedEditor : BuildCommand
 		// set up config sections and the like
 		ConfigHelper = new ConfigHelper(BuildHostPlatform.Current.Platform, ProjectFile, bIsCookedCooker);
 
-
 		ProjectParams BuildParams = GetParams();
 
-		LogInformation("Build? {0}", BuildParams.Build);
 
 		Project.Build(this, BuildParams);
+
+		// after the editor is built, if we are building against a release, and if desired, make that release first to make sure we have it to build against
+		ProjectParams ReleaseParams = null;
+
+		string MakeReleaseOptions = ParseParamValue("makerelease", "");
+		if (MakeReleaseOptions != "")
+		{
+			if (string.IsNullOrEmpty(BuildParams.BasedOnReleaseVersion))
+			{
+				throw new AutomationException("-makerelease was specified but the project doesn't have bBuildAgainstRelease set to true");
+			}
+
+			ReleaseParams = GetReleaseParams(BuildParams, MakeReleaseOptions.Split(','));
+			Project.Build(this, ReleaseParams);
+			Project.Cook(ReleaseParams);
+			Project.CopyBuildToStagingDirectory(ReleaseParams);
+
+			FinalizeRelease(ReleaseParams);
+		}
+
 		Project.Cook(BuildParams);
 		Project.CopyBuildToStagingDirectory(BuildParams);
 
@@ -377,8 +403,25 @@ public class MakeCookedEditor : BuildCommand
 
 		Project.Archive(BuildParams);
 		PrintRunTime();
-		Project.Deploy(BuildParams);
+//		Project.Deploy(BuildParams);
 
+		if (ReleaseParams != null)
+		{
+			string CombinedPath = ParseParamValue("CombineBuilds", "");
+			if (CombinedPath != "")
+			{
+				if (CookedEditorStageDirectory == null || ReleaseStageDirectory == null)
+				{
+					LogError("Combining Release and CookedEditor together currently requires that both are staged this run (-stage -makerelease=stage)");
+					return;
+				}
+
+				LogInformation($"Combinging {ReleaseStageDirectory} + {CookedEditorStageDirectory} -> {CombinedPath}");
+				DirectoryReference Combined = new DirectoryReference(CombinedPath);
+				CopyDirectory_NoExceptions(ReleaseStageDirectory.FullName, Combined.FullName, CopyDirectoryOptions.Default);
+				CopyDirectory_NoExceptions(CookedEditorStageDirectory.FullName, Combined.FullName, CopyDirectoryOptions.Merge);
+			}
+		}
 
 	}
 
@@ -587,6 +630,15 @@ public class MakeCookedEditor : BuildCommand
 	{
 	}
 
+	protected virtual void ModifyReleaseParams(ProjectParams ReleaseParams)
+	{
+	}
+
+	protected virtual void FinalizeRelease(ProjectParams ReleaseParams)
+	{
+
+	}
+
 	protected virtual void PreModifyDeploymentContext(ProjectParams Params, DeploymentContext SC)
 	{
 		ModifyStageContext Context = CreateContext(Params, SC);
@@ -616,11 +668,19 @@ public class MakeCookedEditor : BuildCommand
 			SC.FilesToStage.NonUFSDebugFiles.Remove(new StagedFileReference(Path.Combine(Context.ProjectName, "Binaries", "Linux", MainCookedTarget + ".sym")));
 		}
 
+		CookedEditorStageDirectory = SC.StageDirectory;
 	}
 
 	protected virtual void SetupDLCMode(FileReference ProjectFile, out string DLCName, out string ReleaseVersion, out TargetType Type)
 	{
 		bool bBuildAgainstRelease = ConfigHelper.GetBool("bBuildAgainstRelease");
+
+		if (ParseParamValue("MakeRelease", null) != null && !bBuildAgainstRelease)
+		{
+			LogWarning("-makerelease is meant for projects that have bBuildAgainstRelease set. Will force it on, with default settings for [DLCPluginName, ReleaseName, ReleaseTargetType]");
+			bBuildAgainstRelease = true;
+		}
+
 		if (bBuildAgainstRelease)
 		{
 			DLCName = ConfigHelper.GetString("DLCPluginName");
@@ -854,6 +914,20 @@ public class MakeCookedEditor : BuildCommand
 		Context.UFSFilesToStage.RemoveAll(x => UFSIncompatibleExtensions.Contains(x.GetExtension()));
 	}
 
+	protected virtual string GetReleaseTargetName(UnrealTargetPlatform Platform, TargetType ReleaseType)
+	{
+		// make the platform name, like "WindowsClient", or "LinuxGame", of the premade build we are cooking/staging against
+		string IniPlatformName = ConfigHierarchy.GetIniPlatformName(Platform);
+		string ReleaseTargetName = IniPlatformName + (ReleaseType == TargetType.Game ? "" : ReleaseType.ToString());
+
+		return ReleaseTargetName;
+	}
+
+	protected virtual string GetReleaseVersionPath(string ReleaseVersionName, string ReleaseTargetName)
+	{
+		return CommandUtils.CombinePaths(ProjectFile.Directory.FullName, "Releases", ReleaseVersionName, ReleaseTargetName);
+	}
+
 	private ProjectParams GetParams()
 	{
 		// setup DLC defaults, then ask project if it should 
@@ -864,24 +938,23 @@ public class MakeCookedEditor : BuildCommand
 
 		var Params = new ProjectParams
 		(
-			Command: this,
-			RawProjectPath: ProjectFile
+			Command: this
+			, RawProjectPath: ProjectFile
 
-			// standard cookededitor settings
-			//			, Client:false
-			//			, EditorTargets: new ParamList<string>()
-			// , SkipBuildClient: true
 			, NoBootstrapExe: true
-			// , Client: true
 			, DLCName: DLCName
 			, BasedOnReleaseVersion: BasedOnReleaseVersion
 			, DedicatedServer: bIsCookedCooker
 			, NoClient: bIsCookedCooker
 		);
 
+		// cook the cooked editor targetplatorm as the "client"
+		//Params.ClientCookedTargets.Add("CrashReportClientEditor");
+
 		string TargetPlatformType = bIsCookedCooker ? "CookedCooker" : "CookedEditor";
 		string TargetName = ConfigHelper.GetString(bIsCookedCooker ? "CookedCookerTargetName" : "CookedEditorTargetName");
-		UnrealTargetPlatform Platform;
+		UnrealTargetPlatform Platform = bIsCookedCooker ? Params.ServerTargetPlatforms[0].Type : Params.ClientTargetPlatforms[0].Type;
+
 
 		// look to see if ini didn't override target name
 		if (string.IsNullOrEmpty(TargetName))
@@ -890,24 +963,20 @@ public class MakeCookedEditor : BuildCommand
 			TargetName = ProjectFile.GetFileNameWithoutAnyExtensions() + TargetPlatformType;
 		}
 
-		// cook the cooked editor targetplatorm as the "client"
-		//Params.ClientCookedTargets.Add("CrashReportClientEditor");
-
 		// control the server/client taregts
 		Params.ServerCookedTargets.Clear();
 		Params.ClientCookedTargets.Clear();
+		List<TargetPlatformDescriptor> TargetPlatformList = new List<TargetPlatformDescriptor>() { new TargetPlatformDescriptor(Platform, TargetPlatformType) };
 		if (bIsCookedCooker)
 		{
-			Platform = Params.ServerTargetPlatforms[0].Type;
 			Params.EditorTargets.Add(TargetName);
 			Params.ServerCookedTargets.Add(TargetName);
-			Params.ServerTargetPlatforms = new List<TargetPlatformDescriptor>() { new TargetPlatformDescriptor(Platform, TargetPlatformType) };
+			Params.ServerTargetPlatforms = TargetPlatformList;
 		}
 		else
 		{
-			Platform = Params.ClientTargetPlatforms[0].Type;
 			Params.ClientCookedTargets.Add(TargetName);
-			Params.ClientTargetPlatforms = new List<TargetPlatformDescriptor>() { new TargetPlatformDescriptor(Platform, TargetPlatformType) };
+			Params.ClientTargetPlatforms = TargetPlatformList;
 		}
 
 
@@ -918,21 +987,26 @@ public class MakeCookedEditor : BuildCommand
 		// cause the issues
 		Params.AdditionalCookerOptions += " -AllowUnsafeBlueprintCalls";
 
+		// cook and stage into our project, instead of the Engine's plugins
+		DirectoryReference BaseOutputDirectory = DirectoryReference.Combine(ProjectFile.Directory, "Saved", "CookedEditor");
+		string TargetPlatformName = ConfigHierarchy.GetIniPlatformName(Platform) + TargetPlatformList[0].CookFlavor;
+		Params.CookOutputDir = DirectoryReference.Combine(BaseOutputDirectory, "Cooked", TargetPlatformName).FullName;
+		Params.StageDirectoryParam = DirectoryReference.Combine(BaseOutputDirectory, "Staged").FullName;
+
 		// Params.AdditionalCookerOptions += " -NoFilterAssetRegistry";
 
 		// set up cooking against a client, as DLC
 		if (BasedOnReleaseVersion != null)
 		{
-			// make the platform name, like "WindowsClient", or "LinuxGame", of the premade build we are cooking/staging against
-			string IniPlatformName = ConfigHierarchy.GetIniPlatformName(Platform);
-			string ReleaseTargetName = IniPlatformName + (ReleaseType == TargetType.Game ? "NoEditor" : ReleaseType.ToString());
+			// make WindowsClient or LinuxGame, etc
+			string ReleaseTargetName = GetReleaseTargetName(Platform, ReleaseType);
 
 			Params.AdditionalCookerOptions += " -CookAgainstFixedBase";
 			Params.AdditionalCookerOptions += $" -DevelopmentAssetRegistryPlatformOverride={ReleaseTargetName}";
 			Params.AdditionalIoStoreOptions += $" -DevelopmentAssetRegistryPlatformOverride={ReleaseTargetName}";
 
 			// point to where the premade asset registry can be found
-			Params.BasedOnReleaseVersionPathOverride = CommandUtils.CombinePaths(ProjectFile.Directory.FullName, "Releases", BasedOnReleaseVersion, ReleaseTargetName);
+			Params.BasedOnReleaseVersionPathOverride = GetReleaseVersionPath(BasedOnReleaseVersion, ReleaseTargetName);
 
 			Params.DLCOverrideStagedSubDir = "";
 			Params.DLCIncludeEngineContent = true;
@@ -950,6 +1024,40 @@ public class MakeCookedEditor : BuildCommand
 		return Params;
 	}
 
+	private ProjectParams GetReleaseParams(ProjectParams MainParams, string[] Options)
+	{
+		ProjectParams ReleaseParams = new ProjectParams(
+			Command: this
+			, RawProjectPath: MainParams.RawProjectPath
+
+			, Client: MainParams.Client
+			, CreateReleaseVersion: MainParams.BasedOnReleaseVersion
+			// tell the Params that we want cooked data
+			, Build: true
+			, Cook: true
+			, Stage: true
+			// MainParams already builds the editor
+			, SkipBuildEditor: true
+			, NoBootstrapExe: true
+			// if the param to build/cook is specified, then actually build/cook, otherwise assume the Release is already built/cooked
+			, SkipBuildClient: !Options.Contains("build", StringComparer.InvariantCultureIgnoreCase)
+			, SkipCook: !Options.Contains("cook", StringComparer.InvariantCultureIgnoreCase)
+			, SkipStage: !Options.Contains("stage", StringComparer.InvariantCultureIgnoreCase)
+		);
+
+		// if the MainParams override the ReleaseVersion path, use it directly
+		ReleaseParams.BasedOnReleaseVersionPathOverride = MainParams.BasedOnReleaseVersionPathOverride;
+		// if the MainParams have specified a base location to read the release info from, use that as the location to write 
+		// to when creating the Release
+		ReleaseParams.CreateReleaseVersionBasePath = MainParams.BasedOnReleaseVersionBasePath;
+
+		// copy off the staging dir
+		ReleaseParams.PreModifyDeploymentContextCallback = new Action<ProjectParams, DeploymentContext>((ProjectParams P, DeploymentContext SC) => { ReleaseStageDirectory = SC.StageDirectory; });
+
+		ModifyReleaseParams(ReleaseParams);
+
+		return ReleaseParams;
+	}
 
 
 	protected static void GatherTargetDependencies(ProjectParams Params, DeploymentContext SC, ModifyStageContext Context, string ReceiptName)
