- Allow for merging of directories with options to CopyDirectory_NoExceptions

diff --git a/Engine/Source/Programs/AutomationTool/AutomationUtils/CommandUtils.cs b/Engine/Source/Programs/AutomationTool/AutomationUtils/CommandUtils.cs
--- a/Engine/Source/Programs/AutomationTool/AutomationUtils/CommandUtils.cs
+++ b/Engine/Source/Programs/AutomationTool/AutomationUtils/CommandUtils.cs
@@ -1429,6 +1429,18 @@ namespace AutomationTool
 			File.SetLastWriteTimeUtc(Dest.FullName, File.GetLastWriteTimeUtc(Source.FullName));
 		}
 
+		[Flags]
+		public enum CopyDirectoryOptions
+		{
+			None = 0,
+			SuppressLogging = 1,
+			KeepExistingDirectories = 2,
+			KeepExistingFiles = 4,
+
+			Default = None,
+			Merge = KeepExistingFiles | KeepExistingDirectories
+		};
+
 		/// <summary>
 		/// Copies a directory and all of it's contents recursively. Does not throw exceptions.
 		/// </summary>
@@ -1450,11 +1462,26 @@ namespace AutomationTool
 		/// <returns>True if the operation was successful, false otherwise.</returns>
 		public static bool CopyDirectory_NoExceptions(string Source, string Dest, bool bQuiet = false)
 		{
+			return CopyDirectory_NoExceptions(Source, Dest, CopyDirectoryOptions.Default | (bQuiet ? CopyDirectoryOptions.SuppressLogging : CopyDirectoryOptions.None));
+		}
+
+		/// <summary>
+		/// Copies a directory and all of it's contents recursively. Does not throw exceptions.
+		/// </summary>
+		/// <param name="Source"></param>
+		/// <param name="Dest"></param>
+		/// <param name="Options">Options to control logging and overwriting</param>
+		/// <returns>True if the operation was successful, false otherwise.</returns>
+		public static bool CopyDirectory_NoExceptions(string Source, string Dest, CopyDirectoryOptions Options)
+		{
+			bool bQuiet = !Options.HasFlag(CopyDirectoryOptions.SuppressLogging);
+
 			Source = ConvertSeparators(PathSeparator.Default, Source);
 			Dest = ConvertSeparators(PathSeparator.Default, Dest);
 			Dest = Dest.TrimEnd(new char[] { Path.DirectorySeparatorChar, Path.AltDirectorySeparatorChar });
 
-			if (InternalUtils.SafeDirectoryExists(Dest))
+			// if we want to overwrite directories, and it already exists, delete it first
+			if (!Options.HasFlag(CopyDirectoryOptions.KeepExistingDirectories) && InternalUtils.SafeDirectoryExists(Dest))
 			{
 				InternalUtils.SafeDeleteDirectory(Dest, bQuiet);
 				if (InternalUtils.SafeDirectoryExists(Dest, true))
@@ -1473,7 +1500,7 @@ namespace AutomationTool
 				foreach (var SourceSubDirectory in Directory.GetDirectories(Source))
 				{
 					string DestPath = Dest + GetPathSeparatorChar(PathSeparator.Default) + GetLastDirectoryName(SourceSubDirectory + GetPathSeparatorChar(PathSeparator.Default));
-					if (!CopyDirectory_NoExceptions(SourceSubDirectory, DestPath, bQuiet))
+					if (!CopyDirectory_NoExceptions(SourceSubDirectory, DestPath, Options))
 					{
 						return false;
 					}
@@ -1482,6 +1509,11 @@ namespace AutomationTool
 				{
 					int FilenameStart = SourceFile.LastIndexOf(GetPathSeparatorChar(PathSeparator.Default));
 					string DestPath = Dest + SourceFile.Substring(FilenameStart);
+					// if we don't want to overwrite files, then if it already exists, skip it
+					if (Options.HasFlag(CopyDirectoryOptions.KeepExistingFiles) && InternalUtils.SafeFileExists(DestPath, true))
+					{
+						continue;
+					}
 					if (!CopyFile_NoExceptions(SourceFile, DestPath, bQuiet))
 					{
 						return false;
