Fix for LowLevelTests crash shutdown when running tests locally, and using -set:TestArgs="-device=<address>" instead of using a device pool

#rnx

diff --git a/Engine/Source/Programs/AutomationTool/LowLevelTests/RunLowLevelTests.cs b/Engine/Source/Programs/AutomationTool/LowLevelTests/RunLowLevelTests.cs
--- a/Engine/Source/Programs/AutomationTool/LowLevelTests/RunLowLevelTests.cs
+++ b/Engine/Source/Programs/AutomationTool/LowLevelTests/RunLowLevelTests.cs
@@ -104,8 +104,8 @@ namespace LowLevelTests
 				if (!string.IsNullOrEmpty(Options.Device))
 				{
 					(LowLevelTestNode as LowLevelTests)
-						.LowLevelTestsApp
-						.UnrealDeviceReservation
+						.LowLevelTestsApp?
+						.UnrealDeviceReservation?
 						.ReleaseDevices();
 				}
 
