Undo the -asan changes added to Gauntlet.

The Engine currently no longer spits out a binary for Editor (it does for game) with -ASan appended. Something to resolve at a later date
We can still run ASan tests, but now the binary will be UnrealEditor vs UnrealEditor-ASan

diff --git a/Engine/Source/Programs/AutomationTool/Gauntlet/Unreal/BuildSource/Gauntlet.UnrealBuildSource.cs b/Engine/Source/Programs/AutomationTool/Gauntlet/Unreal/BuildSource/Gauntlet.UnrealBuildSource.cs
--- a/Engine/Source/Programs/AutomationTool/Gauntlet/Unreal/BuildSource/Gauntlet.UnrealBuildSource.cs
+++ b/Engine/Source/Programs/AutomationTool/Gauntlet/Unreal/BuildSource/Gauntlet.UnrealBuildSource.cs
@@ -622,28 +622,6 @@ namespace Gauntlet
 			return CommandlineToReturn;
 		}
 
-		private string GetSanitizerSuffix()
-		{
-			if (Globals.Params.ParseParam("asan"))
-			{
-				return "ASan";
-			}
-			else if (Globals.Params.ParseParam("msan"))
-			{
-				return "MSan";
-			}
-			else if (Globals.Params.ParseParam("ubsan"))
-			{
-				return "UBSan";
-			}
-			else if (Globals.Params.ParseParam("tsan"))
-			{
-				return "TSan";
-			}
-
-			return "";
-		}
-
 		/// <summary>
 		/// Given a role, platform, and config, returns the path to the binary for that config. E.g. Binaries\Win64\FooServer-Win64-Shipping.exe
 		/// </summary>
@@ -653,8 +631,6 @@ namespace Gauntlet
 		/// <returns></returns>
 		virtual public string GetRelativeExecutablePath(UnrealTargetRole TargetRole, UnrealTargetPlatform TargetPlatform, UnrealTargetConfiguration TargetConfiguration)
 		{
-			string SanitizerSuffix = GetSanitizerSuffix();
-			bool HasSanitizerSuffix = !string.IsNullOrEmpty(SanitizerSuffix);
 			string ExePath;
 
 			if (TargetRole.UsesEditor() || TargetRole.IsEditor())
@@ -669,29 +645,13 @@ namespace Gauntlet
 				if (EditorExe != null)
 				{
 					ExePath = EditorExe.FullName;
-
-					if (HasSanitizerSuffix)
-					{
-						ExePath += "-" + SanitizerSuffix;
-					}
 				}
 				else
 				{
 					string ExeFileName = "UnrealEditor";
 					if (TargetConfiguration != UnrealTargetConfiguration.Development)
 					{
-						if (HasSanitizerSuffix)
-						{
-							ExeFileName += string.Format("-{0}-{1}-{2}", SanitizerSuffix, TargetPlatform.ToString(), TargetConfiguration.ToString());
-						}
-						else
-						{
-							ExeFileName += string.Format("-{0}-{1}", TargetPlatform.ToString(), TargetConfiguration.ToString());
-						}
-					}
-					else if (HasSanitizerSuffix)
-					{
-						ExeFileName += "-" + SanitizerSuffix;
+						ExeFileName += string.Format("-{0}-{1}", TargetPlatform.ToString(), TargetConfiguration.ToString());
 					}
 
 					ExeFileName += Platform.GetExeExtension(TargetPlatform);
@@ -763,18 +723,7 @@ namespace Gauntlet
 
 					if (TargetConfiguration != UnrealTargetConfiguration.Development)
 					{
-						if (HasSanitizerSuffix)
-						{
-							ExeFileName += string.Format("-{0}-{1}-{2}", SanitizerSuffix, TargetPlatform.ToString(), TargetConfiguration.ToString());
-						}
-						else
-						{
-							ExeFileName += string.Format("-{0}-{1}", TargetPlatform.ToString(), TargetConfiguration.ToString());
-						}
-					}
-					else if (HasSanitizerSuffix)
-					{
-						ExeFileName += "-" + SanitizerSuffix;
+						ExeFileName += string.Format("-{0}-{1}", TargetPlatform.ToString(), TargetConfiguration.ToString());
 					}
 
 					ExeFileName += Platform.GetExeExtension(TargetPlatform);
