Custom config now supports non extensions style platform configs

diff --git a/Engine/Source/Programs/AutomationTool/Scripts/CopyBuildToStagingDirectory.Automation.cs b/Engine/Source/Programs/AutomationTool/Scripts/CopyBuildToStagingDirectory.Automation.cs
--- a/Engine/Source/Programs/AutomationTool/Scripts/CopyBuildToStagingDirectory.Automation.cs
+++ b/Engine/Source/Programs/AutomationTool/Scripts/CopyBuildToStagingDirectory.Automation.cs
@@ -1459,30 +1459,41 @@ namespace AutomationScripts
 			const string CustomPrefix = "custom/";
 			string NormalizedPath = ConfigFile.MakeRelativeTo(ConfigDir).ToLowerInvariant().Replace('\\', '/');
 
-			if (NormalizedPath.StartsWith(CustomPrefix))
+			bool? ProcessCustomConfigPath(ref string Path)
 			{
-				// Strip custom prefix, then check if it matches our specified custom config
-				NormalizedPath = NormalizedPath.Substring(CustomPrefix.Length);
-
-				if (!String.IsNullOrEmpty(SC.CustomConfig))
+				if (Path.StartsWith(CustomPrefix))
 				{
-					string PrefixToStrip = SC.CustomConfig.ToLowerInvariant() + "/";
-					if (NormalizedPath.StartsWith(PrefixToStrip))
+					// Strip custom prefix, then check if it matches our specified custom config
+					Path = Path.Substring(CustomPrefix.Length);
+
+					if (!string.IsNullOrEmpty(SC.CustomConfig))
 					{
-						// Strip custom config path off the front
-						NormalizedPath = NormalizedPath.Substring(PrefixToStrip.Length);
+						string PrefixToStrip = SC.CustomConfig.ToLowerInvariant() + "/";
+						if (Path.StartsWith(PrefixToStrip))
+						{
+							// Strip custom config path off the front
+							Path = Path.Substring(PrefixToStrip.Length);
+							return true;
+						}
+						else
+						{
+							// If custom config is specified, we want to ignore others automatically
+							return false;
+						}
 					}
 					else
 					{
-						// If custom config is specified, we want to ignore others automatically
+						// Ignore all custom config directories
 						return false;
 					}
 				}
-				else
-				{
-					// Ignore all custom config directories
-					return false;
-				}
+
+				return null;
+			}
+
+			if(ProcessCustomConfigPath(ref NormalizedPath) == false)
+			{
+				return false;
 			}
 
 			int DirectoryIdx = NormalizedPath.IndexOf('/');
@@ -1571,10 +1582,36 @@ namespace AutomationScripts
 					return false;
 				}
 
-				string PlatformPrefix = String.Format("{0}/{0}", NormalizedPath.Substring(0, DirectoryIdx));
+				int INIStartIdx = NormalizedPath.LastIndexOf('/');
+				string PlatformDirectoryName = NormalizedPath.Substring(0, DirectoryIdx);
+
+				// Test for a platform file under a custom config.
+				// Custom config files could have special setups per platform.  For non extension platforms
+				// this will take the form of:
+				// PlatformName/custom/CustomConfig/PlatformName{ini type}.ini
+				if (INIStartIdx != DirectoryIdx)
+				{
+					string PlatformCustomPath = NormalizedPath.Substring(DirectoryIdx + 1);
+					{
+						// Test for a custom path under the platform name
+						bool? CustomPlatformTest = ProcessCustomConfigPath(ref PlatformCustomPath);
+						if (CustomPlatformTest == false)
+						{
+							return false;
+						}
+						else if (CustomPlatformTest == true)
+						{
+							// Valid custom path, strip out the customization parts and test
+							NormalizedPath = $"{PlatformDirectoryName}/{PlatformCustomPath}";
+						}
+					}
+				}
+
+				string PlatformPrefix = $"{PlatformDirectoryName}/{PlatformDirectoryName}";
 				if (NormalizedPath.StartsWith(PlatformPrefix))
 				{
-					return ShouldStageConfigSuffix(SC, ConfigFile, NormalizedPath.Substring(PlatformPrefix.Length));
+					string subPath = NormalizedPath.Substring(PlatformPrefix.Length);
+					return ShouldStageConfigSuffix(SC, ConfigFile, subPath);
 				}
 
 				string PlatformBasePrefix = String.Format("{0}/base{0}", NormalizedPath.Substring(0, DirectoryIdx));
