Add support for per platform FilterPlugin inis for plugins which was originally in a code review for CL 18874143.

#rnx

diff --git a/Engine/Plugins/Bridge/Config/FilterPlugin.ini b/Engine/Plugins/Bridge/Config/FilterPlugin.ini
--- a/Engine/Plugins/Bridge/Config/FilterPlugin.ini
+++ b/Engine/Plugins/Bridge/Config/FilterPlugin.ini
@@ -1,8 +1,2 @@
 [FilterPlugin]
-; This section lists additional files which will be packaged along with your plugin. Paths should be listed relative to the root plugin directory, and
-; may include "...", "*", and "?" wildcards to match directories, files, and individual characters respectively.
-;
-; Examples:
-;    /README.txt
-;    /Extras/...
-;    /Binaries/ThirdParty/*.dll
+/ThirdParty/megascans/...
\ No newline at end of file
diff --git a/Engine/Plugins/Bridge/Config/FilterPluginMac.ini b/Engine/Plugins/Bridge/Config/FilterPluginMac.ini
new file mode 100644
--- /dev/null
+++ b/Engine/Plugins/Bridge/Config/FilterPluginMac.ini
@@ -0,0 +1,2 @@
+[FilterPlugin]
+/ThirdParty/Mac/...
\ No newline at end of file
diff --git a/Engine/Plugins/Bridge/Config/FilterPluginWin64.ini b/Engine/Plugins/Bridge/Config/FilterPluginWin64.ini
new file mode 100644
--- /dev/null
+++ b/Engine/Plugins/Bridge/Config/FilterPluginWin64.ini
@@ -0,0 +1,2 @@
+[FilterPlugin]
+/ThirdParty/Win/...
\ No newline at end of file
diff --git a/Engine/Source/Programs/AutomationTool/Scripts/BuildPluginCommand.Automation.cs b/Engine/Source/Programs/AutomationTool/Scripts/BuildPluginCommand.Automation.cs
--- a/Engine/Source/Programs/AutomationTool/Scripts/BuildPluginCommand.Automation.cs
+++ b/Engine/Source/Programs/AutomationTool/Scripts/BuildPluginCommand.Automation.cs
@@ -17,14 +17,21 @@ using UnrealBuildBase;
 [Help("TargetPlatforms", "Specify a list of target platforms to build, separated by '+' characters (eg. -TargetPlatforms=Win32+Win64). Default is all the Rocket target platforms.")]
 [Help("Package", "The path which the build artifacts should be packaged to, ready for distribution.")]
 [Help("StrictIncludes", "Disables precompiled headers and unity build in order to check all source files have self-contained headers.")]
+[Help("EngineDir=<RootDirectory>", "Root Directory of the engine that will be used to build plugin(s) (optional)")]
 [Help("Unversioned", "Do not embed the current engine version into the descriptor")]
-class BuildPlugin : BuildCommand
+public sealed class BuildPlugin : BuildCommand
 {
 	const string AndroidArchitectures = "armv7+arm64";
 	const string HoloLensArchitecture = "arm64+x64";
+	string UBTExeRelativePath = @"Engine/Binaries/DotNET/UnrealBuildTool/UnrealBuildTool" + (RuntimePlatform.IsWindows ? ".exe" : "");
+	string UBTExe;
 
 	public override void ExecuteBuild()
 	{
+		// See if an engine dir was specified, fall back on default if not.
+		DirectoryReference EngineDirParam = ParseOptionalDirectoryReferenceParam("EngineDir");
+		UBTExe = EngineDirParam == null ? UnrealBuild.GetUBTExecutable() : CommandUtils.CombinePaths(EngineDirParam.ToString(), UBTExeRelativePath);
+
 		// Get the plugin filename
 		string PluginParam = ParseParamValue("Plugin");
 		if(PluginParam == null)
@@ -106,10 +113,10 @@ class BuildPlugin : BuildCommand
 		// Compile the plugin for all the target platforms
 		List<UnrealTargetPlatform> HostPlatforms = ParseParam("NoHostPlatform")? new List<UnrealTargetPlatform>() : new List<UnrealTargetPlatform> { BuildHostPlatform.Current.Platform };
 		List<UnrealTargetPlatform> TargetPlatforms = GetTargetPlatforms(this, BuildHostPlatform.Current.Platform);
-		FileReference[] BuildProducts = CompilePlugin(HostProjectFile, HostProjectPluginFile, Plugin, HostPlatforms, TargetPlatforms, AdditionalArgs.ToString());
+		FileReference[] BuildProducts = CompilePlugin(UBTExe, HostProjectFile, HostProjectPluginFile, Plugin, HostPlatforms, TargetPlatforms, AdditionalArgs.ToString());
 
 		// Package up the final plugin data
-		PackagePlugin(HostProjectPluginFile, BuildProducts, PackageDir, ParseParam("unversioned"));
+		PackagePlugin(HostProjectPluginFile, BuildProducts, PackageDir, ParseParam("unversioned"), TargetPlatforms);
 
 		// Remove the host project
 		if(!ParseParam("NoDeleteHostProject"))
@@ -135,7 +142,7 @@ class BuildPlugin : BuildCommand
 		return FileReference.Combine(HostProjectPluginDir, PluginFile.GetFileName());
 	}
 
-	FileReference[] CompilePlugin(FileReference HostProjectFile, FileReference HostProjectPluginFile, PluginDescriptor Plugin, List<UnrealTargetPlatform> HostPlatforms, List<UnrealTargetPlatform> TargetPlatforms, string AdditionalArgs)
+	public static FileReference[] CompilePlugin(string UBTExe, FileReference HostProjectFile, FileReference HostProjectPluginFile, PluginDescriptor Plugin, List<UnrealTargetPlatform> HostPlatforms, List<UnrealTargetPlatform> TargetPlatforms, string AdditionalArgs = "")
 	{
 		List<FileReference> ManifestFileNames = new List<FileReference>();
 
@@ -147,9 +154,9 @@ class BuildPlugin : BuildCommand
 			{
 				if (Plugin.SupportedPrograms != null && Plugin.SupportedPrograms.Contains("UnrealHeaderTool"))
 				{
-					CompilePluginWithUBT(HostProjectFile, HostProjectPluginFile, Plugin, "UnrealHeaderTool", TargetType.Program, HostPlatform, UnrealTargetConfiguration.Development, ManifestFileNames, String.Format("{0} -plugin={1}", AdditionalArgs, CommandUtils.MakePathSafeToUseWithCommandLine(HostProjectPluginFile.FullName)));
+					CompilePluginWithUBT(UBTExe, HostProjectFile, HostProjectPluginFile, Plugin, "UnrealHeaderTool", TargetType.Program, HostPlatform, UnrealTargetConfiguration.Development, ManifestFileNames, String.Format("{0} -plugin={1}", AdditionalArgs, CommandUtils.MakePathSafeToUseWithCommandLine(HostProjectPluginFile.FullName)));
 				}
-				CompilePluginWithUBT(HostProjectFile, HostProjectPluginFile, Plugin, "UnrealEditor", TargetType.Editor, HostPlatform, UnrealTargetConfiguration.Development, ManifestFileNames, AdditionalArgs);
+				CompilePluginWithUBT(UBTExe, HostProjectFile, HostProjectPluginFile, Plugin, "UnrealEditor", TargetType.Editor, HostPlatform, UnrealTargetConfiguration.Development, ManifestFileNames, AdditionalArgs);
 			}
 		}
 
@@ -161,8 +168,8 @@ class BuildPlugin : BuildCommand
 			foreach (UnrealTargetPlatform TargetPlatform in SupportedTargetPlatforms)
 			{
 				string AdditionalTargetArgs = AdditionalArgs;
-				CompilePluginWithUBT(HostProjectFile, HostProjectPluginFile, Plugin, "UnrealGame", TargetType.Game, TargetPlatform, UnrealTargetConfiguration.Development, ManifestFileNames, AdditionalTargetArgs);
-				CompilePluginWithUBT(HostProjectFile, HostProjectPluginFile, Plugin, "UnrealGame", TargetType.Game, TargetPlatform, UnrealTargetConfiguration.Shipping, ManifestFileNames, AdditionalTargetArgs);
+				CompilePluginWithUBT(UBTExe, HostProjectFile, HostProjectPluginFile, Plugin, "UnrealGame", TargetType.Game, TargetPlatform, UnrealTargetConfiguration.Development, ManifestFileNames, AdditionalTargetArgs);
+				CompilePluginWithUBT(UBTExe, HostProjectFile, HostProjectPluginFile, Plugin, "UnrealGame", TargetType.Game, TargetPlatform, UnrealTargetConfiguration.Shipping, ManifestFileNames, AdditionalTargetArgs);
 			}
 		}
 
@@ -176,7 +183,7 @@ class BuildPlugin : BuildCommand
 		return BuildProducts.ToArray();
 	}
 
-	void CompilePluginWithUBT(FileReference HostProjectFile, FileReference HostProjectPluginFile, PluginDescriptor Plugin, string TargetName, TargetType TargetType, UnrealTargetPlatform Platform, UnrealTargetConfiguration Configuration, List<FileReference> ManifestFileNames, string InAdditionalArgs)
+	static void CompilePluginWithUBT(string UBTExe, FileReference HostProjectFile, FileReference HostProjectPluginFile, PluginDescriptor Plugin, string TargetName, TargetType TargetType, UnrealTargetPlatform Platform, UnrealTargetConfiguration Configuration, List<FileReference> ManifestFileNames, string InAdditionalArgs)
 	{
 		// Find a list of modules that need to be built for this plugin
 		bool bCompilePlatform = false;
@@ -212,7 +219,7 @@ class BuildPlugin : BuildCommand
 					{
 						Arguments += InAdditionalArgs;
 					}
-					CommandUtils.RunUBT(CmdEnv, UnrealBuild.GetUBTExecutable(), HostProjectFile, TargetName, Platform, Configuration, Arguments);
+					CommandUtils.RunUBT(CmdEnv, UBTExe, HostProjectFile, TargetName, Platform, Configuration, Arguments);
 				}
 			}
 			else
@@ -231,17 +238,17 @@ class BuildPlugin : BuildCommand
 					Arguments += InAdditionalArgs;
 				}
 
-				CommandUtils.RunUBT(CmdEnv, UnrealBuild.GetUBTExecutable(), HostProjectFile, TargetName, Platform, Configuration, Arguments);
+				CommandUtils.RunUBT(CmdEnv, UBTExe, HostProjectFile, TargetName, Platform, Configuration, Arguments);
 			}
 		}
 	}
 
-	static void PackagePlugin(FileReference SourcePluginFile, IEnumerable<FileReference> BuildProducts, DirectoryReference TargetDir, bool bUnversioned)
+	public static void PackagePlugin(FileReference SourcePluginFile, IEnumerable<FileReference> BuildProducts, DirectoryReference TargetDir, bool bUnversioned, IEnumerable<UnrealTargetPlatform> TargetPlatforms)
 	{
 		DirectoryReference SourcePluginDir = SourcePluginFile.Directory;
 
 		// Copy all the files to the output directory
-		FileReference[] SourceFiles = FilterPluginFiles(SourcePluginFile, BuildProducts).ToArray();
+		FileReference[] SourceFiles = FilterPluginFiles(SourcePluginFile, BuildProducts, TargetPlatforms).ToArray();
 		foreach(FileReference SourceFile in SourceFiles)
 		{
 			FileReference TargetFile = FileReference.Combine(TargetDir, SourceFile.MakeRelativeTo(SourcePluginDir));
@@ -265,7 +272,16 @@ class BuildPlugin : BuildCommand
 		NewDescriptor.Save(TargetPluginFile.FullName);
 	}
 
-	static IEnumerable<FileReference> FilterPluginFiles(FileReference PluginFile, IEnumerable<FileReference> BuildProducts)
+	static void AddRulesFromFileToFilter(FileFilter Filter, FileReference FilterFile)
+	{
+		if (FileReference.Exists(FilterFile))
+		{
+			CommandUtils.LogInformation("Reading filter rules from {0}", FilterFile);
+			Filter.ReadRulesFromFile(FilterFile, "FilterPlugin");
+		}
+	}
+
+	static IEnumerable<FileReference> FilterPluginFiles(FileReference PluginFile, IEnumerable<FileReference> BuildProducts, IEnumerable<UnrealTargetPlatform> TargetPlatforms)
 	{
 		// Set up the default filter
 		FileFilter Filter = new FileFilter();
@@ -278,12 +294,14 @@ class BuildPlugin : BuildCommand
 		Filter.Include("/Shaders/...");
 		Filter.Include("/Source/...");
 
-		// Add custom rules for each platform
-		FileReference FilterFile = FileReference.Combine(PluginFile.Directory, "Config", "FilterPlugin.ini");
-		if(FileReference.Exists(FilterFile))
+		// Add custom rules for all platforms
+		AddRulesFromFileToFilter(Filter, FileReference.Combine(PluginFile.Directory, "Config", "FilterPlugin.ini"));
+
+		// Add custom rules for targeted platforms.
+		foreach (UnrealTargetPlatform Platform in TargetPlatforms)
 		{
-			CommandUtils.LogInformation("Reading filter rules from {0}", FilterFile);
-			Filter.ReadRulesFromFile(FilterFile, "FilterPlugin");
+			Log.TraceInformation("Loading FilterPlugin{0}.ini", Platform);
+			AddRulesFromFileToFilter(Filter, FileReference.Combine(PluginFile.Directory, "Config", $"FilterPlugin{Platform}.ini"));
 		}
 
 		// Apply the standard exclusion rules
diff --git a/Engine/Source/Programs/Shared/EpicGames.BuildGraph/BgSchema.cs b/Engine/Source/Programs/Shared/EpicGames.BuildGraph/BgSchema.cs
--- a/Engine/Source/Programs/Shared/EpicGames.BuildGraph/BgSchema.cs
+++ b/Engine/Source/Programs/Shared/EpicGames.BuildGraph/BgSchema.cs
@@ -831,6 +831,7 @@ namespace EpicGames.BuildGraph
 			MacroChoice.Items.Add(CreateSchemaElement("Warning", ScriptSchemaStandardType.Warning));
 			MacroChoice.Items.Add(CreateSchemaElement("Error", ScriptSchemaStandardType.Error));
 			MacroChoice.Items.Add(CreateSchemaElement("Expand", ScriptSchemaStandardType.Expand));
+			MacroChoice.Items.Add(CreateSchemaElement("Label", ScriptSchemaStandardType.Label));
 
 			// Agent scope
 			MacroChoice.Items.Add(CreateSchemaElement("Node", ScriptSchemaStandardType.Node));
