Enable Content on Demand for PC

diff --git a/Engine/Source/Programs/AutomationTool/Scripts/CopyBuildToStagingDirectory.Automation.cs b/Engine/Source/Programs/AutomationTool/Scripts/CopyBuildToStagingDirectory.Automation.cs
--- a/Engine/Source/Programs/AutomationTool/Scripts/CopyBuildToStagingDirectory.Automation.cs
+++ b/Engine/Source/Programs/AutomationTool/Scripts/CopyBuildToStagingDirectory.Automation.cs
@@ -1934,6 +1934,8 @@ namespace AutomationScripts
 
 			ConfigHierarchy PakRulesConfig = ConfigCache.ReadHierarchy(ConfigHierarchyType.PakFileRules, DirectoryReference.FromFile(Params.RawProjectPath), SC.StageTargetPlatform.IniPlatformType, SC.CustomConfig);
 
+			string IniPlatformName = ConfigHierarchy.GetIniPlatformName(SC.StageTargetPlatform.IniPlatformType);
+
 			bool bChunkedBuild = SC.PlatformUsesChunkManifests && DoesChunkPakManifestExist(Params, SC);
 
 			List<PakFileRules> RulesList = new List<PakFileRules>();
@@ -1972,12 +1974,7 @@ namespace AutomationScripts
 					// Check platform string
 					foreach (string Platform in PlatformStrings)
 					{
-						if (SC.StageTargetPlatform.PlatformType.ToString().Equals(Platform, StringComparison.OrdinalIgnoreCase))
-						{
-							bMatches = true;
-							break;
-						}
-						else if (SC.StageTargetPlatform.IniPlatformType.ToString().Equals(Platform, StringComparison.OrdinalIgnoreCase))
+						if (IniPlatformName.Equals(Platform, StringComparison.OrdinalIgnoreCase))
 						{
 							bMatches = true;
 							break;
@@ -1986,7 +1983,7 @@ namespace AutomationScripts
 
 					if (!bMatches)
 					{
-						//LogInformation("No matching platform for PakFileRules for Section {0} : {1}, {2}", SectionName, SC.StageTargetPlatform.PlatformType.ToString(), SC.StageTargetPlatform.IniPlatformType.ToString());
+						LogInformation("No matching platform for PakFileRules for Section {0} : {1}", SectionName, IniPlatformName);
 						continue;
 					}
 				}
