Android file server for deployment and installs

#android

diff --git a/Engine/Build/InstalledEngineFilters.xml b/Engine/Build/InstalledEngineFilters.xml
--- a/Engine/Build/InstalledEngineFilters.xml
+++ b/Engine/Build/InstalledEngineFilters.xml
@@ -424,6 +424,7 @@
 	<!-- Android -->
 	<Property Name="CopyAndroidFilter">
 		<!-- Build Files -->
+		Engine/Binaries/DotNET/Android/...
 		Engine/Build/Android/...
 		Engine/Extras/Android/...
 		Engine/Source/ThirdParty/Android/...
diff --git a/Engine/Plugins/Runtime/AndroidFileServer/AndroidFileServer.uplugin b/Engine/Plugins/Runtime/AndroidFileServer/AndroidFileServer.uplugin
new file mode 100644
--- /dev/null
+++ b/Engine/Plugins/Runtime/AndroidFileServer/AndroidFileServer.uplugin
@@ -0,0 +1,31 @@
+{
+	"FileVersion": 3,
+	"Version": 1,
+	"VersionName": "1.0",
+	"FriendlyName": "AndroidFileServer",
+	"Description": "Adds support for remote file management to Android projects.",
+	"Category": "Android",
+	"CreatedBy": "Epic Games, Inc.",
+	"CreatedByURL": "http://epicgames.com",
+	"DocsURL": "",
+	"MarketplaceURL": "",
+	"SupportURL": "",
+	"CanContainContent": false,
+	"IsBetaVersion": false,
+	"EnabledByDefault": true,
+	"Installed": false,
+	"Modules": [
+		{
+			"Name": "AndroidFileServer",
+			"Type": "Runtime",
+			"LoadingPhase": "PreLoadingScreen",
+                        "WhitelistPlatforms" : [ "Android", "Win64", "Mac", "Linux" ]
+                },
+                {
+                        "Name" : "AndroidFileServerEditor",
+                        "Type" : "Editor",
+                        "LoadingPhase" : "PostEngineInit",
+                        "WhitelistPlatforms" : [ "Win64", "Mac", "Linux" ]
+                }
+        ]
+}
diff --git a/Engine/Plugins/Runtime/AndroidFileServer/Source/AFSStub/app/src/main/java/com/epicgames/ue4/GameActivity.java b/Engine/Plugins/Runtime/AndroidFileServer/Source/AFSStub/app/src/main/java/com/epicgames/ue4/GameActivity.java
new file mode 100644
--- /dev/null
+++ b/Engine/Plugins/Runtime/AndroidFileServer/Source/AFSStub/app/src/main/java/com/epicgames/ue4/GameActivity.java
@@ -0,0 +1,134 @@
+// Copyright Epic Games, Inc. All Rights Reserved.
+
+package com.epicgames.ue4;
+
+import android.app.Activity;
+import android.content.pm.ApplicationInfo;
+import android.content.pm.PackageManager;
+import android.os.Bundle;
+import android.util.Log;
+import android.view.ViewGroup.LayoutParams;
+import android.view.ViewGroup.MarginLayoutParams;
+import android.widget.LinearLayout;
+
+import com.epicgames.unreal.RemoteFileManager;
+
+import java.io.BufferedReader;
+import java.io.IOException;
+import java.io.InputStreamReader;
+import java.io.PrintWriter;
+import java.net.ServerSocket;
+import java.net.Socket;
+import java.util.Set;
+
+import static android.os.Environment.getExternalStorageDirectory;
+
+public class GameActivity extends Activity {
+    private static final int SERVER_PORT = 57099;
+
+	public static boolean AndroidFileServer_Verify(String Token)
+	{
+		return true;
+	}
+
+	@Override
+	protected void onCreate(Bundle savedInstanceState) {
+		super.onCreate(savedInstanceState);
+
+		MarginLayoutParams params = new MarginLayoutParams(LayoutParams.WRAP_CONTENT, LayoutParams.WRAP_CONTENT);
+		params.setMargins(0, 0, 0, 0);
+		LinearLayout activityLayout = new LinearLayout(this);
+		setContentView(activityLayout, params);
+
+		int port = SERVER_PORT;
+
+		// make sure UnrealFileServer not already running
+		boolean bUSBRunning = false;
+		boolean bWiFiRunning = false;
+				
+		Set<Thread> threads = Thread.getAllStackTraces().keySet();
+		for (Thread thread : threads) {
+			String name = thread.getName();
+			if (name.equals("UnrealAFS-USB"))
+			{
+				if (thread.isAlive())
+				{
+					bUSBRunning = true;
+				}
+			}
+			else if (name.equals("UnrealAFS-WiFi"))
+			{
+				if (thread.isAlive())
+				{
+					bWiFiRunning = true;
+				}
+			}
+		}
+
+		String packageName = getPackageName();
+		PackageManager pm = getPackageManager();
+
+		int versionCode = 0;
+		try
+		{
+			versionCode = pm.getPackageInfo(packageName, 0).versionCode;
+		}
+		catch (Exception e)
+		{
+			// if the above failed, then, we can't use obbs
+		}
+
+		String EngineVersion = "4.27.3";
+		String ProjectName = packageName;
+		ProjectName = ProjectName.substring(ProjectName.lastIndexOf('.') + 1);
+		boolean IsShipping = false;
+		boolean UseExternalFilesDir = false;
+		boolean PublicLogFiles = false;
+		try {
+			ApplicationInfo ai = pm.getApplicationInfo(packageName, PackageManager.GET_META_DATA);
+			Bundle bundle = ai.metaData;
+
+			if (bundle.containsKey("com.epicgames.ue4.GameActivity.EngineVersion"))
+			{
+				EngineVersion = bundle.getString("com.epicgames.ue4.GameActivity.EngineVersion");
+			}
+			if (bundle.containsKey("com.epicgames.ue4.GameActivity.ProjectName"))
+			{
+				ProjectName = bundle.getString("com.epicgames.ue4.GameActivity.ProjectName");
+			}
+			if (bundle.containsKey("com.epicgames.ue4.GameActivity.BuildConfiguration"))
+			{
+				String Configuration = bundle.getString("com.epicgames.ue4.GameActivity.BuildConfiguration");
+				IsShipping = Configuration.contains("Shipping");
+			}
+			if (bundle.containsKey("com.epicgames.ue4.GameActivity.bUseExternalFilesDir"))
+			{
+				UseExternalFilesDir = bundle.getBoolean("com.epicgames.ue4.GameActivity.bUseExternalFilesDir");
+			}
+			if (bundle.containsKey("com.epicgames.ue4.GameActivity.bPublicLogFiles"))
+			{
+				PublicLogFiles = bundle.getBoolean("com.epicgames.ue4.GameActivity.bPublicLogFiles");
+			}
+		}
+		catch (PackageManager.NameNotFoundException | NullPointerException e)
+		{
+			Log.e("UE","Error when accessing application metadata", e);
+		}
+
+		String internal = getFilesDir().getAbsolutePath();
+		String external = getExternalFilesDir(null).getAbsolutePath();
+		String storage = getExternalStorageDirectory().getAbsolutePath();
+		String obbdir = getObbDir().getAbsolutePath();
+
+		if (!bUSBRunning)
+		{
+			RemoteFileManager fileManager = new RemoteFileManager(true, port, internal, external, storage, obbdir, packageName, versionCode, ProjectName, EngineVersion, IsShipping, PublicLogFiles);
+			new Thread(fileManager, "UnrealAFS-USB").start();
+		}
+		if (!bWiFiRunning)
+		{
+			RemoteFileManager fileManager = new RemoteFileManager(false, port, internal, external, storage, obbdir, packageName, versionCode, ProjectName, EngineVersion, IsShipping, PublicLogFiles);
+			new Thread(fileManager, "UnrealAFS-WiFi").start();
+		}
+    }
+}
\ No newline at end of file
diff --git a/Engine/Plugins/Runtime/AndroidFileServer/Source/AFSStub/app/src/main/java/com/epicgames/ue4/GameApplication.java b/Engine/Plugins/Runtime/AndroidFileServer/Source/AFSStub/app/src/main/java/com/epicgames/ue4/GameApplication.java
new file mode 100644
--- /dev/null
+++ b/Engine/Plugins/Runtime/AndroidFileServer/Source/AFSStub/app/src/main/java/com/epicgames/ue4/GameApplication.java
@@ -0,0 +1,70 @@
+package com.epicgames.ue4;
+
+import android.app.Application;
+import android.content.Context;
+import android.content.res.Configuration;
+import androidx.lifecycle.Lifecycle;
+import androidx.lifecycle.LifecycleObserver;
+import androidx.lifecycle.OnLifecycleEvent;
+import androidx.lifecycle.ProcessLifecycleOwner;
+
+
+public class GameApplication extends Application implements LifecycleObserver {
+	private static final Logger Log = new Logger("UE4", "GameApp");
+
+	private static boolean isForeground = false;
+
+	@Override
+	public void onCreate() {
+		super.onCreate();
+
+
+		ProcessLifecycleOwner.get().getLifecycle().addObserver(this);
+
+	}
+
+	@Override
+	public void attachBaseContext(Context base) {
+		super.attachBaseContext(base);
+
+	}
+
+	@Override
+	public void onLowMemory() {
+		super.onLowMemory();
+
+	}
+
+	@Override
+	public void onTrimMemory(int level) {
+		super.onTrimMemory(level);
+
+	}
+
+	@Override
+	public void onConfigurationChanged (Configuration newConfig) {
+		super.onConfigurationChanged(newConfig);
+		
+	}
+
+	@OnLifecycleEvent(Lifecycle.Event.ON_START)
+	void onEnterForeground() {
+		Log.verbose("App in foreground");
+		isForeground = true;
+	}
+
+	@OnLifecycleEvent(Lifecycle.Event.ON_STOP)
+	void onEnterBackground() {
+		Log.verbose("App in background");
+		isForeground = false;
+	}
+
+	@SuppressWarnings("unused")
+	public static boolean isAppInForeground() {
+		return isForeground;
+	}
+
+	public static boolean isAppInBackground() {
+		return !isForeground;
+	}
+}
diff --git a/Engine/Plugins/Runtime/AndroidFileServer/Source/AFSStub/app/src/main/java/com/epicgames/ue4/Logger.java b/Engine/Plugins/Runtime/AndroidFileServer/Source/AFSStub/app/src/main/java/com/epicgames/ue4/Logger.java
new file mode 100644
--- /dev/null
+++ b/Engine/Plugins/Runtime/AndroidFileServer/Source/AFSStub/app/src/main/java/com/epicgames/ue4/Logger.java
@@ -0,0 +1,150 @@
+package com.epicgames.ue4;
+
+import android.text.TextUtils;
+import android.util.Log;
+
+public class Logger
+{
+	public interface ILoggerCallback
+	{
+		void LoggerCallback(String Level, String Tag, String Message);
+	}
+
+	private static ILoggerCallback mCallback = null;
+	private final String mTag;
+	private final String mSecondaryTag;
+	private boolean bHasSecondaryTag;
+	private String mFormattedTag;
+	private String mFormattedMessageTag;
+	
+	private static boolean bAllowLogging			= true;
+	@SuppressWarnings({"FieldCanBeLocal", "unused"})
+	private static boolean bAllowExceptionLogging	= true;
+	private static boolean bPrependSecondaryTag	= false;
+	private static boolean bPrependChanged	= false;
+	
+	@SuppressWarnings({"unused", "WeakerAccess"})
+	public static void RegisterCallback(ILoggerCallback callback)
+	{
+		mCallback = callback;
+	}
+
+	@SuppressWarnings("WeakerAccess")
+	public static void SuppressLogs ()
+	{
+		bAllowLogging = bAllowExceptionLogging = false;
+	}
+
+	@SuppressWarnings({"unused", "WeakerAccess"})
+	public static void prependSecondaryTag(boolean prependSecondaryTag) {
+		bPrependSecondaryTag = prependSecondaryTag;
+		bPrependChanged = true;
+	}
+	
+	public Logger(String Tag)
+	{
+		this(Tag, "");
+	}
+
+	public Logger(String Tag, String secondaryTag)
+	{
+		mTag = Tag;
+		if(secondaryTag == null) {
+			mSecondaryTag = "";
+		} else {
+			mSecondaryTag = secondaryTag;
+		}
+		bHasSecondaryTag = !TextUtils.isEmpty(mSecondaryTag);
+	}
+
+	public void verbose(String Message)
+	{
+		if (bAllowLogging)
+		{
+			Log.v(getFormattedTag(), getFormattedMessage(Message));
+		}
+		if (mCallback != null)
+		{
+			mCallback.LoggerCallback("V/", getFormattedTag(), getFormattedMessage(Message));
+		}
+	}
+
+	public void debug(String Message)
+	{
+		if (bAllowLogging)
+		{
+			Log.d(getFormattedTag(), getFormattedMessage(Message));
+		}
+		if (mCallback != null)
+		{
+			mCallback.LoggerCallback("D/", getFormattedTag(), getFormattedMessage(Message));
+		}
+	}
+	
+	public void warn(String Message)
+	{
+		if (bAllowLogging)
+		{
+			Log.w(getFormattedTag(), getFormattedMessage(Message));
+		}
+		if (mCallback != null)
+		{
+			mCallback.LoggerCallback("W/", getFormattedTag(), getFormattedMessage(Message));
+		}
+	}
+	
+	public void error(String Message)
+	{
+		if (bAllowLogging)
+		{
+			Log.e(getFormattedTag(), getFormattedMessage(Message));
+		}
+		if (mCallback != null)
+		{
+			mCallback.LoggerCallback("E/", getFormattedTag(), getFormattedMessage(Message));
+		}
+	}
+
+	public void error(String Message, Throwable Throwable)
+	{
+		if (bAllowLogging)
+		{
+			Log.e(getFormattedTag(), getFormattedMessage(Message), Throwable);
+		}
+		if (mCallback != null)
+		{
+			mCallback.LoggerCallback("E/", getFormattedTag(), getFormattedMessage(Message));
+		}
+	}
+	
+	private String getFormattedTag()
+	{
+		if(mFormattedTag == null || bPrependChanged)
+		{
+			bPrependChanged = false;
+			if (bHasSecondaryTag && bPrependSecondaryTag)
+			{
+				mFormattedTag = mTag + "-" + mSecondaryTag;
+			}
+			else
+			{
+				mFormattedTag = mTag;
+			}
+		}
+		return mFormattedTag;
+	}
+
+	private String getFormattedMessage(String message)
+	{
+		if (bHasSecondaryTag && !bPrependSecondaryTag)
+		{
+			if(mFormattedMessageTag == null || bPrependChanged)
+			{
+				bPrependChanged = false;
+				mFormattedMessageTag = "[" + mSecondaryTag + "] ";
+			}
+			return mFormattedMessageTag + message;
+		}
+		return message;
+	}
+}
diff --git a/Engine/Plugins/Runtime/AndroidFileServer/Source/AFSStub/app/src/main/java/com/epicgames/ue4/SplashActivity.java b/Engine/Plugins/Runtime/AndroidFileServer/Source/AFSStub/app/src/main/java/com/epicgames/ue4/SplashActivity.java
new file mode 100644
--- /dev/null
+++ b/Engine/Plugins/Runtime/AndroidFileServer/Source/AFSStub/app/src/main/java/com/epicgames/ue4/SplashActivity.java
@@ -0,0 +1,129 @@
+// Copyright Epic Games, Inc. All Rights Reserved.
+
+package com.epicgames.ue4;
+
+import android.app.Activity;
+import android.content.pm.ApplicationInfo;
+import android.content.pm.PackageManager;
+import android.os.Bundle;
+import android.util.Log;
+import android.view.ViewGroup.LayoutParams;
+import android.view.ViewGroup.MarginLayoutParams;
+import android.widget.LinearLayout;
+
+import com.epicgames.unreal.RemoteFileManager;
+
+import java.io.BufferedReader;
+import java.io.IOException;
+import java.io.InputStreamReader;
+import java.io.PrintWriter;
+import java.net.ServerSocket;
+import java.net.Socket;
+import java.util.Set;
+
+import static android.os.Environment.getExternalStorageDirectory;
+
+public class SplashActivity extends Activity {
+    private static final int SERVER_PORT = 57099;
+
+	@Override
+	protected void onCreate(Bundle savedInstanceState) {
+		super.onCreate(savedInstanceState);
+
+		MarginLayoutParams params = new MarginLayoutParams(LayoutParams.WRAP_CONTENT, LayoutParams.WRAP_CONTENT);
+		params.setMargins(0, 0, 0, 0);
+		LinearLayout activityLayout = new LinearLayout(this);
+		setContentView(activityLayout, params);
+
+		int port = SERVER_PORT;
+
+		// make sure UnrealFileServer not already running
+		boolean bUSBRunning = false;
+		boolean bWiFiRunning = false;
+				
+		Set<Thread> threads = Thread.getAllStackTraces().keySet();
+		for (Thread thread : threads) {
+			String name = thread.getName();
+			if (name.equals("UnrealAFS-USB"))
+			{
+				if (thread.isAlive())
+				{
+					bUSBRunning = true;
+				}
+			}
+			else if (name.equals("UnrealAFS-WiFi"))
+			{
+				if (thread.isAlive())
+				{
+					bWiFiRunning = true;
+				}
+			}
+		}
+
+		String packageName = getPackageName();
+		PackageManager pm = getPackageManager();
+
+		int versionCode = 0;
+		try
+		{
+			versionCode = pm.getPackageInfo(packageName, 0).versionCode;
+		}
+		catch (Exception e)
+		{
+			// if the above failed, then, we can't use obbs
+		}
+
+		String EngineVersion = "4.27.3";
+		String ProjectName = packageName;
+		ProjectName = ProjectName.substring(ProjectName.lastIndexOf('.') + 1);
+		boolean IsShipping = false;
+		boolean UseExternalFilesDir = false;
+		boolean PublicLogFiles = false;
+		try {
+			ApplicationInfo ai = pm.getApplicationInfo(packageName, PackageManager.GET_META_DATA);
+			Bundle bundle = ai.metaData;
+
+			if (bundle.containsKey("com.epicgames.ue4.GameActivity.EngineVersion"))
+			{
+				EngineVersion = bundle.getString("com.epicgames.ue4.GameActivity.EngineVersion");
+			}
+			if (bundle.containsKey("com.epicgames.ue4.GameActivity.ProjectName"))
+			{
+				ProjectName = bundle.getString("com.epicgames.ue4.GameActivity.ProjectName");
+			}
+			if (bundle.containsKey("com.epicgames.ue4.GameActivity.BuildConfiguration"))
+			{
+				String Configuration = bundle.getString("com.epicgames.ue4.GameActivity.BuildConfiguration");
+				IsShipping = Configuration.contains("Shipping");
+			}
+			if (bundle.containsKey("com.epicgames.ue4.GameActivity.bUseExternalFilesDir"))
+			{
+				UseExternalFilesDir = bundle.getBoolean("com.epicgames.ue4.GameActivity.bUseExternalFilesDir");
+			}
+			if (bundle.containsKey("com.epicgames.ue4.GameActivity.bPublicLogFiles"))
+			{
+				PublicLogFiles = bundle.getBoolean("com.epicgames.ue4.GameActivity.bPublicLogFiles");
+			}
+		}
+		catch (PackageManager.NameNotFoundException | NullPointerException e)
+		{
+			Log.e("UE","Error when accessing application metadata", e);
+		}
+
+		String internal = getFilesDir().getAbsolutePath();
+		String external = getExternalFilesDir(null).getAbsolutePath();
+		String storage = getExternalStorageDirectory().getAbsolutePath();
+		String obbdir = getObbDir().getAbsolutePath();
+
+		if (!bUSBRunning)
+		{
+			RemoteFileManager fileManager = new RemoteFileManager(true, port, internal, external, storage, obbdir, packageName, versionCode, ProjectName, EngineVersion, IsShipping, PublicLogFiles);
+			new Thread(fileManager, "UnrealAFS-USB").start();
+		}
+		if (!bWiFiRunning)
+		{
+			RemoteFileManager fileManager = new RemoteFileManager(false, port, internal, external, storage, obbdir, packageName, versionCode, ProjectName, EngineVersion, IsShipping, PublicLogFiles);
+			new Thread(fileManager, "UnrealAFS-WiFi").start();
+		}
+    }
+}
\ No newline at end of file
diff --git a/Engine/Plugins/Runtime/AndroidFileServer/Source/AFSStub/app/src/main/java/com/epicgames/unreal/GameActivity.java b/Engine/Plugins/Runtime/AndroidFileServer/Source/AFSStub/app/src/main/java/com/epicgames/unreal/GameActivity.java
new file mode 100644
--- /dev/null
+++ b/Engine/Plugins/Runtime/AndroidFileServer/Source/AFSStub/app/src/main/java/com/epicgames/unreal/GameActivity.java
@@ -0,0 +1,134 @@
+// Copyright Epic Games, Inc. All Rights Reserved.
+
+package com.epicgames.unreal;
+
+import android.app.Activity;
+import android.content.pm.ApplicationInfo;
+import android.content.pm.PackageManager;
+import android.os.Bundle;
+import android.util.Log;
+import android.view.ViewGroup.LayoutParams;
+import android.view.ViewGroup.MarginLayoutParams;
+import android.widget.LinearLayout;
+
+import com.epicgames.unreal.RemoteFileManager;
+
+import java.io.BufferedReader;
+import java.io.IOException;
+import java.io.InputStreamReader;
+import java.io.PrintWriter;
+import java.net.ServerSocket;
+import java.net.Socket;
+import java.util.Set;
+
+import static android.os.Environment.getExternalStorageDirectory;
+
+public class GameActivity extends Activity {
+    private static final int SERVER_PORT = 57099;
+
+	public static boolean AndroidFileServer_Verify(String Token)
+	{
+		return true;
+	}
+
+	@Override
+	protected void onCreate(Bundle savedInstanceState) {
+		super.onCreate(savedInstanceState);
+
+		MarginLayoutParams params = new MarginLayoutParams(LayoutParams.WRAP_CONTENT, LayoutParams.WRAP_CONTENT);
+		params.setMargins(0, 0, 0, 0);
+		LinearLayout activityLayout = new LinearLayout(this);
+		setContentView(activityLayout, params);
+
+		int port = SERVER_PORT;
+
+		// make sure UnrealFileServer not already running
+		boolean bUSBRunning = false;
+		boolean bWiFiRunning = false;
+				
+		Set<Thread> threads = Thread.getAllStackTraces().keySet();
+		for (Thread thread : threads) {
+			String name = thread.getName();
+			if (name.equals("UnrealAFS-USB"))
+			{
+				if (thread.isAlive())
+				{
+					bUSBRunning = true;
+				}
+			}
+			else if (name.equals("UnrealAFS-WiFi"))
+			{
+				if (thread.isAlive())
+				{
+					bWiFiRunning = true;
+				}
+			}
+		}
+
+		String packageName = getPackageName();
+		PackageManager pm = getPackageManager();
+
+		int versionCode = 0;
+		try
+		{
+			versionCode = pm.getPackageInfo(packageName, 0).versionCode;
+		}
+		catch (Exception e)
+		{
+			// if the above failed, then, we can't use obbs
+		}
+
+		String EngineVersion = "5.0.0";
+		String ProjectName = packageName;
+		ProjectName = ProjectName.substring(ProjectName.lastIndexOf('.') + 1);
+		boolean IsShipping = false;
+		boolean UseExternalFilesDir = false;
+		boolean PublicLogFiles = false;
+		try {
+			ApplicationInfo ai = pm.getApplicationInfo(packageName, PackageManager.GET_META_DATA);
+			Bundle bundle = ai.metaData;
+
+			if (bundle.containsKey("com.epicgames.unreal.GameActivity.EngineVersion"))
+			{
+				EngineVersion = bundle.getString("com.epicgames.unreal.GameActivity.EngineVersion");
+			}
+			if (bundle.containsKey("com.epicgames.unreal.GameActivity.ProjectName"))
+			{
+				ProjectName = bundle.getString("com.epicgames.unreal.GameActivity.ProjectName");
+			}
+			if (bundle.containsKey("com.epicgames.unreal.GameActivity.BuildConfiguration"))
+			{
+				String Configuration = bundle.getString("com.epicgames.unreal.GameActivity.BuildConfiguration");
+				IsShipping = Configuration.contains("Shipping");
+			}
+			if (bundle.containsKey("com.epicgames.unreal.GameActivity.bUseExternalFilesDir"))
+			{
+				UseExternalFilesDir = bundle.getBoolean("com.epicgames.unreal.GameActivity.bUseExternalFilesDir");
+			}
+			if (bundle.containsKey("com.epicgames.unreal.GameActivity.bPublicLogFiles"))
+			{
+				PublicLogFiles = bundle.getBoolean("com.epicgames.unreal.GameActivity.bPublicLogFiles");
+			}
+		}
+		catch (PackageManager.NameNotFoundException | NullPointerException e)
+		{
+			Log.e("UE","Error when accessing application metadata", e);
+		}
+
+		String internal = getFilesDir().getAbsolutePath();
+		String external = getExternalFilesDir(null).getAbsolutePath();
+		String storage = getExternalStorageDirectory().getAbsolutePath();
+		String obbdir = getObbDir().getAbsolutePath();
+
+		if (!bUSBRunning)
+		{
+			RemoteFileManager fileManager = new RemoteFileManager(true, port, internal, external, storage, obbdir, packageName, versionCode, ProjectName, EngineVersion, IsShipping, PublicLogFiles);
+			new Thread(fileManager, "UnrealAFS-USB").start();
+		}
+		if (!bWiFiRunning)
+		{
+			RemoteFileManager fileManager = new RemoteFileManager(false, port, internal, external, storage, obbdir, packageName, versionCode, ProjectName, EngineVersion, IsShipping, PublicLogFiles);
+			new Thread(fileManager, "UnrealAFS-WiFi").start();
+		}
+    }
+}
\ No newline at end of file
diff --git a/Engine/Plugins/Runtime/AndroidFileServer/Source/AFSStub/app/src/main/java/com/epicgames/unreal/GameApplication.java b/Engine/Plugins/Runtime/AndroidFileServer/Source/AFSStub/app/src/main/java/com/epicgames/unreal/GameApplication.java
new file mode 100644
--- /dev/null
+++ b/Engine/Plugins/Runtime/AndroidFileServer/Source/AFSStub/app/src/main/java/com/epicgames/unreal/GameApplication.java
@@ -0,0 +1,70 @@
+package com.epicgames.unreal;
+
+import android.app.Application;
+import android.content.Context;
+import android.content.res.Configuration;
+import androidx.lifecycle.Lifecycle;
+import androidx.lifecycle.LifecycleObserver;
+import androidx.lifecycle.OnLifecycleEvent;
+import androidx.lifecycle.ProcessLifecycleOwner;
+
+
+public class GameApplication extends Application implements LifecycleObserver {
+	private static final Logger Log = new Logger("UE", "GameApp");
+
+	private static boolean isForeground = false;
+
+	@Override
+	public void onCreate() {
+		super.onCreate();
+
+
+		ProcessLifecycleOwner.get().getLifecycle().addObserver(this);
+
+	}
+
+	@Override
+	public void attachBaseContext(Context base) {
+		super.attachBaseContext(base);
+
+	}
+
+	@Override
+	public void onLowMemory() {
+		super.onLowMemory();
+
+	}
+
+	@Override
+	public void onTrimMemory(int level) {
+		super.onTrimMemory(level);
+
+	}
+
+	@Override
+	public void onConfigurationChanged (Configuration newConfig) {
+		super.onConfigurationChanged(newConfig);
+		
+	}
+
+	@OnLifecycleEvent(Lifecycle.Event.ON_START)
+	void onEnterForeground() {
+		Log.verbose("App in foreground");
+		isForeground = true;
+	}
+
+	@OnLifecycleEvent(Lifecycle.Event.ON_STOP)
+	void onEnterBackground() {
+		Log.verbose("App in background");
+		isForeground = false;
+	}
+
+	@SuppressWarnings("unused")
+	public static boolean isAppInForeground() {
+		return isForeground;
+	}
+
+	public static boolean isAppInBackground() {
+		return !isForeground;
+	}
+}
diff --git a/Engine/Plugins/Runtime/AndroidFileServer/Source/AFSStub/app/src/main/java/com/epicgames/unreal/Logger.java b/Engine/Plugins/Runtime/AndroidFileServer/Source/AFSStub/app/src/main/java/com/epicgames/unreal/Logger.java
new file mode 100644
--- /dev/null
+++ b/Engine/Plugins/Runtime/AndroidFileServer/Source/AFSStub/app/src/main/java/com/epicgames/unreal/Logger.java
@@ -0,0 +1,150 @@
+package com.epicgames.unreal;
+
+import android.text.TextUtils;
+import android.util.Log;
+
+public class Logger
+{
+	public interface ILoggerCallback
+	{
+		void LoggerCallback(String Level, String Tag, String Message);
+	}
+
+	private static ILoggerCallback mCallback = null;
+	private final String mTag;
+	private final String mSecondaryTag;
+	private boolean bHasSecondaryTag;
+	private String mFormattedTag;
+	private String mFormattedMessageTag;
+	
+	private static boolean bAllowLogging			= true;
+	@SuppressWarnings({"FieldCanBeLocal", "unused"})
+	private static boolean bAllowExceptionLogging	= true;
+	private static boolean bPrependSecondaryTag	= false;
+	private static boolean bPrependChanged	= false;
+	
+	@SuppressWarnings({"unused", "WeakerAccess"})
+	public static void RegisterCallback(ILoggerCallback callback)
+	{
+		mCallback = callback;
+	}
+
+	@SuppressWarnings("WeakerAccess")
+	public static void SuppressLogs ()
+	{
+		bAllowLogging = bAllowExceptionLogging = false;
+	}
+
+	@SuppressWarnings({"unused", "WeakerAccess"})
+	public static void prependSecondaryTag(boolean prependSecondaryTag) {
+		bPrependSecondaryTag = prependSecondaryTag;
+		bPrependChanged = true;
+	}
+	
+	public Logger(String Tag)
+	{
+		this(Tag, "");
+	}
+
+	public Logger(String Tag, String secondaryTag)
+	{
+		mTag = Tag;
+		if(secondaryTag == null) {
+			mSecondaryTag = "";
+		} else {
+			mSecondaryTag = secondaryTag;
+		}
+		bHasSecondaryTag = !TextUtils.isEmpty(mSecondaryTag);
+	}
+
+	public void verbose(String Message)
+	{
+		if (bAllowLogging)
+		{
+			Log.v(getFormattedTag(), getFormattedMessage(Message));
+		}
+		if (mCallback != null)
+		{
+			mCallback.LoggerCallback("V/", getFormattedTag(), getFormattedMessage(Message));
+		}
+	}
+
+	public void debug(String Message)
+	{
+		if (bAllowLogging)
+		{
+			Log.d(getFormattedTag(), getFormattedMessage(Message));
+		}
+		if (mCallback != null)
+		{
+			mCallback.LoggerCallback("D/", getFormattedTag(), getFormattedMessage(Message));
+		}
+	}
+	
+	public void warn(String Message)
+	{
+		if (bAllowLogging)
+		{
+			Log.w(getFormattedTag(), getFormattedMessage(Message));
+		}
+		if (mCallback != null)
+		{
+			mCallback.LoggerCallback("W/", getFormattedTag(), getFormattedMessage(Message));
+		}
+	}
+	
+	public void error(String Message)
+	{
+		if (bAllowLogging)
+		{
+			Log.e(getFormattedTag(), getFormattedMessage(Message));
+		}
+		if (mCallback != null)
+		{
+			mCallback.LoggerCallback("E/", getFormattedTag(), getFormattedMessage(Message));
+		}
+	}
+
+	public void error(String Message, Throwable Throwable)
+	{
+		if (bAllowLogging)
+		{
+			Log.e(getFormattedTag(), getFormattedMessage(Message), Throwable);
+		}
+		if (mCallback != null)
+		{
+			mCallback.LoggerCallback("E/", getFormattedTag(), getFormattedMessage(Message));
+		}
+	}
+	
+	private String getFormattedTag()
+	{
+		if(mFormattedTag == null || bPrependChanged)
+		{
+			bPrependChanged = false;
+			if (bHasSecondaryTag && bPrependSecondaryTag)
+			{
+				mFormattedTag = mTag + "-" + mSecondaryTag;
+			}
+			else
+			{
+				mFormattedTag = mTag;
+			}
+		}
+		return mFormattedTag;
+	}
+
+	private String getFormattedMessage(String message)
+	{
+		if (bHasSecondaryTag && !bPrependSecondaryTag)
+		{
+			if(mFormattedMessageTag == null || bPrependChanged)
+			{
+				bPrependChanged = false;
+				mFormattedMessageTag = "[" + mSecondaryTag + "] ";
+			}
+			return mFormattedMessageTag + message;
+		}
+		return message;
+	}
+}
diff --git a/Engine/Plugins/Runtime/AndroidFileServer/Source/AFSStub/app/src/main/java/com/epicgames/unreal/SplashActivity.java b/Engine/Plugins/Runtime/AndroidFileServer/Source/AFSStub/app/src/main/java/com/epicgames/unreal/SplashActivity.java
new file mode 100644
--- /dev/null
+++ b/Engine/Plugins/Runtime/AndroidFileServer/Source/AFSStub/app/src/main/java/com/epicgames/unreal/SplashActivity.java
@@ -0,0 +1,129 @@
+// Copyright Epic Games, Inc. All Rights Reserved.
+
+package com.epicgames.unreal;
+
+import android.app.Activity;
+import android.content.pm.ApplicationInfo;
+import android.content.pm.PackageManager;
+import android.os.Bundle;
+import android.util.Log;
+import android.view.ViewGroup.LayoutParams;
+import android.view.ViewGroup.MarginLayoutParams;
+import android.widget.LinearLayout;
+
+import com.epicgames.unreal.RemoteFileManager;
+
+import java.io.BufferedReader;
+import java.io.IOException;
+import java.io.InputStreamReader;
+import java.io.PrintWriter;
+import java.net.ServerSocket;
+import java.net.Socket;
+import java.util.Set;
+
+import static android.os.Environment.getExternalStorageDirectory;
+
+public class SplashActivity extends Activity {
+    private static final int SERVER_PORT = 57099;
+
+	@Override
+	protected void onCreate(Bundle savedInstanceState) {
+		super.onCreate(savedInstanceState);
+
+		MarginLayoutParams params = new MarginLayoutParams(LayoutParams.WRAP_CONTENT, LayoutParams.WRAP_CONTENT);
+		params.setMargins(0, 0, 0, 0);
+		LinearLayout activityLayout = new LinearLayout(this);
+		setContentView(activityLayout, params);
+
+		int port = SERVER_PORT;
+
+		// make sure UnrealFileServer not already running
+		boolean bUSBRunning = false;
+		boolean bWiFiRunning = false;
+				
+		Set<Thread> threads = Thread.getAllStackTraces().keySet();
+		for (Thread thread : threads) {
+			String name = thread.getName();
+			if (name.equals("UnrealAFS-USB"))
+			{
+				if (thread.isAlive())
+				{
+					bUSBRunning = true;
+				}
+			}
+			else if (name.equals("UnrealAFS-WiFi"))
+			{
+				if (thread.isAlive())
+				{
+					bWiFiRunning = true;
+				}
+			}
+		}
+
+		String packageName = getPackageName();
+		PackageManager pm = getPackageManager();
+
+		int versionCode = 0;
+		try
+		{
+			versionCode = pm.getPackageInfo(packageName, 0).versionCode;
+		}
+		catch (Exception e)
+		{
+			// if the above failed, then, we can't use obbs
+		}
+
+		String EngineVersion = "5.0.0";
+		String ProjectName = packageName;
+		ProjectName = ProjectName.substring(ProjectName.lastIndexOf('.') + 1);
+		boolean IsShipping = false;
+		boolean UseExternalFilesDir = false;
+		boolean PublicLogFiles = false;
+		try {
+			ApplicationInfo ai = pm.getApplicationInfo(packageName, PackageManager.GET_META_DATA);
+			Bundle bundle = ai.metaData;
+
+			if (bundle.containsKey("com.epicgames.unreal.GameActivity.EngineVersion"))
+			{
+				EngineVersion = bundle.getString("com.epicgames.unreal.GameActivity.EngineVersion");
+			}
+			if (bundle.containsKey("com.epicgames.unreal.GameActivity.ProjectName"))
+			{
+				ProjectName = bundle.getString("com.epicgames.unreal.GameActivity.ProjectName");
+			}
+			if (bundle.containsKey("com.epicgames.unreal.GameActivity.BuildConfiguration"))
+			{
+				String Configuration = bundle.getString("com.epicgames.unreal.GameActivity.BuildConfiguration");
+				IsShipping = Configuration.contains("Shipping");
+			}
+			if (bundle.containsKey("com.epicgames.unreal.GameActivity.bUseExternalFilesDir"))
+			{
+				UseExternalFilesDir = bundle.getBoolean("com.epicgames.unreal.GameActivity.bUseExternalFilesDir");
+			}
+			if (bundle.containsKey("com.epicgames.unreal.GameActivity.bPublicLogFiles"))
+			{
+				PublicLogFiles = bundle.getBoolean("com.epicgames.unreal.GameActivity.bPublicLogFiles");
+			}
+		}
+		catch (PackageManager.NameNotFoundException | NullPointerException e)
+		{
+			Log.e("UE","Error when accessing application metadata", e);
+		}
+
+		String internal = getFilesDir().getAbsolutePath();
+		String external = getExternalFilesDir(null).getAbsolutePath();
+		String storage = getExternalStorageDirectory().getAbsolutePath();
+		String obbdir = getObbDir().getAbsolutePath();
+
+		if (!bUSBRunning)
+		{
+			RemoteFileManager fileManager = new RemoteFileManager(true, port, internal, external, storage, obbdir, packageName, versionCode, ProjectName, EngineVersion, IsShipping, PublicLogFiles);
+			new Thread(fileManager, "UnrealAFS-USB").start();
+		}
+		if (!bWiFiRunning)
+		{
+			RemoteFileManager fileManager = new RemoteFileManager(false, port, internal, external, storage, obbdir, packageName, versionCode, ProjectName, EngineVersion, IsShipping, PublicLogFiles);
+			new Thread(fileManager, "UnrealAFS-WiFi").start();
+		}
+    }
+}
\ No newline at end of file
diff --git a/Engine/Plugins/Runtime/AndroidFileServer/Source/AndroidFileServer/Android/java/com/epicgames/unreal/RemoteFileManager.java b/Engine/Plugins/Runtime/AndroidFileServer/Source/AndroidFileServer/Android/java/com/epicgames/unreal/RemoteFileManager.java
new file mode 100644
--- /dev/null
+++ b/Engine/Plugins/Runtime/AndroidFileServer/Source/AndroidFileServer/Android/java/com/epicgames/unreal/RemoteFileManager.java
@@ -0,0 +1,1525 @@
+// Copyright Epic Games, Inc. All Rights Reserved.
+
+package com.epicgames.unreal;
+
+import android.app.DownloadManager;
+import android.util.Log;
+
+import java.io.ByteArrayInputStream;
+import java.io.ByteArrayOutputStream;
+import java.io.BufferedInputStream;
+import java.io.BufferedOutputStream;
+import java.io.BufferedReader;
+import java.io.BufferedWriter;
+import java.io.File;
+import java.io.FileReader;
+import java.io.FileOutputStream;
+import java.io.InterruptedIOException;
+import java.io.IOException;
+import java.io.InputStream;
+import java.io.InputStreamReader;
+import java.io.OutputStream;
+import java.io.PrintWriter;
+import java.net.InetAddress;
+import java.net.InetSocketAddress;
+import java.net.NetworkInterface;
+import java.nio.ByteBuffer;
+import java.nio.charset.StandardCharsets;
+import java.nio.file.FileVisitResult;
+import java.nio.file.FileVisitor;
+import java.nio.file.Files;
+import java.nio.file.FileAlreadyExistsException;
+import java.nio.file.NoSuchFileException;
+import java.nio.file.Path;
+import java.nio.file.Paths;
+import java.nio.file.SimpleFileVisitor;
+import java.nio.file.StandardCopyOption;
+import java.nio.file.attribute.BasicFileAttributes;
+import java.nio.file.attribute.PosixFilePermission;
+import java.nio.file.attribute.PosixFilePermissions;
+import java.util.ArrayList;
+import java.util.Collections;
+import java.util.HashMap;
+import java.util.List;
+import java.util.Map;
+import java.util.Set;
+import java.util.zip.GZIPInputStream;
+import java.util.zip.GZIPOutputStream;
+
+import java.net.ServerSocket;
+import java.net.Socket;
+import java.net.SocketException;
+
+public class RemoteFileManager implements Runnable
+{
+	private String TAG = "UEFS";
+	private static final int SERVER_VERSION = 100;
+	private static final int ACCEPT_TIMEOUT = 2000;		// milliseconds
+
+	private static final int Command_Terminate = 0;
+	private static final int Command_Close = 1;
+	private static final int Command_Info = 2;
+	private static final int Command_Query = 3;
+	private static final int Command_GetProp = 4;
+	private static final int Command_SetBaseDir = 5;
+
+	private static final int Command_DirExists = 10;
+	private static final int Command_DirList = 11;
+	private static final int Command_DirListFlat = 12;
+	private static final int Command_DirCreate = 13;
+	private static final int Command_DirDelete = 14;
+	private static final int Command_DirDeleteRecurse = 15;
+
+	private static final int Command_FileExists = 20;
+	private static final int Command_FileDelete = 21;
+	private static final int Command_FileCopy = 22;
+	private static final int Command_FileMove = 23;
+	private static final int Command_FileRead = 24;
+	private static final int Command_FileWrite = 25;
+	private static final int Command_FileWriteCompressed = 26;
+
+	private String externalPackageDir;
+	private String internalPackageDir;
+	private String externalStorageDir;
+
+	private String obbPackageDir;
+	private String unrealPackageDir;
+	private String projectPackageDir;
+	private String enginePackageDir;
+	private String gamePackageDir;
+	private String logsPackageDir;
+	private String savedPackageDir;
+	private String logfile;
+	private String commandfile;
+	private String mainobb;
+	private String patchobb;
+	private String overflow1obb;
+	private String overflow2obb;
+	private String overflow3obb;
+	private String overflow4obb;
+
+	private String packageName;
+	private int versionCode;
+	private String projectName;
+	private String engineVersion;
+	private String unrealGame;
+
+	private boolean bIsShipping = false;
+	private boolean bPublicLogs = false;
+
+	private HashMap<String, String> replacements;
+	private String LastBaseDir = "";
+
+	private boolean bLocalConnection;
+	private int serverPort;
+	private String ipAddress;
+
+	//	private InputStream inputStream;
+//	private OutputStream outputStream;
+	private BufferedInputStream inputStream;
+	private BufferedOutputStream outputStream;
+	private long dataRemaining;
+	private boolean ioException;
+
+	public RemoteFileManager(boolean bLocal, int port, String internal, String external, String storage, String obbdir, String PackageName,
+							 int VersionCode, String ProjectName, String EngineVersion, boolean IsShipping, boolean PublicLogs)
+	{
+		TAG = bLocal ? "UEFS" : "UEFSW";
+
+		ipAddress = getIPAddress(false);
+		bLocalConnection = bLocal;
+		serverPort = port;
+
+		externalPackageDir = external + "/";
+		internalPackageDir = internal + "/";
+		externalStorageDir = storage + "/";
+		packageName = PackageName;
+		versionCode = VersionCode;
+		String versionString = "" + versionCode;
+
+		boolean bIsUE4 = EngineVersion.startsWith("4.");
+		projectName = ProjectName;
+		engineVersion = EngineVersion;
+		unrealGame = bIsUE4 ? "UE4Game/" : "UnrealGame/";
+		bIsShipping = IsShipping;
+		bPublicLogs = PublicLogs;
+
+		obbPackageDir = obbdir + "/";
+		mainobb = obbPackageDir + "main." + versionString + "." + packageName + ".obb";
+		patchobb = obbPackageDir + "patch." + versionString + "." + packageName + ".obb";
+		overflow1obb = obbPackageDir + "overflow1." + versionString + "." + packageName + ".obb";
+		overflow2obb = obbPackageDir + "overflow2." + versionString + "." + packageName + ".obb";
+		overflow3obb = obbPackageDir + "overflow3." + versionString + "." + packageName + ".obb";
+		overflow4obb = obbPackageDir + "overflow4." + versionString + "." + packageName + ".obb";
+
+		unrealPackageDir = (bIsShipping ? internalPackageDir : externalPackageDir) + unrealGame;
+		projectPackageDir = unrealPackageDir + projectName + "/";
+		enginePackageDir = projectPackageDir + "Engine/";
+		gamePackageDir = projectPackageDir + projectName + "/";
+		savedPackageDir = gamePackageDir + "Saved/";
+		logsPackageDir = (bIsShipping ? (bPublicLogs ? externalPackageDir : internalPackageDir) : externalPackageDir) +
+				unrealGame + projectName + "/" + projectName + "/Saved/Logs/";
+
+		commandfile = projectPackageDir + (bIsUE4 ? "UE4CommandLine.txt" : "UECommandLine.txt");
+		logfile = logsPackageDir + projectName + ".log";
+
+		replacements = new HashMap<String, String>();
+		replacements.put("^ext/", externalPackageDir);
+		replacements.put("^int/", internalPackageDir);
+		replacements.put("^storage/", externalStorageDir);
+
+		replacements.put("^packagename", packageName);
+		replacements.put("^version", versionString);
+		replacements.put("^obb", obbPackageDir);
+		replacements.put("^mainobb", mainobb);
+		replacements.put("^patchobb", patchobb);
+		replacements.put("^overflow1obb", overflow1obb);
+		replacements.put("^overflow2obb", overflow2obb);
+		replacements.put("^overflow3obb", overflow3obb);
+		replacements.put("^overflow4obb", overflow4obb);
+
+		replacements.put("^engineversion", engineVersion);
+		replacements.put("^projectname", projectName);
+		replacements.put("^unreal/", unrealPackageDir);
+		replacements.put("^project/", projectPackageDir);
+		replacements.put("^engine/", enginePackageDir);
+		replacements.put("^game/", gamePackageDir);
+		replacements.put("^saved/", savedPackageDir);
+		replacements.put("^logs/", logsPackageDir);
+
+		replacements.put("^commandfile", commandfile);
+		replacements.put("^logfile", logfile);
+		replacements.put("^ip", ipAddress);
+
+		LastBaseDir = "";
+	}
+
+	String getIPAddress(boolean getIPv6)
+	{
+		try
+		{
+			List<NetworkInterface> interfaces = Collections.list(NetworkInterface.getNetworkInterfaces());
+			for (NetworkInterface iface : interfaces)
+			{
+				List<InetAddress> addresses = Collections.list(iface.getInetAddresses());
+				for (InetAddress address : addresses) {
+					if (!address.isLoopbackAddress()) {
+						String result = address.getHostAddress();
+
+						// IP6 contains :'s
+						boolean bIPv6 = result.indexOf(':') >= 0;
+						if (bIPv6) {
+							if (getIPv6) {
+								// remove zone suffix if present
+								int zoneIndex = result.indexOf('%');
+								return zoneIndex < 0 ? result : result.substring(0, zoneIndex);
+							}
+						} else if (!getIPv6) {
+							return result;
+						}
+					}
+				}
+			}
+		}
+		catch (Exception ex)
+		{
+		}
+		return "";
+	}
+
+	private String getProp(String property)
+	{
+		// first try SystemProperties reflection
+		try
+		{
+			java.lang.reflect.Method get = Class.forName("android.os.SystemProperties").getMethod("get", new Class[]{String.class});
+			String propValue = (String)get.invoke(null, property);
+			if (propValue != null)
+			{
+				return propValue;
+			}
+			return "";
+		}
+		catch (Exception e)
+		{
+		}
+
+		// fall back to getprop command
+		Process process = null;
+		BufferedReader reader = null;
+		String result = "";
+
+		try
+		{
+			process = new ProcessBuilder().command("/system/bin/getprop", property).redirectErrorStream(true).start();
+			reader = new BufferedReader(new InputStreamReader(process.getInputStream()));
+			String line = reader.readLine();
+			result = (line == null) ? "" : line;
+		}
+		catch (Exception e)
+		{
+			Log.d(TAG, "Unable to use getprop: " + e);
+		}
+		if (reader != null)
+		{
+			try
+			{
+				reader.close();
+			}
+			catch (IOException e)
+			{
+			}
+		}
+		if (process != null)
+		{
+			process.destroy();
+		}
+		return result;
+	}
+
+	void SendResponse(String results)
+	{
+//		Log.d(TAG, "Results: " + results);
+
+		// make sure we send back UTF8
+		byte[] rawBytes = results.getBytes(StandardCharsets.UTF_8);
+		long size = rawBytes.length;
+		long buffersize = size + 8;
+
+		byte[] buffer = new byte[(int)buffersize];
+		buffer[0] = (byte)(size & 255);
+		buffer[1] = (byte)((size >> 8) & 255);
+		buffer[2] = (byte)((size >> 16) & 255);
+		buffer[3] = (byte)((size >> 24) & 255);
+		buffer[4] = (byte)((size >> 32) & 255);
+		buffer[5] = (byte)((size >> 40) & 255);
+		buffer[6] = (byte)((size >> 48) & 255);
+		buffer[7] = (byte)((size >> 56) & 255);
+
+		System.arraycopy(rawBytes, 0, buffer, 8, (int)size);
+
+		try {
+			outputStream.write(buffer, 0, (int)buffersize);
+			outputStream.flush();
+		}
+		catch (Exception e)
+		{
+			ioException = true;
+		}
+	}
+
+	String FixupPath(String path)
+	{
+		for (Map.Entry<String, String> entry : replacements.entrySet())
+		{
+			String key = entry.getKey();
+			String value = entry.getValue();
+
+			if (path.startsWith(key)) {
+				return entry.getValue() + path.substring(key.length());
+			}
+		}
+		return path;
+	}
+
+	String FixupPathLastDir(String path, boolean bPreserve)
+	{
+		String oldLastBaseDir = LastBaseDir;
+
+		if (path.startsWith("^^/"))
+		{
+			path = LastBaseDir + path.substring(3);
+		}
+		else if (path.startsWith("^-"))
+		{
+			int slashIndex = path.indexOf("/");
+			if (slashIndex > 1)
+			{
+				int dropCount =  Integer.parseInt(path.substring(2, slashIndex)) + 1;
+				int index;
+				while ((dropCount-- > 0) && ((index = LastBaseDir.lastIndexOf("/")) >= 0))
+				{
+					LastBaseDir = LastBaseDir.substring(0, index);
+				}
+				path = LastBaseDir + path.substring(slashIndex);
+			}
+		}
+
+		int lastSlashIndex = path.lastIndexOf("/");
+		LastBaseDir = path.substring(0, lastSlashIndex + 1);
+
+		if (bPreserve)
+		{
+			LastBaseDir = oldLastBaseDir;
+		}
+
+		return FixupPath(path);
+	}
+
+	public void CommandQuery(String params)
+	{
+		if (params.isEmpty())
+		{
+			ClearResponse();
+			for (Map.Entry<String, String> entry : replacements.entrySet())
+			{
+				AppendResponse(entry.getKey() + "\t" + entry.getValue() + "\n");
+			}
+			AppendResponse("^^/\t" + LastBaseDir + "\n");
+			SendResponse(response);
+			return;
+		}
+		String request = params;
+		String result = FixupPathLastDir(request, true);
+		if (result.equals(request))
+		{
+			request = request + "/";
+			result = FixupPathLastDir(request, true);
+		}
+
+		SendResponse(result.equals(request) ? "Invalid" : result);
+	}
+
+	public void CommandDirExists(String params)
+	{
+		Path path = Paths.get(FixupPath(params));
+
+		SendResponse(Files.isDirectory(path) ? "true" : "false");
+	}
+
+	public String response;
+
+	public void ClearResponse()
+	{
+		response = "";
+	}
+
+	public void AppendResponse(String data)
+	{
+		response += data;
+	}
+
+	public void CommandDirList(String params)
+	{
+		boolean recurse = false;
+		boolean filesize = false;
+		boolean attribs = false;
+		ClearResponse();
+
+		int optionsIndex = params.indexOf(":");
+		if (optionsIndex > 0)
+		{
+			String options = params.substring(0, optionsIndex);
+			params = params.substring(optionsIndex+1);
+
+			if (options.contains("R"))
+			{
+				recurse = true;
+			}
+			if (options.contains("S"))
+			{
+				filesize = true;
+			}
+			if (options.contains("A"))
+			{
+				attribs = true;
+			}
+		}
+
+		if (!params.endsWith("/"))
+		{
+			params += "/";
+		}
+		Path path = Paths.get(FixupPath(params));
+		final String pathstring = path.toString();
+		final int pathlen = path.toString().length() + 1;
+		final boolean bRecursive = recurse;
+		final boolean bFilesize = filesize;
+		final boolean bAttributes = attribs;
+
+		if (!Files.isDirectory(path))
+		{
+			SendResponse("");
+			return;
+		}
+
+		try {
+			Files.walkFileTree(path, new FileVisitor<Path>() {
+				@Override
+				public FileVisitResult preVisitDirectory(Path dir, BasicFileAttributes attrs) throws IOException {
+					if (dir.toString().equals(pathstring))
+					{
+						return FileVisitResult.CONTINUE;
+					}
+					String attrib = "";
+					String size = "";
+					if (bAttributes)
+					{
+						try {
+							Set<PosixFilePermission> permissions = Files.getPosixFilePermissions(dir);
+							attrib = "d";
+							attrib += (permissions.contains(PosixFilePermission.OWNER_READ)) ? "r" : "-";
+							attrib += (permissions.contains(PosixFilePermission.OWNER_WRITE)) ? "w" : "-";
+							attrib += (permissions.contains(PosixFilePermission.OWNER_EXECUTE)) ? "x" : "-";
+							attrib += (permissions.contains(PosixFilePermission.GROUP_READ)) ? "r" : "-";
+							attrib += (permissions.contains(PosixFilePermission.GROUP_WRITE)) ? "w" : "-";
+							attrib += (permissions.contains(PosixFilePermission.GROUP_EXECUTE)) ? "x" : "-";
+							attrib += (permissions.contains(PosixFilePermission.OTHERS_READ)) ? "r" : "-";
+							attrib += (permissions.contains(PosixFilePermission.OTHERS_WRITE)) ? "w" : "-";
+							attrib += (permissions.contains(PosixFilePermission.OTHERS_EXECUTE)) ? "x " : "- ";
+						}
+						catch (Exception e)
+						{
+						}
+					}
+					if (bFilesize) {
+						try {
+							size = Files.size(dir) + " ";
+						} catch (Exception e) {
+							size = "0 ";
+						}
+					}
+					AppendResponse("d " + attrib + size + dir.toString().substring(pathlen) + "\n");
+					return bRecursive ? FileVisitResult.CONTINUE : FileVisitResult.SKIP_SUBTREE;
+				}
+
+				@Override
+				public FileVisitResult visitFile(Path file, BasicFileAttributes attrs) throws IOException {
+					String attrib = "";
+					String size = "";
+					if (bAttributes)
+					{
+						try {
+							Set<PosixFilePermission> permissions = Files.getPosixFilePermissions(file);
+							attrib = "-";
+							attrib += (permissions.contains(PosixFilePermission.OWNER_READ)) ? "r" : "-";
+							attrib += (permissions.contains(PosixFilePermission.OWNER_WRITE)) ? "w" : "-";
+							attrib += (permissions.contains(PosixFilePermission.OWNER_EXECUTE)) ? "x" : "-";
+							attrib += (permissions.contains(PosixFilePermission.GROUP_READ)) ? "r" : "-";
+							attrib += (permissions.contains(PosixFilePermission.GROUP_WRITE)) ? "w" : "-";
+							attrib += (permissions.contains(PosixFilePermission.GROUP_EXECUTE)) ? "x" : "-";
+							attrib += (permissions.contains(PosixFilePermission.OTHERS_READ)) ? "r" : "-";
+							attrib += (permissions.contains(PosixFilePermission.OTHERS_WRITE)) ? "w" : "-";
+							attrib += (permissions.contains(PosixFilePermission.OTHERS_EXECUTE)) ? "x " : "- ";
+						}
+						catch (Exception e)
+						{
+						}
+					}
+					if (bFilesize) {
+						try {
+							size = Files.size(file) + " ";
+						} catch (Exception e) {
+							size = "0 ";
+						}
+					}
+
+					AppendResponse( "f " + attrib + size + file.toString().substring(pathlen) + "\n");
+					return FileVisitResult.CONTINUE;
+				}
+
+				@Override
+				public FileVisitResult visitFileFailed(Path file, IOException exc) throws IOException {
+					return FileVisitResult.CONTINUE;
+				}
+
+				@Override
+				public FileVisitResult postVisitDirectory(Path dir, IOException exc) throws IOException {
+					return FileVisitResult.CONTINUE;
+				}
+			});
+		}
+		catch (IOException e)
+		{
+		}
+		catch (Exception e)
+		{
+		}
+		SendResponse(response);
+	}
+
+	private ArrayList<Path> directories;
+	private void VisitDirectories(Path baseDir)
+	{
+		AppendResponse(baseDir.toString() + ":\n");
+
+		directories = new ArrayList<Path>();
+		final String pathstring = baseDir.toString();
+		try {
+			Files.walkFileTree(baseDir, new FileVisitor<Path>() {
+				@Override
+				public FileVisitResult preVisitDirectory(Path dir, BasicFileAttributes attrs) throws IOException {
+					if (dir.toString().equals(pathstring))
+					{
+						return FileVisitResult.CONTINUE;
+					}
+					directories.add(dir);
+					return FileVisitResult.SKIP_SUBTREE;
+				}
+
+				@Override
+				public FileVisitResult visitFile(Path file, BasicFileAttributes attrs) throws IOException {
+					AppendResponse( file.getFileName() + "\n");
+					return FileVisitResult.CONTINUE;
+				}
+
+				@Override
+				public FileVisitResult visitFileFailed(Path file, IOException exc) throws IOException {
+					return FileVisitResult.CONTINUE;
+				}
+
+				@Override
+				public FileVisitResult postVisitDirectory(Path dir, IOException exc) throws IOException {
+					return FileVisitResult.SKIP_SUBTREE;
+				}
+			});
+		}
+		catch (IOException e)
+		{
+		}
+		catch (Exception e)
+		{
+		}
+
+		for (Path dir : directories)
+		{
+			VisitDirectories(dir);
+		}
+	}
+
+	public void CommandDirListFlat(String params)
+	{
+		boolean recurse = true;
+		ClearResponse();
+
+		if (!params.endsWith("/"))
+		{
+			params += "/";
+		}
+
+		Path path = Paths.get(FixupPath(params));
+		if (!Files.isDirectory(path))
+		{
+			SendResponse("");
+			return;
+		}
+		VisitDirectories(path);
+		SendResponse(response);
+	}
+
+	private boolean DirCreateRecursive(String params)
+	{
+		try
+		{
+			File file = new File(params);
+			return file.mkdirs();
+		}
+		catch (Exception e)
+		{
+		}
+		return false;
+	}
+
+	/*  Note: file.mkdirs() is faster
+	private boolean DirCreateRecursive(String params)
+	{
+		Path path = Paths.get(params);
+		int lastIndex = path.getNameCount();
+
+		// first walk find existing root
+		int rootIndex = lastIndex;
+		while (rootIndex > 0)
+		{
+			if (Files.isDirectory(path.subpath(0, rootIndex)))
+			{
+				break;
+			}
+			rootIndex--;
+		}
+
+		// create each directory until end
+		while (rootIndex < lastIndex)
+		{
+			Path dirPath = path.subpath(0, ++rootIndex);
+			try
+			{
+				Path dir = Files.createDirectory(dirPath);
+			}
+			catch (FileAlreadyExistsException e)
+			{
+				// this is ok
+			}
+			catch (NoSuchFileException e)
+			{
+				return false;
+			}
+			catch (IOException e)
+			{
+				return false;
+			}
+		}
+		return true;
+	}
+	*/
+
+	public void CommandDirCreateRecursive(String params)
+	{
+		if (!DirCreateRecursive(FixupPath(params))) {
+			SendResponse("false");
+			return;
+		}
+		SendResponse("true");
+	}
+
+	public void CommandDirCreate(String params)
+	{
+		Path path = Paths.get(FixupPath(params));
+
+		try
+		{
+			Path dir = Files.createDirectory(path);
+		}
+		catch (FileAlreadyExistsException e)
+		{
+			// this is ok
+		}
+		catch (NoSuchFileException e)
+		{
+			// part of the path likely doesn't exist
+			CommandDirCreateRecursive(params);
+			return;
+		}
+		catch (IOException e)
+		{
+			SendResponse("false");
+			return;
+		}
+		SendResponse("true");
+	}
+
+	public void CommandDirDelete(String params)
+	{
+		Path path = Paths.get(FixupPath(params));
+
+		if (!Files.isDirectory(path))
+		{
+			SendResponse("false");
+			return;
+		}
+
+		try
+		{
+			Files.delete(path);
+		}
+		catch (IOException e)
+		{
+			SendResponse("false");
+			return;
+		}
+		SendResponse("true");
+	}
+
+	public void CommandDirDeleteRecurse(String params)
+	{
+		Path path = Paths.get(FixupPath(params));
+
+		if (!Files.isDirectory(path))
+		{
+			SendResponse("false");
+			return;
+		}
+
+		try
+		{
+			Files.walkFileTree(path, new SimpleFileVisitor<Path>() {
+				@Override
+				public FileVisitResult visitFile(Path file, BasicFileAttributes attrs) throws IOException
+				{
+					Files.delete(file);
+					return FileVisitResult.CONTINUE;
+				}
+
+				@Override
+				public FileVisitResult postVisitDirectory(Path dir, IOException exc) throws IOException
+				{
+					Files.delete(dir);
+					return FileVisitResult.CONTINUE;
+				}
+			});
+		}
+		catch (IOException e)
+		{
+			SendResponse("false");
+			return;
+		}
+
+		SendResponse("true");
+	}
+
+	public void CommandFileExists(String params)
+	{
+		Path path = Paths.get(FixupPath(params));
+
+		SendResponse(Files.exists(path) ? "true" : "false");
+	}
+
+	public void CommandFileDelete(String params)
+	{
+		Path path = Paths.get(FixupPath(params));
+//Log.d(TAG, "FileDelete: " + path.toString());
+
+		if (!Files.exists(path) || Files.isDirectory(path))
+		{
+			SendResponse("false");
+			return;
+		}
+
+		try
+		{
+			Files.delete(path);
+		}
+		catch (IOException e)
+		{
+			SendResponse("false");
+			return;
+		}
+		SendResponse("true");
+	}
+
+	public void CommandFileCopy(String source, String dest)
+	{
+		Path sourcePath = Paths.get(FixupPath(source));
+		Path destPath = Paths.get(FixupPath(dest));
+
+		try
+		{
+			Files.copy(sourcePath, destPath, StandardCopyOption.REPLACE_EXISTING);
+		}
+		catch (FileAlreadyExistsException e)
+		{
+			SendResponse("false");
+			return;
+		}
+		catch (IOException e)
+		{
+			SendResponse("false");
+			return;
+		}
+
+		SendResponse("true");
+	}
+
+	public void CommandFileMove(String source, String dest)
+	{
+		Path sourcePath = Paths.get(FixupPath(source));
+		Path destPath = Paths.get(FixupPath(dest));
+
+		try
+		{
+			Files.move(sourcePath, destPath, StandardCopyOption.REPLACE_EXISTING);
+		}
+		catch (IOException e)
+		{
+			SendResponse("false");
+			return;
+		}
+
+		SendResponse("true");
+	}
+
+	private void SendSize(long size)
+	{
+		byte[] buffer = new byte[8];
+		buffer[0] = (byte) (size & 255);
+		buffer[1] = (byte) ((size >> 8) & 255);
+		buffer[2] = (byte) ((size >> 16) & 255);
+		buffer[3] = (byte) ((size >> 24) & 255);
+		buffer[4] = (byte) ((size >> 32) & 255);
+		buffer[5] = (byte) ((size >> 40) & 255);
+		buffer[6] = (byte) ((size >> 48) & 255);
+		buffer[7] = (byte) ((size >> 56) & 255);
+
+		try {
+			outputStream.write(buffer, 0, 8);
+			outputStream.flush();
+		} catch (Exception e) {
+			ioException = true;
+		}
+	}
+
+	public void CommandFileRead(String params)
+	{
+		Path path = Paths.get(FixupPathLastDir(params, false));
+//Log.d(TAG, "FileRead: " + path.toString());
+
+		if (Files.isDirectory(path))
+		{
+			SendSize(-1);
+			return;
+		}
+		if (!Files.exists(path))
+		{
+			SendSize(-1);
+			return;
+		}
+
+		try {
+			long filesize = Files.size(path);
+
+			// special case stream - cannot predetermine size
+			if (filesize == 0 && path.toString().startsWith("/proc/"))
+			{
+				ClearResponse();
+				try {
+					BufferedReader reader = new BufferedReader(new FileReader(new File(path.toString())));
+					String line;
+					while ((line = reader.readLine()) != null)
+					{
+						AppendResponse(line + "\n");
+					}
+					reader.close();
+				}
+				catch (IOException ie)
+				{
+					SendSize(-1);
+					return;
+				}
+				SendResponse(response);
+				return;
+			}
+
+			SendSize(filesize);
+
+			int bufferSize = 65536;
+			byte[] buffer = new byte[bufferSize];
+
+			InputStream reader = Files.newInputStream(path);
+			while (filesize > 0)
+			{
+				int chunk = filesize < bufferSize ? (int)filesize : bufferSize;
+				int readBytes = reader.read(buffer, 0, chunk);
+				if (readBytes > 0)
+				{
+					outputStream.write(buffer, 0, readBytes);
+					filesize -= readBytes;
+				}
+			}
+			outputStream.flush();
+			reader.close();
+		}
+		catch (IOException e)
+		{
+
+		}
+	}
+
+	private void IgnoreData(long size)
+	{
+		try
+		{
+			inputStream.skip(size);
+			dataRemaining -= size;
+		}
+		catch (Exception e)
+		{
+			ioException = true;
+		}
+	}
+
+	public void CommandFileWrite(String params)
+	{
+		Path path = Paths.get(FixupPathLastDir(params, false));
+//Log.d(TAG, "FileWrite: " + path.toString());
+
+		long filesize = dataRemaining;
+
+		while (true)
+		{
+			try
+			{
+				int bufferSize = 65536 * 16;
+				byte[] buffer = new byte[bufferSize];
+
+				OutputStream writer = Files.newOutputStream(path);
+				while (filesize > 0)
+				{
+					int chunk = filesize < bufferSize ? (int)filesize : bufferSize;
+					int readBytes = inputStream.read(buffer, 0, chunk);
+					if (readBytes == -1)
+					{
+						writer.close();
+						ioException = true;
+						return;
+					}
+					if (readBytes > 0)
+					{
+						writer.write(buffer, 0, readBytes);
+						filesize -= readBytes;
+						dataRemaining -= readBytes;
+					}
+				}
+				writer.close();
+				break;
+			}
+			catch (NoSuchFileException fe)
+			{
+				int lastIndex = path.getNameCount();
+				if (!DirCreateRecursive(path.subpath(0, lastIndex - 1).toString()))
+				{
+					IgnoreData(dataRemaining);
+//					SendResponse("false");
+					return;
+				}
+			}
+			catch (Exception e)
+			{
+				IgnoreData(dataRemaining);
+//				SendResponse("false");
+				return;
+			}
+		}
+
+//		SendResponse("true");
+	}
+
+	public void CommandFileWriteCompressedOodle(String params)
+	{
+		Path path = Paths.get(FixupPathLastDir(params.substring(1), false));
+//Log.d(TAG, "FileWriteCompressed: " + path.toString());
+
+		long filesize = dataRemaining;
+		boolean bUSB = TAG == "UEFS";
+		dataRemaining = 0;
+
+		int bufferSize = 1048576;
+		byte[] buffer = new byte[bufferSize + 3];
+
+		int handle = nativeCreateFile(path.toString(), bufferSize, bUSB);
+		if (handle == -1)
+		{
+			return;
+		}
+
+		byte[] header = new byte[3];
+
+		try {
+			while (filesize > 0) {
+				int readBytes;
+				int headerOffset = 0;
+				int headerBytes = 3;
+				while (headerBytes > 0) {
+					readBytes = inputStream.read(header, headerOffset, headerBytes);
+					if (readBytes == -1) {
+						nativeCloseFile(handle, bUSB);
+						ioException = true;
+						return;
+					}
+					headerOffset += readBytes;
+					headerBytes -= readBytes;
+				}
+
+				// zero indicates uncompressed 1M chunk
+				int chunk = ((int) header[0] & 0xff) + (((int) header[1] & 0xff) << 8) + (((int) header[2] & 0xff) << 16);
+				boolean bUncompressed = (chunk == 0);
+				chunk = (bUncompressed ? 1048576 : chunk);
+				int offset = 0;
+				int remaining = chunk;
+				if (bUncompressed) {
+					while (remaining > 0) {
+						readBytes = inputStream.read(buffer, offset, remaining);
+						if (readBytes == -1) {
+							nativeCloseFile(handle, bUSB);
+							ioException = true;
+							return;
+						}
+						filesize -= readBytes;
+						offset += readBytes;
+						remaining -= readBytes;
+					}
+					nativeWrite(handle, buffer, offset);
+					continue;
+				}
+
+				remaining += 3;
+				while (remaining > 0) {
+					readBytes = inputStream.read(buffer, offset, remaining);
+					if (readBytes == -1) {
+						nativeCloseFile(handle, bUSB);
+						ioException = true;
+						return;
+					}
+					offset += readBytes;
+					remaining -= readBytes;
+				}
+
+				// use Ooodle decompression
+				int rawBytes = ((int) buffer[0] & 0xff) + (((int) buffer[1] & 0xff) << 8) + (((int) buffer[2] & 0xff) << 16);
+				readBytes = nativeOodleDecompress(buffer, chunk, handle, rawBytes, bUSB);
+				filesize -= readBytes;
+			}
+			nativeCloseFile(handle, bUSB);
+		}
+		catch (Exception e)
+		{
+			//IgnoreData(dataRemaining);
+			//SendResponse("false");
+			return;
+		}
+
+		//SendResponse("true");
+	}
+
+	public void CommandFileWriteCompressed(String params)
+	{
+		char CompressType = params.charAt(0);
+		if (CompressType == 'O')
+		{
+			CommandFileWriteCompressedOodle(params);
+			return;
+		}
+
+		Path path = Paths.get(FixupPathLastDir(params.substring(1), false));
+//Log.d(TAG, "FileWriteCompressed: " + path.toString());
+
+		long filesize = dataRemaining;
+		dataRemaining = 0;
+
+		int bufferSize = 1048576;
+		byte[] buffer = new byte[bufferSize];
+		byte[] decompbuffer = new byte[bufferSize];
+
+		while (true)
+		{
+			try
+			{
+				byte[] header = new byte[3];
+
+				OutputStream writer = Files.newOutputStream(path);
+				while (filesize > 0)
+				{
+					int readBytes;
+					int headerOffset = 0;
+					int headerBytes = 3;
+					while (headerBytes > 0)
+					{
+						readBytes = inputStream.read(header, headerOffset, headerBytes);
+						if (readBytes == -1)
+						{
+							writer.close();
+							ioException = true;
+							return;
+						}
+						headerOffset += readBytes;
+						headerBytes -= readBytes;
+					}
+
+					// zero indicates uncompressed 1M chunk
+					int chunk = ((int)header[0] & 0xff) + (((int)header[1] & 0xff) << 8) + (((int)header[2] & 0xff) << 16);
+					boolean bUncompressed = (chunk == 0);
+					chunk = (bUncompressed ? 1048576 : chunk);
+					int offset = 0;
+					int remaining = chunk;
+					if (bUncompressed)
+					{
+						while (remaining > 0)
+						{
+							readBytes = inputStream.read(buffer, offset, remaining);
+							if (readBytes == -1)
+							{
+								writer.close();
+								ioException = true;
+								return;
+							}
+							filesize -= readBytes;
+							offset += readBytes;
+							remaining -= readBytes;
+						}
+						writer.write(buffer, 0, offset);
+						continue;
+					}
+
+					// decompress chunk
+					if (CompressType == 'Z')
+					{
+						while (remaining > 0)
+						{
+							readBytes = inputStream.read(buffer, offset, remaining);
+							if (readBytes == -1)
+							{
+								writer.close();
+								ioException = true;
+								return;
+							}
+							offset += readBytes;
+							remaining -= readBytes;
+						}
+
+						ByteArrayInputStream byteArrayStream = new ByteArrayInputStream(buffer, 0, chunk);
+						try
+						{
+							GZIPInputStream gzipStream = new GZIPInputStream(byteArrayStream);
+							while ((readBytes = gzipStream.read(decompbuffer)) != -1)
+							{
+								writer.write(decompbuffer, 0, readBytes);
+								filesize -= readBytes;
+							}
+							gzipStream.close();
+						}
+						catch (IOException e)
+						{
+							writer.close();
+							ioException = true;
+							return;
+						}
+					}
+					else
+					{
+						// unknown compression type !
+						writer.close();
+						ioException = true;
+						return;
+					}
+				}
+				writer.close();
+				break;
+			}
+			catch (NoSuchFileException fe)
+			{
+				int lastIndex = path.getNameCount();
+				if (!DirCreateRecursive(path.subpath(0, lastIndex - 1).toString()))
+				{
+					IgnoreData(dataRemaining);
+					//SendResponse("false");
+					return;
+				}
+			}
+			catch (Exception e)
+			{
+				//IgnoreData(dataRemaining);
+				//SendResponse("false");
+				return;
+			}
+		}
+
+		//SendResponse("true");
+	}
+
+	private String readString()
+	{
+		// 256 should handle most input, and usually only grow once if not
+		int index = 0;
+		int capacity = 256;
+		byte[] stringBuffer = new byte[capacity];
+
+		int value = -1;
+		while (dataRemaining > 0)
+		{
+			dataRemaining--;
+			try
+			{
+				value = inputStream.read();
+				if (value == 0)
+				{
+					break;
+				}
+
+				// grow the buffer if at capacity
+				if (index == capacity)
+				{
+					capacity += 256;
+					byte[] newBuffer = new byte[capacity];
+					System.arraycopy(stringBuffer, 0, newBuffer, 0, index);
+					stringBuffer = newBuffer;
+				}
+				stringBuffer[index++] = (byte)value;
+			}
+			catch (IOException e)
+			{
+				ioException = true;
+				return "";
+			}
+		}
+
+		// convert from byte array to string as UTF8
+		return new String(stringBuffer, 0, index, StandardCharsets.UTF_8);
+	}
+
+	private boolean handleCommand(int command)
+	{
+		String params;
+		String dest;
+
+		switch (command)
+		{
+			case Command_Info:
+				LastBaseDir = "";
+
+				byte[] version = new byte[10];
+				version[0] = 2;	version[1] = 0;	version[2] = 0;	version[3] = 0;	version[4] = 0;	version[5] = 0;
+				version[6] = 0; version[7] = 0;
+				version[8] = SERVER_VERSION & 255;
+				version[9] = SERVER_VERSION >> 8;
+				try {
+					outputStream.write(version, 0, 10);
+					outputStream.flush();
+				}
+				catch (IOException e)
+				{
+					Log.d(TAG, "Exception: " + e.toString());
+					ioException = true;
+					return false;
+				}
+				catch (Exception e)
+				{
+					Log.d(TAG, "Exception: " + e.toString());
+				}
+				return true;
+
+			case Command_Query:
+				params = readString();
+				if (ioException) break;
+				CommandQuery(params);
+				return true;
+
+			case Command_GetProp:
+				params = readString();
+				if (ioException) break;
+				SendResponse(getProp(params));
+				return true;
+
+			case Command_SetBaseDir:
+				params = readString();
+				if (ioException) break;
+				LastBaseDir = params.endsWith("/") ? params.substring(0, LastBaseDir.length() - 1) : params;
+				SendResponse("true");
+				return true;
+
+			case Command_DirExists:
+				params = readString();
+				if (ioException) break;
+				CommandDirExists(params);
+				return true;
+
+			case Command_DirList:
+				params = readString();
+				if (ioException) break;
+				CommandDirList(params);
+				return true;
+
+			case Command_DirListFlat:
+				params = readString();
+				if (ioException) break;
+				CommandDirListFlat(params);
+				return true;
+
+			case Command_DirCreate:
+				params = readString();
+				if (ioException) break;
+				CommandDirCreate(params);
+				return true;
+
+			case Command_DirDelete:
+				params = readString();
+				if (ioException) break;
+				CommandDirDelete(params);
+				return true;
+
+			case Command_DirDeleteRecurse:
+				params = readString();
+				if (ioException) break;
+				CommandDirDeleteRecurse(params);
+				return true;
+
+			case Command_FileExists:
+				params = readString();
+				if (ioException) break;
+				CommandFileExists(params);
+				return true;
+
+			case Command_FileDelete:
+				params = readString();
+				if (ioException) break;
+				CommandFileDelete(params);
+				return true;
+
+			case Command_FileCopy:
+				params = readString();
+				if (ioException) break;
+				dest = readString();
+				if (ioException) break;
+				CommandFileCopy(params, dest);
+				return true;
+
+			case Command_FileMove:
+				params = readString();
+				if (ioException) break;
+				dest = readString();
+				if (ioException) break;
+				CommandFileMove(params, dest);
+				return true;
+
+			case Command_FileRead:
+				params = readString();
+				if (ioException) break;
+				CommandFileRead(params);
+				return false;
+
+			case Command_FileWrite:
+				params = readString();
+				if (ioException) break;
+				CommandFileWrite(params);
+				return true;
+
+			case Command_FileWriteCompressed:
+				params = readString();
+				if (ioException) break;
+				CommandFileWriteCompressed(params);
+				return true;
+		}
+		return false;
+	}
+
+	public void run()
+	{
+		boolean bRunning = true;
+
+		String serverIP = bLocalConnection ? "localhost" : ipAddress;
+		if (serverIP.equals(""))
+		{
+			Log.d(TAG, "Unable to start remote file server for WiFi");
+			return;
+		}
+
+		Log.d(TAG, "Starting remote file server " + packageName);
+		try
+		{
+			ServerSocket serverSocket = new ServerSocket();
+			serverSocket.setSoTimeout(ACCEPT_TIMEOUT);
+//			Log.d(TAG, "Current recv buffer size: " + serverSocket.getReceiveBufferSize().toString());
+			serverSocket.setReceiveBufferSize(65536 * 4);
+			serverSocket.bind(new InetSocketAddress(serverIP, serverPort));
+
+			while (bRunning && !Thread.currentThread().isInterrupted())
+			{
+				ioException = false;
+				Log.d(TAG, "Starting server listen: " + serverPort);
+				Socket socket = null;
+				while (true)
+				{
+					try
+					{
+						Socket waitSocket = serverSocket.accept();
+						socket = waitSocket;
+						break;
+					}
+					catch (InterruptedIOException e1)
+					{
+						if (Thread.currentThread().isInterrupted())
+						{
+							bRunning = false;
+							break;
+						}
+					}
+					catch (SocketException e1)
+					{
+					}
+					catch (IOException e1)
+					{
+					}
+				}
+				if (!bRunning)
+				{
+					break;
+				}
+				if (socket == null)
+				{
+					continue;
+				}
+				Log.d(TAG, "Client connected: " + socket.toString());
+
+				outputStream = new BufferedOutputStream(socket.getOutputStream());
+				inputStream = new BufferedInputStream(socket.getInputStream());
+
+				// process commands
+				byte[] commandHeader = new byte[8];
+				while (!ioException && !Thread.currentThread().isInterrupted())
+				{
+					int readBytes = inputStream.read(commandHeader, 0, 8);
+					if (readBytes == -1)
+					{
+						break;
+					}
+					if (readBytes == 8)
+					{
+						int command = commandHeader[0]  + (commandHeader[1] << 8);
+						long datasize = ((long)commandHeader[2] & 0xff) + (((long)commandHeader[3] & 0xff) << 8) +
+								(((long)commandHeader[4] & 0xff) << 16) + (((long)commandHeader[5] & 0xff) << 24) +
+								(((long)commandHeader[6] & 0xff) << 32) + (((long)commandHeader[7] & 0xff) << 40);
+
+//Log.d(TAG, "Command: " + command);
+
+						// handle termination request
+						if (command == Command_Terminate)// && datasize == 0)
+						{
+							bRunning = false;
+							break;
+						}
+
+						// handle close request
+						if (command == Command_Close)// && datasize == 0)
+						{
+							break;
+						}
+
+						// handle other commands
+						dataRemaining = datasize;
+						if (!handleCommand(command))
+						{
+						}
+
+						// ignore remaining data
+						if (!ioException && dataRemaining > 0)
+						{
+							inputStream.skip(dataRemaining);
+						}
+					}
+				}
+
+				// close connection and listen again
+				try
+				{
+					socket.close();
+				}
+				catch (IOException e5)
+				{
+				}
+
+				Log.d(TAG, "Client connection closed");
+			}
+			try
+			{
+				serverSocket.close();
+			}
+			catch (Exception sse)
+			{
+			}
+		}
+		catch (IOException e)
+		{
+			Log.d(TAG, "Server Exception: " + e.toString());
+		}
+
+		Log.d(TAG, "Server terminated " + packageName);
+	}
+
+	public native int nativeOodleMemSizeNeeded(int compressor, int rawLen);
+	public native int nativeCreateFile(String filename, int bufferSize, boolean bUSB);
+	public native int nativeWrite(int handle, byte[] buffer, int bytes);
+	public native int nativeCloseFile(int handle, boolean bUSB);
+	public native int nativeOodleDecompress(byte[] compbuffer, int compsize, int handle, int decompsize, boolean bUSB);
+
+//	static
+//	{
+//		System.loadLibrary("fileserver");
+//	}
+}
diff --git a/Engine/Plugins/Runtime/AndroidFileServer/Source/AndroidFileServer/Android/java/com/epicgames/unreal/RemoteFileManagerReceiver.java b/Engine/Plugins/Runtime/AndroidFileServer/Source/AndroidFileServer/Android/java/com/epicgames/unreal/RemoteFileManagerReceiver.java
new file mode 100644
--- /dev/null
+++ b/Engine/Plugins/Runtime/AndroidFileServer/Source/AndroidFileServer/Android/java/com/epicgames/unreal/RemoteFileManagerReceiver.java
@@ -0,0 +1,149 @@
+// Copyright Epic Games, Inc. All Rights Reserved.
+
+package com.epicgames.unreal;
+
+import android.content.BroadcastReceiver;
+import android.content.Context;
+import android.content.Intent;
+import android.content.pm.ApplicationInfo;
+import android.content.pm.PackageManager;
+import android.os.Bundle;
+import android.util.Log;
+
+import java.util.Set;
+
+import static android.os.Environment.getExternalStorageDirectory;
+
+public class RemoteFileManagerReceiver extends BroadcastReceiver {
+    private static String TAG = "UEFS";
+    private static final int SERVER_PORT = 57099;
+
+	private boolean ValidateToken(String token)
+	{
+		return true;
+	}
+
+    @Override
+    public void onReceive(Context context, Intent intent) {
+		String action = intent.getAction();
+		if (action != null && action == "com.epicgames.unreal.RemoteFileManager.intent.COMMAND") {
+			String cmd = intent.getStringExtra("cmd");
+			String token = intent.getStringExtra("token");
+			int port = intent.getIntExtra("port", SERVER_PORT);
+			if (cmd != null && cmd.equals("start") && GameActivity.AndroidFileServer_Verify(token)) {
+
+				// make sure UnrealFileServer not already running
+				boolean bUSBRunning = false;
+				boolean bWiFiRunning = false;
+				
+				Set<Thread> threads = Thread.getAllStackTraces().keySet();
+				for (Thread thread : threads) {
+					String name = thread.getName();
+					if (name.equals("UnrealAFS-USB"))
+					{
+						if (thread.isAlive())
+						{
+							bUSBRunning = true;
+						}
+					}
+					else if (name.equals("UnrealAFS-WiFi"))
+					{
+						if (thread.isAlive())
+						{
+							bWiFiRunning = true;
+						}
+					}
+				}
+
+				if (!(bUSBRunning && bWiFiRunning)) {
+                    String packageName = context.getPackageName();
+                    PackageManager pm = context.getPackageManager();
+
+					int versionCode = 0;
+					try
+					{
+						versionCode = pm.getPackageInfo(packageName, 0).versionCode;
+					}
+					catch (Exception e)
+					{
+						// if the above failed, then, we can't use obbs
+					}
+
+					String EngineVersion = "5.0.0";
+					String ProjectName = packageName;
+					ProjectName = ProjectName.substring(ProjectName.lastIndexOf('.') + 1);
+					boolean IsShipping = false;
+					boolean UseExternalFilesDir = false;
+					boolean PublicLogFiles = false;
+					boolean AllowNetworkConnection = false;
+					try {
+						ApplicationInfo ai = pm.getApplicationInfo(packageName, PackageManager.GET_META_DATA);
+						Bundle bundle = ai.metaData;
+
+						if (bundle.containsKey("com.epicgames.unreal.GameActivity.EngineVersion"))
+						{
+							EngineVersion = bundle.getString("com.epicgames.unreal.GameActivity.EngineVersion");
+						}
+						if (bundle.containsKey("com.epicgames.unreal.GameActivity.ProjectName"))
+						{
+							ProjectName = bundle.getString("com.epicgames.unreal.GameActivity.ProjectName");
+						}
+						if (bundle.containsKey("com.epicgames.unreal.GameActivity.BuildConfiguration"))
+						{
+							String Configuration = bundle.getString("com.epicgames.unreal.GameActivity.BuildConfiguration");
+							IsShipping = Configuration.contains("Shipping");
+						}
+						if (bundle.containsKey("com.epicgames.unreal.GameActivity.bUseExternalFilesDir"))
+						{
+							UseExternalFilesDir = bundle.getBoolean("com.epicgames.unreal.GameActivity.bUseExternalFilesDir");
+						}
+						if (bundle.containsKey("com.epicgames.unreal.GameActivity.bPublicLogFiles"))
+						{
+							PublicLogFiles = bundle.getBoolean("com.epicgames.unreal.GameActivity.bPublicLogFiles");
+						}
+						if (bundle.containsKey("com.epicgames.unreal.RemoteFileManager.bAllowNetworkConnection"))
+						{
+							AllowNetworkConnection = bundle.getBoolean("com.epicgames.unreal.RemoteFileManager.bAllowNetworkConnection");
+						}
+					}
+					catch (PackageManager.NameNotFoundException | NullPointerException e)
+					{
+						Log.e(TAG,"Error when accessing application metadata", e);
+					}
+
+					String internal = context.getFilesDir().getAbsolutePath();
+					String external = context.getExternalFilesDir(null).getAbsolutePath();
+					String storage = android.os.Environment.getExternalStorageDirectory().getAbsolutePath();
+					String obbdir = context.getObbDir().getAbsolutePath();
+
+					if (!bUSBRunning)
+					{
+						RemoteFileManager fileManager = new RemoteFileManager(true, port, internal, external, storage, obbdir, packageName, versionCode, ProjectName, EngineVersion, IsShipping, PublicLogFiles);
+						new Thread(fileManager, "UnrealAFS-USB").start();
+					}
+					if (AllowNetworkConnection && !bWiFiRunning)
+					{
+						RemoteFileManager fileManager = new RemoteFileManager(false, port, internal, external, storage, obbdir, packageName, versionCode, ProjectName, EngineVersion, IsShipping, PublicLogFiles);
+						new Thread(fileManager, "UnrealAFS-WiFi").start();
+					}
+                }
+            }
+			else if (cmd != null && cmd.equals("stop")) {
+
+				// send stop request to threads
+				Set<Thread> threads = Thread.getAllStackTraces().keySet();
+				for (Thread thread : threads) {
+					String name = thread.getName();
+					if (name.equals("UnrealAFS-USB") || name.equals("UnrealAFS-WiFi"))
+					{
+						if (thread.isAlive())
+						{
+							Log.d(TAG, "Sent stop request to " + context.getPackageName() + "/" + name);
+							thread.interrupt();
+						}
+					}
+				}
+            }
+        }
+    }
+}
\ No newline at end of file
diff --git a/Engine/Plugins/Runtime/AndroidFileServer/Source/AndroidFileServer/AndroidFileServer.Build.cs b/Engine/Plugins/Runtime/AndroidFileServer/Source/AndroidFileServer/AndroidFileServer.Build.cs
new file mode 100644
--- /dev/null
+++ b/Engine/Plugins/Runtime/AndroidFileServer/Source/AndroidFileServer/AndroidFileServer.Build.cs
@@ -0,0 +1,39 @@
+// Some copyright should be here...
+
+using UnrealBuildTool;
+using System.IO;
+
+public class AndroidFileServer : ModuleRules
+{
+	public AndroidFileServer(ReadOnlyTargetRules Target) : base(Target)
+	{
+		PCHUsage = ModuleRules.PCHUsageMode.UseExplicitOrSharedPCHs;
+		
+		PublicDependencyModuleNames.AddRange(
+			new string[]
+			{
+				"Core",
+			}
+			);
+			
+		
+		PrivateDependencyModuleNames.AddRange(
+			new string[]
+			{
+				"CoreUObject",
+				"Engine",
+				"Slate",
+				"SlateCore",
+			}
+			);
+		
+		
+		if (Target.Platform == UnrealTargetPlatform.Android)
+		{
+			string PluginPath = Utils.MakePathRelativeTo(ModuleDirectory, Target.RelativeEnginePath);
+			AdditionalPropertiesForReceipt.Add("AndroidPlugin", Path.Combine(PluginPath, "AndroidFileServer_UPL.xml"));
+
+			PrivateDependencyModuleNames.AddRange(new string[] { "Launch" });
+		}
+	}
+}
diff --git a/Engine/Plugins/Runtime/AndroidFileServer/Source/AndroidFileServer/AndroidFileServer_UPL.xml b/Engine/Plugins/Runtime/AndroidFileServer/Source/AndroidFileServer/AndroidFileServer_UPL.xml
new file mode 100644
--- /dev/null
+++ b/Engine/Plugins/Runtime/AndroidFileServer/Source/AndroidFileServer/AndroidFileServer_UPL.xml
@@ -0,0 +1,511 @@
+<?xml version="1.0" encoding="utf-8"?>
+<root xmlns:android="http://schemas.android.com/apk/res/android">
+	<init>
+		<log text="AndroidFileServer Plugin Init"/>
+		<setBoolFromProperty result="bEnabled" ini="Engine" section="/Script/AndroidFileServerEditor.AndroidFileServerRuntimeSettings" property="bEnablePlugin" default="true"/>
+		<setBoolFromProperty result="bAllowNetworkConnection" ini="Engine" section="/Script/AndroidFileServerEditor.AndroidFileServerRuntimeSettings" property="bAllowNetworkConnection" default="true"/>
+		<setBoolFromProperty result="bIncludeInShipping" ini="Engine" section="/Script/AndroidFileServerEditor.AndroidFileServerRuntimeSettings" property="bIncludeInShipping" default="false"/>
+		<setBoolFromProperty result="bAllowExternalStartInShipping" ini="Engine" section="/Script/AndroidFileServerEditor.AndroidFileServerRuntimeSettings" property="bAllowExternalStartInShipping" default="false"/>
+		<setBoolFromProperty result="bCompileAFSProject" ini="Engine" section="/Script/AndroidFileServerEditor.AndroidFileServerRuntimeSettings" property="bCompileAFSProject" default="false"/>
+		<setStringFromProperty result="SecurityToken" ini="Engine" section="/Script/AndroidFileServerEditor.AndroidFileServerRuntimeSettings" property="SecurityToken" default=""/>
+
+		<!-- create the AFS project if plugin enabled -->
+		<setBool result="bCreateAFSProject" value="$B(bEnabled)"/>
+
+		<!-- token check may be disabled by clearing the security token field -->
+		<setBoolIsEqual result="bTokenOptional" arg1="$S(SecurityToken)" arg2=""/>
+		
+		<!-- by default embed AFS in APK if plugin enabled -->
+		<setBool result="bEmbedAFS" value="$B(bEnabled)"/>
+
+		<!-- by default external start of AFS allowed if plugin enabled -->
+		<setBool result="bExternalStartAFS" value="$B(bEnabled)"/>
+			 
+		<if condition="bEnabled">
+			<false>
+				<setBool result="bCreateAFSProject" value="false"/>
+			</false>
+		</if>
+
+		<!-- disable embedding AFS if shipping but "embed in shipping" not checked  -->
+		<setBoolIsEqual result="bIsShipping" arg1="$S(Configuration)" arg2="Shipping"/>
+		<if condition="bIsShipping">
+			<true>
+				<if condition="bIncludeInShipping">
+					<false>
+						<log text="Disabled embedding AFS because shipping"/>
+						<setBool result="bEmbedAFS" value="false"/>
+						
+						<!-- but we need to compile AFSProject to install data if using AFS -->
+						<setBool result="bCompileAFSProject" value="$B(bCreateAFSProject)"/>
+					</false>
+				</if>
+				<if condition="bAllowExternalStartInShipping">
+					<false>
+						<setBool result="bExternalStartAFS" value="false"/>
+					</false>
+				</if>
+			</true>
+		</if>
+	</init>
+
+	<androidManifestUpdates>
+		<if condition="bEnabled">
+			<true>
+				<if condition="bEmbedAFS">
+					<true>
+						<if condition="bExternalStartAFS">
+							<true>
+								<addElements tag="application">
+									<receiver android:name="com.epicgames.unreal.RemoteFileManagerReceiver" android:exported="true">
+										<intent-filter>
+											<action android:name="com.epicgames.unreal.RemoteFileManager.intent.COMMAND" />
+										</intent-filter>
+									</receiver>
+								</addElements>
+							</true>
+						</if>
+
+						<!-- allow or disallow network file server -->
+						<if condition="bAllowNetworkConnection">
+							<true>
+								<addElements tag="application">
+									<meta-data android:name="com.epicgames.unreal.RemoteFileManager.bAllowNetworkConnection" android:value="true"/>
+								</addElements>
+							</true>
+							<false>
+								<addElements tag="application">
+									<meta-data android:name="com.epicgames.unreal.RemoteFileManager.bAllowNetworkConnection" android:value="false"/>
+								</addElements>
+							</false>
+						</if>
+					</true>
+				</if>
+				
+				<!-- force bExternalFilesDir to true for all builds -->
+				<setBool result="bUpdatedExternal" value="false"/>
+				<loopElements tag="meta-data">
+					<setStringFromAttribute result="metaName" tag="$" name="android:name" />
+					<setBoolIsEqual result="bIsExternal" arg1="$S(metaName)" arg2="com.epicgames.unreal.GameActivity.bUseExternalFilesDir" />
+					<if condition="bIsExternal">
+						<true>
+							<removeAttribute tag="$" name="android:value"/>
+							<addAttribute tag="$" name="android:value" value="true"/>
+							<setBool result="bUpdatedExternal" value="true"/>
+						</true>
+					</if>
+				</loopElements>
+				<if condition="bUpdatedExternal">
+					<false>
+						<addElements tag="application">
+							<meta-data android:name="com.epicgames.unreal.GameActivity.bUseExternalFilesDir" android:value="true"/>
+						</addElements>
+					</false>
+				</if>
+			</true>
+		</if>
+	</androidManifestUpdates>
+
+	<gradleProperties>
+		<if condition="bEnabled">
+			<true>
+			</true>
+		</if>
+	</gradleProperties>
+
+	<resourceCopies>
+		<if condition="bEmbedAFS">
+			<true>
+				<copyFile src="$S(PluginDir)/Android/java/com/epicgames/unreal/RemoteFileManager.java" dst="$S(BuildDir)/src/com/epicgames/unreal/RemoteFileManager.java" />
+				<if condition="bExternalStartAFS">
+					<true>
+						<copyFile src="$S(PluginDir)/Android/java/com/epicgames/unreal/RemoteFileManagerReceiver.java" dst="$S(BuildDir)/src/com/epicgames/unreal/RemoteFileManagerReceiver.java" />
+					</true>
+				</if>
+<!--				<copyDir src="$S(PluginDir)/Android/libs" dst="$S(BuildDir)/libs" /> -->
+			</true>
+		</if>
+	</resourceCopies>
+
+	<gradleCopies>
+		<if condition="bCreateAFSProject">
+			<true>
+				<!-- create AFSStub project -->
+				<copyDir src="$S(PluginDir)/../AFSStub" dst="$S(BuildDir)/gradle/AFSProject" />
+				<copyDir src="$S(PluginDir)/Android" dst="$S(BuildDir)/gradle/AFSProject/app/src/main" />
+				<copyDir src="$S(BuildDir)/gradle/app/src/main/res" dst="$S(BuildDir)/gradle/AFSProject/app/src/main/res" />
+				<copyDir src="$S(BuildDir)/gradle/gradle" dst="$S(BuildDir)/gradle/AFSProject/gradle" />
+			</true>
+		</if>
+	</gradleCopies>
+
+	<buildGradleAdditions>
+		<if condition="bCompileAFSProject">
+			<true>
+				<insert>
+<![CDATA[
+// --- Begin AFSProject ---
+tasks.whenTaskAdded { task ->
+	if (task.name == 'assembleRelease') {
+		tasks.create(name:'ueAFSProjectAssembleRelease', type: Exec) {
+			description "Compile AFSProject for Release"
+			workingDir "$rootDir" + "/AFSProject"
+			if (System.getProperty('os.name').toLowerCase(Locale.ROOT).contains('windows')) {
+				commandLine 'cmd', '/c', 'gradlew.bat', 'app:assembleRelease'
+			} else {
+				commandLine 'sh', '-c', 'gradlew', 'app:assembleRelease'
+			}
+		}
+		task.finalizedBy 'ueAFSProjectAssembleRelease'
+	}
+	if (task.name == 'assembleDebug') {
+		tasks.create(name:'ueAFSProjectAssembleDebug', type: Exec) {
+			description "Compile AFSProject for Debug"
+			workingDir "$rootDir" + "/AFSProject"
+			if (System.getProperty('os.name').toLowerCase(Locale.ROOT).contains('windows')) {
+				commandLine 'cmd', '/c', 'gradlew.bat', 'app:assembleDebug'
+			} else {
+				commandLine 'sh', '-c', 'gradlew', 'app:assembleDebug'
+			}
+		}
+		task.finalizedBy 'ueAFSProjectAssembleDebug'
+	}
+	if (task.name == 'bundleRelease') {
+		tasks.create(name:'ueAFSProjectBundleRelease', type: Exec) {
+			description "Compile AFSProject for Release"
+			workingDir "$rootDir" + "/AFSProject"
+			if (System.getProperty('os.name').toLowerCase(Locale.ROOT).contains('windows')) {
+				commandLine 'cmd', '/c', 'gradlew.bat', 'app:assembleRelease'
+			} else {
+				commandLine 'sh', '-c', 'gradlew', 'app:assembleRelease'
+			}
+		}
+		task.finalizedBy 'ueAFSProjectBundleRelease'
+	}
+	if (task.name == 'bundleDebug') {
+		tasks.create(name:'ueAFSProjectBundleDebug', type: Exec) {
+			description "Compile AFSProject for Debug"
+			workingDir "$rootDir" + "/AFSProject"
+			if (System.getProperty('os.name').toLowerCase(Locale.ROOT).contains('windows')) {
+				commandLine 'cmd', '/c', 'gradlew.bat', 'app:assembleDebug'
+			} else {
+				commandLine 'sh', '-c', 'gradlew', 'app:assembleDebug'
+			}
+		}
+		task.finalizedBy 'ueAFSProjectBundleDebug'
+	}
+}
+// --- End AFSProject ---
+]]></insert>
+			</true>
+		</if>
+	</buildGradleAdditions>
+	
+	<settingsGradleAdditions>
+		<if condition="bCreateAFSProject">
+			<true>
+				<insert>
+<![CDATA[
+	// copy files into AFSStub
+	File afs_infile = new File('gradle.properties')
+	File afs_outfile = new File('AFSProject/gradle.properties')
+	def afs_line = "";
+	afs_outfile.withWriter { writer ->
+		afs_infile.withReader { reader ->
+			while ((afs_line = reader.readLine()) != null) {
+				if (afs_line.startsWith('OUTPUT_FILENAME=')) {
+					writer.writeLine("OUTPUT_FILENAME=AFS_" + afs_line.substring(16));
+					continue;
+				}
+				else if (afs_line.startsWith('OUTPUT_BUNDLEFILENAME=')) {
+					writer.writeLine("OUTPUT_BUNDLEFILENAME=AFS_" + afs_line.substring(22));
+					continue;
+				}
+				writer.writeLine("${afs_line}")
+			}
+		}
+	}
+
+	afs_infile = new File('app/buildAdditions.gradle')
+	afs_outfile = new File('AFSProject/app/buildAdditions.gradle')
+	afs_line = "";
+	def bRemovingAFSBlock = false;
+	afs_outfile.withWriter { writer ->
+		afs_infile.withReader { reader ->
+			while ((afs_line = reader.readLine()) != null) {
+				if (bRemovingAFSBlock)
+				{
+					if (afs_line.startsWith('// --- End AFSProject ---'))
+					{
+						bRemovingAFSBlock = false;
+					}
+					continue;
+				}
+				if (afs_line.startsWith('// --- Begin AFSProject ---'))
+				{
+					bRemovingAFSBlock = true;
+					continue;
+				}
+				writer.writeLine("${afs_line}")
+			}
+		}
+	}
+
+	afs_infile = new File('app/src/main/AndroidManifest.xml')
+	afs_outfile = new File('AFSProject/app/src/main/AndroidManifest.xml')
+	def afs_hasReceiver = false
+	afs_outfile.withWriter { writer ->
+		afs_infile.withReader { reader ->
+			while ((afs_line = reader.readLine()) != null) {
+				if (afs_line.contains('com.epicgames.unreal.RemoteFileManagerReceiver')) {
+					afs_hasReceiver = true
+				} else if (afs_line.contains('</application>')) {
+					if (afs_hasReceiver == false) {
+						writer.writeLine('    <receiver android:name="com.epicgames.unreal.RemoteFileManagerReceiver" android:exported="true">')
+						writer.writeLine('      <intent-filter>')
+						writer.writeLine('        <action android:name="com.epicgames.unreal.RemoteFileManager.intent.COMMAND" />')
+						writer.writeLine('      </intent-filter>')
+						writer.writeLine('    </receiver>')
+					}
+				}
+				writer.writeLine("${afs_line}")
+			}
+		}
+	}
+
+	def afs_files = ["app/aar-imports.gradle", "app/abi.gradle", "app/build.gradle",
+		"baseBuildAdditions.gradle", "build.gradle", "buildscriptAdditions.gradle", "local.properties",
+		"gradlew", "gradlew.bat", "rungradle.bat"] as String[]
+	afs_files.eachWithIndex { v, i ->
+		afs_infile = new File(v)
+		if (afs_infile.exists()) {
+			afs_outfile = new File("AFSProject/" + v)
+			afs_outfile.withWriter { writer ->
+				afs_infile.withReader { reader ->
+					while ((afs_line = reader.readLine()) != null) {
+						writer.writeLine("${afs_line}")
+					}
+				}
+			}
+		}
+	}
+
+]]></insert>
+			</true>
+		</if>
+	</settingsGradleAdditions>
+
+	<gameActivityImportAdditions>
+		<if condition="bEmbedAFS">
+			<true>
+				<insert>
+import java.util.Set;
+import com.epicgames.unreal.RemoteFileManager;
+</insert>
+			</true>
+		</if>
+	</gameActivityImportAdditions>
+
+	<gameActivityClassAdditions>
+		<if condition="bEmbedAFS">
+			<true>
+				<insert>
+	public static boolean AndroidFileServer_Verify(String Token)
+	{
+				</insert>
+				<if condition="bTokenOptional">
+					<true>
+						<insertValue value="return true;"/>
+					</true>
+					<false>
+						<insertValue value="return (Token.equals(&quot;$S(SecurityToken)&quot;));"/>
+					</false>
+				</if>
+				<insertNewline/>
+				<insert>
+	}
+
+<![CDATA[
+	public int AndroidThunkJava_AndroidFileServer_IsRunning()
+	{
+        Set<Thread> threads = Thread.getAllStackTraces().keySet();
+		int ConnectionType = 0;
+        for (Thread thread : threads) {
+            String name = thread.getName();
+            if (name.equals("UnrealAFS-USB"))
+            {
+                if (thread.isAlive()) {
+					ConnectionType |= 1;
+                }
+            }
+            else if (name.equals("UnrealAFS-WiFi"))
+            {
+                if (thread.isAlive()) {
+					ConnectionType |= 2;
+                }
+            }
+        }
+		
+		return ConnectionType;
+	}
+
+	public boolean AndroidThunkJava_AndroidFileServer_Stop(boolean bUSB, boolean bNetwork)
+	{
+        Set<Thread> threads = Thread.getAllStackTraces().keySet();
+        boolean bFoundThread = false;
+        for (Thread thread : threads) {
+            String name = thread.getName();
+            if (bUSB && name.equals("UnrealAFS-USB"))
+            {
+                if (thread.isAlive()) {
+					thread.interrupt();
+                    bFoundThread = true;
+                }
+            }
+            else if (bNetwork && name.equals("UnrealAFS-WiFi"))
+            {
+                if (thread.isAlive()) {
+					thread.interrupt();
+                    bFoundThread = true;
+                }
+            }
+        }
+		
+		return bFoundThread;
+	}
+
+	public boolean AndroidThunkJava_AndroidFileServer_Start(boolean bUSB, boolean bNetwork, int port)
+	{
+		boolean bUSBRunning = false;
+		boolean bWiFiRunning = false;
+		
+		if (port == 0)
+		{
+			port = 57099;
+		}
+		
+        Set<Thread> threads = Thread.getAllStackTraces().keySet();
+        boolean bFoundThread = false;
+        for (Thread thread : threads) {
+            String name = thread.getName();
+            if (name.equals("UnrealAFS-USB"))
+            {
+                if (thread.isAlive()) {
+                    bUSBRunning = true;
+                }
+            }
+            else if (name.equals("UnrealAFS-WiFi"))
+            {
+                if (thread.isAlive()) {
+                    bWiFiRunning = true;
+                    break;
+                }
+            }
+        }
+
+		if (bUSBRunning && bWiFiRunning)
+		{
+			return true;
+		}
+
+		String packageName = getPackageName();
+		PackageManager pm = getPackageManager();
+		
+		int versionCode = 0;
+		try
+		{
+			versionCode = pm.getPackageInfo(packageName, 0).versionCode;
+		}
+		catch (Exception e)
+		{
+			// if the above failed, then, we can't use obbs
+		}
+
+		String EngineVersion = "5.0.0";
+		String ProjectName = packageName;
+		ProjectName = ProjectName.substring(ProjectName.lastIndexOf('.') + 1);
+		boolean IsShipping = false;
+		boolean UseExternalFilesDir = false;
+		boolean PublicLogFiles = false;
+		boolean AllowNetworkConnection = false;
+		try {
+			ApplicationInfo ai = pm.getApplicationInfo(packageName, PackageManager.GET_META_DATA);
+			Bundle bundle = ai.metaData;
+
+			if (bundle.containsKey("com.epicgames.unreal.GameActivity.EngineVersion"))
+			{
+				EngineVersion = bundle.getString("com.epicgames.unreal.GameActivity.EngineVersion");
+			}
+			if (bundle.containsKey("com.epicgames.unreal.GameActivity.ProjectName"))
+			{
+				ProjectName = bundle.getString("com.epicgames.unreal.GameActivity.ProjectName");
+			}
+			if (bundle.containsKey("com.epicgames.unreal.GameActivity.BuildConfiguration"))
+			{
+				String Configuration = bundle.getString("com.epicgames.unreal.GameActivity.BuildConfiguration");
+				IsShipping = Configuration.contains("Shipping");
+			}
+			if (bundle.containsKey("com.epicgames.unreal.GameActivity.bUseExternalFilesDir"))
+			{
+				UseExternalFilesDir = bundle.getBoolean("com.epicgames.unreal.GameActivity.bUseExternalFilesDir");
+			}
+			if (bundle.containsKey("com.epicgames.unreal.GameActivity.bPublicLogFiles"))
+			{
+				PublicLogFiles = bundle.getBoolean("com.epicgames.unreal.GameActivity.bPublicLogFiles");
+			}
+			if (bundle.containsKey("com.epicgames.unreal.RemoteFileManager.bAllowNetworkConnection"))
+			{
+				AllowNetworkConnection = bundle.getBoolean("com.epicgames.unreal.RemoteFileManager.bAllowNetworkConnection");
+			}
+		}
+		catch (PackageManager.NameNotFoundException | NullPointerException e)
+		{
+			Log.debug("Error when accessing application metadata");
+		}
+
+		String internal = getFilesDir().getAbsolutePath();
+		String external = getExternalFilesDir(null).getAbsolutePath();
+		String storage = android.os.Environment.getExternalStorageDirectory().getAbsolutePath();
+		String obbdir = getObbDir().getAbsolutePath();
+
+		if (bUSB && !bUSBRunning)
+		{
+			RemoteFileManager fileManager = new RemoteFileManager(true, port, internal, external, storage, obbdir, packageName, versionCode, ProjectName, EngineVersion, IsShipping, PublicLogFiles);
+			new Thread(fileManager, "UnrealAFS-USB").start();
+		}
+		if (AllowNetworkConnection && bNetwork && !bWiFiRunning)
+		{
+			RemoteFileManager fileManager = new RemoteFileManager(false, port, internal, external, storage, obbdir, packageName, versionCode, ProjectName, EngineVersion, IsShipping, PublicLogFiles);
+			new Thread(fileManager, "UnrealAFS-WiFi").start();
+		}
+		
+		return true;
+	}
+]]>
+				</insert>
+			</true>
+			<false>
+	public static boolean AndroidFileServer_Verify(String Token)
+	{
+		return false;
+	}
+
+	public boolean AndroidThunkJava_AndroidFileServer_IsRunning()
+	{
+		return false;
+	}
+
+	public boolean AndroidThunkJava_AndroidFileServer_Start()
+	{
+		return false;
+	}
+
+	public boolean AndroidThunkJava_AndroidFileServer_Stop()
+	{
+		return false;
+	}
+			</false>
+		</if>
+	</gameActivityClassAdditions>
+
+</root>
diff --git a/Engine/Plugins/Runtime/AndroidFileServer/Source/AndroidFileServer/Classes/AndroidFileServerBPLibrary.h b/Engine/Plugins/Runtime/AndroidFileServer/Source/AndroidFileServer/Classes/AndroidFileServerBPLibrary.h
new file mode 100644
--- /dev/null
+++ b/Engine/Plugins/Runtime/AndroidFileServer/Source/AndroidFileServer/Classes/AndroidFileServerBPLibrary.h
@@ -0,0 +1,38 @@
+// Copyright Epic Games, Inc. All Rights Reserved.
+
+#pragma once
+
+#include "CoreMinimal.h"
+#include "UObject/ObjectMacros.h"
+#include "Kismet/BlueprintFunctionLibrary.h"
+#include "AndroidFileServerBPLibrary.generated.h"
+
+UENUM()
+namespace EAFSActiveType
+{
+	enum Type
+	{
+		None = 0 UMETA(DisplayName = "None"),
+		USBOnly = 1 UMETA(DisplayName = "USB only"),
+		NetworkOnly = 2 UMETA(DisplayName = "Network only"),
+		Combined = 3 UMETA(DisplayName = "USB and Network combined"),
+	};
+}
+
+UCLASS()
+class UAndroidFileServerBPLibrary : public UBlueprintFunctionLibrary
+{
+	GENERATED_BODY()
+public:
+	/** Request startup of Android FileServer */
+	UFUNCTION(BlueprintCallable, Category = "AndroidFileServer")
+	static bool StartFileServer(bool bUSB = true, bool bNetwork = false, int32 Port = 57099);
+
+	/** Request termination of Android FileServer */
+	UFUNCTION(BlueprintCallable, Category = "AndroidFileServer")
+	static bool StopFileServer(bool bUSB = true, bool bNetwork = true);
+
+	/** Check if Android FileServer is running */
+	UFUNCTION(BlueprintCallable, Category = "AndroidFileServer")
+	static TEnumAsByte<EAFSActiveType::Type> IsFileServerRunning();
+};
diff --git a/Engine/Plugins/Runtime/AndroidFileServer/Source/AndroidFileServer/Private/AndroidFileServer.cpp b/Engine/Plugins/Runtime/AndroidFileServer/Source/AndroidFileServer/Private/AndroidFileServer.cpp
new file mode 100644
--- /dev/null
+++ b/Engine/Plugins/Runtime/AndroidFileServer/Source/AndroidFileServer/Private/AndroidFileServer.cpp
@@ -0,0 +1,22 @@
+// Copyright Epic Games, Inc. All Rights Reserved.
+
+#include "AndroidFileServer.h"
+
+#define LOCTEXT_NAMESPACE "FAndroidFileServerModule"
+
+void FAndroidFileServerModule::StartupModule()
+{
+	// This code will execute after your module is loaded into memory; the exact timing is specified in the .uplugin file per-module
+	
+}
+
+void FAndroidFileServerModule::ShutdownModule()
+{
+	// This function may be called during shutdown to clean up your module.  For modules that support dynamic reloading,
+	// we call this function before unloading the module.
+	
+}
+
+#undef LOCTEXT_NAMESPACE
+	
+IMPLEMENT_MODULE(FAndroidFileServerModule, AndroidFileServer)
diff --git a/Engine/Plugins/Runtime/AndroidFileServer/Source/AndroidFileServer/Private/AndroidFileServerBPLibrary.cpp b/Engine/Plugins/Runtime/AndroidFileServer/Source/AndroidFileServer/Private/AndroidFileServerBPLibrary.cpp
new file mode 100644
--- /dev/null
+++ b/Engine/Plugins/Runtime/AndroidFileServer/Source/AndroidFileServer/Private/AndroidFileServerBPLibrary.cpp
@@ -0,0 +1,63 @@
+// Copyright Epic Games, Inc. All Rights Reserved.
+
+#include "AndroidFileServerBPLibrary.h"
+#include "AndroidFileServer.h"
+
+#if PLATFORM_ANDROID
+
+#include "Android/AndroidJNI.h"
+#include "Android/AndroidApplication.h"
+
+#if USE_ANDROID_JNI
+#include <jni.h>
+#endif
+#endif
+
+bool UAndroidFileServerBPLibrary::StartFileServer(bool bUSB, bool bNetwork, int32 Port)
+{
+	bool result = false;
+#if PLATFORM_ANDROID
+	if (JNIEnv* Env = FAndroidApplication::GetJavaEnv())
+	{
+		static jmethodID StartFunc = FJavaWrapper::FindMethod(Env, FJavaWrapper::GameActivityClassID, "AndroidThunkJava_AndroidFileServer_Start", "(ZZI)Z", false);
+		if (StartFunc != nullptr)
+		{
+			result = FJavaWrapper::CallBooleanMethod(Env, FJavaWrapper::GameActivityThis, StartFunc, bUSB, bNetwork, Port);
+		}
+	}
+#endif
+	return result;
+}
+
+bool UAndroidFileServerBPLibrary::StopFileServer(bool bUSB, bool bNetwork)
+{
+	bool result = false;
+#if PLATFORM_ANDROID
+	if (JNIEnv* Env = FAndroidApplication::GetJavaEnv())
+	{
+		static jmethodID StopFunc = FJavaWrapper::FindMethod(Env, FJavaWrapper::GameActivityClassID, "AndroidThunkJava_AndroidFileServer_Stop", "(ZZ)Z", false);
+		if (StopFunc != nullptr)
+		{
+			result = FJavaWrapper::CallBooleanMethod(Env, FJavaWrapper::GameActivityThis, StopFunc, bUSB, bNetwork);
+		}
+	}
+#endif
+	return result;
+}
+
+TEnumAsByte<EAFSActiveType::Type> UAndroidFileServerBPLibrary::IsFileServerRunning()
+{
+	TEnumAsByte<EAFSActiveType::Type> result = EAFSActiveType::None;
+#if PLATFORM_ANDROID
+	if (JNIEnv* Env = FAndroidApplication::GetJavaEnv())
+	{
+		static jmethodID IsRunningFunc = FJavaWrapper::FindMethod(Env, FJavaWrapper::GameActivityClassID, "AndroidThunkJava_AndroidFileServer_IsRunning", "()I", false);
+		if (IsRunningFunc != nullptr)
+		{
+			result = (TEnumAsByte<EAFSActiveType::Type>)FJavaWrapper::CallIntMethod(Env, FJavaWrapper::GameActivityThis, IsRunningFunc);
+		}
+	}
+#endif
+	return result;
+}
+
diff --git a/Engine/Plugins/Runtime/AndroidFileServer/Source/AndroidFileServer/Public/AndroidFileServer.h b/Engine/Plugins/Runtime/AndroidFileServer/Source/AndroidFileServer/Public/AndroidFileServer.h
new file mode 100644
--- /dev/null
+++ b/Engine/Plugins/Runtime/AndroidFileServer/Source/AndroidFileServer/Public/AndroidFileServer.h
@@ -0,0 +1,14 @@
+// Copyright Epic Games, Inc. All Rights Reserved.
+
+#pragma once
+
+#include "Modules/ModuleManager.h"
+
+class FAndroidFileServerModule : public IModuleInterface
+{
+public:
+
+	/** IModuleInterface implementation */
+	virtual void StartupModule() override;
+	virtual void ShutdownModule() override;
+};
diff --git a/Engine/Plugins/Runtime/AndroidFileServer/Source/AndroidFileServerEditor/AndroidFileServerEditor.Build.cs b/Engine/Plugins/Runtime/AndroidFileServer/Source/AndroidFileServerEditor/AndroidFileServerEditor.Build.cs
new file mode 100644
--- /dev/null
+++ b/Engine/Plugins/Runtime/AndroidFileServer/Source/AndroidFileServerEditor/AndroidFileServerEditor.Build.cs
@@ -0,0 +1,22 @@
+// Copyright Epic Games, Inc. All Rights Reserved.
+
+namespace UnrealBuildTool.Rules
+{
+	public class AndroidFileServerEditor : ModuleRules
+	{
+		public AndroidFileServerEditor(ReadOnlyTargetRules Target) : base(Target)
+		{
+			PrivateDependencyModuleNames.AddRange(
+				new string[] {
+					"Core",
+					"CoreUObject",
+					"UnrealEd",
+				});
+
+			PrivateIncludePaths.AddRange(
+				new string[] {
+					"AndroidFileServerEditor/Private",
+				});
+		}
+	}
+}
diff --git a/Engine/Plugins/Runtime/AndroidFileServer/Source/AndroidFileServerEditor/Private/AndroidFileServerEditor.cpp b/Engine/Plugins/Runtime/AndroidFileServer/Source/AndroidFileServerEditor/Private/AndroidFileServerEditor.cpp
new file mode 100644
--- /dev/null
+++ b/Engine/Plugins/Runtime/AndroidFileServer/Source/AndroidFileServerEditor/Private/AndroidFileServerEditor.cpp
@@ -0,0 +1,42 @@
+// Copyright Epic Games, Inc. All Rights Reserved.
+
+#include "AndroidFileServerEditor.h"
+#include "ISettingsModule.h"
+#include "AndroidFileServerRuntimeSettings.h"
+
+/**
+ * Implements the AndroidFileServerEditor module.
+ */
+
+#define LOCTEXT_NAMESPACE "AndroidFileServer"
+
+void FAndroidFileServerEditorModule::StartupModule()
+{
+	// register settings
+	ISettingsModule* SettingsModule = FModuleManager::GetModulePtr<ISettingsModule>("Settings");
+
+	if (SettingsModule != nullptr)	
+	{
+		SettingsModule->RegisterSettings("Project", "Plugins", "AndroidFileServer",
+			LOCTEXT("AndroidFileServerSettingsName", "AndroidFileServer"),
+			LOCTEXT("AndroidFileServerSettingsDescription", "Project settings for AndroidFileServer plugin"),
+			GetMutableDefault<UAndroidFileServerRuntimeSettings>()
+		);
+	}
+}
+
+void FAndroidFileServerEditorModule::ShutdownModule()
+{
+	// unregister settings
+	ISettingsModule* SettingsModule = FModuleManager::GetModulePtr<ISettingsModule>("Settings");
+
+	if (SettingsModule != nullptr)
+	{
+       		SettingsModule->UnregisterSettings("Project", "Plugins", "AndroidFileServer");
+	}
+}
+
+
+IMPLEMENT_MODULE(FAndroidFileServerEditorModule, AndroidFileServerEditor);
+
+#undef LOCTEXT_NAMESPACE
diff --git a/Engine/Plugins/Runtime/AndroidFileServer/Source/AndroidFileServerEditor/Private/AndroidFileServerRuntimeSettings.cpp b/Engine/Plugins/Runtime/AndroidFileServer/Source/AndroidFileServerEditor/Private/AndroidFileServerRuntimeSettings.cpp
new file mode 100644
--- /dev/null
+++ b/Engine/Plugins/Runtime/AndroidFileServer/Source/AndroidFileServerEditor/Private/AndroidFileServerRuntimeSettings.cpp
@@ -0,0 +1,36 @@
+// Copyright Epic Games, Inc. All Rights Reserved.
+
+#include "AndroidFileServerRuntimeSettings.h"
+
+//////////////////////////////////////////////////////////////////////////
+// UAndroidFileServerRuntimeSettings
+
+UAndroidFileServerRuntimeSettings::UAndroidFileServerRuntimeSettings(const FObjectInitializer& ObjectInitializer)
+        : Super(ObjectInitializer)
+        , bEnablePlugin(true)
+		, bAllowNetworkConnection(true)
+		, SecurityToken(TEXT("[AUTO]"))
+		, bIncludeInShipping(false)
+		, bAllowExternalStartInShipping(false)
+		, bCompileAFSProject(false)
+		, bUseCompression(false)
+		, bLogFiles(false)
+		, bReportStats(false)
+		, ConnectionType(EAFSConnectionType::USBOnly)
+		, bUseManualIPAddress(false)
+		, ManualIPAddress(TEXT(""))
+{
+}
+
+void UAndroidFileServerRuntimeSettings::PostInitProperties()
+{
+	Super::PostInitProperties();
+
+	// Auto-generate a security token first time
+	if (SecurityToken == "[AUTO]")
+	{
+		FGuid UniqueId = FGuid::NewGuid();
+		SecurityToken = UniqueId.ToString();
+		TryUpdateDefaultConfigFile();
+	}
+}
\ No newline at end of file
diff --git a/Engine/Plugins/Runtime/AndroidFileServer/Source/AndroidFileServerEditor/Public/AndroidFileServerEditor.h b/Engine/Plugins/Runtime/AndroidFileServer/Source/AndroidFileServerEditor/Public/AndroidFileServerEditor.h
new file mode 100644
--- /dev/null
+++ b/Engine/Plugins/Runtime/AndroidFileServer/Source/AndroidFileServerEditor/Public/AndroidFileServerEditor.h
@@ -0,0 +1,19 @@
+// Copyright Epic Games, Inc. All Rights Reserved.
+
+#pragma once
+
+#include "Modules/ModuleInterface.h"
+#include "Modules/ModuleManager.h"
+
+#include "AndroidFileServerRuntimeSettings.h"
+
+class FAndroidFileServerEditorModule
+	: public IModuleInterface
+{
+public:
+
+	//~ IModuleInterface interface
+
+	virtual void StartupModule() override;
+	virtual void ShutdownModule() override;
+};
diff --git a/Engine/Plugins/Runtime/AndroidFileServer/Source/AndroidFileServerEditor/Public/AndroidFileServerRuntimeSettings.h b/Engine/Plugins/Runtime/AndroidFileServer/Source/AndroidFileServerEditor/Public/AndroidFileServerRuntimeSettings.h
new file mode 100644
--- /dev/null
+++ b/Engine/Plugins/Runtime/AndroidFileServer/Source/AndroidFileServerEditor/Public/AndroidFileServerRuntimeSettings.h
@@ -0,0 +1,83 @@
+// Copyright Epic Games, Inc. All Rights Reserved.
+
+#pragma once
+
+#include "CoreMinimal.h"
+#include "UObject/ObjectMacros.h"
+#include "UObject/Object.h"
+#include "Engine/EngineTypes.h"
+#include "AndroidFileServerRuntimeSettings.generated.h"
+
+/**
+* Implements the settings for the AndroidFileServer plugin.
+*/
+
+UENUM()
+namespace EAFSConnectionType
+{
+	enum Type
+	{
+		USBOnly = 0 UMETA(DisplayName = "USB only"),
+		NetworkOnly = 1 UMETA(DisplayName = "Network only"),
+		Combined = 2 UMETA(DisplayName = "USB and Network combined"),
+	};
+}
+
+
+UCLASS(Config = Engine, DefaultConfig)
+class UAndroidFileServerRuntimeSettings : public UObject
+{
+	GENERATED_UCLASS_BODY()
+
+	// Enable Android FileServer plugin
+	UPROPERTY(EditAnywhere, config, Category = Packaging)
+	bool bEnablePlugin;
+
+	// Allow FileServer connection using network
+	UPROPERTY(EditAnywhere, config, Category = Packaging, Meta = (EditCondition = "bEnablePlugin"))
+	bool bAllowNetworkConnection;
+
+	// Optional security token required to start FileServer (leave empty to disable)
+	UPROPERTY(EditAnywhere, config, Category = Packaging, Meta = (EditCondition = "bEnablePlugin"))
+	FString SecurityToken;
+
+	// Embed FileServer in Shipping builds
+	UPROPERTY(EditAnywhere, config, Category = Packaging, Meta=(EditCondition = "bEnablePlugin"))
+	bool bIncludeInShipping;
+
+	// Allow FileServer to be started in Shipping builds with UnrealAndroidFileTool
+	UPROPERTY(EditAnywhere, config, Category = Packaging, Meta = (EditCondition = "bEnablePlugin && bIncludeInShipping"))
+	bool bAllowExternalStartInShipping;
+
+	// Compile standalone AFS project
+	UPROPERTY(EditAnywhere, config, Category = Packaging, Meta=(EditCondition = "bEnablePlugin"))
+	bool bCompileAFSProject;
+
+	// Enable compression during data transfer
+	UPROPERTY(EditAnywhere, config, Category = Deployment, Meta=(EditCondition = "bEnablePlugin"))
+	bool bUseCompression;
+
+	// Log files transferred
+	UPROPERTY(EditAnywhere, config, Category = Deployment, Meta=(EditCondition = "bEnablePlugin"))
+	bool bLogFiles;
+
+	// Report transfer rate statistics
+	UPROPERTY(EditAnywhere, config, Category = Deployment, Meta=(EditCondition = "bEnablePlugin"))
+	bool bReportStats;
+
+	// How to connect to file server (USB cable, Network, or combined)
+	UPROPERTY(EditAnywhere, config, Category = Connection, Meta=(EditCondition = "bEnablePlugin"))
+	TEnumAsByte<EAFSConnectionType::Type> ConnectionType;
+
+	// Use manual IP address instead of automatic query from device (only for single device deploys!)
+	UPROPERTY(EditAnywhere, config, Category = Connection, Meta=(EditCondition = "ConnectionType == EConnectionType::NetworkOnly || ConnectionType == EConnectionType::Combined"), Meta=(DisplayName = "Use Manual IP Address?"))
+	bool bUseManualIPAddress;
+
+	// IP address of device to use
+	UPROPERTY(EditAnywhere, config, Category = Connection, Meta=(EditCondition = "bUseManualIPAddress"), Meta=(DisplayName = "Manual IP Address"))
+	FString ManualIPAddress;
+
+private:
+	// UObject interface
+	virtual void PostInitProperties() override;
+};
diff --git a/Engine/Source/Programs/AutomationTool/Android/AndroidFileClient.cs b/Engine/Source/Programs/AutomationTool/Android/AndroidFileClient.cs
new file mode 100644
--- /dev/null
+++ b/Engine/Source/Programs/AutomationTool/Android/AndroidFileClient.cs
@@ -0,0 +1,2302 @@
+
+// Copyright Epic Games, Inc. All Rights Reserved.
+
+using System;
+using System.Collections.Generic;
+using System.Linq;
+using System.Text;
+using System.IO;
+using System.Diagnostics;
+using System.Net.NetworkInformation;
+using System.Threading;
+using AutomationTool;
+using UnrealBuildTool;
+using Ionic.Zip;
+using EpicGames.Core;
+using UnrealBuildBase;
+using System.Text.RegularExpressions;
+using AutomationScripts;
+using System.Net;
+using System.Net.Sockets;
+using System.IO.Compression;
+
+
+namespace AutomationTool
+{
+	public class AndroidFileClient
+	{
+		private const bool TimeBatching = false;
+
+		private const int DefaultPort = 57099;
+		private const int ServerVersion = 100;
+		private const int FILE_BUFFERSIZE = 65536 * 16;
+		private const int SOCKET_READBUFFER = 625 * 1024;
+		private const int SOCKET_SENDBUFFER = 625 * 1024;
+		private const int BATCH_BUFFERSIZE = !TimeBatching ? 65536 * 16 : 200000000;
+		private const int BATCH_SENDSIZE = 65536 * 16;
+
+		private string Device;
+		private int ServerPort;
+		private int HostPort;
+
+		private static OptimalADB adb = new OptimalADB();
+		static string LastListenDevice = "";
+		static bool bHave_ss = true;
+		static bool bHave_netstat = true;
+
+		private Socket ClientSocket = null;
+
+		private FileStream RecordFile = null;
+		private BufferedStream RecordStream = null;
+		private bool bBlockSend = false;
+
+		private byte[] BatchBuffer = null;
+		private int BatchBufferIndex = 0;
+
+		private string LastBaseDir = "";
+
+		private static int Command_Terminate = 0;
+		private static int Command_Close = 1;
+		private static int Command_Info = 2;
+		private static int Command_Query = 3;
+		private static int Command_GetProp = 4;
+		private static int Command_SetBaseDir = 5;
+
+		private static int Command_DirExists = 10;
+		private static int Command_DirList = 11;
+		private static int Command_DirListFlat = 12;
+		private static int Command_DirCreate = 13;
+		private static int Command_DirDelete = 14;
+		private static int Command_DirDeleteRecurse = 15;
+
+		private static int Command_FileExists = 20;
+		private static int Command_FileDelete = 21;
+		private static int Command_FileCopy = 22;
+		private static int Command_FileMove = 23;
+		private static int Command_FileRead = 24;
+		private static int Command_FileWrite = 25;
+		private static int Command_FileWriteCompressed = 26;
+
+		// ------------------ Helpers to deal with ADB ------------------ 
+
+		// optional way to talk directly with the adbd instead of shelling out to adb.exe
+		public class OptimalADB
+		{
+			private const int DefaultPort = 5037;
+
+			private string Executable;
+			private int Port = 6037;
+			private Socket ClientSocket = null;
+
+			public OptimalADB(string ADBExecutable = "", int inPort = DefaultPort)
+			{
+				Executable = (ADBExecutable != "") ? ADBExecutable : Environment.ExpandEnvironmentVariables("%ANDROID_HOME%/platform-tools/adb" + (RuntimePlatform.IsWindows ? ".exe" : ""));
+				Port = inPort;
+			}
+
+			~OptimalADB()
+			{
+				CloseConnection();
+			}
+
+			private bool OpenConnection()
+			{
+				if (ClientSocket != null)
+				{
+					return true;
+				}
+
+				// create TCP/IP socket
+				try
+				{
+					IPHostEntry ipHost = Dns.GetHostEntry(Dns.GetHostName());
+					IPAddress ipAddr = IPAddress.Parse("127.0.0.1");
+					IPEndPoint localEndPoint = new IPEndPoint(ipAddr, Port);
+
+					// create TCP/IP socket
+					ClientSocket = new Socket(ipAddr.AddressFamily, SocketType.Stream, ProtocolType.Tcp);
+					ClientSocket.Connect(localEndPoint);
+				}
+				catch (Exception)
+				{
+					// adb server not running so start it
+					IProcessResult Result = CommandUtils.Run(Executable, "start-server", null, CommandUtils.ERunOptions.NoLoggingOfRunCommand);
+					if (Result.ExitCode != 0)
+					{
+						Log.TraceError("Unable to execute adb to start adb server");
+						return false;
+					}
+
+					// try again
+					return OpenConnection();
+				}
+
+				return true;
+			}
+
+			private void CloseConnection()
+			{
+				if (ClientSocket != null)
+				{
+					try
+					{
+						ClientSocket.Shutdown(SocketShutdown.Both);
+						ClientSocket.Close();
+						ClientSocket.Dispose();
+						ClientSocket = null;
+					}
+					catch (SocketException se)
+					{
+						Log.TraceWarning("SocketException: {0}", se.ToString());
+					}
+					catch (Exception e)
+					{
+						Log.TraceError("Unexpected Exception: {0}", e.ToString());
+					}
+				}
+			}
+
+			private bool sendData(string Data)
+			{
+				if (!OpenConnection())
+				{
+					return false;
+				}
+
+				byte[] buffer = Encoding.UTF8.GetBytes(Data.Length.ToString("X4") + Data);
+				int sentBytes = ClientSocket.Send(buffer);
+
+				return true;
+			}
+
+			private bool receiveData(out string Output, bool bOnce = false)
+			{
+				Output = "";
+
+				if (ClientSocket == null)
+				{
+					return false;
+				}
+
+				byte[] buffer = new byte[65536];
+
+				int bytesRecv = ClientSocket.Receive(buffer);
+				string WorkBuffer = Encoding.UTF8.GetString(buffer, 0, bytesRecv);
+
+				if (WorkBuffer.Length >= 4)
+				{
+					if (WorkBuffer.Substring(0, 4) == "OKAY")
+					{
+						WorkBuffer = WorkBuffer.Substring(4);
+						if (bOnce)
+						{
+							while (ClientSocket.Available > 0)
+							{
+								bytesRecv = ClientSocket.Receive(buffer);
+								WorkBuffer += Encoding.UTF8.GetString(buffer, 0, bytesRecv);
+							}
+							Output += WorkBuffer;
+							return true;
+						}
+						int remaining = 0;
+						while (true)
+						{
+							while (ClientSocket.Available > 0 || WorkBuffer.Length < 4)
+							{
+								bytesRecv = ClientSocket.Receive(buffer);
+								WorkBuffer += Encoding.UTF8.GetString(buffer, 0, bytesRecv);
+							}
+							if (WorkBuffer.Length >= 4)
+							{
+								if (WorkBuffer.Substring(0, 4) == "OKAY")
+								{
+									WorkBuffer = WorkBuffer.Substring(4);
+									continue;
+								}
+								try
+								{
+									remaining = int.Parse(WorkBuffer.Substring(0, 4), System.Globalization.NumberStyles.HexNumber);
+									WorkBuffer = WorkBuffer.Substring(4);
+									remaining -= WorkBuffer.Length;
+								}
+								catch (Exception)
+								{
+									while (ClientSocket.Available > 0)
+									{
+										bytesRecv = ClientSocket.Receive(buffer);
+										WorkBuffer += Encoding.UTF8.GetString(buffer, 0, bytesRecv);
+									}
+									Output = WorkBuffer;
+									return true;
+								}
+								break;
+							}
+						}
+						while (remaining > 0)
+						{
+							bytesRecv = ClientSocket.Receive(buffer);
+							WorkBuffer += Encoding.UTF8.GetString(buffer, 0, bytesRecv);
+							remaining -= bytesRecv;
+						}
+						Output = WorkBuffer;
+						return true;
+					}
+				}
+
+				// fail
+				Output = WorkBuffer;
+				return false;
+			}
+
+			private bool execute(string Command, out string Output, bool bOnce = false)
+			{
+				bool Result = true;
+				Output = "";
+
+				if (sendData(Command))
+				{
+					Result = receiveData(out Output, bOnce);
+				}
+				return Result;
+			}
+
+			private bool executeDevice(string Device, string Command, out string Output, bool bOnce = false)
+			{
+				if (!execute("host:transport" + (Device == "" ? "-any" : ":" + Device), out Output, true))
+				{
+					return false;
+				}
+				return execute(Command, out Output, bOnce);
+			}
+
+			public string GetVersion()
+			{
+				string Result;
+				execute("host:version", out Result);
+				CloseConnection();
+				return Result;
+			}
+
+			public string GetDevices()
+			{
+				string Result;
+				execute("host:devices", out Result);
+				CloseConnection();
+				return Result;
+			}
+
+			private string shellEscape(string source)
+			{
+				// wraps whole thing in ' ' after escaping it with \'
+				return "'" + source.Replace("'", "\'") + "'";
+			}
+
+			public string Shell(string Device, string Command)
+			{
+				string Result = "FAIL";
+
+				if (!execute("host:transport" + (Device == "" ? "-any" : ":" + Device), out Result, true))
+				{
+					CloseConnection();
+					return Result;
+				}
+
+				Result = "FAIL";
+				//if (!sendData("shell:" + shellEscape(Command)))
+				if (!sendData("shell:" + Command))
+				{
+					CloseConnection();
+					return Result;
+				}
+
+				byte[] buffer = new byte[65536];
+				Result = "";
+				while (Result.Length < 4)
+				{
+					int bytesRecv = ClientSocket.Receive(buffer);
+					Result += Encoding.UTF8.GetString(buffer, 0, bytesRecv);
+				}
+				if (Result.Substring(0, 4) != "OKAY")
+				{
+					CloseConnection();
+					return Result;
+				}
+				Result = Result.Substring(4);
+
+				// there can be a delay before any more results (and may not be any)
+				if (ClientSocket.Available == 0)
+				{
+					Thread.Sleep(50);
+				}
+				while (ClientSocket.Available > 0)
+				{
+					int bytesRecv = ClientSocket.Receive(buffer);
+					Result += Encoding.UTF8.GetString(buffer, 0, bytesRecv);
+					if (ClientSocket.Available == 0)
+					{
+						Thread.Sleep(50);
+					}
+				}
+
+				CloseConnection();
+				return Result;
+			}
+
+			public string ListForward(string Device)
+			{
+				string Result;
+
+				executeDevice(Device, "host:list-forward", out Result);
+				CloseConnection();
+				return Result;
+			}
+
+			public string AddForward(string Device, string Local, string Remote)
+			{
+				string Result;
+
+				executeDevice(Device, "host:forward:" + Local + ";" + Remote, out Result);
+				CloseConnection();
+				return Result;
+			}
+
+			public string RemoveForward(string Device, string Local)
+			{
+				string Result;
+
+				executeDevice(Device, "host:killforward:" + Local, out Result, true);
+				CloseConnection();
+				return Result;
+			}
+		}
+
+		private static string GetAdbCommandLine(string SerialNumber, string Args)
+		{
+			if (string.IsNullOrEmpty(SerialNumber) == false)
+			{
+				SerialNumber = "-s " + SerialNumber;
+			}
+
+			return string.Format("{0} {1}", SerialNumber, Args);
+		}
+
+		static string LastSpewFilename = "";
+
+		public static string ADBSpewFilter(string Message)
+		{
+			if (Message.StartsWith("[") && Message.Contains("%]"))
+			{
+				int LastIndex = Message.IndexOf(":");
+				LastIndex = LastIndex == -1 ? Message.Length : LastIndex;
+
+				if (Message.Length > 7)
+				{
+					string Filename = Message.Substring(7, LastIndex - 7);
+					if (Filename == LastSpewFilename)
+					{
+						return null;
+					}
+					LastSpewFilename = Filename;
+				}
+				return Message;
+			}
+			return Message;
+		}
+
+		public static IProcessResult RunAdbCommand(string SerialNumber, string Args, string Input = null, CommandUtils.ERunOptions Options = CommandUtils.ERunOptions.NoLoggingOfRunCommand, bool bShouldLogCommand = false)
+		{
+			string AdbCommand = Environment.ExpandEnvironmentVariables("%ANDROID_HOME%/platform-tools/adb" + (RuntimePlatform.IsWindows ? ".exe" : ""));
+			if (Options.HasFlag(CommandUtils.ERunOptions.AllowSpew) || Options.HasFlag(CommandUtils.ERunOptions.SpewIsVerbose))
+			{
+				LastSpewFilename = "";
+				return CommandUtils.Run(AdbCommand, GetAdbCommandLine(SerialNumber, Args), Input, Options, SpewFilterCallback: new ProcessResult.SpewFilterCallbackType(ADBSpewFilter));
+			}
+			return CommandUtils.Run(AdbCommand, GetAdbCommandLine(SerialNumber, Args), Input, Options);
+		}
+
+		public static void GetConnectedDevices(out List<string> Devices)
+		{
+			Devices = new List<string>();
+
+
+			string Result = adb.GetDevices();
+			if (Result.Length > 0)
+			{
+				string[] LogLines = Result.Split(new char[] { '\n', '\r' });
+				for (int i = 0; i < LogLines.Length; ++i)
+				{
+					string[] DeviceLine = LogLines[i].Split(new char[] { '\t' });
+
+					if (DeviceLine.Length == 2)
+					{
+						// the second param should be "device"
+						// if it's not setup correctly it might be "unattached" or "powered off" or something like that
+						// warning in that case
+						if (DeviceLine[1] == "device")
+						{
+							Devices.Add("@" + DeviceLine[0]);
+						}
+						else
+						{
+							CommandUtils.LogWarning("Device attached but in bad state {0}:{1}", DeviceLine[0], DeviceLine[1]);
+						}
+					}
+				}
+			}
+		}
+
+		// Checks for existing forward of device port and returns local host port if found, otherwise sets one
+		public static int ForwardPort_FindOrAdd(string Device, int DevicePort)
+		{
+			int HostPort = -1;
+
+			// Check if port already forwarded to this device
+			string Result = adb.ListForward(Device);
+			if (Result.Length > 0)
+			{
+				string[] LogLines = Result.Split(new char[] { '\n', '\r' });
+				for (int i = 0; i < LogLines.Length; ++i)
+				{
+					string[] DeviceLine = LogLines[i].Split(new char[] { '\t', ' ' });
+					if (DeviceLine.Length == 3)
+					{
+						if (DeviceLine[1].StartsWith("tcp:") && DeviceLine[2].StartsWith("tcp:"))
+						{
+							int ListDevicePort = -1;
+							if (int.TryParse(DeviceLine[2].Substring(4), out ListDevicePort))
+							{
+								if (ListDevicePort == DevicePort)
+								{
+									if (int.TryParse(DeviceLine[1].Substring(4), out HostPort))
+									{
+										return HostPort;
+									}
+								}
+							}
+						}
+					}
+				}
+			}
+
+			// request forwarding from an available host port
+			Result = adb.AddForward(Device, "tcp:0", "tcp:" + DevicePort.ToString());
+			if (Result.Length > 0)
+			{
+				string[] LogLines = Result.Split(new char[] { '\n', '\r' });
+				for (int i = 0; i < LogLines.Length; ++i)
+				{
+					if (int.TryParse(LogLines[i], out HostPort))
+					{
+						return HostPort;
+					}
+				}
+			}
+
+			// did not work, return -1 as error
+			return -1;
+		}
+
+		// Checks for existing forwards of device port and removes them then requests a new one to return
+		// This deals better with stale binds and leftover forward requests
+		public static int ForwardPort_Add(string Device, int DevicePort)
+		{
+			int HostPort = -1;
+
+			// Check if port already forwarded to this device
+			string Result = adb.ListForward(Device);
+			if (Result.Length > 0)
+			{
+				string[] LogLines = Result.Split(new char[] { '\n', '\r' });
+				for (int i = 0; i < LogLines.Length; ++i)
+				{
+					string[] DeviceLine = LogLines[i].Split(new char[] { '\t', ' ' });
+					if (DeviceLine.Length == 3)
+					{
+						if (DeviceLine[1].StartsWith("tcp:") && DeviceLine[2].StartsWith("tcp:"))
+						{
+							int ListDevicePort = -1;
+							if (int.TryParse(DeviceLine[2].Substring(4), out ListDevicePort))
+							{
+								if (ListDevicePort == DevicePort)
+								{
+									adb.RemoveForward(Device, DeviceLine[1]);
+								}
+							}
+						}
+					}
+				}
+			}
+
+			// request forwarding from an available host port
+			Result = adb.AddForward(Device, "tcp:0", "tcp:" + DevicePort.ToString());
+			if (Result.Length > 0)
+			{
+				string[] LogLines = Result.Split(new char[] { '\n', '\r' });
+				for (int i = 0; i < LogLines.Length; ++i)
+				{
+					if (int.TryParse(LogLines[i], out HostPort))
+					{
+						return HostPort;
+					}
+				}
+			}
+
+			// did not work, return -1 as error
+			return -1;
+		}
+
+		public static void ForwardPort_Remove(string Device, int HostPort)
+		{
+			// request removal of forwarding from host port
+			adb.RemoveForward(Device, "tcp:" + HostPort.ToString());
+		}
+
+		public static void ReversePort_Add(string Device, int DevicePort, int HostPort)
+		{
+			// request reversing from device port to local host port
+			IProcessResult Result = RunAdbCommand(Device, "reverse tcp:" + DevicePort.ToString() + " tcp:" + HostPort.ToString());
+			Result.DisposeProcess();
+		}
+
+		public static void ReversePort_Remove(string Device, int DevicePort)
+		{
+			// request removal of forwarding from host port
+			IProcessResult Result = RunAdbCommand(Device, "reverse --remove tcp:" + DevicePort.ToString());
+			Result.DisposeProcess();
+		}
+
+		// ------------------ Helpers to deal with packets ------------------ 
+
+		private long Stats_FilesRead = 0;
+		private long Stats_FilesWrite = 0;
+
+		private long Stats_TotalBytesReceived = 0;
+		private long Stats_TotalBytesSent = 0;
+
+		private long Stats_PayloadBytesReceived = 0;
+		private long Stats_PayloadBytesSent = 0;
+
+		private long Stats_CompressedBytesReceived = 0;
+		private long Stats_CompressedBytesSent = 0;
+
+		private Stopwatch Stats_Stopwatch = null;
+
+		public void Stats_Clear(bool bStartTimer = true)
+		{
+			Stats_FilesRead = 0;
+			Stats_FilesWrite = 0;
+
+			Stats_TotalBytesReceived = 0;
+			Stats_TotalBytesSent = 0;
+
+			Stats_PayloadBytesReceived = 0;
+			Stats_PayloadBytesSent = 0;
+
+			Stats_CompressedBytesReceived = 0;
+			Stats_CompressedBytesSent = 0;
+
+			if (bStartTimer)
+			{
+				Stats_StartStopwatch();
+			}
+		}
+
+		public void Stats_StartStopwatch()
+		{
+			Stats_Stopwatch = System.Diagnostics.Stopwatch.StartNew();
+		}
+
+		public void Stats_Combine(AndroidFileClient client)
+		{
+			Stats_FilesRead += client.Stats_FilesRead;
+			Stats_FilesWrite += client.Stats_FilesWrite;
+
+			Stats_TotalBytesReceived += client.Stats_TotalBytesReceived;
+			Stats_TotalBytesSent += client.Stats_TotalBytesSent;
+
+			Stats_PayloadBytesReceived += client.Stats_PayloadBytesReceived;
+			Stats_PayloadBytesSent += client.Stats_PayloadBytesSent;
+
+			Stats_CompressedBytesReceived += client.Stats_CompressedBytesReceived;
+			Stats_CompressedBytesSent += client.Stats_CompressedBytesSent;
+
+			client.Stats_Clear(false);
+		}
+
+		public void Stats_Report()
+		{
+			Log.TraceInformation("Total bytes received:    {0}", Stats_TotalBytesReceived);
+			Log.TraceInformation("Payload bytes received:  {0}", Stats_PayloadBytesReceived);
+			Log.TraceInformation("Total files read:        {0}", Stats_FilesRead);
+			Log.TraceInformation(".");
+			Log.TraceInformation("Total bytes sent:        {0}", Stats_TotalBytesSent);
+			Log.TraceInformation("Payload bytes sent:      {0}", Stats_PayloadBytesSent);
+			Log.TraceInformation("Total files written:     {0}", Stats_FilesWrite);
+			Log.TraceInformation(".");
+			Log.TraceInformation("Uncompressed bytes recv: {0}", Stats_CompressedBytesReceived);
+			Log.TraceInformation("Uncompressed bytes sent: {0}", Stats_CompressedBytesSent);
+			Log.TraceInformation(".");
+
+			if (Stats_Stopwatch != null)
+			{
+				Stats_Stopwatch.Stop();
+
+				var ElapsedMs = Stats_Stopwatch.ElapsedMilliseconds;
+				double MBPerSec = (double)Stats_TotalBytesSent / ((double)ElapsedMs / 1000.0) / 1024.0 / 1024.0;
+				Log.TraceInformation("Actual sent:  {0:N1} MB/s ({1} in {2:N3}s)", MBPerSec, Stats_TotalBytesSent, (float)(ElapsedMs / 1000.0));
+				MBPerSec = (double)Stats_CompressedBytesSent / ((double)ElapsedMs / 1000.0) / 1024.0 / 1024.0;
+				Log.TraceInformation("Uncompressed: {0:N1} MB/s ({1} in {2:N3}s)", MBPerSec, Stats_CompressedBytesSent, (float)(ElapsedMs / 1000.0));
+			}
+		}
+
+		public void Record_Stop()
+		{
+			bBlockSend = false;
+			if (RecordStream != null)
+			{
+				RecordStream.Close();
+				RecordStream.Dispose();
+				RecordStream = null;
+			}
+			if (RecordFile != null)
+			{
+				RecordFile.Close();
+				RecordFile.Dispose();
+				RecordFile = null;
+			}
+		}
+
+		public bool Record_Begin(string Filename, bool BlockSend = false)
+		{
+			bBlockSend = BlockSend;
+			Record_Stop();
+
+			try
+			{
+				RecordFile = File.Open(Filename, FileMode.Create, FileAccess.Write);
+				try
+				{
+					RecordStream = new BufferedStream(RecordFile);
+				}
+				catch (Exception)
+				{
+					RecordFile.Close();
+					RecordFile.Dispose();
+					RecordFile = null;
+					return false;
+				}
+			}
+			catch (Exception)
+			{
+				RecordFile = null;
+				return false;
+			}
+			return true;
+		}
+
+		public bool Record_Play(string Filepath)
+		{
+			Boolean Result = false;
+
+			if (ClientSocket == null)
+			{
+				return Result;
+			}
+			if (!File.Exists(Filepath))
+			{
+				return Result;
+			}
+
+			using (FileStream fileStream = File.Open(Filepath, FileMode.Open, FileAccess.Read, FileShare.Read))
+			{
+				try
+				{
+					int bytesRead;
+					byte[] buffer = new byte[FILE_BUFFERSIZE];
+
+					while ((bytesRead = fileStream.Read(buffer, 0, FILE_BUFFERSIZE)) != 0)
+					{
+						SocketSend(true, buffer, bytesRead);
+					}
+					fileStream.Close();
+					Result = true;
+				}
+				catch (Exception e)
+				{
+					Log.TraceError("Unexpected Exception: {0}", e.ToString());
+					CloseConnection();
+				}
+			}
+			return Result;
+		}
+
+		private int SocketReceive(bool bPayload, byte[] packet)
+		{
+			int received = 0;
+			try
+			{
+				received = ClientSocket.Receive(packet);
+
+				Stats_TotalBytesReceived += received;
+				if (bPayload)
+				{
+					Stats_PayloadBytesReceived += received;
+				}
+
+				return received;
+			}
+			catch (Exception)
+			{
+			}
+			return received;
+		}
+
+		private int SocketReceive(bool bPayload, byte[] packet, int readSize)
+		{
+			int received = ClientSocket.Receive(packet, readSize, SocketFlags.None);
+
+			Stats_TotalBytesReceived += received;
+			if (bPayload)
+			{
+				Stats_PayloadBytesReceived += received;
+			}
+
+			return received;
+		}
+
+		public bool Batch_Start()
+		{
+			if (BatchBuffer != null)
+			{
+				return false;
+			}
+
+			BatchBuffer = new byte[BATCH_BUFFERSIZE];
+			BatchBufferIndex = 0;
+
+			return true;
+		}
+
+		public bool Batch_Flush()
+		{
+			if (BatchBuffer == null)
+			{
+				return false;
+			}
+
+			if (BatchBufferIndex > 0)
+			{
+				int index = 0;
+				int remaining = BatchBufferIndex;
+				while (remaining > 0)
+				{
+					int chunk = remaining < BATCH_SENDSIZE ? remaining : BATCH_SENDSIZE;
+					int sent = ClientSocket.Send(BatchBuffer, index, chunk, SocketFlags.None);
+					remaining -= sent;
+					index += sent;
+				}
+				BatchBufferIndex = 0;
+			}
+
+			return true;
+		}
+
+		public void Batch_Stop()
+		{
+			Batch_Flush();
+
+			BatchBuffer = null;
+			BatchBufferIndex = 0;
+		}
+
+		private void Batch_Append(byte[] packet, int writeSize)
+		{
+			int readOffset = 0;
+			while (writeSize > 0)
+			{
+				int copySize = (BatchBufferIndex + writeSize > BATCH_BUFFERSIZE) ? BATCH_BUFFERSIZE - BatchBufferIndex : writeSize;
+				Buffer.BlockCopy(packet, readOffset, BatchBuffer, BatchBufferIndex, copySize);
+
+				BatchBufferIndex += copySize;
+				readOffset += copySize;
+				writeSize -= copySize;
+
+				if (BatchBufferIndex == BATCH_BUFFERSIZE)
+				{
+					Batch_Flush();
+				}
+			}
+		}
+
+		private int SocketSend(bool bPayload, byte[] packet)
+		{
+			int writeSize = packet.Length;
+			int sent = writeSize;
+
+			if (BatchBuffer != null)
+			{
+				Batch_Append(packet, writeSize);
+			}
+			else if (!bBlockSend)
+			{
+				sent = ClientSocket.Send(packet);
+			}
+
+			if (RecordStream != null)
+			{
+				RecordStream.Write(packet);
+			}
+
+			Stats_TotalBytesSent += sent;
+			if (bPayload)
+			{
+				Stats_PayloadBytesSent += sent;
+			}
+
+			if (sent < writeSize)
+			{
+				// did not send all
+			}
+
+			return sent;
+		}
+
+		private int SocketSend(bool bPayload, byte[] packet, int writeSize)
+		{
+			int sent = writeSize;
+
+			if (BatchBuffer != null)
+			{
+				Batch_Append(packet, writeSize);
+			}
+			else if (!bBlockSend)
+			{
+				sent = ClientSocket.Send(packet, writeSize, SocketFlags.None);
+			}
+
+			if (RecordStream != null)
+			{
+				RecordStream.Write(packet, 0, writeSize);
+			}
+
+			Stats_TotalBytesSent += sent;
+			if (bPayload)
+			{
+				Stats_PayloadBytesSent += sent;
+			}
+
+			if (sent < writeSize)
+			{
+				// did not send all
+			}
+
+			return sent;
+		}
+		
+		private byte[] CommandPacket(int Command, long Size)
+		{
+			byte[] Result = { (byte)(Command & 255),
+								(byte)((Command >> 8) & 255),
+								(byte)(Size & 255),
+								(byte)((Size >> 8) & 255),
+								(byte)((Size >> 16) & 255),
+								(byte)((Size >> 24) & 255),
+								(byte)((Size >> 32) & 255),
+								(byte)((Size >> 40) & 255) };
+			return Result;
+		}
+
+		private long ResultSize()
+		{
+			long Result = -1;
+			try
+			{
+				byte[] sizePacket = new byte[8];
+				int byteRecv = SocketReceive(false, sizePacket);
+				if (byteRecv == 8)
+				{
+					Result = sizePacket[0] + (sizePacket[1] << 8) + (sizePacket[2] << 16) + (sizePacket[3] << 24) +
+							(sizePacket[4] << 32) + ((long)sizePacket[5] << 40) + ((long)sizePacket[6] << 48) + ((long)sizePacket[7] << 56);
+				}
+			}
+			catch (Exception e)
+			{
+				Log.TraceError("Unexpected Exception: {0}", e.ToString());
+				CloseConnection();
+			}
+			return Result;
+		}
+
+		private string GetStringResult(int Command, string Params)
+		{
+			string Result = null;
+
+			if (ClientSocket == null)
+			{
+				return Result;
+			}
+
+			byte[] paramMessage = Encoding.UTF8.GetBytes(Params + "\0");
+			long paramSize = paramMessage.Length;
+			byte[] commandMessage = CommandPacket(Command, paramSize);
+
+			try
+			{
+				int bytesSent = SocketSend(false, commandMessage);
+				bytesSent = SocketSend(false, paramMessage);
+
+				int resultSize = (int)ResultSize();
+				if (resultSize > 0)
+				{
+					byte[] buffer = new byte[resultSize < FILE_BUFFERSIZE ? resultSize : FILE_BUFFERSIZE ];
+
+					Result = "";
+					long remaining = resultSize;
+					while (remaining > 0)
+					{
+						int readSize = remaining > FILE_BUFFERSIZE ? FILE_BUFFERSIZE : (int)remaining;
+						int bytesRecv = SocketReceive(true, buffer, readSize);
+						remaining -= bytesRecv;
+						Result += Encoding.UTF8.GetString(buffer, 0, bytesRecv);
+					}
+				}
+			}
+			catch (Exception e)
+			{
+				Log.TraceError("Unexpected Exception: {0}", e.ToString());
+				CloseConnection();
+			}
+			return Result;
+		}
+
+		private bool GetBoolResult(int Command, string Params)
+		{
+			string Result = GetStringResult(Command, Params);
+			if (Result == null)
+			{
+				return false;
+			}
+
+			return (Result == "true");
+		}
+
+		// ------------------ Optimizes destination paths -----------------
+
+		private string OptimizePath(string inPath)
+		{
+			string inDir = Path.GetDirectoryName(inPath).Replace("\\", "/") + "/";
+			string inFilename = Path.GetFileName(inPath);
+
+			if (inDir == "/")
+			{
+				return inPath;
+			}
+
+			if (LastBaseDir == "")
+			{
+				LastBaseDir = inDir;
+				return inPath;
+			}
+
+			// trivial case of same as last base
+			if (LastBaseDir == inDir)
+			{
+				return "^^/" + inFilename;
+			}
+
+			// basedir matches start of new path
+			if (inDir.StartsWith(LastBaseDir))
+			{
+				int BaseDirLength = LastBaseDir.Length;
+				LastBaseDir = inDir;
+				return "^^/" + inDir.Substring(BaseDirLength) + inFilename;
+			}
+
+			int startIndex;
+			int nextIndex;
+			int dropCount;
+
+			// if indir is below basedir, figure out how many directories to drop
+			if (LastBaseDir.StartsWith(inDir))
+			{
+				startIndex = inDir.Length;
+				dropCount = 0;
+				while ((nextIndex = LastBaseDir.Substring(startIndex).IndexOf('/')) >= 0)
+				{
+					dropCount++;
+					startIndex += nextIndex + 1;
+				}
+				LastBaseDir = LastBaseDir.Substring(0, inDir.Length);
+				return "^-" + dropCount + "/" + inFilename;
+			}
+
+			// lastly, see what portion of the two paths match (if any)
+			int searchIndex = 0;
+			int lastSlash = -1;
+			int maxIndex = inDir.Length < LastBaseDir.Length ? inDir.Length : LastBaseDir.Length;
+			while (searchIndex < maxIndex && inDir[searchIndex] == LastBaseDir[searchIndex])
+			{
+				if (inDir[searchIndex] == '/')
+				{
+					lastSlash = searchIndex;
+				}
+				searchIndex++;
+			}
+
+			if (lastSlash < 1)
+			{
+				LastBaseDir = inDir;
+				return inPath;
+			}
+
+			startIndex = lastSlash + 1;
+			dropCount = 0;
+			while ((nextIndex = LastBaseDir.Substring(startIndex).IndexOf('/')) >= 0)
+			{
+				dropCount++;
+				startIndex += nextIndex + 1;
+			}
+			LastBaseDir = inDir;
+			return "^-" + dropCount + inDir.Substring(lastSlash) + inFilename;
+		}
+
+		// ------------------ Interace ------------------ 
+
+		public AndroidFileClient(string inDevice, int inPort = DefaultPort)
+		{
+			Device = inDevice;
+			ServerPort = inPort != 0 ? inPort : DefaultPort;
+			HostPort = -1;
+		}
+
+		~AndroidFileClient()
+		{
+			CloseConnection();
+		}
+
+		public static int GetDefaultPort()
+		{
+			return DefaultPort;
+		}
+
+		public string GetDevice()
+		{
+			return Device;
+		}
+
+		public int GetServerPort()
+		{
+			return ServerPort;
+		}
+
+		public bool OpenConnection(String Address = "127.0.0.1")
+		{
+			int UsePort = ServerPort;
+			if (Address == null || Address == "")
+			{
+				Address = "127.0.0.1";
+			}
+			if (Address == "127.0.0.1")
+			{
+				HostPort = ForwardPort_Add(Device, ServerPort);
+				if (HostPort == -1)
+				{
+					return false;
+				}
+				UsePort = HostPort;
+			}
+
+			// create TCP/IP socket
+			try
+			{
+				IPHostEntry ipHost = Dns.GetHostEntry(Dns.GetHostName());
+				IPAddress ipAddr = IPAddress.Parse(Address);
+				IPEndPoint localEndPoint = new IPEndPoint(ipAddr, UsePort);
+
+				// create TCP/IP socket
+				ClientSocket = new Socket(ipAddr.AddressFamily, SocketType.Stream, ProtocolType.Tcp);
+
+				ClientSocket.Connect(localEndPoint);
+				ClientSocket.NoDelay = false;
+				ClientSocket.SendBufferSize = SOCKET_SENDBUFFER;
+				ClientSocket.ReceiveBufferSize = SOCKET_READBUFFER;
+
+				Stats_Clear();
+				LastBaseDir = "";
+
+				int bytesSent = SocketSend(false, CommandPacket(Command_Info, 0));
+
+				int defaultTimeout = ClientSocket.ReceiveTimeout;
+				ClientSocket.ReceiveTimeout = 1000;
+
+				byte[] sizePacket = new byte[8];
+				int bytesRecv = SocketReceive(false, sizePacket);
+				if (bytesRecv == 8)
+				{
+					ClientSocket.ReceiveTimeout = defaultTimeout;
+
+					int resultSize = sizePacket[0] + (sizePacket[1] << 8) + (sizePacket[2] << 8) + (sizePacket[3] << 16) +
+									(sizePacket[4] << 24) + (sizePacket[5] << 32) + (sizePacket[6] << 40) + (sizePacket[7] << 48);
+					byte[] message = new byte[resultSize];
+					bytesRecv = SocketReceive(false, message);
+					if (bytesRecv == resultSize)
+					{
+						int version = message[0] + (message[1] << 8);
+						if (version == ServerVersion)
+						{
+							Log.TraceInformation("Connected to RemoteFileManager");
+							return true;
+						}
+					}
+				}
+				else
+				{
+					//Log.TraceInformation("Did not get response from RemoteFileManager");
+				}
+
+				CloseConnection();
+				return false;
+			}
+			catch (ArgumentNullException ane)
+			{
+				Log.TraceWarning("ArgumentNullException: {0}", ane.ToString());
+			}
+			catch (SocketException se)
+			{
+				string message = se.ToString();
+				if (message.Contains("No connection could be made because the target machine actively refused it"))
+				{
+					// ok not to display this one
+				}
+				else
+				{
+					Log.TraceWarning("SocketException: {0}", message);
+				}
+			}
+			catch (Exception e)
+			{
+				Log.TraceWarning("OpenConnection failed: {0}", e.ToString());
+			}
+
+			if (ClientSocket != null)
+			{
+				ClientSocket.Dispose();
+				ClientSocket = null;
+			}
+			if (HostPort != -1)
+			{
+				ForwardPort_Remove(Device, HostPort);
+				HostPort = -1;
+			}
+			return false;
+		}
+
+		public static List<string> GetInstalledReceivers(string Device)
+		{
+			List<string> Result = new List<string>();
+
+			String InstalledResult = adb.Shell(Device, "cmd package query-receivers --components -a com.epicgames.unreal.RemoteFileManager.intent.COMMAND");
+			foreach (string Line in InstalledResult.Split('\n'))
+			{
+				int SlashIndex = Line.IndexOf('/');
+				if (SlashIndex > 0 && (Line.Substring(SlashIndex + 1) == "com.epicgames.unreal.RemoteFileManagerReceiver"))
+				{
+					Result.Add(Line.Substring(0, SlashIndex));
+				}
+			}
+			return Result;
+		}
+
+		public static bool GetListenStatus_netstat(string Device, int Port, out bool bUSB, out bool bWifi, out string WifiAddress)
+		{
+			bUSB = false;
+			bWifi = false;
+			WifiAddress = "";
+
+			string Search = ":" + Port;
+			string netstatResult = adb.Shell(Device, "netstat -atn");
+			if (netstatResult.Contains("not found") || netstatResult.Contains("FAIL"))
+			{
+				return false;
+			}
+			foreach (string Line in netstatResult.Split('\n'))
+			{
+				if (Line.Contains("LISTEN") && Line.Contains(Search))
+				{
+					String[] parts = Line.Split(new char[0], StringSplitOptions.RemoveEmptyEntries);
+					if (parts.Length > 3 && parts[3].EndsWith(Search))
+					{
+						int SearchIndex = parts[3].LastIndexOf(Search);
+						string LineIP = parts[3].Substring(0, SearchIndex);
+						if (LineIP.Contains("127.0.0.1"))
+						{
+							bUSB = true;
+						}
+						else
+						{
+							bWifi = true;
+							// ip address is partial so don't try to parse it
+						}
+					}
+				}
+			}
+			return true;
+		}
+
+		public static bool GetListenStatus_ss(string Device, int Port, out bool bUSB, out bool bWifi, out string WifiAddress)
+		{
+			bUSB = false;
+			bWifi = false;
+			WifiAddress = "";
+
+			string Search = ":" + Port;
+			string ssResult = adb.Shell(Device, "ss -ta");
+			if (ssResult.Contains("not found") || ssResult.Contains("FAIL"))
+			{
+				return false;
+			}
+			foreach (string Line in ssResult.Split('\n'))
+			{
+				if (Line.Contains("LISTEN") && Line.Contains(Search))
+				{
+					String[] parts = Line.Split(new char[0], StringSplitOptions.RemoveEmptyEntries);
+					if (parts.Length > 3 && parts[3].EndsWith(Search))
+					{
+						// either ip:port or [*]:port
+						int SearchIndex = parts[3].LastIndexOf(Search);
+						string LineIP = parts[3].Substring(0, SearchIndex);
+						if (LineIP.StartsWith("["))
+						{
+							// handle [*:ip]:port or [ip]:port
+							SearchIndex = LineIP.LastIndexOf(":");
+							if (SearchIndex > 0)
+							{
+								LineIP = LineIP.Substring(SearchIndex + 1, LineIP.Length - SearchIndex - 2);
+							}
+							else
+							{
+								LineIP = LineIP.Substring(1, LineIP.Length - 2);
+							}
+						}
+						if (LineIP == "127.0.0.1")
+						{
+							bUSB = true;
+						}
+						else
+						{
+							WifiAddress = LineIP;
+							bWifi = true;
+						}
+					}
+				}
+			}
+
+			return true;
+		}
+
+		public static bool GetListenStatus(string Device, int Port, out bool bUSB, out bool bWifi, out string WifiAddress)
+		{
+			bUSB = false;
+			bWifi = false;
+			WifiAddress = "";
+
+			if (LastListenDevice != Device)
+			{
+				// reset availability of commands if different device
+				LastListenDevice = Device;
+				bHave_ss = true;
+				bHave_netstat = true;
+			}
+
+			// try ss first and remember if available for next time
+			if (bHave_ss)
+			{
+				bHave_ss = GetListenStatus_ss(Device, Port == 0 ? DefaultPort : Port, out bUSB, out bWifi, out WifiAddress);
+				if (bHave_ss)
+				{
+					return true;
+				}
+			}
+
+			// next fall back to netstat and remember if available for next time
+			if (bHave_netstat)
+			{
+				bHave_netstat = GetListenStatus_netstat(Device, Port == 0 ? DefaultPort : Port, out bUSB, out bWifi, out WifiAddress);
+				if (bHave_netstat)
+				{
+					return true;
+				}
+			}
+
+			// neither available so cannot determine listen status
+			return false;
+		}
+
+		public static bool StopAnyServers(string Device, int Port, bool bWaitForStop = true)
+		{
+			// get a list of installed receivers and verify package is available first
+			List<string> InstalledReceivers = GetInstalledReceivers(Device);
+
+			// deal with stopping all client servers if requested
+			bool bDidSendStops = false;
+			foreach (string Receiver in InstalledReceivers)
+			{
+				bDidSendStops = true;
+				adb.Shell(Device, "am broadcast -a com.epicgames.unreal.RemoteFileManager.intent.COMMAND -n " + Receiver + "/com.epicgames.unreal.RemoteFileManagerReceiver -e cmd 'stop'");
+			}
+
+			// it can take up to 2 seconds for running servers to terminate so check if any binds to port are still active
+			if (bDidSendStops && bWaitForStop)
+			{
+				// we will try ss first, but it may not be available
+				bool bUSB;
+				bool bWifi;
+				string WifiAddress;
+
+				if (GetListenStatus(Device, Port, out bUSB, out bWifi, out WifiAddress))
+				{
+					// wait for both to terminate (up to 10 seconds, then give up)
+					long StartTime = DateTimeOffset.Now.ToUnixTimeSeconds();
+					while (bUSB || bWifi)
+					{
+						GetListenStatus(Device, Port, out bUSB, out bWifi, out WifiAddress);
+						if (DateTimeOffset.Now.ToUnixTimeSeconds() - StartTime > 9)
+						{
+							// there are still binds active that won't terminate
+							return false;
+						}
+						Thread.Sleep(100);
+					}
+				}
+				else
+				{
+					// command not available so will just need to wait 2 seconds to be sure
+					Thread.Sleep(2000);
+				}
+			}
+
+			return true;
+		}
+
+		public bool StartServer(string PackageName, string Token = "", string IPAddress = "127.0.0.1", bool bStopAnyServers = true)
+		{
+			// already connected?
+			if (ClientSocket != null)
+			{
+				return true;
+			}
+
+			bool bIsUSB = (IPAddress == null || IPAddress == "" || IPAddress == "127.0.0.1");
+
+			// get a list of installed clients and verify package is available first
+			List<string> InstalledReceivers = GetInstalledReceivers(Device);
+			if (!InstalledReceivers.Contains(PackageName))
+			{
+				return false;
+			}
+
+			// deal with stopping all servers if requested
+			if (bStopAnyServers)
+			{
+				if (!StopAnyServers(Device, ServerPort, true))
+				{
+					// there are still binds active that won't terminate so can't start a new one
+					return false;
+				}
+			}
+
+			int tries = 5;
+			bool bUSB;
+			bool bWifi;
+			string WifiAddress;
+
+			// sent start request (won't do anything if already started)
+			string StartCommand = "am broadcast -a com.epicgames.unreal.RemoteFileManager.intent.COMMAND -n " + PackageName + "/com.epicgames.unreal.RemoteFileManagerReceiver -e cmd 'start' -e token '" + Token + "' -ei port " + ServerPort;
+			adb.Shell(Device, StartCommand);
+
+			// see if we can check listen status
+			if (GetListenStatus(Device, ServerPort, out bUSB, out bWifi, out WifiAddress))
+			{
+				while (tries-- > 0)
+				{
+					// try to connect if requested server is listening
+					if ((bUSB && bIsUSB) || (bWifi && !bIsUSB))
+					{
+						if (OpenConnection(IPAddress))
+						{
+							return true;
+						}
+					}
+
+					// wait before trying to again
+					Thread.Sleep(250);
+
+					// sent start request again (won't do anything if already started)
+					adb.Shell(Device, StartCommand);
+
+					GetListenStatus(Device, ServerPort, out bUSB, out bWifi, out WifiAddress);
+				}
+
+				return false;
+			}
+
+			// try 5 times, this is faster than looking response in logcat
+			while (tries-- > 0)
+			{
+				// wait before trying to connect
+				Thread.Sleep(250);
+
+				if (OpenConnection(IPAddress))
+				{
+					return true;
+				}
+
+				// sent start request again (won't do anything if already started)
+				adb.Shell(Device, StartCommand);
+			}
+
+			return false;
+		}
+
+		public void TerminateServer()
+		{
+			// try to connect if not already connected
+			if (ClientSocket == null)
+			{
+				if (!OpenConnection())
+				{
+					return;
+				}
+			}
+
+			if (ClientSocket != null)
+			{
+				try
+				{
+					int bytesSent = SocketSend(false, CommandPacket(Command_Terminate, 0));
+
+					// close socket
+					ClientSocket.Shutdown(SocketShutdown.Both);
+					ClientSocket.Close();
+					ClientSocket.Dispose();
+					ClientSocket = null;
+				}
+				catch (SocketException se)
+				{
+					Log.TraceWarning("SocketException: {0}", se.ToString());
+				}
+				catch (ObjectDisposedException)
+				{
+					// allow it
+				}
+				catch (Exception e)
+				{
+					Log.TraceError("Unexpected Exception: {0}", e.ToString());
+				}
+			}
+		}
+
+		public void CloseConnection()
+		{
+			Stats_Stopwatch = null;
+			if (ClientSocket != null)
+			{
+				try
+				{
+					int bytesSent = SocketSend(false, CommandPacket(Command_Close, 0));
+
+					// close socket
+					ClientSocket.Shutdown(SocketShutdown.Both);
+					ClientSocket.Close();
+					ClientSocket.Dispose();
+					ClientSocket = null;
+				}
+				catch (SocketException se)
+				{
+					Log.TraceWarning("SocketException: {0}", se.ToString());
+				}
+				catch (ObjectDisposedException)
+				{
+					// allow it
+				}
+				catch (Exception e)
+				{
+					Log.TraceError("Unexpected Exception: {0}", e.ToString());
+				}
+			}
+
+			if (HostPort != -1)
+			{
+				ForwardPort_Remove(Device, HostPort);
+				HostPort = -1;
+			}
+		}
+
+		public string Query(string Filepath, bool bStripEndingSlash = false)
+		{
+			string result = GetStringResult(Command_Query, Filepath);
+			if (result != null && bStripEndingSlash && result.EndsWith("/"))
+			{
+				return result.Substring(0, result.Length - 1);
+			}
+			return result;
+		}
+
+		public string GetProp(string Filepath)
+		{
+			return GetStringResult(Command_GetProp, Filepath);
+		}
+
+		public bool SetBaseDir(string Filepath)
+		{
+			if (Filepath.EndsWith("/"))
+			{
+				Filepath = Filepath.Substring(0, Filepath.Length - 1);
+			}
+			LastBaseDir = Filepath;
+			return GetBoolResult(Command_SetBaseDir, Filepath);
+		}
+
+		public bool DirExists(string Filepath)
+		{
+			return GetBoolResult(Command_DirExists, Filepath);
+		}
+
+		public string DirList(string Filepath)
+		{
+			return GetStringResult(Command_DirList, Filepath);
+		}
+
+		public string DirListFlat(string Filepath)
+		{
+			return GetStringResult(Command_DirListFlat, Filepath);
+		}
+
+		public bool DirCreate(string Filepath)
+		{
+			return GetBoolResult(Command_DirCreate, Filepath);
+		}
+
+		public bool DirDelete(string Filepath)
+		{
+			return GetBoolResult(Command_DirDelete, Filepath);
+		}
+
+		public bool DirDeleteRecurse(string Filepath)
+		{
+			return GetBoolResult(Command_DirDeleteRecurse, Filepath);
+		}
+
+		public bool FileExists(string Filepath)
+		{
+			return GetBoolResult(Command_FileExists, Filepath);
+		}
+
+		public bool FileDelete(string Filepath)
+		{
+			return GetBoolResult(Command_FileDelete, Filepath);
+		}
+
+		public bool FileCopy(string SourcePath, string DestPath)
+		{
+			return GetBoolResult(Command_FileCopy, SourcePath + "\0" + DestPath);
+		}
+
+		public bool FileMove(string SourcePath, string DestPath)
+		{
+			return GetBoolResult(Command_FileMove, SourcePath + "\0" + DestPath);
+		}
+
+		public bool FileRead(string SourcePath, string DestPath)
+		{
+			Boolean Result = false;
+
+			Log.TraceInformation("FileRead {0}", SourcePath);
+
+			if (ClientSocket == null)
+			{
+				return Result;
+			}
+
+			if (File.Exists(DestPath))
+			{
+				File.Delete(DestPath);
+			}
+
+			try
+			{
+				// deal with LastBaseDir substitution
+				SourcePath = OptimizePath(SourcePath);
+
+				byte[] paramMessage = Encoding.UTF8.GetBytes(SourcePath + "\0");
+				long paramSize = paramMessage.Length;
+				byte[] commandMessage = CommandPacket(Command_FileRead, paramSize);
+
+				int bytesSent = SocketSend(false, commandMessage);
+				bytesSent = SocketSend(false, paramMessage);
+
+				long resultSize = ResultSize();
+				long startSize = resultSize;
+				if (resultSize >= 0)
+				{
+					int bytesRead;
+					byte[] buffer = new byte[FILE_BUFFERSIZE];
+
+					Stats_FilesRead++;
+
+					try
+					{
+						using (FileStream fileStream = File.Open(DestPath, FileMode.Create, FileAccess.Write))
+						using (BufferedStream bufferedStream = new BufferedStream(fileStream))
+						{
+							long remaining = resultSize;
+							while (remaining > 0)
+							{
+								int readSize = remaining > FILE_BUFFERSIZE ? FILE_BUFFERSIZE : (int)remaining;
+								bytesRead = SocketReceive(true, buffer, readSize);
+								remaining -= bytesRead;
+								bufferedStream.Write(buffer, 0, bytesRead);
+								Stats_CompressedBytesReceived += bytesRead;
+							}
+							bufferedStream.Close();
+							fileStream.Close();
+						}
+					}
+					catch (Exception e)
+					{
+						Log.TraceError("Unexpected Exception: {0}", e.ToString());
+						// ignore received data
+						long remainingSkip = resultSize;
+						while (remainingSkip > 0)
+						{
+							int readSize = remainingSkip > FILE_BUFFERSIZE ? FILE_BUFFERSIZE : (int)remainingSkip;
+							bytesRead = SocketReceive(true, buffer, readSize);
+							remainingSkip -= bytesRead;
+						}
+						return Result;
+					}
+
+					Result = true;
+				}
+			}
+			catch (Exception e)
+			{
+				Log.TraceError("Unexpected Exception: {0}", e.ToString());
+				CloseConnection();
+			}
+			return Result;
+		}
+
+		public bool FileWrite(string SourcePath, string DestPath, int bLog = 0)
+		{
+			Boolean Result = false;
+
+			if (ClientSocket == null)
+			{
+				return Result;
+			}
+
+			try
+			{
+				using (FileStream fileStream = new FileStream(SourcePath, FileMode.Open, FileAccess.Read, FileShare.None, 4096, FileOptions.SequentialScan))
+				{
+					// deal with LastBaseDir substitution
+					DestPath = OptimizePath(DestPath);
+					if (bLog > 0)
+					{
+						Log.TraceInformation("{0}> Writing {1}", bLog, DestPath);
+					}
+
+					long fileSize = fileStream.Length;
+					Stats_FilesWrite++;
+
+					int bytesSent;
+					int bytesRead;
+					long remaining = fileSize;
+
+					byte[] paramMessage = Encoding.UTF8.GetBytes(DestPath + "\0");
+					int headerSize = paramMessage.Length;
+					long paramSize = headerSize + fileSize;
+
+					if (BatchBuffer != null)
+					{
+						byte[] commandMessage = CommandPacket(Command_FileWrite, paramSize);
+						bytesSent = SocketSend(false, commandMessage);
+						bytesSent = SocketSend(false, paramMessage);
+
+						// directly read into batch buffer for minimal copies
+						while (remaining > 0)
+						{
+							int batchRemaining = BATCH_BUFFERSIZE - BatchBufferIndex;
+							int chunkSize = (int)(BatchBufferIndex + remaining > BATCH_BUFFERSIZE ? BATCH_BUFFERSIZE - BatchBufferIndex : remaining);
+							bytesRead = fileStream.Read(BatchBuffer, BatchBufferIndex, chunkSize);
+
+							// need to do this here since SocketSend not called
+							if (RecordStream != null)
+							{
+								RecordStream.Write(BatchBuffer, BatchBufferIndex, bytesRead);
+							}
+
+							Stats_CompressedBytesSent += bytesRead;
+							Stats_TotalBytesSent += bytesRead;
+							Stats_PayloadBytesSent += bytesRead;
+
+							BatchBufferIndex += bytesRead;
+							remaining -= bytesRead;
+
+							if (BatchBufferIndex == BATCH_BUFFERSIZE)
+							{
+								Batch_Flush();
+							}
+						}
+					}
+					else
+					{
+						// reduce the number of SocketSends if not batching
+						byte[] commandMessage = new byte[8 + headerSize];
+						commandMessage[0] = (byte)(Command_FileWrite & 255);
+						commandMessage[1] = (byte)((Command_FileWrite >> 8) & 255);
+						commandMessage[2] = (byte)(paramSize & 255);
+						commandMessage[3] = (byte)((paramSize >> 8) & 255);
+						commandMessage[4] = (byte)((paramSize >> 16) & 255);
+						commandMessage[5] = (byte)((paramSize >> 24) & 255);
+						commandMessage[6] = (byte)((paramSize >> 32) & 255);
+						commandMessage[7] = (byte)((paramSize >> 40) & 255);
+						Buffer.BlockCopy(paramMessage, 0, commandMessage, 8, headerSize);
+						bytesSent = SocketSend(false, commandMessage);
+
+						byte[] buffer = new byte[FILE_BUFFERSIZE];
+
+						while ((bytesRead = fileStream.Read(buffer, 0, FILE_BUFFERSIZE)) != 0)
+						{
+							bytesSent = SocketSend(true, buffer, bytesRead);
+							Stats_CompressedBytesSent += bytesSent;
+						}
+					}
+					fileStream.Close();
+				}
+				Result = true;
+			}
+			catch (IOException)
+			{
+				// file not found
+				return false;
+			}
+			catch (Exception e)
+			{
+				Log.TraceError("{0}> Unexpected Exception: {1}", bLog, e.ToString());
+				CloseConnection();
+			}
+			return Result;
+		}
+
+		public bool FileWriteCompressed(string SourcePath, string DestPath, int bLog = 0)
+		{
+			Boolean Result = false;
+
+			if (ClientSocket == null)
+			{
+				return Result;
+			}
+
+			try
+			{
+				using (FileStream fileStream = File.Open(SourcePath, FileMode.Open, FileAccess.Read))
+				{
+					long fileSize = fileStream.Length;
+					if (fileSize < 1024)
+					{
+						fileStream.Close();
+						return FileWrite(SourcePath, DestPath, bLog);
+					}
+
+					// deal with LastBaseDir substitution
+					DestPath = OptimizePath(DestPath);
+					if (bLog > 0)
+					{
+						Log.TraceInformation("{0}> Writing Compressed {1}", bLog, DestPath);
+					}
+
+					Stats_FilesWrite++;
+					int COMPRESS_BUFFERSIZE = 1024 * 1024;
+
+					int bytesSent;
+					int bytesRead;
+					long remaining = fileSize;
+
+					byte[] paramMessage = Encoding.UTF8.GetBytes("Z" + DestPath + "\0");
+					int headerSize = paramMessage.Length;
+					long paramSize = headerSize + fileSize;
+					if (BatchBuffer != null)
+					{
+						byte[] commandMessage = CommandPacket(Command_FileWriteCompressed, paramSize);
+						bytesSent = SocketSend(false, commandMessage);
+						bytesSent = SocketSend(false, paramMessage);
+					}
+					else
+					{
+						// reduce the number of SocketSends if not batching
+						byte[] commandMessage = new byte[8 + headerSize];
+						commandMessage[0] = (byte)(Command_FileWriteCompressed & 255);
+						commandMessage[1] = (byte)((Command_FileWriteCompressed >> 8) & 255);
+						commandMessage[2] = (byte)(paramSize & 255);
+						commandMessage[3] = (byte)((paramSize >> 8) & 255);
+						commandMessage[4] = (byte)((paramSize >> 16) & 255);
+						commandMessage[5] = (byte)((paramSize >> 24) & 255);
+						commandMessage[6] = (byte)((paramSize >> 32) & 255);
+						commandMessage[7] = (byte)((paramSize >> 40) & 255);
+						Buffer.BlockCopy(paramMessage, 0, commandMessage, 8, headerSize);
+						bytesSent = SocketSend(false, commandMessage);
+					}
+
+					byte[] buffer = new byte[COMPRESS_BUFFERSIZE + 3];
+					byte[] header = new byte[3];
+
+					// first part of buffer contains zero which flags COMPRESS_BUFFERSIZE uncompressed
+					buffer[0] = 0;
+					buffer[1] = 0;
+					buffer[2] = 0;
+
+					while ((bytesRead = fileStream.Read(buffer, 3, COMPRESS_BUFFERSIZE)) != 0)
+					{
+						Stats_CompressedBytesSent += bytesRead;
+
+						var compressedStream = new MemoryStream();
+						var zipStream = new GZipStream(compressedStream, CompressionLevel.Fastest);
+						zipStream.Write(buffer, 3, bytesRead);
+						zipStream.Close();
+						byte[] compressedArray = compressedStream.ToArray();
+						int compsize = compressedArray.Length;
+
+						if (compsize < COMPRESS_BUFFERSIZE)
+						{
+							if (BatchBuffer != null)
+							{
+								// batching will do a block copy so better to do two calls
+								header[0] = (byte)(compsize & 255);
+								header[1] = (byte)((compsize >> 8) & 255);
+								header[2] = (byte)((compsize >> 16) & 255);
+								bytesSent = SocketSend(true, header);
+								bytesSent = SocketSend(true, compressedArray);
+							}
+							else
+							{
+								// more efficient to send one packet
+								byte[] compPacket = new byte[compsize + 3];
+								compPacket[0] = (byte)(compsize & 255);
+								compPacket[1] = (byte)((compsize >> 8) & 255);
+								compPacket[2] = (byte)((compsize >> 16) & 255);
+								Buffer.BlockCopy(compressedArray, 0, compPacket, 3, compsize);
+								bytesSent = SocketSend(true, compPacket);
+							}
+						}
+						else
+						{
+							bytesSent = SocketSend(true, buffer);
+						}
+					}
+					fileStream.Close();
+
+					/*
+					int resultSize = (int)ResultSize();
+					if (resultSize > 0)
+					{
+						byte[] message = new byte[resultSize];
+						int bytesRecv = SocketReceive(false, message);
+
+						string Outcome = Encoding.UTF8.GetString(message, 0, bytesRecv);
+						Result = Outcome == "true";
+					}
+					*/
+					Result = true;
+				}
+			}
+			catch (IOException)
+			{
+				// file not found
+				Result = false;
+			}
+			catch (Exception e)
+			{
+				Log.TraceError("{0}> Unexpected Exception: {1}", bLog, e.ToString());
+				CloseConnection();
+			}
+			return Result;
+		}
+
+		public bool FileWriteString(string Contents, string DestPath)
+		{
+			Boolean Result = false;
+
+			if (ClientSocket == null)
+			{
+				return Result;
+			}
+
+			try
+			{
+				// deal with LastBaseDir substitution
+				DestPath = OptimizePath(DestPath);
+
+				byte[] paramMessage = Encoding.UTF8.GetBytes(DestPath + "\0");
+				byte[] contentsMessage = Encoding.UTF8.GetBytes(Contents);
+				long paramSize = paramMessage.Length + contentsMessage.Length;
+				byte[] commandMessage = CommandPacket(Command_FileWrite, paramSize);
+
+				Stats_FilesWrite++;
+
+				int bytesSent = SocketSend(false, commandMessage);
+				bytesSent = SocketSend(false, paramMessage);
+
+				Stats_CompressedBytesSent += bytesSent;
+
+				bytesSent = SocketSend(true, contentsMessage);
+
+				/*				int resultSize = (int)ResultSize();
+								if (resultSize > 0)
+								{
+									byte[] message = new byte[resultSize];
+									int bytesRecv = SocketReceive(false, message);
+
+									string Outcome = Encoding.UTF8.GetString(message, 0, bytesRecv);
+									Result = Outcome == "true";
+								}*/
+				Result = true;
+			}
+			catch (Exception e)
+			{
+				Log.TraceError("Unexpected Exception: {0}", e.ToString());
+				CloseConnection();
+			}
+			return Result;
+		}
+
+		public bool PushFile(string InSource, string InDest, bool bCompress = false, int bLog = 0)
+		{
+			InSource = InSource.Replace("\\", "/");
+			InDest = InDest.Replace("\\", "/");
+
+			bool Result = false;
+			if (bCompress)
+			{
+				Result = FileWriteCompressed(InSource, InDest, bLog);
+			}
+			else
+			{
+				Result = FileWrite(InSource, InDest, bLog);
+			}
+			if (!Result && bLog > 0)
+			{
+				Log.TraceInformation("{0}> Failed to copy {1} to {2}", bLog, InSource, InDest);
+				return false;
+			}
+
+			/*
+			if (bLog > 0)
+			{
+				Log.TraceInformation("{0}> Copied: {1}", bLog, InSource);
+			}
+			*/
+			return true;
+		}
+
+		public bool PushDirectory(string InSource, string InDest, bool bCompress = false, int bLog = 0)
+		{
+			if (Directory.Exists(InSource))
+			{
+				foreach (string File in Directory.EnumerateFiles(InSource))
+				{
+					if (!PushFile(File, File.Replace(InSource, InDest), bCompress, bLog))
+					{
+						return false;
+					}
+				}
+			}
+			return true;
+		}
+
+		public bool PushDirectories(string InSource, string InDest, bool bCompress = false, int bLog = 0)
+		{
+			if (Directory.Exists(InSource))
+			{
+				foreach (string File in Directory.EnumerateFiles(InSource))
+				{
+					if (!PushFile(File, File.Replace(InSource, InDest), bCompress, bLog))
+					{
+						return false;
+					}
+				}
+
+				foreach (string Directory in Directory.EnumerateDirectories(InSource))
+				{
+					if (!PushDirectories(Directory, Directory.Replace(InSource, InDest), bCompress, bLog))
+					{
+						return false;
+					}
+				}
+			}
+			return true;
+		}
+
+		private void AddDeployDirectories(ref List<KeyValuePair<string, string>> USBClientFiles, ref List<KeyValuePair<string, string>> WiFiClientFiles, ref bool bAlternate, string InSource, string InDest)
+		{
+			if (Directory.Exists(InSource))
+			{
+				foreach (string File in Directory.EnumerateFiles(InSource))
+				{
+					if (bAlternate)
+					{
+						WiFiClientFiles.Add(new KeyValuePair<string, string>(File, File.Replace(InSource, InDest)));
+					}
+					else
+					{
+						USBClientFiles.Add(new KeyValuePair<string, string>(File, File.Replace(InSource, InDest)));
+
+					}
+					bAlternate = !bAlternate;
+				}
+
+				foreach (string Directory in Directory.EnumerateDirectories(InSource))
+				{
+					AddDeployDirectories(ref USBClientFiles, ref WiFiClientFiles, ref bAlternate, Directory, Directory.Replace(InSource, InDest));
+				}
+			}
+		}
+
+		private bool bDeployThreadInterruptRequested = false;
+		private bool bDeployThreadError = false;
+
+		private void DeployThread(AndroidFileClient Client, ref List<KeyValuePair<string, string>> Files, bool bCompress, int bLog)
+		{
+			// Push all the WiFi client files
+			foreach (KeyValuePair<string, string> Entry in Files)
+			{
+				if (!Client.PushFile(Entry.Key, Entry.Value, bCompress, bLog))
+				{
+					Volatile.Write(ref bDeployThreadError, true);
+					break;
+				}
+				if (Volatile.Read(ref bDeployThreadInterruptRequested))
+				{
+					break;
+				}
+			}
+		}
+
+		private void AddSourceDirectories(ref List<string> Directories, string InSource, string InDest)
+		{
+			if (Directory.Exists(InSource))
+			{
+				foreach (string Directory in Directory.EnumerateDirectories(InSource))
+				{
+					AddSourceDirectories(ref Directories, Directory, Directory.Replace(InSource, InDest));
+				}
+				string Entry = InDest.Replace("\\", "/");
+				int LastIndex = Directories.Count - 1;
+				if (LastIndex < 0 || !Directories[LastIndex].StartsWith(Entry))
+				{
+					Directories.Add(Entry);
+				}
+			}
+		}
+
+		public bool Deploy(HashSet<string> EntriesToDeploy, string SourceDir, string DestDir, bool bCompress = false, bool bLog = false, bool bReportStats = true, AndroidFileClient WiFiClient = null)
+		{
+			bool Result = true;
+			int logId = bLog ? 1 : 0;
+
+			Stopwatch totalwatch = System.Diagnostics.Stopwatch.StartNew();
+			Stopwatch stopwatch = System.Diagnostics.Stopwatch.StartNew();
+
+			// pre-create the directories (send only deepest)
+			List<string> Directories = new List<string>();
+			foreach (string Entry in EntriesToDeploy)
+			{
+				string RemotePath = Entry.Replace(SourceDir, DestDir).Replace("\\", "/");
+
+				FileAttributes attributes = File.GetAttributes(Entry);
+				if ((attributes & FileAttributes.Directory) == FileAttributes.Directory)
+				{
+					AddSourceDirectories(ref Directories, Entry, RemotePath);
+				}
+				else
+				{
+					Directories.Add(Path.GetDirectoryName(RemotePath).Replace("\\", "/"));
+				}
+			}
+			foreach (string Entry in Directories)
+			{
+				DirCreate(Entry);
+			}
+
+			stopwatch.Stop();
+			var ElapsedMs = stopwatch.ElapsedMilliseconds;
+			if (bReportStats)
+			{
+				Log.TraceInformation("Time to create directories:  {0:N3}s", (float)(ElapsedMs / 1000.0));
+			}
+
+			Stats_Clear();
+			Batch_Start();
+
+			stopwatch.Restart();
+
+			if (WiFiClient != null)
+			{
+				WiFiClient.Stats_Clear(false);
+				WiFiClient.Batch_Start();
+
+				// collect all the files into two lists, one for each client
+				List<KeyValuePair<string, string>> USBClientFiles = new List<KeyValuePair<string, string>>();
+				List<KeyValuePair<string, string>> WiFiClientFiles = new List<KeyValuePair<string, string>>();
+				bool bAlternate = false;
+
+				foreach (string Entry in EntriesToDeploy)
+				{
+					string RemotePath = Entry.Replace(SourceDir, DestDir).Replace("\\", "/");
+
+					FileAttributes attributes = File.GetAttributes(Entry);
+					if ((attributes & FileAttributes.Directory) == FileAttributes.Directory)
+					{
+						AddDeployDirectories(ref USBClientFiles, ref WiFiClientFiles, ref bAlternate, Entry, RemotePath);
+					}
+					else
+					{
+						if (bAlternate)
+						{
+							WiFiClientFiles.Add(new KeyValuePair<string, string>(Entry, RemotePath));
+						}
+						else
+						{
+							USBClientFiles.Add(new KeyValuePair<string, string>(Entry, RemotePath));
+
+						}
+						bAlternate = !bAlternate;
+					}
+				}
+
+				Thread WiFiDeployThread = new Thread(() => DeployThread(WiFiClient, ref WiFiClientFiles, bCompress, bLog ? 2 : 0))
+				{
+					Name = "WiFiDeployThread"
+				};
+
+				bDeployThreadInterruptRequested = false;
+				bDeployThreadError = false;
+
+				WiFiDeployThread.Start();
+
+				// Push all the USB client files
+				foreach (KeyValuePair<string, string> Entry in USBClientFiles)
+				{
+					if (!PushFile(Entry.Key, Entry.Value, bCompress, logId))
+					{
+						Volatile.Write(ref bDeployThreadInterruptRequested, true);
+						Result = false;
+						break;
+					}
+					if (Volatile.Read(ref bDeployThreadError))
+						if (bDeployThreadError)
+						{
+							Result = false;
+							break;
+						}
+				}
+
+				// Wait for WiFi thread to finish
+				WiFiDeployThread.Join();
+
+				WiFiClient.Batch_Stop();
+				Stats_Combine(WiFiClient);
+			}
+			else
+			{
+				// push all files on USB
+				foreach (string Entry in EntriesToDeploy)
+				{
+					string RemotePath = Entry.Replace(SourceDir, DestDir).Replace("\\", "/");
+
+					FileAttributes attributes = File.GetAttributes(Entry);
+					if ((attributes & FileAttributes.Directory) == FileAttributes.Directory)
+					{
+						if (!PushDirectories(Entry, RemotePath, bCompress, logId))
+						{
+							Result = false;
+							break;
+						}
+					}
+					else
+					{
+						if (!PushFile(Entry, RemotePath, bCompress, logId))
+						{
+							Result = false;
+							break;
+						}
+					}
+				}
+			}
+
+			// time to batch (if buffer wasn't flushed)
+			stopwatch.Stop();
+
+			// this code will be removed if TimeBatching is false but will give an unreachable code warning
+			#pragma warning disable
+			if (TimeBatching)
+			{
+				if (bReportStats)
+				{
+					ElapsedMs = stopwatch.ElapsedMilliseconds;
+					Log.TraceInformation("Time to build batch:  {0:N3}s", (float)(ElapsedMs / 1000.0));
+				}
+				Stats_Stopwatch.Stop();
+				Stats_StartStopwatch();
+			}
+			#pragma warning restore
+
+			// stop and flush batch buffer
+			Batch_Stop();
+
+			// this causes a sync so totalwatch is valid
+			Query("^^");
+			totalwatch.Stop();
+
+			if (bReportStats)
+			{
+				Stats_Report();
+
+				ElapsedMs = totalwatch.ElapsedMilliseconds;
+				Log.TraceInformation("Total time to Deploy AFS:  {0:N3}s", (float)(ElapsedMs / 1000.0));
+			}
+
+			return Result;
+		}
+
+		/*
+		public void TestOptimizePath()
+		{
+			Log.TraceInformation("{0}", OptimizePath("/storage/emulated/0/Android/data/com.epicgames.SaveTest/files/UE4Game/SaveTest/Manifest_DebugFiles_Android.txt"));
+			Log.TraceInformation("{0}", OptimizePath("/storage/emulated/0/Android/data/com.epicgames.SaveTest/files/UE4Game/SaveTest/DataDrivenPlatformInfo.ini"));
+			Log.TraceInformation("{0}", OptimizePath("/storage/emulated/0/Android/data/com.epicgames.SaveTest/files/UE4Game/SaveTest/Engine/Config/Layouts/DefaultLayout.ini"));
+			Log.TraceInformation("{0}", OptimizePath("/storage/emulated/0/Android/data/com.epicgames.SaveTest/files/UE4Game/SaveTest/Engine/Content/ArtTools/RenderToTexture/Materials/Debug/M_Emissive_Color.uasset"));
+
+			Log.TraceInformation("{0}", OptimizePath("/storage/0/emulated/commandfile.txt"));
+			Log.TraceInformation("{0}", OptimizePath("/storage/0/emulated/123.txt"));
+			Log.TraceInformation("{0}", OptimizePath("/storage/0/emulated/game/123.txt"));
+			Log.TraceInformation("{0}", OptimizePath("/storage/0/emulated/every.txt"));
+
+			Log.TraceInformation("{0}", OptimizePath("^commandfile"));
+			Log.TraceInformation("{0}", OptimizePath("^commandfile"));
+
+			Log.TraceInformation("{0}", OptimizePath("^ext/z.txt"));
+			Log.TraceInformation("{0}", OptimizePath("^ext/abc.txt"));
+			Log.TraceInformation("{0}", OptimizePath("^ext/cde.txt"));
+			Log.TraceInformation("{0}", OptimizePath("^ext/1/abc.txt"));
+			Log.TraceInformation("{0}", OptimizePath("^ext/1/2/abc.txt"));
+			Log.TraceInformation("{0}", OptimizePath("^ext/xyz.txt"));
+			Log.TraceInformation("{0}", OptimizePath("^ext/1/2/stu.txt"));
+			Log.TraceInformation("{0}", OptimizePath("^ext/1/wer.txt"));
+			Log.TraceInformation("{0}", OptimizePath("^int/1/2/stu.txt"));
+			Log.TraceInformation("{0}", OptimizePath("^int/1/abc.txt"));
+		}
+		*/
+	}
+}
diff --git a/Engine/Source/Programs/AutomationTool/Android/AndroidPlatform.Automation.cs b/Engine/Source/Programs/AutomationTool/Android/AndroidPlatform.Automation.cs
--- a/Engine/Source/Programs/AutomationTool/Android/AndroidPlatform.Automation.cs
+++ b/Engine/Source/Programs/AutomationTool/Android/AndroidPlatform.Automation.cs
@@ -12,6 +12,7 @@ using UnrealBuildTool;
 using Ionic.Zip;
 using EpicGames.Core;
 using UnrealBuildBase;
+using AutomationUtils.Automation;
 using System.Text.RegularExpressions;
 using AutomationScripts;
 
@@ -546,6 +547,86 @@ public class AndroidPlatform : Platform
 		return PluginExtras;
 	}
 
+	private bool UsingAndroidFileServer(DeploymentContext SC, out bool bEnablePlugin, out string AFSToken, out bool bIsShipping, out bool bIncludeInShipping, out bool bAllowExternalStartInShipping)
+	{
+		UnrealTargetConfiguration TargetConfiguration = SC.StageTargetConfigurations[0];
+		bIsShipping = TargetConfiguration == UnrealTargetConfiguration.Shipping;
+
+		ConfigHierarchy Ini = ConfigCache.ReadHierarchy(ConfigHierarchyType.Engine, DirectoryReference.FromFile(SC.RawProjectPath), SC.StageTargetPlatform.PlatformType);
+		bEnablePlugin = true;
+		AFSToken = "";
+		bIncludeInShipping = false;
+		bAllowExternalStartInShipping = false;
+		Ini.GetBool("/Script/AndroidFileServerEditor.AndroidFileServerRuntimeSettings", "bEnablePlugin", out bEnablePlugin);
+		Ini.GetString("/Script/AndroidFileServerEditor.AndroidFileServerRuntimeSettings", "SecurityToken", out AFSToken);
+		Ini.GetBool("/Script/AndroidFileServerEditor.AndroidFileServerRuntimeSettings", "bIncludeInShipping", out bIncludeInShipping);
+		Ini.GetBool("/Script/AndroidFileServerEditor.AndroidFileServerRuntimeSettings", "bAllowExternalStartInShipping", out bAllowExternalStartInShipping);
+
+		if (bIsShipping && !(bIncludeInShipping && bAllowExternalStartInShipping))
+		{
+			return false;
+		}
+		return bEnablePlugin;
+	}
+
+	enum EConnectionType
+	{
+		USBOnly,
+		NetworkOnly,
+		Combined
+	}
+
+	private EConnectionType GetAndroidFileServerNetworkConfig(DeploymentContext SC, out bool bUseCompression, out bool bLogFiles, out bool bReportStats, out bool bUseManualIPAddress, out string ManualIPAddress)
+	{
+		EConnectionType ConnectionType = EConnectionType.USBOnly;
+
+		UnrealTargetConfiguration TargetConfiguration = SC.StageTargetConfigurations[0];
+		ConfigHierarchy Ini = ConfigCache.ReadHierarchy(ConfigHierarchyType.Engine, DirectoryReference.FromFile(SC.RawProjectPath), SC.StageTargetPlatform.PlatformType);
+
+		string ConnectionString = "";
+		Ini.GetString("/Script/AndroidFileServerEditor.AndroidFileServerRuntimeSettings", "ConnectionType", out ConnectionString);
+		switch (ConnectionString)
+		{
+			case "USBOnly":
+				ConnectionType = EConnectionType.USBOnly;
+				break;
+			case "NetworkOnly":
+				ConnectionType = EConnectionType.NetworkOnly;
+				break;
+			case "Combined":
+				ConnectionType = EConnectionType.Combined;
+				break;
+			default:
+				ConnectionType = EConnectionType.USBOnly;
+				break;
+		}
+
+		bUseCompression = false;
+		bLogFiles = false;
+		bReportStats = false;
+		bUseManualIPAddress = false;
+		ManualIPAddress = "";
+		bool bAllowNetworkConnection = true;
+
+		Ini.GetBool("/Script/AndroidFileServerEditor.AndroidFileServerRuntimeSettings", "bUseCompression", out bUseCompression);
+		Ini.GetBool("/Script/AndroidFileServerEditor.AndroidFileServerRuntimeSettings", "bLogFiles", out bLogFiles);
+		Ini.GetBool("/Script/AndroidFileServerEditor.AndroidFileServerRuntimeSettings", "bReportStats", out bReportStats);
+		Ini.GetBool("/Script/AndroidFileServerEditor.AndroidFileServerRuntimeSettings", "bUseManualIPAddress", out bUseManualIPAddress);
+		if (!Ini.GetString("/Script/AndroidFileServerEditor.AndroidFileServerRuntimeSettings", "ManualIPAddress", out ManualIPAddress))
+		{
+			ManualIPAddress = "127.0.0.1";
+		}
+
+		Ini.GetBool("/Script/AndroidFileServerEditor.AndroidFileServerRuntimeSettings", "bAllowNetworkConnection", out bAllowNetworkConnection);
+		if (!bAllowNetworkConnection && ConnectionType != EConnectionType.USBOnly)
+		{
+			Log.TraceWarning("AFS will only use USB connection due to network connection disabled");
+			ConnectionType = EConnectionType.USBOnly;
+		}
+
+		return ConnectionType;
+	}
+
 	private bool BuildWithHiddenSymbolVisibility(DeploymentContext SC)
 	{
 		UnrealTargetConfiguration TargetConfiguration = SC.StageTargetConfigurations[0];
@@ -683,6 +764,26 @@ public class AndroidPlatform : Platform
 		var Deploy = AndroidExports.CreateDeploymentHandler(Params.RawProjectPath, Params.ForcePackageData);
 		bool bPackageDataInsideApk = Deploy.GetPackageDataInsideApk();
 
+		bool bUseAFS = false;
+		bool bUseAFSProject = false;
+
+		bool bAFSEnablePlugin;
+		string AFSToken;
+		bool bIsShipping;
+		bool bAFSIncludeInShipping;
+		bool bAFSAllowExternalStartInShipping;
+		UsingAndroidFileServer(SC, out bAFSEnablePlugin, out AFSToken, out bIsShipping, out bAFSIncludeInShipping, out bAFSAllowExternalStartInShipping);
+
+		if (bAFSEnablePlugin)
+		{
+			bUseAFS = true;
+			// AFSProject APK should be used if shipping and AFS wasn't included
+			if (bIsShipping && !(bAFSIncludeInShipping && bAFSAllowExternalStartInShipping))
+			{
+				bUseAFSProject = true;
+			}
+		}
+
 		string BaseApkName = GetFinalApkName(Params, SC.StageExecutables[0], true, "");
 		LogInformation("BaseApkName = {0}", BaseApkName);
 
@@ -1029,7 +1130,7 @@ public class AndroidPlatform : Platform
 				string PackageName = GetPackageInfo(ApkName, SC, false);
 				string BatchName = GetFinalBatchName(ApkName, SC, bMakeSeparateApks ? Architecture : "", false, EBatchType.Install, Target);
 				string[] BatchLines = GenerateInstallBatchFile(bPackageDataInsideApk, PackageName, ApkName, Params, ObbName, DeviceObbName, false, PatchName, DevicePatchName, false, 
-					Overflow1Name, DeviceOverflow1Name, false, Overflow2Name, DeviceOverflow2Name, false, bIsPC, Params.Distribution, TargetSDKVersion > 22, bDisablePerfHarden);
+					Overflow1Name, DeviceOverflow1Name, false, Overflow2Name, DeviceOverflow2Name, false, bIsPC, Params.Distribution, TargetSDKVersion > 22, bDisablePerfHarden, bUseAFS, bUseAFSProject, AFSToken, Target);
 				if (bHaveAPK)
 				{
 					// make a batch file that can be used to install the .apk and .obb files
@@ -1049,7 +1150,7 @@ public class AndroidPlatform : Platform
 					UniversalBatchName = GetFinalBatchName(UniversalApkName, SC, "", false, EBatchType.Install, Target);
 					// make a batch file that can be used to install the .apk
 					string[] UniversalBatchLines = GenerateInstallBatchFile(bPackageDataInsideApk, PackageName, UniversalApkName, Params, ObbName, DeviceObbName, false, PatchName, DevicePatchName, false,
-						Overflow1Name, DeviceOverflow1Name, false, Overflow2Name, DeviceOverflow2Name, false, bIsPC, Params.Distribution, TargetSDKVersion > 22, bDisablePerfHarden);
+						Overflow1Name, DeviceOverflow1Name, false, Overflow2Name, DeviceOverflow2Name, false, bIsPC, Params.Distribution, TargetSDKVersion > 22, bDisablePerfHarden, bUseAFS, bUseAFSProject, AFSToken, Target);
 					File.WriteAllText(UniversalBatchName, string.Join(LineEnding, UniversalBatchLines) + LineEnding);
 				}
 
@@ -1111,9 +1212,27 @@ public class AndroidPlatform : Platform
 		PrintRunTime();
 	}
 
-    private string[] GenerateInstallBatchFile(bool bPackageDataInsideApk, string PackageName, string ApkName, ProjectParams Params, string ObbName, string DeviceObbName, bool bNoObbInstall,
+	string GetAFSExecutable(UnrealTargetPlatform Target)
+	{
+		if (Target == UnrealTargetPlatform.Win64)
+		{
+			return "win-x64/UnrealAndroidFileTool.exe";
+		}
+		if (Target == UnrealTargetPlatform.Mac)
+		{
+			return "osx-x64/UnrealAndroidFileTool";
+		}
+		if (Target == UnrealTargetPlatform.Linux)
+		{
+			return "linux-x64/UnrealAndroidFileTool";
+		}
+		Log.TraceWarning("GetAFSExecutable unsupported target, assuming Win64");
+		return "win-x64/UnrealAndroidFileTool.exe";
+	}
+
+	private string[] GenerateInstallBatchFile(bool bPackageDataInsideApk, string PackageName, string ApkName, ProjectParams Params, string ObbName, string DeviceObbName, bool bNoObbInstall,
 		string PatchName, string DevicePatchName, bool bNoPatchInstall, string Overflow1Name, string DeviceOverflow1Name, bool bNoOverflow1Install, string Overflow2Name, string DeviceOverflow2Name, bool bNoOverflow2Install,
-		bool bIsPC, bool bIsDistribution, bool bRequireRuntimeStoragePermission, bool bDisablePerfHarden)
+		bool bIsPC, bool bIsDistribution, bool bRequireRuntimeStoragePermission, bool bDisablePerfHarden, bool bUseAFS, bool bUseAFSProject, string AFSToken, UnrealTargetPlatform Target)
     {
         string[] BatchLines = null;
         string ReadPermissionGrantCommand = "shell pm grant " + PackageName + " android.permission.READ_EXTERNAL_STORAGE";
@@ -1125,26 +1244,56 @@ public class AndroidPlatform : Platform
 		bool bNeedGrantStoragePermission = bRequireRuntimeStoragePermission && !bIsDistribution;
 
 		// We can't always push directly to Android/obb so uploads to Download then moves it
-		bool bDontMoveOBB = bPackageDataInsideApk || !bIsDistribution;
+		bool bDontMoveOBB = bUseAFS ? true : (bPackageDataInsideApk || !bIsDistribution);
 
 		bool bHavePatch = (PatchName != "");
 		bool bHaveOverflow1 = (Overflow1Name != "");
 		bool bHaveOverflow2 = (Overflow2Name != "");
 
+		string AFSExecutable = GetAFSExecutable(Target);
+		string AFSCommonArg = "-p " + PackageName;
+		if (AFSToken != "")
+		{
+			AFSCommonArg += " -k " + AFSToken;
+		}
+
 		if (!bIsPC)
         {
+			string APKInstallCommand = "$ADB $DEVICE install " + Path.GetFileName(ApkName);
+			string APKReinstallCommand = "";
+
 			// If it is a distribution build, push to $STORAGE/Android/obb folder instead of $STORAGE/obb folder.
 			// Note that $STORAGE/Android/obb will be the folder that contains the obb if you download the app from playstore.
-			string OBBInstallCommand = bNoObbInstall ? "shell 'rm -r $EXTERNAL_STORAGE/" + DeviceObbName + "'" : "push " + Path.GetFileName(ObbName) + (bIsDistribution ? " " + TargetAndroidTemp : " $STORAGE/") + DeviceObbName;
-			string PatchInstallCommand = bNoPatchInstall ? "shell 'rm -r $EXTERNAL_STORAGE/" + DevicePatchName + "'" : "push " + Path.GetFileName(PatchName) + (bIsDistribution ? " " + TargetAndroidTemp : " $STORAGE/") + DevicePatchName;
-			string Overflow1InstallCommand = bNoOverflow1Install ? "shell 'rm -r $EXTERNAL_STORAGE/" + DeviceOverflow1Name + "'" : "push " + Path.GetFileName(Overflow1Name) + (bIsDistribution ? " " + TargetAndroidTemp : " $STORAGE/") + DeviceOverflow1Name;
-			string Overflow2InstallCommand = bNoOverflow2Install ? "shell 'rm -r $EXTERNAL_STORAGE/" + DeviceOverflow2Name + "'" : "push " + Path.GetFileName(Overflow2Name) + (bIsDistribution ? " " + TargetAndroidTemp : " $STORAGE/") + DeviceOverflow2Name;
+			string OBBInstallCommand = bNoObbInstall ? "\t$ADB $DEVICE shell 'rm -r $EXTERNAL_STORAGE/" + DeviceObbName + "'" : "\t$ADB $DEVICE push " + Path.GetFileName(ObbName) + (bIsDistribution ? " " + TargetAndroidTemp : " $STORAGE/") + DeviceObbName;
+			string PatchInstallCommand = bNoPatchInstall ? "\t$ADB $DEVICE shell 'rm -r $EXTERNAL_STORAGE/" + DevicePatchName + "'" : "\t$ADB $DEVICE push " + Path.GetFileName(PatchName) + (bIsDistribution ? " " + TargetAndroidTemp : " $STORAGE/") + DevicePatchName;
+			string Overflow1InstallCommand = bNoOverflow1Install ? "\t$ADB $DEVICE shell 'rm -r $EXTERNAL_STORAGE/" + DeviceOverflow1Name + "'" : "\t$ADB $DEVICE push " + Path.GetFileName(Overflow1Name) + (bIsDistribution ? " " + TargetAndroidTemp : " $STORAGE/") + DeviceOverflow1Name;
+			string Overflow2InstallCommand = bNoOverflow2Install ? "\t$ADB $DEVICE shell 'rm -r $EXTERNAL_STORAGE/" + DeviceOverflow2Name + "'" : "\t$ADB $DEVICE push " + Path.GetFileName(Overflow2Name) + (bIsDistribution ? " " + TargetAndroidTemp : " $STORAGE/") + DeviceOverflow2Name;
+
+			if (bUseAFS)
+			{
+				if (bUseAFSProject)
+				{
+					APKInstallCommand = "$ADB $DEVICE install AFS_" + Path.GetFileName(ApkName);
+					APKReinstallCommand = "$ADB $DEVICE install -r " + Path.GetFileName(ApkName);
+				}
+				else
+				{
+					// stop the fileserver (not needed on reinstall above
+					APKReinstallCommand = "$AFS $DEVICE " + AFSCommonArg + " stop-all";
+				}
+				string AFSCommand = "\t$AFS $DEVICE " + AFSCommonArg; 
+				OBBInstallCommand = bNoObbInstall ? AFSCommand + " deletefile '^mainobb'" : AFSCommand + " push " + Path.GetFileName(ObbName) + " '^mainobb'";
+				PatchInstallCommand = bNoPatchInstall ? AFSCommand + " deletefile '^patchobb'" : AFSCommand + " push " + Path.GetFileName(PatchName) + " '^patchobb'";
+				Overflow1InstallCommand = bNoOverflow1Install ? AFSCommand + " deletefile '^overflow1obb'" : AFSCommand + " push " + Path.GetFileName(Overflow1Name) + " '^overflow1obb'";
+				Overflow1InstallCommand = bNoOverflow2Install ? AFSCommand + " deletefile '^overflow2obb'" : AFSCommand + " push " + Path.GetFileName(Overflow2Name) + " '^overflow2obb'";
+			}
 
 			LogInformation("Writing shell script for install with {0}", bPackageDataInsideApk ? "data in APK" : "separate obb");
             BatchLines = new string[] {
 						"#!/bin/sh",
 						"cd \"`dirname \"$0\"`\"",
-                        "ADB=",
+						"AFS=./" + AFSExecutable,
+						"ADB=",
 						"if [ \"$ANDROID_HOME\" != \"\" ]; then ADB=$ANDROID_HOME/platform-tools/adb; else ADB=" +Environment.GetEnvironmentVariable("ANDROID_HOME") + "/platform-tools/adb; fi",
 						"DEVICE=",
 						"if [ \"$1\" != \"\" ]; then DEVICE=\"-s $1\"; fi",
@@ -1153,9 +1302,10 @@ public class AndroidPlatform : Platform
 						"$ADB $DEVICE uninstall " + PackageName,
 						"echo",
 						"echo Installing existing application. Failures here indicate a problem with the device \\(connection or storage permissions\\) and are fatal.",
-						"$ADB $DEVICE install " + Path.GetFileName(ApkName),
+						APKInstallCommand,
 						"if [ $? -eq 0 ]; then",
                         "\techo",
+						"\t$ADB $DEVICE shell pm list packages " + PackageName,
 						bNeedGrantStoragePermission ? "\techo Grant READ_EXTERNAL_STORAGE and WRITE_EXTERNAL_STORAGE to the apk for reading OBB or game file in external storage." : "",
 						bNeedGrantStoragePermission ? "\t$ADB $DEVICE " + ReadPermissionGrantCommand : "",
 						bNeedGrantStoragePermission ? "\t$ADB $DEVICE " + WritePermissionGrantCommand : "",
@@ -1170,14 +1320,15 @@ public class AndroidPlatform : Platform
 						bPackageDataInsideApk ? "" : "\techo",
 						bPackageDataInsideApk ? "" : "\techo Installing new data. Failures here indicate storage problems \\(missing SD card or bad permissions\\) and are fatal.",
 						bPackageDataInsideApk ? "" : "\tSTORAGE=$(echo \"`$ADB $DEVICE shell 'echo $EXTERNAL_STORAGE'`\" | cat -v | tr -d '^M')",
-						bPackageDataInsideApk ? "" : "\t$ADB $DEVICE " + OBBInstallCommand,
+						bPackageDataInsideApk ? "" : OBBInstallCommand,
 						bPackageDataInsideApk ? "if [ 1 ]; then" : "\tif [ $? -eq 0 ]; then",
-						!bHavePatch ? "" : (bPackageDataInsideApk ? "" : "\t$ADB $DEVICE " + PatchInstallCommand),
-						!bHaveOverflow1 ? "" : (bPackageDataInsideApk ? "" : "\t$ADB $DEVICE " + Overflow1InstallCommand),
-						!bHaveOverflow2 ? "" : (bPackageDataInsideApk ? "" : "\t$ADB $DEVICE " + Overflow2InstallCommand),
+						!bHavePatch ? "" : (bPackageDataInsideApk ? "" : PatchInstallCommand),
+						!bHaveOverflow1 ? "" : (bPackageDataInsideApk ? "" : Overflow1InstallCommand),
+						!bHaveOverflow2 ? "" : (bPackageDataInsideApk ? "" : Overflow2InstallCommand),
 						bDontMoveOBB ? "" : "\t\t$ADB $DEVICE shell mkdir $STORAGE/Android/" + TargetAndroidLocation + PackageName, // don't check for error since installing may create the obb directory
 						bDontMoveOBB ? "" : "\t\t$ADB $DEVICE shell mv " + TargetAndroidTemp + TargetAndroidLocation + PackageName + " $STORAGE/Android/" + TargetAndroidLocation,
 						bDontMoveOBB ? "" : "\t\t$ADB $DEVICE shell rm -r " + TargetAndroidTemp + TargetAndroidLocation,
+						APKReinstallCommand,
 						"\t\techo",
 						"\t\techo Installation successful",
 						"\t\texit 0",
@@ -1195,10 +1346,32 @@ public class AndroidPlatform : Platform
         }
         else
         {
-			string OBBInstallCommand = bNoObbInstall ? "shell rm -r %STORAGE%/" + DeviceObbName : "push " + Path.GetFileName(ObbName) + (bIsDistribution ? " " + TargetAndroidTemp : " %STORAGE%/") + DeviceObbName;
-			string PatchInstallCommand = bNoPatchInstall ? "shell rm -r %STORAGE%/" + DevicePatchName : "push " + Path.GetFileName(PatchName) + (bIsDistribution ? " " + TargetAndroidTemp : " %STORAGE%/") + DevicePatchName;
-			string Overflow1InstallCommand = bNoOverflow1Install ? "shell rm -r %STORAGE%/" + DeviceOverflow1Name : "push " + Path.GetFileName(Overflow1Name) + (bIsDistribution ? " " + TargetAndroidTemp : " %STORAGE%/") + DeviceOverflow1Name;
-			string Overflow2InstallCommand = bNoOverflow2Install ? "shell rm -r %STORAGE%/" + DeviceOverflow2Name : "push " + Path.GetFileName(Overflow2Name) + (bIsDistribution ? " " + TargetAndroidTemp : " %STORAGE%/") + DeviceOverflow2Name;
+			string APKInstallCommand = "%ADB% %DEVICE% install " + Path.GetFileName(ApkName);
+			string APKReinstallCommand = "";
+
+			string OBBInstallCommand = bNoObbInstall ? "%ADB% %DEVICE% shell rm -r %STORAGE%/" + DeviceObbName : "%ADB% %DEVICE% push " + Path.GetFileName(ObbName) + (bIsDistribution ? " " + TargetAndroidTemp : " %STORAGE%/") + DeviceObbName;
+			string PatchInstallCommand = bNoPatchInstall ? "%ADB% %DEVICE% shell rm -r %STORAGE%/" + DevicePatchName : "%ADB% %DEVICE% push " + Path.GetFileName(PatchName) + (bIsDistribution ? " " + TargetAndroidTemp : " %STORAGE%/") + DevicePatchName;
+			string Overflow1InstallCommand = bNoOverflow1Install ? "%ADB% %DEVICE% shell rm -r %STORAGE%/" + DeviceOverflow1Name : "%ADB% %DEVICE% push " + Path.GetFileName(Overflow1Name) + (bIsDistribution ? " " + TargetAndroidTemp : " %STORAGE%/") + DeviceOverflow1Name;
+			string Overflow2InstallCommand = bNoOverflow2Install ? "%ADB% %DEVICE% shell rm -r %STORAGE%/" + DeviceOverflow2Name : "%ADB% %DEVICE% push " + Path.GetFileName(Overflow2Name) + (bIsDistribution ? " " + TargetAndroidTemp : " %STORAGE%/") + DeviceOverflow2Name;
+
+			if (bUseAFS)
+			{
+				if (bUseAFSProject)
+				{
+					APKInstallCommand = "%ADB% %DEVICE% install AFS_" + Path.GetFileName(ApkName);
+					APKReinstallCommand = "%ADB% %DEVICE% install -r " + Path.GetFileName(ApkName);
+				}
+				else
+				{
+					// stop the fileserver (not needed on reinstall above
+					APKReinstallCommand = "%AFS% %DEVICE% " + AFSCommonArg + " stop-all";
+				}
+				string AFSCommand = "\t%AFS% %DEVICE% " + AFSCommonArg;
+				OBBInstallCommand = bNoObbInstall ? AFSCommand + " deletefile \"^mainobb\"" : AFSCommand + " push " + Path.GetFileName(ObbName) + " \"^mainobb\"";
+				PatchInstallCommand = bNoPatchInstall ? AFSCommand + " deletefile \"^patchobb\"" : AFSCommand + " push " + Path.GetFileName(PatchName) + " \"^patchobb\"";
+				Overflow1InstallCommand = bNoOverflow1Install ? AFSCommand + " deletefile \"^overflow1obb\"" : AFSCommand + " push " + Path.GetFileName(Overflow1Name) + " \"^overflow1obb\"";
+				Overflow1InstallCommand = bNoOverflow2Install ? AFSCommand + " deletefile \"^overflow2obb\"" : AFSCommand + " push " + Path.GetFileName(Overflow2Name) + " \"^overflow2obb\"";
+			}
 
 			LogInformation("Writing bat for install with {0}", bPackageDataInsideApk ? "data in APK" : "separate OBB");
             BatchLines = new string[] {
@@ -1207,6 +1380,7 @@ public class AndroidPlatform : Platform
 						"set ANDROIDHOME=%ANDROID_HOME%",		
 						"if \"%ANDROIDHOME%\"==\"\" set ANDROIDHOME="+Environment.GetEnvironmentVariable("ANDROID_HOME"),
 						"set ADB=%ANDROIDHOME%\\platform-tools\\adb.exe",
+						"set AFS=.\\" + AFSExecutable.Replace("/", "\\"),
 						"set DEVICE=",
                         "if not \"%1\"==\"\" set DEVICE=-s %1",
                         "for /f \"delims=\" %%A in ('%ADB% %DEVICE% " + GetStorageQueryCommand(true) +"') do @set STORAGE=%%A",
@@ -1215,27 +1389,30 @@ public class AndroidPlatform : Platform
 						"%ADB% %DEVICE% uninstall " + PackageName,
 						"@echo.",
 						"@echo Installing existing application. Failures here indicate a problem with the device (connection or storage permissions) and are fatal.",
-						"%ADB% %DEVICE% install " + Path.GetFileName(ApkName),
+						APKInstallCommand,
 						"@if \"%ERRORLEVEL%\" NEQ \"0\" goto Error",
-                        "%ADB% %DEVICE% shell rm -r %STORAGE%/UnrealGame/" + Params.ShortProjectName,
+						"%ADB% %DEVICE% shell pm list packages " + PackageName,
+						"%ADB% %DEVICE% shell rm -r %STORAGE%/UnrealGame/" + Params.ShortProjectName,
 						"%ADB% %DEVICE% shell rm -r %STORAGE%/UnrealGame/UECommandLine.txt", // we need to delete the commandline in UnrealGame or it will mess up loading
 						"%ADB% %DEVICE% shell rm -r %STORAGE%/" + TargetAndroidLocation + PackageName,
 						"%ADB% %DEVICE% shell rm -r %STORAGE%/Android/" + TargetAndroidLocation + PackageName,
 						"%ADB% %DEVICE% shell rm -r %STORAGE%/Download/" + TargetAndroidLocation + PackageName,
 						bPackageDataInsideApk ? "" : "@echo.",
 						bPackageDataInsideApk ? "" : "@echo Installing new data. Failures here indicate storage problems (missing SD card or bad permissions) and are fatal.",
-						bPackageDataInsideApk ? "" : "%ADB% %DEVICE% " + OBBInstallCommand,
+						bPackageDataInsideApk ? "" : OBBInstallCommand,
 						bPackageDataInsideApk ? "" : "if \"%ERRORLEVEL%\" NEQ \"0\" goto Error",
-						!bHavePatch ? "" : (bPackageDataInsideApk ? "" : "%ADB% %DEVICE% " + PatchInstallCommand),
+						!bHavePatch ? "" : (bPackageDataInsideApk ? "" : PatchInstallCommand),
 						!bHavePatch ? "" : (bPackageDataInsideApk ? "" : "if \"%ERRORLEVEL%\" NEQ \"0\" goto Error"),
-						!bHaveOverflow1 ? "" : (bPackageDataInsideApk ? "" : "%ADB% %DEVICE% " + Overflow1InstallCommand),
+						!bHaveOverflow1 ? "" : (bPackageDataInsideApk ? "" : Overflow1InstallCommand),
 						!bHaveOverflow1 ? "" : (bPackageDataInsideApk ? "" : "if \"%ERRORLEVEL%\" NEQ \"0\" goto Error"),
-						!bHaveOverflow2 ? "" : (bPackageDataInsideApk ? "" : "%ADB% %DEVICE% " + Overflow2InstallCommand),
+						!bHaveOverflow2 ? "" : (bPackageDataInsideApk ? "" : Overflow2InstallCommand),
 						!bHaveOverflow2 ? "" : (bPackageDataInsideApk ? "" : "if \"%ERRORLEVEL%\" NEQ \"0\" goto Error"),
 						bDontMoveOBB ? "" : "%ADB% %DEVICE% shell mkdir %STORAGE%/Android/" + TargetAndroidLocation + PackageName, // don't check for error since installing may create the obb directory
 						bDontMoveOBB ? "" : "%ADB% %DEVICE% shell mv " + TargetAndroidTemp + TargetAndroidLocation + PackageName + " %STORAGE%/Android/" + TargetAndroidLocation,
 						bDontMoveOBB ? "" : "if \"%ERRORLEVEL%\" NEQ \"0\" goto Error",
 						bDontMoveOBB ? "" : "%ADB% %DEVICE% shell rm -r " + TargetAndroidTemp + TargetAndroidLocation,
+						APKReinstallCommand,
+						bUseAFSProject ? "if \"%ERRORLEVEL%\" NEQ \"0\" goto Error" : "",
 						"@echo.",
 						bNeedGrantStoragePermission ? "@echo Grant READ_EXTERNAL_STORAGE and WRITE_EXTERNAL_STORAGE to the apk for reading OBB file or game file in external storage." : "",
 						bNeedGrantStoragePermission ? "%ADB% %DEVICE% " + ReadPermissionGrantCommand : "",
@@ -1372,6 +1549,14 @@ public class AndroidPlatform : Platform
 		bool bMakeSeparateApks = UnrealBuildTool.AndroidExports.ShouldMakeSeparateApks();
 		bool bPackageDataInsideApk = UnrealBuildTool.AndroidExports.CreateDeploymentHandler(Params.RawProjectPath, Params.ForcePackageData).GetPackageDataInsideApk();
 
+		bool bAFSEnablePlugin;
+		string AFSToken;
+		bool bIsShipping;
+		bool bAFSIncludeInShipping;
+		bool bAFSAllowExternalStartInShipping;
+		UsingAndroidFileServer(SC, out bAFSEnablePlugin, out AFSToken, out bIsShipping, out bAFSIncludeInShipping, out bAFSAllowExternalStartInShipping);
+		bool bUseAFS = bAFSEnablePlugin && !bPackageDataInsideApk;
+
 		List<string> AddedObbFiles = new List<string>();
 		foreach (string Architecture in Architectures)
 		{
@@ -1442,6 +1627,13 @@ public class AndroidPlatform : Platform
 				}
 			}
 
+			// copy optional AFSProject apk if exists
+			string AFSApkName = Path.Combine(APKDirectory, "AFS_" + APKNameWithoutExtension + ".apk");
+			if (FileExists(AFSApkName))
+			{
+				SC.ArchiveFiles(APKDirectory, Path.GetFileName(AFSApkName));
+			}
+
 			// verify the files exist
 			if (!FileExists(ApkName))
 			{
@@ -1540,6 +1732,10 @@ public class AndroidPlatform : Platform
 					string SymbolizeBatchName = GetFinalBatchName(ApkName, SC, Architecture, false, EBatchType.Symbolize, Target);
 					SC.ArchiveFiles(Path.GetDirectoryName(SymbolizeBatchName), Path.GetFileName(SymbolizeBatchName));
 				}
+				if (bUseAFS && (bHaveAPK || bHaveUniversal))
+				{
+					SC.ArchiveFiles(Path.Combine(SC.EngineRoot.FullName, "Binaries", "DotNet", "Android", "UnrealAndroidFileTool"), GetAFSExecutable(Target));
+				}
 				//SC.ArchiveFiles(Path.GetDirectoryName(NoOBBBatchName), Path.GetFileName(NoOBBBatchName));
 			}
 			);
@@ -1705,8 +1901,92 @@ public class AndroidPlatform : Platform
 	}
 	*/
 
+	private bool RetrieveDeployedManifestsAFS(ProjectParams Params, DeploymentContext SC, string DeviceName, out List<string> UFSManifests, out List<string> NonUFSManifests, string AFSToken)
+	{
+		UFSManifests = null;
+		NonUFSManifests = null;
+
+		string DeviceArchitecture = GetBestDeviceArchitecture(Params, DeviceName);
+		string ApkName = GetFinalApkName(Params, SC.StageExecutables[0], true, DeviceArchitecture);
+		string PackageName = GetPackageInfo(ApkName, SC, false);
+
+		AndroidFileClient client = new AndroidFileClient(DeviceName);
+		if (!client.OpenConnection())
+		{
+			if (PackageName == null || PackageName == "")
+			{
+				Log.TraceInformation("Retrieve Manifests: Unable to start file server without package name, ignoring manifests");
+				return false;
+			}
+			Log.TraceInformation("Retrieve Manifests: Trying to start file server {0}", PackageName);
+			if (!client.StartServer(PackageName, AFSToken))
+			{
+				Log.TraceInformation("Retrieve Manifests: Failed to start server, ignoring manifests");
+				return false;
+			}
+		}
+
+		// verify connection to the correct server
+		string DevicePackageName = client.Query("^packagename");
+		if (DevicePackageName != PackageName)
+		{
+			if (PackageName == null || PackageName == "")
+			{
+				Log.TraceInformation("Retrieve Manifests: Unable to start file server without package name, ignoring manifests");
+				client.CloseConnection();
+				return false;
+			}
+
+			Log.TraceInformation("Connected to wrong server {0}, trying again", DevicePackageName);
+			client.TerminateServer();
+
+			Log.TraceInformation("Retrieve Manifests: Trying to start file server {0}", PackageName);
+			if (!client.StartServer(PackageName, AFSToken))
+			{
+				Log.TraceInformation("Retrieve Manifests: Failed to start server, ignoring manifests");
+				return false;
+			}
+		}
+
+		// Try retrieving the UFS files manifest files from the device
+		string UFSManifestFileName = CombinePaths(SC.StageDirectory.FullName, SC.GetUFSDeployedManifestFileName(DeviceName));
+		if (!client.FileRead("^project/" + SC.GetUFSDeployedManifestFileName(null), UFSManifestFileName))
+		{
+			return false;
+		}
+
+		// Try retrieving the non UFS files manifest files from the device
+		string NonUFSManifestFileName = CombinePaths(SC.StageDirectory.FullName, SC.GetNonUFSDeployedManifestFileName(DeviceName));
+		if (!client.FileRead("^project/" + SC.GetNonUFSDeployedManifestFileName(null), NonUFSManifestFileName))
+		{
+			return false;
+		}
+
+		client.CloseConnection();
+
+		// Return the manifest files
+		UFSManifests = new List<string>();
+		UFSManifests.Add(UFSManifestFileName);
+		NonUFSManifests = new List<string>();
+		NonUFSManifests.Add(NonUFSManifestFileName);
+
+		Log.TraceInformation("Retrieve Manifests: Success!!");
+
+		return true;
+	}
+
 	public override bool RetrieveDeployedManifests(ProjectParams Params, DeploymentContext SC, string DeviceName, out List<string> UFSManifests, out List<string> NonUFSManifests)
 	{
+		bool bAFSEnablePlugin;
+		string AFSToken;
+		bool bIsShipping;
+		bool bAFSIncludeInShipping;
+		bool bAFSAllowExternalStartInShipping;
+		if (UsingAndroidFileServer(SC, out bAFSEnablePlugin, out AFSToken, out bIsShipping, out bAFSIncludeInShipping, out bAFSAllowExternalStartInShipping))
+		{
+			return RetrieveDeployedManifestsAFS(Params, SC, DeviceName, out UFSManifests, out NonUFSManifests, AFSToken);
+		}
+
 		UFSManifests = null;
 		NonUFSManifests = null;
 
@@ -1770,7 +2050,592 @@ public class AndroidPlatform : Platform
 		}
 	}
 
-    public override void Deploy(ProjectParams Params, DeploymentContext SC)
+	private void DeployAndroidFileServer(ProjectParams Params, DeploymentContext SC, string AFSToken)
+	{
+		var AppArchitectures = AndroidExports.CreateToolChain(Params.RawProjectPath).GetAllArchitectures();
+
+		ConfigHierarchy Ini = ConfigCache.ReadHierarchy(ConfigHierarchyType.Engine, DirectoryReference.FromFile(Params.RawProjectPath), UnrealTargetPlatform.Android);
+		bool bDisablePerfHarden = false;
+		Ini.GetBool("/Script/AndroidRuntimeSettings.AndroidRuntimeSettings", "bEnableMaliPerfCounters", out bDisablePerfHarden);
+
+		bool bUseCompression = false;
+		bool bLogFiles = false;
+		bool bReportStats = false;
+		bool bUseManualIPAddress = false;
+		string ManualIPAddress = "";
+		EConnectionType ConnectionType = GetAndroidFileServerNetworkConfig(SC, out bUseCompression, out bLogFiles, out bReportStats, out bUseManualIPAddress, out ManualIPAddress);
+
+		AndroidFileClient.OptimalADB adb = new AndroidFileClient.OptimalADB();
+
+		foreach (var DeviceName in Params.DeviceNames)
+		{
+			string DeviceArchitecture = GetBestDeviceArchitecture(Params, DeviceName);
+			string ApkName = GetFinalApkName(Params, SC.StageExecutables[0], true, DeviceArchitecture);
+
+			// make sure APK is up to date (this is fast if so)
+			var Deploy = AndroidExports.CreateDeploymentHandler(Params.RawProjectPath, Params.ForcePackageData);
+			if (!Params.Prebuilt)
+			{
+				string CookFlavor = SC.FinalCookPlatform.IndexOf("_") > 0 ? SC.FinalCookPlatform.Substring(SC.FinalCookPlatform.IndexOf("_")) : "";
+				string SOName = GetSONameWithoutArchitecture(Params, SC.StageExecutables[0]);
+				Deploy.SetAndroidPluginData(AppArchitectures, CollectPluginDataPaths(SC));
+				Deploy.PrepForUATPackageOrDeploy(Params.RawProjectPath, Params.ShortProjectName, SC.ProjectRoot, SOName, SC.LocalRoot + "/Engine", Params.Distribution, CookFlavor, SC.StageTargets[0].Receipt.Configuration, true, false);
+			}
+
+			// now we can use the apk to get more info
+			string PackageName = GetPackageInfo(ApkName, SC, false);
+
+			// start up AFS connection (allowed to fail here.. may not be a server installed yet)
+			AndroidFileClient client = new AndroidFileClient(DeviceName);
+			AndroidFileClient client2 = null;
+			string IPAddress = bUseManualIPAddress ? ManualIPAddress : "127.0.0.1";
+			if (!client.OpenConnection(IPAddress))
+			{
+				if (!client.StartServer(PackageName, AFSToken, IPAddress))
+				{
+				}
+			}
+
+			// Try to get IP address from device if we connected
+			string DeviceIPAddress = client.Query("^ip");
+
+			// Setup the OBB name and add the storage path (queried from the device) to it
+			string QueryStorageResult = adb.Shell(DeviceName, "echo $EXTERNAL_STORAGE");
+			string StorageLocation = QueryStorageResult.Trim() + "/Android";  // "mnt/sdcard"
+			string DeviceObbName = StorageLocation + "/" + GetDeviceObbName(ApkName, SC);
+			string DevicePatchName = StorageLocation + "/" + GetDevicePatchName(ApkName, SC);
+			string DeviceOverflow1Name = StorageLocation + "/" + GetDeviceOverflowName(ApkName, SC, 1);
+			string DeviceOverflow2Name = StorageLocation + "/" + GetDeviceOverflowName(ApkName, SC, 2);
+			string RemoteDir = client.Query("^project", true);
+			if (RemoteDir == null)
+			{
+				RemoteDir = StorageLocation + "/data/" + PackageName + "/files/UnrealGame/" + Params.ShortProjectName;
+			}
+			string ExtFiles = StorageLocation + "/data/" + PackageName + "/files";
+
+			if (bDisablePerfHarden)
+			{
+				adb.Shell(DeviceName, "setprop security.perf_harden 0");
+			}
+
+			// close connection to since uninstall/reinstall will reset server
+			client.CloseConnection();
+
+			// determine if APK out of date
+			string APKLastUpdateTime = new FileInfo(ApkName).LastWriteTime.ToString();
+			bool bNeedAPKInstall = true;
+			bool bFreshInstall = false;
+			bool bFastDeploy = false;
+			if (Params.IterativeDeploy)
+			{
+				// Check for apk installed with this package name on the device
+				String InstalledResult = adb.Shell(DeviceName, "pm list packages " + PackageName);
+				if (InstalledResult.Contains(PackageName))
+				{
+					Log.TraceInformation("{0} already installed!", PackageName);
+					// already installed so enable --fast-deploy option if need to update apk
+					bFastDeploy = true;
+
+					// See if apk is up to date on device
+					InstalledResult = adb.Shell(DeviceName, "cat " + ExtFiles + "/APKFileStamp.txt");
+					if (InstalledResult.StartsWith("APK: "))
+					{
+						Log.TraceInformation("Found APKFileStamp.txt! {0}", InstalledResult);
+						if (InstalledResult.Substring(5).Trim() == APKLastUpdateTime)
+							bNeedAPKInstall = false;
+
+						if (InstalledResult.Substring(5).Trim() != APKLastUpdateTime)
+						{
+							Log.TraceInformation("{0} != {1}", InstalledResult.Substring(5).Trim(), APKLastUpdateTime);
+						}
+
+						// Stop the previously running copy (uninstall/install did this before)
+						InstalledResult = adb.Shell(DeviceName, "am force-stop " + PackageName);
+						if (InstalledResult.Contains("Error"))
+						{
+							// force-stop not supported (Android < 3.0) so check if package is actually running
+							// Note: cannot use grep here since it may not be installed on device
+							InstalledResult = adb.Shell(DeviceName, "ps");
+							if (InstalledResult.Contains(PackageName))
+							{
+								// it is actually running so use the slow way to kill it (uninstall and reinstall)
+								bNeedAPKInstall = true;
+							}
+						}
+					}
+				}
+				else
+				{
+					// If not already installed we must do a full deploy
+					bFreshInstall = true;
+				}
+			}
+
+			// install new APK if needed
+			if (bNeedAPKInstall)
+			{
+				// try reinstall the apk to preserve data
+				int SuccessCode = 0;
+				string InstallCommandline = "install -r " + (bFastDeploy ? "--fastdeploy \"" : "\"") + ApkName + "\"";
+				string InstallOutput = RunAndLogAdbCommand(Params, DeviceName, InstallCommandline, out SuccessCode);
+				int FailureIndex = InstallOutput.IndexOf("Failure");
+				// adb install doesn't always return an error code on failure, and instead prints "Failure", followed by an error code.
+				if (SuccessCode != 0 || FailureIndex != -1)
+				{
+					string ErrorMessage = string.Format("Installation of apk '{0}' failed", ApkName);
+					if (FailureIndex != -1)
+					{
+						string FailureString = InstallOutput.Substring(FailureIndex + 7).Trim();
+						if (FailureString != "")
+						{
+							ErrorMessage += ": " + FailureString;
+						}
+					}
+					if (ErrorMessage.Contains("OLDER_SDK"))
+					{
+						LogError("minSdkVersion is higher than Android version installed on device, possibly due to NDK API Level");
+						throw new AutomationException(ExitCode.Error_AppInstallFailed, ErrorMessage);
+					}
+
+					// try uninstalling an old app with the same identifier.
+					// NOTE: uninstall -k will preserve data/cache.. consider using instead?
+					bFreshInstall = true;
+					SuccessCode = 0;
+					adb.Shell(DeviceName, "pm uninstall " + PackageName);
+
+					// install the apk
+					InstallCommandline = "install \"" + ApkName + "\"";
+					InstallOutput = RunAndLogAdbCommand(Params, DeviceName, InstallCommandline, out SuccessCode);
+					FailureIndex = InstallOutput.IndexOf("Failure");
+
+					// adb install doesn't always return an error code on failure, and instead prints "Failure", followed by an error code.
+					if (SuccessCode != 0 || FailureIndex != -1)
+					{
+						ErrorMessage = string.Format("Installation of apk '{0}' failed", ApkName);
+						if (FailureIndex != -1)
+						{
+							string FailureString = InstallOutput.Substring(FailureIndex + 7).Trim();
+							if (FailureString != "")
+							{
+								ErrorMessage += ": " + FailureString;
+							}
+						}
+						if (ErrorMessage.Contains("OLDER_SDK"))
+						{
+							LogError("minSdkVersion is higher than Android version installed on device, possibly due to NDK API Level");
+						}
+						throw new AutomationException(ExitCode.Error_AppInstallFailed, ErrorMessage);
+					}
+				}
+
+				// giving EXTERNAL_STORAGE_WRITE permission to the apk for API23+
+				// without this permission apk can't access to the assets put into the device
+				string ReadPermissionCommandLine = "pm grant " + PackageName + " android.permission.READ_EXTERNAL_STORAGE";
+				string WritePermissionCommandLine = "pm grant " + PackageName + " android.permission.WRITE_EXTERNAL_STORAGE";
+				adb.Shell(DeviceName, ReadPermissionCommandLine);
+				adb.Shell(DeviceName, WritePermissionCommandLine);
+
+				// time for receivers to be registered by pm after install
+				Thread.Sleep(350);
+			}
+
+			// reopen file server connection (either USB or Network)
+			IPAddress = (ConnectionType == EConnectionType.NetworkOnly) ? (bUseManualIPAddress ? ManualIPAddress : (DeviceIPAddress != null ? DeviceIPAddress : "127.0.0.1")) : "127.0.0.1";
+			Log.TraceInformation("Attempting to connect to file server [{0}]", (IPAddress == "127.0.0.1" ? "USB" : IPAddress));
+			if (!client.OpenConnection(IPAddress))
+			{
+				Log.TraceInformation("Not connected, attempting to start file server");
+				if (!client.StartServer(PackageName, AFSToken, IPAddress))
+				{
+					// try one more time with longer delay
+					Log.TraceInformation("Trying again");
+					Thread.Sleep(1000);
+					if (!client.StartServer(PackageName, AFSToken, IPAddress))
+					{
+						Log.TraceWarning("Failed to start Android file server for {0}, skipping deploy for {1}", PackageName, DeviceName);
+						continue;
+					}
+				}
+			}
+
+			// verify we connected to the right server
+			string DevicePackageName = client.Query("^packagename");
+			if (DevicePackageName != PackageName)
+			{
+				Log.TraceInformation("Connected to wrong server {0}, trying again", DevicePackageName);
+				client.TerminateServer();
+
+				Log.TraceInformation("Trying to start file server {0}", PackageName);
+				if (!client.StartServer(PackageName, AFSToken, IPAddress))
+				{
+					// try one more time with longer delay
+					Log.TraceInformation("Trying again");
+					Thread.Sleep(1000);
+					if (!client.StartServer(PackageName, AFSToken, IPAddress))
+					{
+						Log.TraceWarning("Failed to start Android file server for {0}, skipping deploy for {1}", PackageName, DeviceName);
+						continue;
+					}
+				}
+			}
+
+			if (ConnectionType == EConnectionType.Combined)
+			{
+				IPAddress = (bUseManualIPAddress ? ManualIPAddress : client.Query("^ip"));
+				client2 = new AndroidFileClient(DeviceName);
+
+				Log.TraceInformation("Attempting to connect to file server [{0}]", IPAddress);
+				if (!client2.OpenConnection(IPAddress))
+				{
+					Log.TraceInformation("Not connected, attempting to start file server");
+					if (!client2.StartServer(PackageName, AFSToken, IPAddress, false))
+					{
+						Log.TraceWarning("Failed to start Android file server for {0}, only using one connection for {1}", PackageName, DeviceName);
+						client2 = null;
+					}
+				}
+
+				// verify we connected to the right server
+				DevicePackageName = client2.Query("^packagename");
+				if (DevicePackageName != PackageName)
+				{
+					Log.TraceInformation("Connected to wrong server {0} for [{1}], not using network", DevicePackageName, IPAddress, DeviceName);
+					client2.CloseConnection();
+					client2 = null;
+				}
+			}
+
+			// get RemoteDir again (should be valid now after install / restart)
+			RemoteDir = client.Query("^project", true);
+
+			// write new timestamp for APK (do it here since RemoteDir now available)
+			if (bNeedAPKInstall)
+			{
+				client.FileWriteString("APK: " + APKLastUpdateTime + "\n", "^ext/APKFileStamp.txt");
+			}
+
+			// update the uecommandline.txt
+			// update and deploy uecommandline.txt
+			// always delete the existing commandline text file, so it doesn't reuse an old one
+			FileReference IntermediateCmdLineFile = FileReference.Combine(SC.StageDirectory, "UECommandLine.txt");
+			Project.WriteStageCommandline(IntermediateCmdLineFile, Params, SC);
+
+			// copy files to device if we were staging
+			if (SC.Stage)
+			{
+				HashSet<string> EntriesToDeploy = new HashSet<string>();
+
+				// Fresh install always needs full deploy
+				if (Params.IterativeDeploy && !bFreshInstall)
+				{
+					// always send UECommandLine.txt (it was written above after delta checks applied)
+					EntriesToDeploy.Add(IntermediateCmdLineFile.FullName);
+
+					// Add non UFS files if any to deploy
+					String NonUFSManifestPath = SC.GetNonUFSDeploymentDeltaPath(DeviceName);
+					if (File.Exists(NonUFSManifestPath))
+					{
+						string NonUFSFiles = File.ReadAllText(NonUFSManifestPath);
+						foreach (string Filename in NonUFSFiles.Split('\n'))
+						{
+							if (!string.IsNullOrEmpty(Filename) && !string.IsNullOrWhiteSpace(Filename))
+							{
+								EntriesToDeploy.Add(CombinePaths(SC.StageDirectory.FullName, Filename.Trim()));
+							}
+						}
+					}
+
+					// Add UFS files if any to deploy
+					String UFSManifestPath = SC.GetUFSDeploymentDeltaPath(DeviceName);
+					if (File.Exists(UFSManifestPath))
+					{
+						string UFSFiles = File.ReadAllText(UFSManifestPath);
+						foreach (string Filename in UFSFiles.Split('\n'))
+						{
+							if (!string.IsNullOrEmpty(Filename) && !string.IsNullOrWhiteSpace(Filename))
+							{
+								EntriesToDeploy.Add(CombinePaths(SC.StageDirectory.FullName, Filename.Trim()));
+							}
+						}
+					}
+
+					// For now, if too many files may be better to just push them all
+					if (EntriesToDeploy.Count > 500)
+					{
+						// make sure device is at a clean state
+						client.DirDeleteRecurse(RemoteDir);
+
+						EntriesToDeploy.Clear();
+						EntriesToDeploy.TrimExcess();
+						EntriesToDeploy.Add(SC.StageDirectory.FullName);
+					}
+					else
+					{
+						// Discover & remove any files on device that are not in staging
+
+						// get listing of remote directory from device
+						string CommandResult = client.DirListFlat(RemoteDir);
+
+						if (CommandResult == null)
+						{
+							Log.TraceWarning("Failed to read remote dir: {0}", RemoteDir);
+							RemoteDir = client.Query("^project", true);
+							CommandResult = client.DirListFlat(RemoteDir);
+							if (CommandResult == null)
+							{
+								Log.TraceWarning("Failed to read remote dir again: {0}", RemoteDir);
+							}
+						}
+
+						{
+							// listing output is of the form
+							// [Samsung]                 [Google]
+							//
+							// RemoteDir/RestOfPath:     RemoteDir/RestOfPath:
+							// - File1.png               File1.png
+							// - File2.txt               File2.txt
+							// d SubDir1                 SubDir1/
+							// d SubDir2                 Subdir2/
+							//
+							// RemoteDir/RestOfPath/SubDir1:
+
+							HashSet<string> DirsToDeleteFromDevice = new HashSet<string>();
+							List<string> FilesToDeleteFromDevice = new List<string>();
+
+							using (var reader = new StringReader(CommandResult))
+							{
+								string ProjectSaved = Params.ShortProjectName + "/Saved";
+								string ProjectConfig = Params.ShortProjectName + "/Config";
+								const string EngineSaved = "Engine/Saved"; // is this safe to use, or should we use SC.EngineRoot.GetDirectoryName()?
+								const string EngineConfig = "Engine/Config";
+								LogWarning("Excluding {0} {1} {2} {3} from clean during deployment.", ProjectSaved, ProjectConfig, EngineSaved, EngineConfig);
+
+								string CurrentDir = "";
+								bool SkipFiles = false;
+								for (string Line = reader.ReadLine(); Line != null; Line = reader.ReadLine())
+								{
+									if (String.IsNullOrWhiteSpace(Line))
+									{
+										continue; // ignore blank lines
+									}
+
+									if (Line.EndsWith(":"))
+									{
+										// RemoteDir/RestOfPath:
+										//      keep ^--------^
+										CurrentDir = Line.Substring(RemoteDir.Length + 1, Math.Max(0, Line.Length - RemoteDir.Length - 2));
+										// Max is there for the case of base "RemoteDir:" --> ""
+
+										// We want to keep config & logs between deployments.
+										if (CurrentDir.StartsWith(ProjectSaved) || CurrentDir.StartsWith(ProjectConfig) || CurrentDir.StartsWith(EngineSaved) || CurrentDir.StartsWith(EngineConfig))
+										{
+											SkipFiles = true;
+											continue;
+										}
+
+										bool DirExistsInStagingArea = Directory.Exists(Path.Combine(SC.StageDirectory.FullName, CurrentDir));
+										if (DirExistsInStagingArea)
+										{
+											SkipFiles = false;
+										}
+										else
+										{
+											// delete directory from device
+											SkipFiles = true;
+											DirsToDeleteFromDevice.Add(CurrentDir);
+										}
+									}
+									else
+									{
+										if (SkipFiles)
+										{
+											continue;
+										}
+
+										string FileName = GetFileNameFromListing(Line);
+										if (FileName != null)
+										{
+											bool FileExistsInStagingArea = File.Exists(Path.Combine(SC.StageDirectory.FullName, CurrentDir, FileName));
+											if (FileExistsInStagingArea)
+											{
+												// keep or overwrite
+											}
+											else if (FileName == "APKFileStamp.txt")
+											{
+												// keep it
+											}
+											else
+											{
+												// delete file from device
+												string FilePath = CurrentDir.Length == 0 ? FileName : (CurrentDir + "/" + FileName); // use / for Android target, no matter the development system
+												LogWarning("Deleting {0} from device; not found in staging area", FilePath);
+												FilesToDeleteFromDevice.Add(FilePath);
+											}
+										}
+										// We ignore subdirs here as each will have its own "RemoteDir/CurrentDir/SubDir:" entry.
+									}
+								}
+							}
+
+							// delete directories
+							foreach (var DirToDelete in DirsToDeleteFromDevice)
+							{
+								// if a whole tree is to be deleted, don't spend extra commands deleting its branches
+								int FinalSlash = DirToDelete.LastIndexOf('/');
+								string ParentDir = FinalSlash >= 0 ? DirToDelete.Substring(0, FinalSlash) : "";
+								bool ParentMarkedForDeletion = DirsToDeleteFromDevice.Contains(ParentDir);
+								if (!ParentMarkedForDeletion)
+								{
+									LogWarning("Deleting {0} and its contents from device; not found in staging area", DirToDelete);
+									client.DirDeleteRecurse(RemoteDir + "/" + DirToDelete);
+								}
+							}
+
+							// delete loose files
+							foreach (var FileToDelete in FilesToDeleteFromDevice)
+							{
+								client.FileDelete(RemoteDir + "/" + FileToDelete);
+							}
+						}
+					}
+				}
+				else
+				{
+					// make sure device is at a clean state
+					client.DirDeleteRecurse(RemoteDir);
+
+					// Copy UFS files..
+					string[] Files = Directory.GetFiles(SC.StageDirectory.FullName, "*", SearchOption.AllDirectories);
+					System.Array.Sort(Files);
+
+					// Find all the files we exclude from copying. And include
+					// the directories we need to individually copy.
+					HashSet<string> ExcludedFiles = new HashSet<string>();
+					SortedSet<string> IndividualCopyDirectories
+						= new SortedSet<string>((IComparer<string>)new LongestFirst());
+					foreach (string Filename in Files)
+					{
+						bool Exclude = false;
+						// Don't push the apk, we install it
+						Exclude |= Path.GetExtension(Filename).Equals(".apk", StringComparison.InvariantCultureIgnoreCase);
+						// For excluded files we add the parent dirs to our
+						// tracking of stuff to individually copy.
+						if (Exclude)
+						{
+							ExcludedFiles.Add(Filename);
+							// We include all directories up to the stage root in having
+							// to individually copy the files.
+							for (string FileDirectory = Path.GetDirectoryName(Filename);
+								!FileDirectory.Equals(SC.StageDirectory);
+								FileDirectory = Path.GetDirectoryName(FileDirectory))
+							{
+								if (!IndividualCopyDirectories.Contains(FileDirectory))
+								{
+									IndividualCopyDirectories.Add(FileDirectory);
+								}
+							}
+							if (!IndividualCopyDirectories.Contains(SC.StageDirectory.FullName))
+							{
+								IndividualCopyDirectories.Add(SC.StageDirectory.FullName);
+							}
+						}
+					}
+
+					// The directories are sorted above in "deepest" first. We can
+					// therefore start copying those individual dirs which will
+					// recreate the tree. As the subtrees will get copied at each
+					// possible individual level.
+					foreach (string DirectoryName in IndividualCopyDirectories)
+					{
+						string[] Entries
+							= Directory.GetFileSystemEntries(DirectoryName, "*", SearchOption.TopDirectoryOnly);
+						foreach (string Entry in Entries)
+						{
+							// We avoid excluded files and the individual copy dirs
+							// (the individual copy dirs will get handled as we iterate).
+							if (ExcludedFiles.Contains(Entry) || IndividualCopyDirectories.Contains(Entry))
+							{
+								continue;
+							}
+							else
+							{
+								EntriesToDeploy.Add(Entry);
+							}
+						}
+					}
+
+					if (EntriesToDeploy.Count == 0)
+					{
+						EntriesToDeploy.Add(SC.StageDirectory.FullName);
+					}
+				}
+
+				// delete the .obb file, since it will cause nothing we deploy to be used
+				client.FileDelete(DeviceObbName);
+				client.FileDelete(DevicePatchName);
+				client.FileDelete(DeviceOverflow1Name);
+				client.FileDelete(DeviceOverflow2Name);
+
+				// We now have a minimal set of file & dir entries we need
+				// to deploy. Files we deploy will get individually copied
+				// and dirs will get the tree copies by default (that's
+				// what ADB does, too).
+				Log.TraceInformation("Deploying files using AFS");
+				string SourceDir = SC.StageDirectory.FullName;
+				client.Deploy(EntriesToDeploy, SourceDir, RemoteDir, bUseCompression, bLogFiles, bReportStats, client2);
+			}
+			else if (SC.Archive)
+			{
+				// deploy the obb if there is one
+				string ObbPath = Path.Combine(SC.StageDirectory.FullName, GetFinalObbName(ApkName, SC));
+				if (File.Exists(ObbPath))
+				{
+					client.FileWrite(ObbPath, DeviceObbName);
+				}
+
+				// deploy the patch if there is one
+				string PatchPath = Path.Combine(SC.StageDirectory.FullName, GetFinalPatchName(ApkName, SC));
+				if (File.Exists(PatchPath))
+				{
+					client.FileWrite(PatchPath, DevicePatchName);
+				}
+
+				// deploy the overflow1 if there is one
+				string Overflow1Path = Path.Combine(SC.StageDirectory.FullName, GetFinalOverflowName(ApkName, SC, 1));
+				if (File.Exists(Overflow1Path))
+				{
+					client.FileWrite(Overflow1Path, DeviceOverflow1Name);
+				}
+
+				// deploy the overflow2 if there is one
+				string Overflow2Path = Path.Combine(SC.StageDirectory.FullName, GetFinalOverflowName(ApkName, SC, 2));
+				if (File.Exists(Overflow2Path))
+				{
+					client.FileWrite(Overflow2Path, DeviceOverflow2Name);
+				}
+			}
+			else
+			{
+				// cache some strings
+				string RemoteFilename = IntermediateCmdLineFile.FullName.Replace(SC.StageDirectory.FullName, RemoteDir).Replace("\\", "/");
+				client.FileWrite(IntermediateCmdLineFile.FullName, RemoteFilename);
+			}
+
+			// terminate server and close AFS connections
+			if (client != null)
+			{
+				client.TerminateServer();
+				client.CloseConnection();
+			}
+			if (client2 != null)
+			{
+				client2.TerminateServer();
+				client2.CloseConnection();
+			}
+		}
+	}
+
+	private void DeployADB(ProjectParams Params, DeploymentContext SC)
     {
 		var AppArchitectures = AndroidExports.CreateToolChain(Params.RawProjectPath).GetAllArchitectures();
 
@@ -2039,6 +2904,10 @@ public class AndroidPlatform : Platform
 										{
 											// keep or overwrite
 										}
+										else if (FileName == "APKFileStamp.txt")
+										{
+											// keep it
+										}
 										else
 										{
 											// delete file from device
@@ -2262,6 +3131,25 @@ public class AndroidPlatform : Platform
         }
     }
 
+	public override void Deploy(ProjectParams Params, DeploymentContext SC)
+	{
+		bool bAFSEnablePlugin;
+		string AFSToken;
+		bool bIsShipping;
+		bool bAFSIncludeInShipping;
+		bool bAFSAllowExternalStartInShipping;
+
+		// Pick the proper deploy method
+		if (UsingAndroidFileServer(SC, out bAFSEnablePlugin, out AFSToken, out bIsShipping, out bAFSIncludeInShipping, out bAFSAllowExternalStartInShipping))
+		{
+			DeployAndroidFileServer(Params, SC, AFSToken);
+		}
+		else
+		{
+			DeployADB(Params, SC);
+		}
+	}
+
 	/** Internal usage for GetPackageName */
 	private static string PackageLine = null;
 	private static Mutex PackageInfoMutex = new Mutex();
diff --git a/Engine/Source/Programs/UnrealAndroidFileTool/App.config b/Engine/Source/Programs/UnrealAndroidFileTool/App.config
new file mode 100644
--- /dev/null
+++ b/Engine/Source/Programs/UnrealAndroidFileTool/App.config
@@ -0,0 +1,11 @@
+<?xml version="1.0"?>
+<configuration>
+  <startup>
+    <supportedRuntime version="v4.0" sku=".NETFramework,Version=v4.6.2"/>
+  </startup>
+  <runtime>
+    <Thread_UseAllCpuGroups enabled="true"/>
+    <GCCpuGroup enabled="true"/>
+    <gcServer enabled="true"/>
+  </runtime>
+</configuration>
diff --git a/Engine/Source/Programs/UnrealAndroidFileTool/Program.cs b/Engine/Source/Programs/UnrealAndroidFileTool/Program.cs
new file mode 100644
--- /dev/null
+++ b/Engine/Source/Programs/UnrealAndroidFileTool/Program.cs
@@ -0,0 +1,1283 @@
+// Copyright Epic Games, Inc. All Rights Reserved.
+
+using System;
+using System.Threading;
+using System.Threading.Tasks;
+using System.Diagnostics;
+using System.Reflection;
+using System.IO;
+using System.Collections.Generic;
+using System.Text;
+using System.Linq;
+using AutomationTool;
+using EpicGames.Core;
+
+namespace UnrealAndroidFileTool
+{
+	class Program
+	{
+		private static string AppVersion = "1.0.0";
+
+		static void ShowHelp(string command)
+		{
+			if (command == "")
+			{
+				Log.TraceInformation("UnrealAndroidFileTool version {0}\n" +
+					"Copyright Epic Games, Inc. All Rights Reserved.\n\n" +
+					"UnrealAndroidFileTool [-s Device] [-ip ipAddress] [-t port] [-p PackageName] [-k token] command\n\n" +
+					"Commands:\n\n" +
+					"devices                Show connected devices (@ prefix for authorized).\n" +
+					"packages               Show list of packages with receiver enabled.\n" +
+					"stop-all [-w]          Send stop request to all packages with receivers.\n" +
+					"help [command]         Get help about a command.\n" +
+					"shell                  Start interactive mode.\n" +
+					"quit                   Exit interactive mode.\n" +
+					"exit                   Exit interactive mode.\n" +
+					"terminate              Stop remote server.\n" +
+					"query [key]            Show all variables or a specific key.\n" +
+					"getprop [key]          Get property for key.\n" + 
+					"cd path                Set base directory (value of ^^).\n" +
+					"pwd                    Show current base directory.\n" +
+					"direxists path         Returns true if directory path exists.\n" +
+//					"dirlist [RSA:]path     Show contents of directory path.\n" +
+					"ls [-lsRf] path        Show contents of directory path.\n" +
+//					"dirlistflat path       Show recursive contents of directories grouped.\n" +
+//					"dircreate path         Create directory path (recursive).\n" +
+					"mkdir path             Create directory path (recursive).\n" +
+//					"dirdelete path         Delete directory path (recursive).\n" +
+					"rmdir path             Delete directory path (recursive).\n" +
+					"fileexists file        Returns true if file exists.\n" +
+//					"filedelete file        Delete file.\n" +
+					"rm file                Delete file.\n" +
+//					"filecopy src dst       Copies file from src to dst on device.\n" +
+					"cp src dst             Copies file from src to dst on device.\n" +
+//					"filemove src dst       Moves file from src to dst on device.\n" +
+					"mv src dst             Moves file from src to dst on device.\n" +
+//					"fileread src dst       Pulls src file from device to local dst.\n" +
+					"pull src dst           Pulls src file from device to local dst.\n" +
+//					"filewrite src dst      Pushes local src file to device dst.\n" +
+					"push [-c] src dst      Pushes local src file or directories to device dst.\n" +
+					"command [data]         Writes data to commandline file, or shows contents.\n" +
+					"addcommand [data]      Adds data to commandline file.\n" +
+					"delcommand [data]      Removes data from commandline file.\n" +
+					"cat file               Writes file contents from device to output.\n" +
+					"deploy [-c] file       Reads a text file with deployment source/dest pairs.\n" +
+					"\n" +
+					"Note: keys may be used at start of a path as shortcut.\n" +
+					"   Ex: fileread ^logfile game.log   - pulls logfile.\n" +
+					"       rm ^project/Manifest.txt     - deletes Manifest.txt from project dir.\n\n"
+					, AppVersion);
+				return;
+			}
+			if (command == "ls")
+			{
+				Log.TraceInformation("-l\tlist permissions");
+				Log.TraceInformation("-s\tlist size");
+				Log.TraceInformation("-R\tlist recursive directory tree");
+				Log.TraceInformation("-f\tlist flat\n");
+				return;
+			}
+			if (command == "stop-all")
+			{
+				Log.TraceInformation("-w\twait for all listen binds to terminate");
+				return;
+			}
+			if (command == "push")
+			{
+				Log.TraceInformation("-c\tcompress files\n");
+				return;
+			}
+			if (command == "deploy")
+			{
+				Log.TraceInformation("-c\tcompress files\n");
+				return;
+			}
+		}
+
+		static string FixPath(AndroidFileClient client, string path)
+		{
+			if (path.StartsWith("^^"))
+			{
+				string cmd = client.Query("^^");
+				path = client.Query("^^") + path.Substring(2);
+			}
+			else if (!(path.StartsWith("^") || path.StartsWith("/")))
+			{
+				string cwd = client.Query("^^");
+				if (cwd != "")
+				{
+					path = cwd + "/" + path;
+				}
+			}
+
+			// remove .. sets by trimming out prior directory
+			int dotdotIndex = path.IndexOf("..");
+			while (dotdotIndex > 0)
+			{
+				int slashIndex = path.Substring(0, dotdotIndex - 1).LastIndexOf("/");
+				if (slashIndex < 0)
+				{
+					break;
+				}
+				path = path.Substring(0, slashIndex) + path.Substring(dotdotIndex + 2);
+				dotdotIndex = path.IndexOf("..");
+			}
+
+			return path;
+		}
+
+		static int ProcessCommand(AndroidFileClient client, bool bInShell, string[] args, int ArgIndex)
+		{
+			if (ArgIndex < args.Length && (args[ArgIndex] == "help" || args[ArgIndex] == "?"))
+			{
+				ArgIndex++;
+				if (ArgIndex < args.Length)
+				{
+					ShowHelp(args[ArgIndex++]);
+				}
+				else
+				{
+					ShowHelp("");
+				}
+				return ArgIndex;
+			}
+			if (ArgIndex < args.Length && args[ArgIndex] == "shell")
+			{
+				ArgIndex++;
+				while (true)
+				{
+					Console.Write("shell> ");
+					string input = Console.ReadLine();
+					if (input == "exit" || input == "quit")
+					{
+						return ArgIndex;
+					}
+
+					if (input != "")
+					{
+						string[] split = input.Split(new char[] { ' ', '\t' });
+						if (split.Length > 0)
+						{
+							if (ProcessCommand(client, true, split, 0) == -1)
+							{
+								return -1;
+							}
+						}
+					}
+				}
+			}
+			if (ArgIndex < args.Length && args[ArgIndex] == "packages")
+			{
+				ArgIndex++;
+				List<string> Receivers = AndroidFileClient.GetInstalledReceivers(client.GetDevice());
+				if (Receivers.Count > 0)
+				{
+					foreach (string Line in Receivers)
+					{
+						Log.TraceInformation("{0}", Line);
+					}
+					return ArgIndex;
+				}
+				Log.TraceInformation("No packages found with receiver.");
+				return ArgIndex;
+			}
+			if (ArgIndex < args.Length && args[ArgIndex] == "stop-all")
+			{
+				ArgIndex++;
+				bool bWait = false;
+
+				if (ArgIndex < args.Length && args[ArgIndex].StartsWith("-w"))
+				{
+					bWait = true;
+				}
+
+				client.CloseConnection();
+				AndroidFileClient.StopAnyServers(client.GetDevice(), client.GetServerPort(), bWait);
+				System.Environment.Exit(0);
+			}
+			if (ArgIndex < args.Length && args[ArgIndex] == "stats")
+			{
+				ArgIndex++;
+				client.Stats_Report();
+				return ArgIndex;
+			}
+			if (ArgIndex < args.Length && args[ArgIndex] == "terminate")
+			{
+				ArgIndex++;
+				client.TerminateServer();
+				Log.TraceInformation("Terminated server");
+				return -1;
+			}
+			if (ArgIndex < args.Length && args[ArgIndex] == "pwd")
+			{
+				ArgIndex++;
+				string result = client.Query("^^");
+				Log.TraceInformation("{0}", result);
+				return ArgIndex;
+			}
+			if (ArgIndex < args.Length && args[ArgIndex] == "query")
+			{
+				ArgIndex++;
+				if (ArgIndex < args.Length)
+				{
+					string result = client.Query(args[ArgIndex++]);
+					Log.TraceInformation("{0}", result);
+				}
+				else
+				{
+					string result = client.Query("");
+					string[] lines = result.Split('\n');
+					foreach (string line in lines)
+					{
+						string[] parts = line.Split("\t");
+						if (parts.Length > 1)
+						{
+							Log.TraceInformation("{0,-20}{1}", parts[0], parts[1]);
+						}
+					}
+				}
+				return ArgIndex;
+			}
+			if (ArgIndex < args.Length && args[ArgIndex] == "getprop")
+			{
+				ArgIndex++;
+				if (ArgIndex < args.Length)
+				{
+					string result = client.GetProp(args[ArgIndex++]);
+					Log.TraceInformation("{0}", result);
+				}
+				else
+				{
+					string result = client.GetProp("");
+					Log.TraceInformation("{0}", result);
+				}
+				return ArgIndex;
+			}
+			if (ArgIndex < args.Length && args[ArgIndex] == "cd")
+			{
+				ArgIndex++;
+				if (ArgIndex < args.Length)
+				{
+					bool result = client.SetBaseDir(FixPath(client, args[ArgIndex]));
+					Log.TraceInformation("{0}", client.Query("^^"));
+					return ArgIndex++;
+				}
+				else
+				{
+					string result = client.Query("^^");
+					Log.TraceInformation("{0}", result);
+				}
+				return ArgIndex;
+			}
+			if (ArgIndex < args.Length && args[ArgIndex] == "direxists")
+			{
+				ArgIndex++;
+				if (ArgIndex < args.Length)
+				{
+					bool result = client.DirExists(args[ArgIndex]);
+					Log.TraceInformation("{0}", (result ? "true" : "false"));
+					if (!bInShell && !result)
+					{
+						client.CloseConnection();
+						System.Environment.Exit(1);
+					}
+					return ArgIndex++;
+				}
+				else
+				{
+					Log.TraceError("missing path");
+					if (!bInShell)
+					{
+						client.CloseConnection();
+						System.Environment.Exit(1);
+					}
+				}
+				return ArgIndex;
+			}
+			if (ArgIndex < args.Length && args[ArgIndex] == "ls")
+			{
+				ArgIndex++;
+				string path = "^^";
+
+				// process flags
+				bool bOptions = false;
+				bool bAttributes = false;
+				bool bRecursive = false;
+				bool bSize = false;
+				bool bFlat = false;
+				while (ArgIndex < args.Length && args[ArgIndex].StartsWith("-"))
+				{
+					if (args[ArgIndex].Contains("l"))
+					{
+						bAttributes = true;
+						bOptions = true;
+					}
+					if (args[ArgIndex].Contains("R"))
+					{
+						bRecursive = true;
+						bOptions = true;
+					}
+					if (args[ArgIndex].Contains("s"))
+					{
+						bSize = true;
+						bOptions = true;
+					}
+					if (args[ArgIndex].Contains("f"))
+					{
+						bFlat = true;
+					}
+					ArgIndex++;
+				}
+				if (ArgIndex < args.Length)
+				{
+					path = args[ArgIndex];
+					ArgIndex++;
+				}
+				path = FixPath(client, path);
+				if (path == "^^")
+				{
+					path = client.Query("^^");
+				}
+				string result;
+				if (bRecursive && bFlat)
+				{
+					result = client.DirListFlat(path);
+				}
+				else
+				{
+					result = client.DirList((bRecursive ? "R" : "") + (bSize ? "S" : "") + (bAttributes ? "A" : "") + (bOptions ? ":" : "") + path);
+				}
+				Log.TraceInformation("{0}", result);
+				return ArgIndex++;
+			}
+			if (ArgIndex < args.Length && args[ArgIndex] == "dirlist")
+			{
+				ArgIndex++;
+				if (ArgIndex < args.Length)
+				{
+					string result = client.DirList(FixPath(client, args[ArgIndex]));
+					Log.TraceInformation("{0}", result);
+					return ArgIndex++;
+				}
+				else
+				{
+					Log.TraceError("missing path");
+					if (!bInShell)
+					{
+						client.CloseConnection();
+						System.Environment.Exit(1);
+
+					}
+				}
+				return ArgIndex;
+			}
+			if (ArgIndex < args.Length && (args[ArgIndex] == "dirlistflat"))
+			{
+				ArgIndex++;
+				if (ArgIndex < args.Length)
+				{
+					string result = client.DirListFlat(FixPath(client, args[ArgIndex]));
+					Log.TraceInformation("{0}", result);
+					return ArgIndex++;
+				}
+				else
+				{
+					Log.TraceError("missing path");
+				}
+				return ArgIndex;
+			}
+			if (ArgIndex < args.Length && (args[ArgIndex] == "dircreate" || args[ArgIndex] == "mkdir"))
+			{
+				ArgIndex++;
+				if (ArgIndex < args.Length)
+				{
+					bool result = client.DirCreate(FixPath(client, args[ArgIndex]));
+					Log.TraceInformation("{0}", (result ? "true" : "false"));
+					if (!bInShell && !result)
+					{
+						client.CloseConnection();
+						System.Environment.Exit(1);
+
+					}
+					return ArgIndex++;
+				}
+				else
+				{
+					Log.TraceError("missing path");
+					if (!bInShell)
+					{
+						client.CloseConnection();
+						System.Environment.Exit(1);
+
+					}
+				}
+				return ArgIndex;
+			}
+			if (ArgIndex < args.Length && (args[ArgIndex] == "dirdelete" || args[ArgIndex] == "rmdir"))
+			{
+				ArgIndex++;
+				if (ArgIndex < args.Length && args[ArgIndex] == "-r")
+				{
+					ArgIndex++;
+					if (ArgIndex < args.Length)
+					{
+						bool result = client.DirDeleteRecurse(FixPath(client, args[ArgIndex]));
+						Log.TraceInformation("{0}", (result ? "true" : "false"));
+						if (!bInShell && !result)
+						{
+							client.CloseConnection();
+							System.Environment.Exit(1);
+
+						}
+					}
+					else
+					{
+						Log.TraceError("missing path");
+						if (!bInShell)
+						{
+							client.CloseConnection();
+							System.Environment.Exit(1);
+						}
+					}
+				}
+				else if (ArgIndex < args.Length)
+				{
+					bool result = client.DirDelete(FixPath(client, args[ArgIndex]));
+					Log.TraceInformation("{0}", (result ? "true" : "false"));
+					if (!bInShell && !result)
+					{
+						client.CloseConnection();
+						System.Environment.Exit(1);
+					}
+					return ArgIndex++;
+				}
+				else
+				{
+					Log.TraceError("missing path");
+					if (!bInShell)
+					{
+						client.CloseConnection();
+						System.Environment.Exit(1);
+
+					}
+				}
+				return ArgIndex;
+			}
+			if (ArgIndex < args.Length && args[ArgIndex] == "dirdeleterecurse")
+			{
+				ArgIndex++;
+				if (ArgIndex < args.Length)
+				{
+					bool result = client.DirDeleteRecurse(args[ArgIndex]);
+					Log.TraceInformation("{0}", (result ? "true" : "false"));
+					if (!bInShell && !result)
+					{
+						client.CloseConnection();
+						System.Environment.Exit(1);
+					}
+					return ArgIndex++;
+				}
+				else
+				{
+					Log.TraceError("missing path");
+					if (!bInShell)
+					{
+						client.CloseConnection();
+						System.Environment.Exit(1);
+					}
+				}
+				return ArgIndex;
+			}
+			if (ArgIndex < args.Length && args[ArgIndex] == "fileexists")
+			{
+				ArgIndex++;
+				if (ArgIndex < args.Length)
+				{
+					bool result = client.FileExists(FixPath(client, args[ArgIndex]));
+					Log.TraceInformation("{0}", (result ? "true" : "false"));
+					if (!bInShell && !result)
+					{
+						client.CloseConnection();
+						System.Environment.Exit(1);
+					}
+					return ArgIndex++;
+				}
+				else
+				{
+					Log.TraceError("missing path");
+					if (!bInShell)
+					{
+						client.CloseConnection();
+						System.Environment.Exit(1);
+					}
+				}
+				return ArgIndex;
+			}
+			if (ArgIndex < args.Length && (args[ArgIndex] == "filedelete" || args[ArgIndex] == "rm"))
+			{
+				ArgIndex++;
+				if (ArgIndex < args.Length)
+				{
+					bool result = client.FileDelete(FixPath(client, args[ArgIndex]));
+					Log.TraceInformation("{0}", (result ? "true" : "false"));
+					if (!bInShell && !result)
+					{
+						client.CloseConnection();
+						System.Environment.Exit(1);
+					}
+					return ArgIndex++;
+				}
+				else
+				{
+					Log.TraceError("missing path");
+					if (!bInShell)
+					{
+						client.CloseConnection();
+						System.Environment.Exit(1);
+					}
+				}
+				return ArgIndex;
+			}
+			if (ArgIndex < args.Length && (args[ArgIndex] == "filecopy" || args[ArgIndex] == "cp"))
+			{
+				ArgIndex++;
+				if (ArgIndex + 1 < args.Length)
+				{
+					bool result = client.FileCopy(FixPath(client, args[ArgIndex]), FixPath(client, args[ArgIndex+1]));
+					Log.TraceInformation("{0}", (result ? "true" : "false"));
+					if (!bInShell && !result)
+					{
+						client.CloseConnection();
+						System.Environment.Exit(1);
+					}
+					return ArgIndex +=2;
+				}
+				else
+				{
+					Log.TraceError("missing paths");
+					if (!bInShell)
+					{
+						client.CloseConnection();
+						System.Environment.Exit(1);
+					}
+				}
+				return ArgIndex;
+			}
+			if (ArgIndex < args.Length && (args[ArgIndex] == "filemove" || args[ArgIndex] == "mv"))
+			{
+				ArgIndex++;
+				if (ArgIndex + 1 < args.Length)
+				{
+					bool result = client.FileMove(FixPath(client, args[ArgIndex]), FixPath(client, args[ArgIndex + 1]));
+					Log.TraceInformation("{0}", (result ? "true" : "false"));
+					if (!bInShell && !result)
+					{
+						client.CloseConnection();
+						System.Environment.Exit(1);
+
+					}
+					return ArgIndex += 2;
+				}
+				else
+				{
+					Log.TraceError("missing paths");
+					if (!bInShell)
+					{
+						client.CloseConnection();
+						System.Environment.Exit(1);
+
+					}
+				}
+				return ArgIndex;
+			}
+			if (ArgIndex < args.Length && (args[ArgIndex] == "fileread" || args[ArgIndex] == "pull"))
+			{
+				ArgIndex++;
+				if (ArgIndex + 1 < args.Length)
+				{
+					bool result = client.FileRead(args[ArgIndex], args[ArgIndex + 1]);
+					Log.TraceInformation("{0}", (result ? "true" : "false"));
+					if (!bInShell && !result)
+					{
+						client.CloseConnection();
+						System.Environment.Exit(1);
+					}
+					return ArgIndex += 2;
+				}
+				else
+				{
+					Log.TraceError("missing paths");
+					if (!bInShell)
+					{
+						client.CloseConnection();
+						System.Environment.Exit(1);
+
+					}
+				}
+				return ArgIndex;
+			}
+			if (ArgIndex < args.Length && args[ArgIndex] == "filewrite")
+			{
+				ArgIndex++;
+				if (ArgIndex + 1 < args.Length)
+				{
+					//Log.TraceInformation("Src: '{0}', Dest: '{1}'", args[ArgIndex], args[ArgIndex + 1]);
+					bool result = client.FileWrite(args[ArgIndex], args[ArgIndex + 1]);
+					Log.TraceInformation("{0}", (result ? "true" : "false"));
+					if (!bInShell && !result)
+					{
+						client.CloseConnection();
+						System.Environment.Exit(1);
+					}
+					return ArgIndex += 2;
+				}
+				else
+				{
+					Log.TraceError("missing paths");
+					if (!bInShell)
+					{
+						client.CloseConnection();
+						System.Environment.Exit(1);
+					}
+				}
+				return ArgIndex;
+			}
+			if (ArgIndex < args.Length && args[ArgIndex] == "push")
+			{
+				ArgIndex++;
+				if (ArgIndex + 1 < args.Length)
+				{
+					bool result = false;
+					bool bCompress = false;
+
+					if (args[ArgIndex] == "-c")
+					{
+						bCompress = true;
+						ArgIndex++;
+					}
+					if (ArgIndex + 1 < args.Length)
+					{
+						string source = args[ArgIndex];
+						string dest = args[ArgIndex + 1];
+						ArgIndex += 2;
+						try
+						{
+							FileAttributes attributes = File.GetAttributes(source);
+							if ((attributes & FileAttributes.Directory) == FileAttributes.Directory)
+							{
+								result = client.PushDirectories(source, dest, bCompress, 1);
+							}
+							else
+							{
+								result = client.PushFile(source, dest, bCompress, 1);
+							}
+						}
+						catch (Exception)
+						{
+							Log.TraceInformation("false");
+							if (!bInShell)
+							{
+								client.CloseConnection();
+								System.Environment.Exit(1);
+
+							}
+							return ArgIndex;
+						}
+
+						Log.TraceInformation("{0}", (result ? "true" : "false"));
+						if (!bInShell && !result)
+						{
+							client.CloseConnection();
+							System.Environment.Exit(1);
+						}
+						return ArgIndex;
+					}
+					else
+					{
+						Log.TraceError("missing paths");
+						if (!bInShell)
+						{
+							client.CloseConnection();
+							System.Environment.Exit(1);
+						}
+					}
+				}
+				else
+				{
+					Log.TraceError("missing paths");
+					if (!bInShell)
+					{
+						client.CloseConnection();
+						System.Environment.Exit(1);
+					}
+				}
+			}
+			if (ArgIndex < args.Length && args[ArgIndex] == "command")
+			{
+				ArgIndex++;
+				if (ArgIndex < args.Length)
+				{
+					string command = args[ArgIndex++];
+					while (ArgIndex < args.Length)
+					{
+						command = command + " " + args[ArgIndex++];
+					}
+					bool result = client.FileWriteString(command, "^commandfile");
+					Log.TraceInformation("{0}", (result ? "true" : "false"));
+					if (!bInShell && !result)
+					{
+						client.CloseConnection();
+						System.Environment.Exit(1);
+					}
+					return ArgIndex;
+				}
+				else
+				{
+					string TempFilename = Path.GetTempFileName();
+					bool result = client.FileRead("^commandfile", TempFilename);
+					if (result)
+					{
+						try
+						{
+							StreamReader myReader = File.OpenText(TempFilename);
+							Console.WriteLine(myReader.ReadToEnd());
+							myReader.Close();
+							File.Delete(TempFilename);
+						}
+						catch (Exception e)
+						{
+							Log.TraceError("Error reading temp file {0}, Exception: {1}", TempFilename, e.ToString());
+						}
+					}
+					else
+					{
+						Log.TraceInformation("No commandline file found.");
+					}
+				}
+				return ArgIndex;
+			}
+			if (ArgIndex < args.Length && args[ArgIndex] == "addcommand")
+			{
+				ArgIndex++;
+				if (ArgIndex < args.Length)
+				{
+					string command = "";
+					string TempFilename = Path.GetTempFileName();
+					bool result = client.FileRead("^commandfile", TempFilename);
+					if (result)
+					{
+						try
+						{
+							StreamReader myReader = File.OpenText(TempFilename);
+							string? line = myReader.ReadLine();
+							command = (line != null) ? line : "";
+							myReader.Close();
+							File.Delete(TempFilename);
+						}
+						catch (Exception)
+						{
+						}
+					}
+					while (ArgIndex < args.Length)
+					{
+						command = command + " " + args[ArgIndex++];
+					}
+					result = client.FileWriteString(command, "^commandfile");
+					if (!bInShell && !result)
+					{
+						Log.TraceError("false");
+						client.CloseConnection();
+						System.Environment.Exit(1);
+					}
+					Log.TraceInformation("New commandline: {0}", command);
+					return ArgIndex;
+				}
+				else
+				{
+					string TempFilename = Path.GetTempFileName();
+					bool result = client.FileRead("^commandfile", TempFilename);
+					if (result)
+					{
+						try
+						{
+							StreamReader myReader = File.OpenText(TempFilename);
+							Console.WriteLine(myReader.ReadToEnd());
+							myReader.Close();
+							File.Delete(TempFilename);
+						}
+						catch (Exception e)
+						{
+							Log.TraceError("Error reading temp file {0}, Exception: {1}", TempFilename, e.ToString());
+							if (!bInShell)
+							{
+								client.CloseConnection();
+								System.Environment.Exit(1);
+							}
+						}
+					}
+					else
+					{
+						Log.TraceInformation("No commandline file found.");
+						if (!bInShell)
+						{
+							client.CloseConnection();
+							System.Environment.Exit(1);
+
+						}
+					}
+				}
+				return ArgIndex;
+			}
+			if (ArgIndex < args.Length && args[ArgIndex] == "delcommand")
+			{
+				ArgIndex++;
+				if (ArgIndex < args.Length)
+				{
+					string command = "";
+					string TempFilename = Path.GetTempFileName();
+					bool result = client.FileRead("^commandfile", TempFilename);
+					if (result)
+					{
+						try
+						{
+							StreamReader myReader = File.OpenText(TempFilename);
+							string? line = myReader.ReadLine();
+							command = (line != null) ? line : "";
+							myReader.Close();
+							File.Delete(TempFilename);
+						}
+						catch (Exception e)
+						{
+							Log.TraceError("Error reading temp file {0}, Exception: {1}", TempFilename, e.ToString());
+							if (!bInShell)
+							{
+								client.CloseConnection();
+								System.Environment.Exit(1);
+							}
+						}
+					}
+					while (ArgIndex < args.Length)
+					{
+						command = command.Replace(args[ArgIndex++], "");
+					}
+					result = client.FileWriteString(command, "^commandfile");
+					if (!bInShell && !result)
+					{
+						Log.TraceError("false");
+						client.CloseConnection();
+						System.Environment.Exit(1);
+					}
+					Log.TraceInformation("New commandline: {0}", command);
+					return ArgIndex;
+				}
+				else
+				{
+					string TempFilename = Path.GetTempFileName();
+					bool result = client.FileRead("^commandfile", TempFilename);
+					if (result)
+					{
+						try
+						{
+							StreamReader myReader = File.OpenText(TempFilename);
+							Console.WriteLine(myReader.ReadToEnd());
+							myReader.Close();
+							File.Delete(TempFilename);
+						}
+						catch (Exception e)
+						{
+							Log.TraceError("Error reading temp file {0}, Exception: {1}", TempFilename, e.ToString());
+							if (!bInShell)
+							{
+								client.CloseConnection();
+								System.Environment.Exit(1);
+							}
+						}
+					}
+					else
+					{
+						Log.TraceInformation("No commandline file found.");
+						if (!bInShell)
+						{
+							client.CloseConnection();
+							System.Environment.Exit(1);
+						}
+					}
+				}
+				return ArgIndex;
+			}
+			if (ArgIndex < args.Length && args[ArgIndex] == "cat")
+			{
+				ArgIndex++;
+				if (ArgIndex < args.Length)
+				{
+					string TempFilename = Path.GetTempFileName();
+					bool result = client.FileRead(FixPath(client, args[ArgIndex]), TempFilename);
+					if (result)
+					{
+						try
+						{
+							StreamReader myReader = File.OpenText(TempFilename);
+							Console.WriteLine(myReader.ReadToEnd());
+							myReader.Close();
+							File.Delete(TempFilename);
+						}
+						catch (Exception e)
+						{
+							Log.TraceError("Error reading temp file {0}, Exception: {1}", TempFilename, e.ToString());
+						}
+					}
+					else
+					{
+						Log.TraceInformation("{0}", (result ? "true" : "false"));
+						if (!bInShell)
+						{
+							client.CloseConnection();
+							System.Environment.Exit(1);
+						}
+					}
+					return ArgIndex++;
+				}
+				else
+				{
+					Log.TraceError("missing path");
+					if (!bInShell)
+					{
+						client.CloseConnection();
+						System.Environment.Exit(1);
+					}
+				}
+				return ArgIndex;
+			}
+			if (ArgIndex < args.Length && args[ArgIndex] == "deploy")
+			{
+				ArgIndex++;
+				if (ArgIndex < args.Length)
+				{
+					bool bCompress = false;
+
+					if (args[ArgIndex] == "-c")
+					{
+						bCompress = true;
+						ArgIndex++;
+					}
+					if (ArgIndex < args.Length)
+					{
+						try
+						{
+							StreamReader myReader = File.OpenText(args[ArgIndex]);
+							ArgIndex++;
+
+							string? SourceDir = myReader.ReadLine();
+							string? DestDir = myReader.ReadLine();
+							HashSet<string> EntriesToDeploy = new HashSet<string>();
+							string? Entry;
+							while ((Entry = myReader.ReadLine()) != null)
+							{
+								if (Entry != null)
+								{
+									EntriesToDeploy.Add(Entry);
+								}
+							}
+							myReader.Close();
+							if (SourceDir != null && DestDir != null)
+							{
+								client.Deploy(EntriesToDeploy, SourceDir, DestDir, bCompress, false, true, null);
+							}
+						}
+						catch (Exception e)
+						{
+							Log.TraceError("Error reading temp file {0}, Exception: {1}", args[ArgIndex], e.ToString());
+							if (!bInShell)
+							{
+								client.CloseConnection();
+								System.Environment.Exit(1);
+							}
+						}
+					}
+					else
+					{
+						Log.TraceError("missing path");
+						if (!bInShell)
+						{
+							client.CloseConnection();
+							System.Environment.Exit(1);
+						}
+					}
+				}
+				else
+				{
+					Log.TraceError("missing path");
+					if (!bInShell)
+					{
+						client.CloseConnection();
+						System.Environment.Exit(1);
+					}
+				}
+				return ArgIndex;
+			}
+
+			Log.TraceWarning("Unknown command");
+			return ArgIndex;
+		}
+
+		static AndroidFileClient? CreateAndConnectClient(string Device, string PackageName, string Token, string IPAddress, int Port)
+		{
+			AndroidFileClient? client = new AndroidFileClient(Device, Port);
+			if (client != null)
+			{
+				if (!client.OpenConnection(IPAddress))
+				{
+					if (PackageName == "")
+					{
+						Log.TraceError("Need package name to start server on {0}!", Device);
+						return null;
+					}
+					else
+					{
+						Log.TraceInformation("Trying to start file server {0}", PackageName);
+						if (!client.StartServer(PackageName, Token, IPAddress))
+						{
+							Log.TraceError("Unable to connect to " + Device);
+							return null;
+						}
+					}
+				}
+
+				// verify connection to the correct server
+				string DevicePackageName = client.Query("^packagename");
+				if (DevicePackageName != PackageName)
+				{
+					if (PackageName == "")
+					{
+						Log.TraceError("Need package name to start server on {0}!", Device);
+						client.CloseConnection();
+						return null;
+					}
+
+					Log.TraceInformation("Connected to wrong server {0}, trying again", DevicePackageName);
+					client.TerminateServer();
+
+					Log.TraceInformation("Trying to start file server {0}", PackageName);
+					if (!client.StartServer(PackageName))
+					{
+						Log.TraceError("Failed to start server");
+						return null;
+					}
+				}
+				return client;
+			}
+			return null;
+		}
+
+		static void Main(string[] args)
+		{
+			string IPAddress = "";
+			string Device = "";
+			int Port = AndroidFileClient.GetDefaultPort();
+			string PackageName = "";
+			string Token = "";
+
+			// Ensure we can resolve any external assemblies as necessary.
+			Assembly? assembly = Assembly.GetEntryAssembly();
+			if (assembly != null)
+			{
+				string? PathToBinariesDotNET = Path.GetDirectoryName(assembly.GetOriginalLocation());
+				if (PathToBinariesDotNET != null)
+				{
+					AssemblyUtils.InstallAssemblyResolver(PathToBinariesDotNET);
+					AssemblyUtils.InstallRecursiveAssemblyResolver(PathToBinariesDotNET);
+				}
+			}
+
+			if (args.Length == 0 || (args.Length == 1 && args[0] == "help"))
+			{
+				ShowHelp("");
+				System.Environment.Exit(0);
+			}
+			if (args.Length > 1 && args[0] == "help")
+			{
+				ShowHelp(args[1]);
+				System.Environment.Exit(0);
+			}
+
+			int ArgIndex = 0;
+			if (ArgIndex < args.Length && args[ArgIndex] == "-s")
+			{
+				ArgIndex++;
+				if (ArgIndex < args.Length)
+				{
+					Device = "@" + args[ArgIndex];
+					ArgIndex++;
+				}
+				else
+				{
+					Log.TraceError("Missing device id");
+					System.Environment.Exit(1);
+				}
+			}
+
+			if (ArgIndex < args.Length && args[ArgIndex] == "-p")
+			{
+				ArgIndex++;
+				if (ArgIndex < args.Length)
+				{
+					PackageName = args[ArgIndex];
+					ArgIndex++;
+				}
+				else
+				{
+					Log.TraceError("Missing package name");
+					System.Environment.Exit(1);
+				}
+			}
+
+			if (ArgIndex < args.Length && args[ArgIndex] == "-k")
+			{
+				ArgIndex++;
+				if (ArgIndex < args.Length)
+				{
+					Token = args[ArgIndex];
+					ArgIndex++;
+				}
+				else
+				{
+					Log.TraceError("Missing security token key");
+					System.Environment.Exit(1);
+				}
+			}
+
+			if (ArgIndex < args.Length && args[ArgIndex] == "-t")
+			{
+				ArgIndex++;
+				if (ArgIndex < args.Length)
+				{
+					try
+					{
+						Port = int.Parse(args[ArgIndex]);
+						ArgIndex++;
+					}
+					catch (Exception)
+					{
+						Log.TraceError("Invalid port");
+						System.Environment.Exit(1);
+					}
+				}
+				else
+				{
+					Log.TraceError("Missing port");
+					System.Environment.Exit(1);
+				}
+			}
+
+			if (ArgIndex < args.Length && args[ArgIndex] == "-ip")
+			{
+				ArgIndex++;
+				if (ArgIndex < args.Length)
+				{
+					IPAddress = args[ArgIndex];
+					ArgIndex++;
+				}
+				else
+				{
+					Log.TraceError("Missing IP address");
+					System.Environment.Exit(1);
+				}
+			}
+
+			if (ArgIndex < args.Length && args[ArgIndex] == "devices")
+			{
+				List<string> devices = new List<string>();
+				AndroidFileClient.GetConnectedDevices(out devices);
+				foreach (string DeviceName in devices)
+				{
+					Log.TraceInformation("{0}", DeviceName);
+				}
+				return;
+			}
+
+			if (Device == "")
+			{
+				// look for connected devices
+				List<string> devices = new List<string>();
+				AndroidFileClient.GetConnectedDevices(out devices);
+				if (devices.Count == 0)
+				{
+					Log.TraceError("No connected devices");
+				}
+				else if (devices.Count == 1)
+				{
+					Device = devices[0];
+				}
+				else
+				{
+					Log.TraceError("Multiple devices connected, select one -s!");
+					foreach (string DeviceName in devices)
+					{
+						Log.TraceInformation("{0}", DeviceName);
+					}
+				}
+			}
+
+			if (Device == "")
+			{
+				Log.TraceError("No devices attached!");
+				System.Environment.Exit(0);
+			}
+
+			// check these here since may not have a running server to connect
+			if (ArgIndex < args.Length && args[ArgIndex] == "packages")
+			{
+				ArgIndex++;
+				List<string> Receivers = AndroidFileClient.GetInstalledReceivers(Device.Substring(1));
+				if (Receivers.Count > 0)
+				{
+					foreach (string Line in Receivers)
+					{
+						Log.TraceInformation("{0}", Line);
+					}
+					System.Environment.Exit(0);
+				}
+				Log.TraceInformation("No packages found with receiver.");
+				System.Environment.Exit(0);
+			}
+			if (ArgIndex < args.Length && args[ArgIndex] == "stop-all")
+			{
+				ArgIndex++;
+				bool bWait = false;
+
+				if (ArgIndex < args.Length && args[ArgIndex].StartsWith("-w"))
+				{
+					bWait = true;
+				}
+
+				AndroidFileClient.StopAnyServers(Device.Substring(1), Port, bWait);
+				System.Environment.Exit(0);
+			}
+
+			AndroidFileClient? client = CreateAndConnectClient(Device.Substring(1), PackageName, Token, IPAddress, Port);
+			if (client == null)
+			{
+				System.Environment.Exit(1);
+			}
+			Log.TraceInformation("Connected!");
+
+			if (ArgIndex < args.Length)
+			{
+				ArgIndex = ProcessCommand(client, false, args, ArgIndex);
+				if (client != null)
+				{
+					client.CloseConnection();
+				}
+				System.Environment.Exit(0);
+			}
+
+			ShowHelp("");
+
+			if (client != null)
+			{
+				client.CloseConnection();
+			}
+			System.Environment.Exit(0);
+		}
+	}
+}
diff --git a/Engine/Source/Programs/UnrealAndroidFileTool/Properties/AssemblyInfo.cs b/Engine/Source/Programs/UnrealAndroidFileTool/Properties/AssemblyInfo.cs
new file mode 100644
--- /dev/null
+++ b/Engine/Source/Programs/UnrealAndroidFileTool/Properties/AssemblyInfo.cs
@@ -0,0 +1,11 @@
+﻿// Copyright Epic Games, Inc. All Rights Reserved.
+
+using System.Reflection;
+using System.Runtime.CompilerServices;
+
+// General Information about an assembly is controlled through the following 
+// set of attributes. Change these attribute values to modify the information
+// associated with an assembly.
+[assembly: AssemblyTitle("UnrealAndroidFileTool")]
+[assembly: AssemblyDescription("A tool to work with Android File Server")]
+[assembly: AssemblyConfiguration("")]
diff --git a/Engine/Source/Programs/UnrealAndroidFileTool/UnrealAndroidFileTool.csproj b/Engine/Source/Programs/UnrealAndroidFileTool/UnrealAndroidFileTool.csproj
new file mode 100644
--- /dev/null
+++ b/Engine/Source/Programs/UnrealAndroidFileTool/UnrealAndroidFileTool.csproj
@@ -0,0 +1,81 @@
+<Project Sdk="Microsoft.NET.Sdk">
+  <PropertyGroup>
+    <TargetFramework>netcoreapp3.1</TargetFramework>
+    <Configuration Condition=" '$(Configuration)' == '' ">Development</Configuration>
+    <OutputType>Exe</OutputType>
+    <AppendTargetFrameworkToOutputPath>false</AppendTargetFrameworkToOutputPath>
+    <GenerateAssemblyInfo>false</GenerateAssemblyInfo>
+    <GenerateTargetFrameworkAttribute>false</GenerateTargetFrameworkAttribute>
+    <Configurations>Debug;Release;Development</Configurations>
+    <ApplicationManifest>app.manifest</ApplicationManifest>
+    <AssemblyName>UnrealAndroidFileTool</AssemblyName>
+    <RootNamespace>UnrealAndroidFileTool</RootNamespace>
+    <SatelliteResourceLanguages>en</SatelliteResourceLanguages> <!-- remove non english resource languages -->
+    <AppendRuntimeIdentifierToOutputPath>false</AppendRuntimeIdentifierToOutputPath>
+    <IsWindows Condition="'$([System.Runtime.InteropServices.RuntimeInformation]::IsOSPlatform($([System.Runtime.InteropServices.OSPlatform]::Windows)))' == 'true'">true</IsWindows> 
+    <IsOSX Condition="'$([System.Runtime.InteropServices.RuntimeInformation]::IsOSPlatform($([System.Runtime.InteropServices.OSPlatform]::OSX)))' == 'true'">true</IsOSX> 
+    <IsLinux Condition="'$([System.Runtime.InteropServices.RuntimeInformation]::IsOSPlatform($([System.Runtime.InteropServices.OSPlatform]::Linux)))' == 'true'">true</IsLinux> 
+    <DebugType>pdbonly</DebugType> 
+    <Nullable>enable</Nullable>
+  </PropertyGroup>
+  <Choose>
+    <When Condition="'$(IsLinux)'=='true'">
+      <PropertyGroup>
+        <RuntimeIdentifier>linux-x64</RuntimeIdentifier>
+      </PropertyGroup>
+    </When>
+    <When Condition="'$(IsOSX)'=='true'">
+      <PropertyGroup>
+        <!-- Set runtie identifier for osx so that it produces a app host executable (consistency with the other platforms) -->
+        <RuntimeIdentifier>osx-x64</RuntimeIdentifier>
+      </PropertyGroup>
+    </When>
+  </Choose>
+  <PropertyGroup Condition=" '$(Configuration)|$(Platform)' == 'Development|AnyCPU' ">
+    <Optimize>false</Optimize>
+    <OutputPath>..\..\..\Binaries\DotNET\Android\UnrealAndroidFileTool</OutputPath>
+    <DocumentationFile>..\..\..\Binaries\DotNET\Android\UnrealAndroidFileTool\UnrealAndroidFileTool.xml</DocumentationFile>
+  </PropertyGroup>
+  <PropertyGroup Condition=" '$(Configuration)|$(Platform)' == 'Release|AnyCPU' ">
+    <Optimize>false</Optimize>
+    <DebugType>pdbonly</DebugType>
+    <PlatformTarget>AnyCPU</PlatformTarget>
+    <ErrorReport>prompt</ErrorReport>
+    <AllowUnsafeBlocks>false</AllowUnsafeBlocks>
+    <Prefer32Bit>false</Prefer32Bit>
+	<OutputPath>..\..\..\Binaries\DotNET\Android\UnrealAndroidFileTool</OutputPath>
+    <DocumentationFile>..\..\..\Binaries\DotNET\Android\UnrealAndroidFileTool\UnrealAndroidFileTool.xml</DocumentationFile>
+  </PropertyGroup>
+  <PropertyGroup Condition=" '$(Configuration)|$(Platform)' == 'Debug|AnyCPU' ">
+    <OutputPath>..\..\..\Binaries\DotNET\Android\UnrealAndroidFileTool</OutputPath>
+    <DocumentationFile>..\..\..\Binaries\DotNET\Android\UnrealAndroidFileTool\UnrealAndroidFileTool.xml</DocumentationFile>
+    <DebugSymbols>true</DebugSymbols>
+    <DebugType>full</DebugType>
+  </PropertyGroup>
+
+  <ItemGroup>
+    <Compile Include="..\Shared\MetaData.cs">
+      <Link>Properties\MetaData.cs</Link>
+    </Compile>
+  </ItemGroup>
+
+  <!-- The following will glob all platform extensions' UnrealAndroidFileTool source
+  files and include them in the build. They will also appear in VisualStudio.
+  RecursiveDir does not work in VS Mac (and maybe Linux?), so only use it in Windows. 
+  See: https://developercommunity.visualstudio.com/content/problem/527478/recursivedir-causes-infinitely-nested-subdirectori.html -->
+  <ItemGroup>
+    <ProjectReference Include="..\Shared\EpicGames.Build\EpicGames.Build.csproj" />
+    <ProjectReference Include="..\Shared\EpicGames.Core\EpicGames.Core.csproj" />
+    <ProjectReference Include="..\AutomationTool\AutomationUtils\AutomationUtils.Automation.csproj" PrivateAssets="All"></ProjectReference>
+    <ProjectReference Include="..\AutomationTool\Android\Android.Automation.csproj" PrivateAssets="All"></ProjectReference>
+  </ItemGroup>
+
+  <ItemGroup>
+    <PackageReference Include="Microsoft.CodeAnalysis.CSharp" Version="3.7.0" />
+    <PackageReference Include="Microsoft.Win32.Registry" Version="4.7.0" />
+    <PackageReference Include="OpenTracing" Version="0.12.1" />
+    <PackageReference Include="System.Management" Version="4.7.0" />
+    <PackageReference Include="System.Security.Cryptography.Csp" Version="4.3.0" />
+    <PackageReference Include="System.ServiceProcess.ServiceController" Version="4.7.0" />
+  </ItemGroup>
+</Project>
diff --git a/Engine/Source/Programs/UnrealAndroidFileTool/app.manifest b/Engine/Source/Programs/UnrealAndroidFileTool/app.manifest
new file mode 100644
--- /dev/null
+++ b/Engine/Source/Programs/UnrealAndroidFileTool/app.manifest
@@ -0,0 +1,30 @@
+<?xml version="1.0" encoding="utf-8"?>
+<assembly manifestVersion="1.0" xmlns="urn:schemas-microsoft-com:asm.v1">
+	<trustInfo xmlns="urn:schemas-microsoft-com:asm.v2">
+		<security>
+			<requestedPrivileges xmlns="urn:schemas-microsoft-com:asm.v3">
+				<requestedExecutionLevel level="asInvoker" uiAccess="false" />
+			</requestedPrivileges>
+		</security>
+	</trustInfo>
+	<compatibility xmlns="urn:schemas-microsoft-com:compatibility.v1">
+		<application>
+			<!-- Windows 7 -->
+			<supportedOS Id="{35138b9a-5d96-4fbd-8e2d-a2440225f93a}" />
+
+			<!-- Windows 8 -->
+			<supportedOS Id="{4a2f28e3-53b9-4441-ba9c-d69d4a4a6e38}" />
+
+			<!-- Windows 8.1 -->
+			<supportedOS Id="{1f676c76-80e1-4239-95bb-83d0f6d0da78}" />
+
+			<!-- Windows 10 -->
+			<supportedOS Id="{8e0f7a12-bfb3-4fe8-b9a5-48fd50a15a9a}" />
+		</application>
+	</compatibility>
+	<application xmlns="urn:schemas-microsoft-com:asm.v3">
+		<windowsSettings>
+			<longPathAware xmlns="http://schemas.microsoft.com/SMI/2016/WindowsSettings">true</longPathAware>
+		</windowsSettings>
+	</application>
+</assembly>
