Gauntlet Android: Making sure that we read the right Metadata field names (ie not prefixed with the package name)

diff --git a/Engine/Source/Programs/AutomationTool/Gauntlet/Platform/Android/Gauntlet.AndroidBuildSource.cs b/Engine/Source/Programs/AutomationTool/Gauntlet/Platform/Android/Gauntlet.AndroidBuildSource.cs
--- a/Engine/Source/Programs/AutomationTool/Gauntlet/Platform/Android/Gauntlet.AndroidBuildSource.cs
+++ b/Engine/Source/Programs/AutomationTool/Gauntlet/Platform/Android/Gauntlet.AndroidBuildSource.cs
@@ -115,8 +115,8 @@ namespace Gauntlet
 				{
 					AndroidPlatform.GetPackageInfo(SourceApkPath, false);
 					// Establish remote directory usage
-					string ApkUsesExternalFilesDir = AndroidPlatform.GetMetadataValue("com.epicgames.unreal.GameActivity.bUseExternalFilesDir");
-					string ApkUsesPublicLogs = AndroidPlatform.GetMetadataValue("com.epicgames.unreal.GameActivity.bPublicLogFiles");
+					string ApkUsesExternalFilesDir = AndroidPlatform.GetMetadataValue("bUseExternalFilesDir");
+					string ApkUsesPublicLogs = AndroidPlatform.GetMetadataValue("bPublicLogFiles");
 					bUsesExternalFilesDir = ApkUsesExternalFilesDir != null ? ApkUsesExternalFilesDir.Contains("1") : false;
 					bUsesPublicLogs = ApkUsesPublicLogs != null ? ApkUsesPublicLogs.Contains("1") : false;
 				}
