Fix for error handling on failed parse of build record file

diff --git a/Engine/Source/Programs/AutomationTool/CompileScriptModules.cs b/Engine/Source/Programs/AutomationTool/CompileScriptModules.cs
--- a/Engine/Source/Programs/AutomationTool/CompileScriptModules.cs
+++ b/Engine/Source/Programs/AutomationTool/CompileScriptModules.cs
@@ -201,7 +201,7 @@ namespace AutomationToolDriver
 						Log.TraceWarning($"[{JsonFile}] Failed to load build record: {Ex.Message}");
 					}
 
-					if (BuildRecord.ProjectPath != null)
+					if (BuildRecord != null && BuildRecord.ProjectPath != null)
 					{
 						LoadedBuildRecords.Add(FileReference.FromString(Path.GetFullPath(BuildRecord.ProjectPath, JsonFile.Directory.FullName)), (BuildRecord, JsonFile));
 					}
