- Only remap cooked output directory for cookededitor when making DLC (it was meant for DLC mode only, but is causing long paths)

diff --git a/Engine/Source/Programs/AutomationTool/CookedEditor/CookedEditor.Automation.cs b/Engine/Source/Programs/AutomationTool/CookedEditor/CookedEditor.Automation.cs
--- a/Engine/Source/Programs/AutomationTool/CookedEditor/CookedEditor.Automation.cs
+++ b/Engine/Source/Programs/AutomationTool/CookedEditor/CookedEditor.Automation.cs
@@ -935,6 +935,7 @@ public class MakeCookedEditor : BuildCommand
 		string BasedOnReleaseVersion;
 		TargetType ReleaseType;
 		SetupDLCMode(ProjectFile, out DLCName, out BasedOnReleaseVersion, out ReleaseType);
+		bool bIsDLC = DLCName != null;
 
 		var Params = new ProjectParams
 		(
@@ -987,17 +988,15 @@ public class MakeCookedEditor : BuildCommand
 		// cause the issues
 		Params.AdditionalCookerOptions += " -AllowUnsafeBlueprintCalls";
 
-		// cook and stage into our project, instead of the Engine's plugins
-		DirectoryReference BaseOutputDirectory = DirectoryReference.Combine(ProjectFile.Directory, "Saved", "CookedEditor");
-		string TargetPlatformName = ConfigHierarchy.GetIniPlatformName(Platform) + TargetPlatformList[0].CookFlavor;
-		Params.CookOutputDir = DirectoryReference.Combine(BaseOutputDirectory, "Cooked", TargetPlatformName).FullName;
-		Params.StageDirectoryParam = DirectoryReference.Combine(BaseOutputDirectory, "Staged").FullName;
-
-		// Params.AdditionalCookerOptions += " -NoFilterAssetRegistry";
-
 		// set up cooking against a client, as DLC
-		if (BasedOnReleaseVersion != null)
+		if (bIsDLC)
 		{
+			// cook and stage into our project, instead of the Engine's plugins
+			DirectoryReference BaseOutputDirectory = DirectoryReference.Combine(ProjectFile.Directory, "Saved", "CookedEditor");
+			string TargetPlatformName = ConfigHierarchy.GetIniPlatformName(Platform) + TargetPlatformList[0].CookFlavor;
+			Params.CookOutputDir = DirectoryReference.Combine(BaseOutputDirectory, "Cooked", TargetPlatformName).FullName;
+			Params.StageDirectoryParam = DirectoryReference.Combine(BaseOutputDirectory, "Staged").FullName;
+
 			// make WindowsClient or LinuxGame, etc
 			string ReleaseTargetName = GetReleaseTargetName(Platform, ReleaseType);
 
