PR #8921: Added support to cross-compile editor plugins through the runUAT BuildPlugin. (Contributed by [REDACTED_PERSON])

diff --git a/Engine/Source/Programs/AutomationTool/Scripts/BuildPluginCommand.Automation.cs b/Engine/Source/Programs/AutomationTool/Scripts/BuildPluginCommand.Automation.cs
--- a/Engine/Source/Programs/AutomationTool/Scripts/BuildPluginCommand.Automation.cs
+++ b/Engine/Source/Programs/AutomationTool/Scripts/BuildPluginCommand.Automation.cs
@@ -14,6 +14,7 @@ using UnrealBuildBase;
 [Help("Builds a plugin, and packages it for distribution")]
 [Help("Plugin", "Specify the path to the descriptor file for the plugin that should be packaged")]
 [Help("NoHostPlatform", "Prevent compiling for the editor platform on the host")]
+[Help("HostPlatforms", "Specify a list of host platforms to build, separated by '+' characters (eg. -HostPlatforms=Win32+Win64). Default is the current host platforms")]
 [Help("TargetPlatforms", "Specify a list of target platforms to build, separated by '+' characters (eg. -TargetPlatforms=Win32+Win64). Default is all the Rocket target platforms.")]
 [Help("Package", "The path which the build artifacts should be packaged to, ready for distribution.")]
 [Help("StrictIncludes", "Disables precompiled headers and unity build in order to check all source files have self-contained headers.")]
@@ -111,7 +112,7 @@ public sealed class BuildPlugin : BuildCommand
 		}
 
 		// Compile the plugin for all the target platforms
-		List<UnrealTargetPlatform> HostPlatforms = ParseParam("NoHostPlatform")? new List<UnrealTargetPlatform>() : new List<UnrealTargetPlatform> { BuildHostPlatform.Current.Platform };
+		List<UnrealTargetPlatform> HostPlatforms = GetHostPlatforms(this);
 		List<UnrealTargetPlatform> TargetPlatforms = GetTargetPlatforms(this, BuildHostPlatform.Current.Platform);
 		FileReference[] BuildProducts = CompilePlugin(UBTExe, HostProjectFile, HostProjectPluginFile, Plugin, HostPlatforms, TargetPlatforms, AdditionalArgs.ToString());
 
@@ -364,5 +365,67 @@ public sealed class BuildPlugin : BuildCommand
 		}
 		return TargetPlatforms;
 	}
-}
 
+	static List<UnrealTargetPlatform> GetHostPlatforms(BuildCommand Command)
+	{
+		List<UnrealTargetPlatform> HostPlatforms = new List<UnrealTargetPlatform>();
+		if (!Command.ParseParam("NoHostPlatform"))
+		{
+			var CurrentPlatform = BuildHostPlatform.Current.Platform;
+
+			if (Command.ParseParam("HostPlatforms"))
+			{
+				string HostPlatformFilter = Command.ParseParamValue("HostPlatforms", null);
+				if(HostPlatformFilter != null)
+				{
+					// Only interested in building for Platforms that support code projects
+					HostPlatforms = PlatformExports.GetRegisteredPlatforms().Where(x => InstalledPlatformInfo.IsValidPlatform(x, EProjectType.Code)).ToList();
+
+					// only build Mac on Mac
+					if (CurrentPlatform != UnrealTargetPlatform.Mac && HostPlatforms.Contains(UnrealTargetPlatform.Mac))
+					{
+						HostPlatforms.Remove(UnrealTargetPlatform.Mac);
+					}
+					// only build Windows on Windows
+					if (CurrentPlatform != UnrealTargetPlatform.Win64 && HostPlatforms.Contains(UnrealTargetPlatform.Win64))
+					{
+						HostPlatforms.Remove(UnrealTargetPlatform.Win64);
+					}
+					// build Linux on Windows and Linux
+					if (CurrentPlatform != UnrealTargetPlatform.Win64 && CurrentPlatform != UnrealTargetPlatform.Linux)
+					{
+						if (HostPlatforms.Contains(UnrealTargetPlatform.Linux))
+							HostPlatforms.Remove(UnrealTargetPlatform.Linux);
+
+						if (HostPlatforms.Contains(UnrealTargetPlatform.LinuxArm64))
+							HostPlatforms.Remove(UnrealTargetPlatform.LinuxArm64);
+					}
+
+					List<UnrealTargetPlatform> NewHostPlatforms = new List<UnrealTargetPlatform>();
+					foreach (string HostPlatformName in HostPlatformFilter.Split(new char[]{ '+' }, StringSplitOptions.RemoveEmptyEntries))
+					{
+						UnrealTargetPlatform HostPlatform;
+						if (!UnrealTargetPlatform.TryParse(HostPlatformName, out HostPlatform))
+						{
+							throw new AutomationException("Unknown host platform '{0}' specified on command line", HostPlatformName);
+						}
+						if(HostPlatforms.Contains(HostPlatform))
+						{
+							NewHostPlatforms.Add(HostPlatform);
+						}
+					}
+					HostPlatforms = NewHostPlatforms;
+				}
+				else
+				{
+					HostPlatforms.Add(CurrentPlatform);
+				}
+			}
+			else
+			{
+				HostPlatforms.Add(CurrentPlatform);
+			}
+		}
+		return HostPlatforms;
+	}
+}
