Provide a way to add abitrary metadata to testreport.
ie: ASan build variation

#rnx

diff --git a/Engine/Source/Programs/AutomationTool/Gauntlet/Unreal/Base/Gauntlet.UnrealTestNode.cs b/Engine/Source/Programs/AutomationTool/Gauntlet/Unreal/Base/Gauntlet.UnrealTestNode.cs
--- a/Engine/Source/Programs/AutomationTool/Gauntlet/Unreal/Base/Gauntlet.UnrealTestNode.cs
+++ b/Engine/Source/Programs/AutomationTool/Gauntlet/Unreal/Base/Gauntlet.UnrealTestNode.cs
@@ -1287,6 +1287,14 @@ namespace Gauntlet
 			Report.SetMetadata("BuildType", string.Join("+", AllRoleTypes.Select(R => R.ToString()).Distinct().OrderBy(B => B)));
 			Report.SetMetadata("Configuration", string.Join("+", AllRoleContexts.Select(R => R.Configuration.ToString()).Distinct().OrderBy(C => C)));
 			Report.SetMetadata("Project", Context.BuildInfo.ProjectName);
+			// Additional metadata passed through the command line arguments
+			foreach( var Meta in Globals.Params.ParseValues("Metadata") )
+			{
+				var Entry = Meta.Split(":", 2);
+				if (Entry.Count() > 1) {
+					Report.SetMetadata(Entry[0], Entry[1]); 
+				}
+			}
 		}
 
 		/// <summary>
