Allow running buildgraph scripts outside the UE root directory.

diff --git a/Engine/Source/Programs/AutomationTool/BuildGraph/BuildGraph.cs b/Engine/Source/Programs/AutomationTool/BuildGraph/BuildGraph.cs
--- a/Engine/Source/Programs/AutomationTool/BuildGraph/BuildGraph.cs
+++ b/Engine/Source/Programs/AutomationTool/BuildGraph/BuildGraph.cs
@@ -107,10 +107,15 @@ namespace AutomationTool
 	/// </summary>
 	class ScriptReaderFileContext : IBgScriptReaderContext
 	{
+		/// <summary>
+		/// Directory to resolve relative paths to
+		/// </summary>
+		public DirectoryReference RootDirectory { get; set; } = Unreal.RootDirectory;
+
 		/// <inheritdoc/>
 		public object GetNativePath(string Path)
 		{
-			return FileReference.Combine(Unreal.RootDirectory, Path).FullName;
+			return FileReference.Combine(RootDirectory, Path).FullName;
 		}
 
 		/// <inheritdoc/>
@@ -118,7 +123,7 @@ namespace AutomationTool
 		{
 			try
 			{
-				return Task.FromResult(FileReference.Exists(FileReference.Combine(Unreal.RootDirectory, Path)) || DirectoryReference.Exists(DirectoryReference.Combine(Unreal.RootDirectory, Path)));
+				return Task.FromResult(FileReference.Exists(FileReference.Combine(RootDirectory, Path)) || DirectoryReference.Exists(DirectoryReference.Combine(RootDirectory, Path)));
 			}
 			catch
 			{
@@ -131,7 +136,7 @@ namespace AutomationTool
 		{
 			try
 			{
-				FileReference File = FileReference.Combine(Unreal.RootDirectory, Path);
+				FileReference File = FileReference.Combine(RootDirectory, Path);
 				if (FileReference.Exists(File))
 				{
 					return await FileReference.ReadAllBytesAsync(File);
@@ -149,7 +154,7 @@ namespace AutomationTool
 			FileFilter Filter = new FileFilter();
 			Filter.AddRule(Pattern, FileFilterType.Include);
 
-			List<string> Files = Filter.ApplyToDirectory(Unreal.RootDirectory, true).ConvertAll(x => x.MakeRelativeTo(Unreal.RootDirectory).Replace('\\', '/'));
+			List<string> Files = Filter.ApplyToDirectory(RootDirectory, true).ConvertAll(x => x.MakeRelativeTo(RootDirectory).Replace('\\', '/'));
 			Files.Sort(StringComparer.OrdinalIgnoreCase);
 			return Task.FromResult(Files.ToArray());
 		}
@@ -459,12 +464,11 @@ namespace AutomationTool
 
 				// Normalize the script filename
 				FileReference FullScriptFile = FileReference.Combine(Unreal.RootDirectory, ScriptFileName);
-				if (!FullScriptFile.IsUnderDirectory(Unreal.RootDirectory))
+				if (!FullScriptFile.IsUnderDirectory(Context.RootDirectory))
 				{
-					LogError("BuildGraph scripts must be under the UE root directory");
-					return ExitCode.Error_Unknown;
+					Context.RootDirectory = FullScriptFile.Directory;
 				}
-				ScriptFileName = FullScriptFile.MakeRelativeTo(Unreal.RootDirectory).Replace('\\', '/');
+				ScriptFileName = FullScriptFile.MakeRelativeTo(Context.RootDirectory).Replace('\\', '/');
 
 				// Read the script from disk
 				Graph = BgScriptReader.ReadAsync(Context, ScriptFileName, Arguments, DefaultProperties, Schema, Logger, SingleNodeName).Result;
