Additional logging when enumerating Gauntlet devices
#rnx

diff --git a/Engine/Source/Programs/AutomationTool/Gauntlet/Framework/Devices/Gauntlet.DevicePool.cs b/Engine/Source/Programs/AutomationTool/Gauntlet/Framework/Devices/Gauntlet.DevicePool.cs
--- a/Engine/Source/Programs/AutomationTool/Gauntlet/Framework/Devices/Gauntlet.DevicePool.cs
+++ b/Engine/Source/Programs/AutomationTool/Gauntlet/Framework/Devices/Gauntlet.DevicePool.cs
@@ -53,7 +53,7 @@ namespace Gauntlet
 		
 		public override string ToString()
 		{
-			return string.Format("{0} @ {1}. Platform={2}", Name, Address, Platform);
+			return string.Format("{0} @ {1}. Platform={2} Model={3}", Name, Address, Platform, string.IsNullOrEmpty(Model) ? "Unspecified" : Model);
 		}
 	}
 
@@ -1126,7 +1126,15 @@ namespace Gauntlet
 			lock (LockObject)
 			{
 				List<ITargetDevice> Selection = new List<ITargetDevice>();
+				
+				Log.Verbose($"Enumerating devices for constraint {Constraint}");
 
+				Log.Verbose($"   Available devices:");
+				AvailableDevices.ForEach(D => Log.Verbose($"      {D.Platform}:{D.Name}"));
+
+				Log.Verbose($"   Unprovisioned devices:");
+				UnprovisionedDevices.ForEach(D => Log.Verbose($"      {D}"));
+				
 				// randomize the order of all devices that are of this platform
 				var MatchingProvisionedDevices = AvailableDevices.Where(D => Constraint.Check(D)).ToList();
 				var MatchingUnprovisionedDevices = UnprovisionedDevices.Where(D => Constraint.Check(D)).ToList();
