Fix handling of AndroidFileServer startup checks (receiver may not be ready immediately after install)

#android

diff --git a/Engine/Source/Programs/AutomationTool/Android/AndroidFileClient.cs b/Engine/Source/Programs/AutomationTool/Android/AndroidFileClient.cs
--- a/Engine/Source/Programs/AutomationTool/Android/AndroidFileClient.cs
+++ b/Engine/Source/Programs/AutomationTool/Android/AndroidFileClient.cs
@@ -1387,13 +1387,6 @@ namespace AutomationTool
 
 			bool bIsUSB = (IPAddress == null || IPAddress == "" || IPAddress == "127.0.0.1");
 
-			// get a list of installed clients and verify package is available first
-			List<string> InstalledReceivers = GetInstalledReceivers(Device);
-			if (!InstalledReceivers.Contains(PackageName))
-			{
-				return false;
-			}
-
 			// deal with stopping all servers if requested
 			if (bStopAnyServers)
 			{
@@ -1404,7 +1397,32 @@ namespace AutomationTool
 				}
 			}
 
-			int tries = 5;
+			int tries = 10;
+
+			// get a list of installed clients and verify package is available
+			List<string> InstalledReceivers = GetInstalledReceivers(Device);
+			if (!InstalledReceivers.Contains(PackageName))
+			{
+				// if this was called right after installed the receiver may not be registered yet, wait and try again
+				bool bFound = false;
+				while (tries-- > 0)
+				{
+					Thread.Sleep(200);
+
+					InstalledReceivers = GetInstalledReceivers(Device);
+					if (InstalledReceivers.Contains(PackageName))
+					{
+						bFound = true;
+						break;
+					}
+				}
+				if (!bFound)
+				{
+					return false;
+				}
+			}
+
+			tries = 5;
 			bool bUSB;
 			bool bWifi;
 			string WifiAddress;
