Added server failure reasons and increased retries on StartServer

#android

diff --git a/Engine/Source/Programs/AutomationTool/Android/AndroidFileClient.cs b/Engine/Source/Programs/AutomationTool/Android/AndroidFileClient.cs
--- a/Engine/Source/Programs/AutomationTool/Android/AndroidFileClient.cs
+++ b/Engine/Source/Programs/AutomationTool/Android/AndroidFileClient.cs
@@ -1418,11 +1418,12 @@ namespace AutomationTool
 				}
 				if (!bFound)
 				{
+					Log.TraceInformation("Did not find package with receiver");
 					return false;
 				}
 			}
 
-			tries = 5;
+			tries = 15;
 			bool bUSB;
 			bool bWifi;
 			string WifiAddress;
@@ -1454,6 +1455,7 @@ namespace AutomationTool
 					GetListenStatus(Device, ServerPort, out bUSB, out bWifi, out WifiAddress);
 				}
 
+				Log.TraceInformation("Did not find a bind listener");
 				return false;
 			}
 
@@ -1472,6 +1474,7 @@ namespace AutomationTool
 				adb.Shell(Device, StartCommand);
 			}
 
+			Log.TraceInformation("Timed out on connection attempts");
 			return false;
 		}
 
