Clear stale saved Engine config files when running the DDC Commandlet without a project.
#rnx

diff --git a/Engine/Source/Programs/AutomationTool/Scripts/BuildDerivedDataCache.Automation.cs b/Engine/Source/Programs/AutomationTool/Scripts/BuildDerivedDataCache.Automation.cs
--- a/Engine/Source/Programs/AutomationTool/Scripts/BuildDerivedDataCache.Automation.cs
+++ b/Engine/Source/Programs/AutomationTool/Scripts/BuildDerivedDataCache.Automation.cs
@@ -97,6 +97,10 @@ public class BuildDerivedDataCache : BuildCommand
 
 		}
 
+		// Before running the Engine, delete any stale saved Config files from previous runs of the Engine. The list of enabled plugins can change, and the config file needs to be
+		// recomputed after they do, but there is currently nothing that makes that happen automatically
+		CommandUtils.DeleteDirectory(CommandUtils.CombinePaths(TempDir, "Engine", "Saved", "Config"));
+
 		// Generate DDC for the editor, and merge all the other PAK files in
 		CommandUtils.LogInformation("Generating DDC data for engine content on {0}", TargetPlatforms);
 		CommandUtils.DDCCommandlet(null, EditorExe, null, TargetPlatforms, String.Format("-fill -DDC={0} -MergePaks={1}{2}", BackendName, CommandUtils.MakePathSafeToUseWithCommandLine(String.Join("+", ProjectPakFiles)), bSkipEngine? " -projectonly" : ""));
