Fix P4WriteConfig error text reference to P4CreateConfig instead of P4WriteConfig.

diff --git a/Engine/Source/Programs/AutomationTool/AutomationUtils/P4WriteConfig.cs b/Engine/Source/Programs/AutomationTool/AutomationUtils/P4WriteConfig.cs
--- a/Engine/Source/Programs/AutomationTool/AutomationUtils/P4WriteConfig.cs
+++ b/Engine/Source/Programs/AutomationTool/AutomationUtils/P4WriteConfig.cs
@@ -50,7 +50,7 @@ namespace AutomationTool
 			catch (Exception Ex)
 			{
 				Log.TraceError(
-					"Unable to find matching Perforce info. If the below does not help try P4CreateConfig -p4port=<server:port> and -p4user=<username> to supply more info");
+					"Unable to find matching Perforce info. If the below does not help try P4WriteConfig -p4port=<server:port> and -p4user=<username> to supply more info");
 				Log.TraceError("{0}", Ex.Message);
 				return ExitCode.Error_Arguments;
 			}
