More adjustments to deal with slower Android devices in StartServer

#android

diff --git a/Engine/Source/Programs/AutomationTool/Android/AndroidFileClient.cs b/Engine/Source/Programs/AutomationTool/Android/AndroidFileClient.cs
--- a/Engine/Source/Programs/AutomationTool/Android/AndroidFileClient.cs
+++ b/Engine/Source/Programs/AutomationTool/Android/AndroidFileClient.cs
@@ -290,7 +290,7 @@ namespace AutomationTool
 				return "'" + source.Replace("'", "\'") + "'";
 			}
 
-			public string Shell(string Device, string Command)
+			public string Shell(string Device, string Command, bool bExpectResponse = true)
 			{
 				string Result = "FAIL";
 
@@ -322,18 +322,24 @@ namespace AutomationTool
 				}
 				Result = Result.Substring(4);
 
-				// there can be a delay before any more results (and may not be any)
-				if (ClientSocket.Available == 0)
+				if (bExpectResponse)
 				{
-					Thread.Sleep(50);
-				}
-				while (ClientSocket.Available > 0)
-				{
-					int bytesRecv = ClientSocket.Receive(buffer);
-					Result += Encoding.UTF8.GetString(buffer, 0, bytesRecv);
-					if (ClientSocket.Available == 0)
+					// there can be a delay before any more results (and may not be any)
+					// wait for up to 250 ms for a response in 25 ms intervals
+					int Waits = 10;
+					while (ClientSocket.Available == 0 && Waits-- > 0)
 					{
-						Thread.Sleep(50);
+						Thread.Sleep(25);
+					}
+					while (ClientSocket.Available > 0)
+					{
+						int bytesRecv = ClientSocket.Receive(buffer);
+						Result += Encoding.UTF8.GetString(buffer, 0, bytesRecv);
+						if (ClientSocket.Available == 0)
+						{
+							// wait to see if another packet follows
+							Thread.Sleep(50);
+						}
 					}
 				}
 
@@ -1341,7 +1347,7 @@ namespace AutomationTool
 			foreach (string Receiver in InstalledReceivers)
 			{
 				bDidSendStops = true;
-				adb.Shell(Device, "am broadcast -a com.epicgames.unreal.RemoteFileManager.intent.COMMAND -n " + Receiver + "/com.epicgames.unreal.RemoteFileManagerReceiver -e cmd 'stop'");
+				adb.Shell(Device, "am broadcast -a com.epicgames.unreal.RemoteFileManager.intent.COMMAND -n " + Receiver + "/com.epicgames.unreal.RemoteFileManagerReceiver -e cmd 'stop'", false);
 			}
 
 			// it can take up to 2 seconds for running servers to terminate so check if any binds to port are still active
@@ -1397,9 +1403,9 @@ namespace AutomationTool
 				}
 			}
 
-			int tries = 10;
 
-			// get a list of installed clients and verify package is available
+			// get a list of installed clients and verify package is available (up to 6 seconds)
+			int tries = 60;
 			List<string> InstalledReceivers = GetInstalledReceivers(Device);
 			if (!InstalledReceivers.Contains(PackageName))
 			{
@@ -1407,7 +1413,7 @@ namespace AutomationTool
 				bool bFound = false;
 				while (tries-- > 0)
 				{
-					Thread.Sleep(200);
+					Thread.Sleep(100);
 
 					InstalledReceivers = GetInstalledReceivers(Device);
 					if (InstalledReceivers.Contains(PackageName))
@@ -1423,14 +1429,15 @@ namespace AutomationTool
 				}
 			}
 
-			tries = 15;
+			// retries up to 4 seconds to start and see listener ready
+			tries = 16;
 			bool bUSB;
 			bool bWifi;
 			string WifiAddress;
 
 			// sent start request (won't do anything if already started)
 			string StartCommand = "am broadcast -a com.epicgames.unreal.RemoteFileManager.intent.COMMAND -n " + PackageName + "/com.epicgames.unreal.RemoteFileManagerReceiver -e cmd 'start' -e token '" + Token + "' -ei port " + ServerPort;
-			adb.Shell(Device, StartCommand);
+			adb.Shell(Device, StartCommand, false);
 
 			// see if we can check listen status
 			if (GetListenStatus(Device, ServerPort, out bUSB, out bWifi, out WifiAddress))
@@ -1450,7 +1457,7 @@ namespace AutomationTool
 					Thread.Sleep(250);
 
 					// sent start request again (won't do anything if already started)
-					adb.Shell(Device, StartCommand);
+					adb.Shell(Device, StartCommand, false);
 
 					GetListenStatus(Device, ServerPort, out bUSB, out bWifi, out WifiAddress);
 				}
@@ -1471,7 +1478,7 @@ namespace AutomationTool
 				}
 
 				// sent start request again (won't do anything if already started)
-				adb.Shell(Device, StartCommand);
+				adb.Shell(Device, StartCommand, false);
 			}
 
 			Log.TraceInformation("Timed out on connection attempts");
