Fallbacks for missing cmd on older Android and more timing adjustments for StartServer

#android

diff --git a/Engine/Source/Programs/AutomationTool/Android/AndroidFileClient.cs b/Engine/Source/Programs/AutomationTool/Android/AndroidFileClient.cs
--- a/Engine/Source/Programs/AutomationTool/Android/AndroidFileClient.cs
+++ b/Engine/Source/Programs/AutomationTool/Android/AndroidFileClient.cs
@@ -322,11 +322,11 @@ namespace AutomationTool
 				}
 				Result = Result.Substring(4);
 
-				if (bExpectResponse)
+//				if (bExpectResponse)
 				{
 					// there can be a delay before any more results (and may not be any)
-					// wait for up to 250 ms for a response in 25 ms intervals
-					int Waits = 10;
+					// wait for up to 1000 ms for a response in 25 ms intervals
+					int Waits = 40;
 					while (ClientSocket.Available == 0 && Waits-- > 0)
 					{
 						Thread.Sleep(25);
@@ -1201,6 +1201,38 @@ namespace AutomationTool
 			List<string> Result = new List<string>();
 
 			String InstalledResult = adb.Shell(Device, "cmd package query-receivers --components -a com.epicgames.unreal.RemoteFileManager.intent.COMMAND");
+			if (InstalledResult.Contains("not found") || InstalledResult.Contains("FAIL"))
+			{
+				InstalledResult = adb.Shell(Device, "dumpsys package");
+				bool bFoundReceiver = false;
+				foreach (string Line in InstalledResult.Split('\n'))
+				{
+					if (!bFoundReceiver)
+					{
+						if (Line.Contains("com.epicgames.unreal.RemoteFileManager.intent.COMMAND:"))
+						{
+							bFoundReceiver = true;
+						}
+					}
+					else
+					{
+						int SlashIndex = Line.IndexOf("/com.epicgames.unreal.RemoteFileManagerReceiver");
+						if (SlashIndex > 0)
+						{
+							int StartIndex = Line.LastIndexOf(' ');
+							if (StartIndex >= 0)
+							{
+								Result.Add(Line.Substring(StartIndex + 1, SlashIndex - StartIndex - 1));
+							}
+						}
+						else if (Line.IndexOf(":") > 0)
+						{
+							break;
+						}
+					}
+				}
+				return Result;
+			}
 			foreach (string Line in InstalledResult.Split('\n'))
 			{
 				int SlashIndex = Line.IndexOf('/');
@@ -1429,8 +1461,8 @@ namespace AutomationTool
 				}
 			}
 
-			// retries up to 4 seconds to start and see listener ready
-			tries = 16;
+			// retries up to 8 seconds to start and see listener ready
+			tries = 40;
 			bool bUSB;
 			bool bWifi;
 			string WifiAddress;
@@ -1454,7 +1486,7 @@ namespace AutomationTool
 					}
 
 					// wait before trying to again
-					Thread.Sleep(250);
+					Thread.Sleep(200);
 
 					// sent start request again (won't do anything if already started)
 					adb.Shell(Device, StartCommand, false);
