Added support for the -trace and -tracehost parameters for BuildCookRun

diff --git a/Engine/Source/Programs/AutomationTool/AutomationUtils/ProjectParams.cs b/Engine/Source/Programs/AutomationTool/AutomationUtils/ProjectParams.cs
--- a/Engine/Source/Programs/AutomationTool/AutomationUtils/ProjectParams.cs
+++ b/Engine/Source/Programs/AutomationTool/AutomationUtils/ProjectParams.cs
@@ -408,6 +408,9 @@ namespace AutomationTool
 			this.bTreatNonShippingBinariesAsDebugFiles = InParams.bTreatNonShippingBinariesAsDebugFiles;
 			this.bUseExtraFlavor = InParams.bUseExtraFlavor;
 			this.AdditionalPackageOptions = InParams.AdditionalPackageOptions;
+			this.Trace = InParams.Trace;
+			this.TraceHost = InParams.TraceHost;
+			this.TraceFile = InParams.TraceFile;
 		}
 
 		/// <summary>
@@ -561,6 +564,9 @@ namespace AutomationTool
 			string Certificate = null,
 		    string Team = null,
 		    bool AutomaticSigning = false,
+			string Trace = null,
+			string TraceHost = null,
+			string TraceFile = null,
 			ParamList<string> InMapsToRebuildLightMaps = null,
             ParamList<string> InMapsToRebuildHLOD = null,
             ParamList<string> TitleID = null
@@ -923,6 +929,39 @@ namespace AutomationTool
 			this.UbtArgs = ParseParamValueIfNotSpecified(Command, UbtArgs, "ubtargs", String.Empty);
 			this.AdditionalPackageOptions = ParseParamValueIfNotSpecified(Command, AdditionalPackageOptions, "AdditionalPackageOptions", String.Empty);
 
+			// -trace can be used with or without a value
+			if (Trace != null || GetParamValueIfNotSpecified(Command, null, false, "trace"))
+			{
+				this.Trace += "-trace";
+				string Value = ParseParamValueIfNotSpecified(Command, Trace, "trace", null);
+				if (!String.IsNullOrWhiteSpace(Value))
+				{
+					this.Trace += "=" + Value;
+				}
+			}
+
+			// -tracehost can be used with or without a value
+			if (TraceHost != null || GetParamValueIfNotSpecified(Command, null, false, "tracehost"))
+			{
+				this.TraceHost += "-tracehost";
+				string Value = ParseParamValueIfNotSpecified(Command, TraceHost, "tracehost", null);
+				if (!String.IsNullOrWhiteSpace(Value))
+				{
+					this.TraceHost += "=" + Value;
+				}
+			}
+
+			// -tracefile can be used with or without a value
+			if (TraceFile != null || GetParamValueIfNotSpecified(Command, null, false, "tracefile"))
+			{
+				this.TraceFile += "-tracefile";
+				string Value = ParseParamValueIfNotSpecified(Command, TraceFile, "tracefile", null);
+				if (!String.IsNullOrWhiteSpace(Value))
+				{
+					this.TraceFile += "=" + Value;
+				}
+			}
+
 			if (ClientConfigsToBuild == null)
 			{
 				if (Command != null)
@@ -2147,6 +2186,15 @@ namespace AutomationTool
         [Help("IgnoreLightMapErrors", "Whether Light Map errors should be treated as critical")]
 		public bool IgnoreLightMapErrors { get; set; }
 
+		[Help("trace", "The list of trace channels to enable")]
+		public string Trace { get; set; }
+
+		[Help("tracehost", "The host address of the trace recorder")]
+		public string TraceHost { get; set; }
+		
+		[Help("tracefile", "The file where the trace will be recorded")]
+		public string TraceFile { get; set; }
+
 		private List<SingleTargetProperties> DetectedTargets;
 		private Dictionary<UnrealTargetPlatform, ConfigHierarchy> LoadedEngineConfigs;
 		private Dictionary<UnrealTargetPlatform, ConfigHierarchy> LoadedGameConfigs;
@@ -3069,6 +3117,9 @@ namespace AutomationTool
 				CommandUtils.LogLog("bUseExtraFlavor={0}", bUseExtraFlavor);
                 CommandUtils.LogLog("StageDirectoryParam={0}", StageDirectoryParam);
 				CommandUtils.LogLog("AdditionalPackageOptions={0}", AdditionalPackageOptions);
+				CommandUtils.LogLog("Trace={0}", Trace);
+				CommandUtils.LogLog("TraceHost={0}", TraceHost);
+				CommandUtils.LogLog("TraceFile={0}", TraceFile);
 				CommandUtils.LogLog("Project Params **************");
 			}
 			bLogged = true;
diff --git a/Engine/Source/Programs/AutomationTool/Scripts/CookCommand.Automation.cs b/Engine/Source/Programs/AutomationTool/Scripts/CookCommand.Automation.cs
--- a/Engine/Source/Programs/AutomationTool/Scripts/CookCommand.Automation.cs
+++ b/Engine/Source/Programs/AutomationTool/Scripts/CookCommand.Automation.cs
@@ -173,6 +173,18 @@ namespace AutomationScripts
 			{
 				CommandletParams += " -basedonreleaseversionroot=" + Params.BasedOnReleaseVersionBasePath;
 			}
+			if (!String.IsNullOrEmpty(Params.Trace))
+			{
+				CommandletParams += " " + Params.Trace;
+			}
+			if (!String.IsNullOrEmpty(Params.TraceHost))
+			{
+				CommandletParams += " " + Params.TraceHost;
+			}
+			if (!String.IsNullOrEmpty(Params.TraceFile))
+			{
+				CommandletParams += " " + Params.TraceFile;
+			}
 
 			// if we are not going to pak but we specified compressed then compress in the cooker ;)
 			// otherwise compress the pak files
@@ -256,6 +268,21 @@ namespace AutomationScripts
 						COTFCommandLine += " -zenstore -messaging";
 					}
 
+					if (!String.IsNullOrEmpty(Params.Trace))
+					{
+						COTFCommandLine += " " + Params.Trace;
+					}
+
+					if (!String.IsNullOrEmpty(Params.TraceHost))
+					{
+						COTFCommandLine += " " + Params.TraceHost;
+					}
+
+					if (!String.IsNullOrEmpty(Params.TraceFile))
+					{
+						COTFCommandLine += " " + Params.TraceFile;
+					}
+
 					var ServerLogFile = CombinePaths(LogFolderOutsideOfSandbox, "Server.log");
 					Platform ClientPlatformInst = Params.ClientTargetPlatformInstances[0];
 					string TargetCook = ClientPlatformInst.GetCookPlatform(false, false); // cook on he fly doesn't support server cook platform... 
diff --git a/Engine/Source/Programs/AutomationTool/Scripts/RunProjectCommand.Automation.cs b/Engine/Source/Programs/AutomationTool/Scripts/RunProjectCommand.Automation.cs
--- a/Engine/Source/Programs/AutomationTool/Scripts/RunProjectCommand.Automation.cs
+++ b/Engine/Source/Programs/AutomationTool/Scripts/RunProjectCommand.Automation.cs
@@ -860,6 +860,19 @@ namespace AutomationScripts
 			{
 				TempCmdLine += "-nullrhi ";
 			}
+			if (!String.IsNullOrEmpty(Params.Trace))
+			{
+				TempCmdLine += Params.Trace + " ";
+			}
+			if (!String.IsNullOrEmpty(Params.TraceHost))
+			{
+				TempCmdLine += Params.TraceHost + " ";
+			}
+			if (!String.IsNullOrEmpty(Params.TraceFile))
+			{
+				TempCmdLine += Params.TraceFile + " ";
+			}
+
 			TempCmdLine += SC.StageTargetPlatform.GetLaunchExtraCommandLine(Params);
 
 			TempCmdLine += "-CrashForUAT ";
@@ -985,6 +998,22 @@ namespace AutomationScripts
 			{
 				Args += " -buildmachine";
 			}
+
+			if (!String.IsNullOrEmpty(Params.Trace))
+			{
+				Args += " " + Params.Trace;
+			}
+
+			if (!String.IsNullOrEmpty(Params.TraceHost))
+			{
+				Args += " " + Params.TraceHost;
+			}
+
+			if (!String.IsNullOrEmpty(Params.TraceFile))
+			{
+				Args += " " + Params.TraceFile;
+			}
+
 			Args += " -CrashForUAT";
 			Args += " " + AdditionalCommandLine;
 
