Remove input keys if they are not valid to remove warnings

diff --git a/Engine/Source/Programs/AutomationTool/Gauntlet/Unreal/Base/Gauntlet.UnrealTestConfiguration.cs b/Engine/Source/Programs/AutomationTool/Gauntlet/Unreal/Base/Gauntlet.UnrealTestConfiguration.cs
--- a/Engine/Source/Programs/AutomationTool/Gauntlet/Unreal/Base/Gauntlet.UnrealTestConfiguration.cs
+++ b/Engine/Source/Programs/AutomationTool/Gauntlet/Unreal/Base/Gauntlet.UnrealTestConfiguration.cs
@@ -855,11 +855,12 @@ namespace Gauntlet
 				}
 			}
 
-			// due to an issue with dotnet being extremely pedantic we have to drop our locks on files so we can read from the log file
-			// [REDACTED_PROFILE_URL]/runtime/issues/34126
 			if (AppConfig.Platform == UnrealTargetPlatform.Linux)
 			{
+				// due to an issue with dotnet being extremely pedantic we have to drop our locks on files so we can read from the log file
+				// [REDACTED_PROFILE_URL]/runtime/issues/34126
 				AppConfig.CommandLine += " -noexclusivelockonwrite";
+				AppConfig.CommandLine += " -RemoveInvalidKeys";
 			}
 
 			// use -log on servers so we get a window..
