Change adb shell communication (another StartServer improvement for AndroidFileServer)

#android

diff --git a/Engine/Source/Programs/AutomationTool/Android/AndroidFileClient.cs b/Engine/Source/Programs/AutomationTool/Android/AndroidFileClient.cs
--- a/Engine/Source/Programs/AutomationTool/Android/AndroidFileClient.cs
+++ b/Engine/Source/Programs/AutomationTool/Android/AndroidFileClient.cs
@@ -87,6 +87,21 @@ namespace AutomationTool
 			private int Port = 6037;
 			private Socket ClientSocket = null;
 
+			private string LastShellDevice = "";
+			private string LastShellFeatures = "";
+
+			// header is 1 byte id, 4 byte uint32 length
+			private static int kHeaderSize = 1 + 4;
+
+			private const int kIdStdin = 0;
+			private const int kIdStdOut = 1;
+			private const int kIdStdErr = 2;
+			private const int kIdExit = 3;
+			private const int kIdCloseStdin = 4;
+			private const int kIdWindowSizeChange = 5;
+			private const int kIdInvalid = 255;
+
+
 			public OptimalADB(string ADBExecutable = "", int inPort = DefaultPort)
 			{
 				Executable = (ADBExecutable != "") ? ADBExecutable : Environment.ExpandEnvironmentVariables("%ANDROID_HOME%/platform-tools/adb" + (RuntimePlatform.IsWindows ? ".exe" : ""));
@@ -284,64 +299,116 @@ namespace AutomationTool
 				return Result;
 			}
 
+			public string GetFeatures(string Device)
+			{
+				string Result;
+
+				executeDevice(Device, "host:features", out Result);
+				CloseConnection();
+				return Result;
+			}
+
 			private string shellEscape(string source)
 			{
 				// wraps whole thing in ' ' after escaping it with \'
-				return "'" + source.Replace("'", "\'") + "'";
+				return "'" + source.Replace("'", "'\\''") + "'";
 			}
 
-			public string Shell(string Device, string Command, bool bExpectResponse = true)
+			public string Shell(string Device, string Command)
 			{
 				string Result = "FAIL";
 
-				if (!execute("host:transport" + (Device == "" ? "-any" : ":" + Device), out Result, true))
+				if (Device != LastShellDevice)
 				{
-					CloseConnection();
-					return Result;
+					LastShellFeatures = GetFeatures(Device);
+					LastShellDevice = Device;
 				}
 
-				Result = "FAIL";
-				//if (!sendData("shell:" + shellEscape(Command)))
-				if (!sendData("shell:" + Command))
+				if (!execute("host:transport" + (Device == "" ? "-any" : ":" + Device), out Result, true))
 				{
 					CloseConnection();
 					return Result;
 				}
 
-				byte[] buffer = new byte[65536];
-				Result = "";
-				while (Result.Length < 4)
-				{
-					int bytesRecv = ClientSocket.Receive(buffer);
-					Result += Encoding.UTF8.GetString(buffer, 0, bytesRecv);
-				}
-				if (Result.Substring(0, 4) != "OKAY")
-				{
-					CloseConnection();
-					return Result;
-				}
-				Result = Result.Substring(4);
+				bool bUseShellv2 = LastShellFeatures.Contains("shell_v2");
+				string ShellCommand = string.Format("shell{0}:{1}", bUseShellv2 ? ",v2" : "", Command);
 
-//				if (bExpectResponse)
+				Result = "FAIL";
+				try
 				{
-					// there can be a delay before any more results (and may not be any)
-					// wait for up to 1000 ms for a response in 25 ms intervals
-					int Waits = 40;
-					while (ClientSocket.Available == 0 && Waits-- > 0)
+					if (!sendData(ShellCommand))
 					{
-						Thread.Sleep(25);
+						CloseConnection();
+						return Result;
 					}
-					while (ClientSocket.Available > 0)
+
+					byte[] buffer = new byte[65536];
+					Result = "";
+
+					while (Result.Length < 4)
 					{
-						int bytesRecv = ClientSocket.Receive(buffer);
+						int bytesRecv = ClientSocket.Receive(buffer, 4, SocketFlags.None);
 						Result += Encoding.UTF8.GetString(buffer, 0, bytesRecv);
-						if (ClientSocket.Available == 0)
+					}
+					if (Result.Substring(0, 4) != "OKAY")
+					{
+						CloseConnection();
+						return Result;
+					}
+					Result = "";
+
+					if (bUseShellv2)
+					{
+						bool bExit = false;
+						while (!bExit)
 						{
-							// wait to see if another packet follows
-							Thread.Sleep(50);
+							int bytesRecv = ClientSocket.Receive(buffer, kHeaderSize, SocketFlags.None);
+							if (bytesRecv <= 0)
+							{
+								break;
+							}
+
+							int id = buffer[0];
+							int length = buffer[1] + (buffer[2] << 8) + (buffer[3] << 16) + (buffer[4] << 24);
+
+							while (length > 0)
+							{
+								int request = length < buffer.Length ? length : buffer.Length;
+								bytesRecv = ClientSocket.Receive(buffer, request, SocketFlags.None);
+								length -= bytesRecv;
+
+								switch (id)
+								{
+									case kIdStdOut:
+										Result += Encoding.UTF8.GetString(buffer, 0, bytesRecv);
+										break;
+									case kIdStdErr:
+										Result += Encoding.UTF8.GetString(buffer, 0, bytesRecv);
+										break;
+									case kIdExit:
+										// exit code is buffer[0]
+										bExit = true;
+										break;
+								}
+							}
+						}
+						CloseConnection();
+						return Result;
+					}
+
+					while (true)
+					{
+						int bytesRecv = ClientSocket.Receive(buffer);
+						if (bytesRecv <= 0)
+						{
+							break;
 						}
+						Result += Encoding.UTF8.GetString(buffer, 0, bytesRecv);
 					}
 				}
+				catch (SocketException)
+				{
+				}
 
 				CloseConnection();
 				return Result;
@@ -897,7 +964,7 @@ namespace AutomationTool
 
 			return sent;
 		}
-		
+
 		private byte[] CommandPacket(int Command, long Size)
 		{
 			byte[] Result = { (byte)(Command & 255),
@@ -953,7 +1020,7 @@ namespace AutomationTool
 				int resultSize = (int)ResultSize();
 				if (resultSize > 0)
 				{
-					byte[] buffer = new byte[resultSize < FILE_BUFFERSIZE ? resultSize : FILE_BUFFERSIZE ];
+					byte[] buffer = new byte[resultSize < FILE_BUFFERSIZE ? resultSize : FILE_BUFFERSIZE];
 
 					Result = "";
 					long remaining = resultSize;
@@ -1379,7 +1446,7 @@ namespace AutomationTool
 			foreach (string Receiver in InstalledReceivers)
 			{
 				bDidSendStops = true;
-				adb.Shell(Device, "am broadcast -a com.epicgames.unreal.RemoteFileManager.intent.COMMAND -n " + Receiver + "/com.epicgames.unreal.RemoteFileManagerReceiver -e cmd 'stop'", false);
+				adb.Shell(Device, "am broadcast -a com.epicgames.unreal.RemoteFileManager.intent.COMMAND -n " + Receiver + "/com.epicgames.unreal.RemoteFileManagerReceiver -e cmd stop");
 			}
 
 			// it can take up to 2 seconds for running servers to terminate so check if any binds to port are still active
@@ -1468,8 +1535,8 @@ namespace AutomationTool
 			string WifiAddress;
 
 			// sent start request (won't do anything if already started)
-			string StartCommand = "am broadcast -a com.epicgames.unreal.RemoteFileManager.intent.COMMAND -n " + PackageName + "/com.epicgames.unreal.RemoteFileManagerReceiver -e cmd 'start' -e token '" + Token + "' -ei port " + ServerPort;
-			adb.Shell(Device, StartCommand, false);
+			string StartCommand = "am broadcast -a com.epicgames.unreal.RemoteFileManager.intent.COMMAND -n " + PackageName + "/com.epicgames.unreal.RemoteFileManagerReceiver -e cmd start -e token " + Token + " -ei port " + ServerPort;
+			adb.Shell(Device, StartCommand);
 
 			// see if we can check listen status
 			if (GetListenStatus(Device, ServerPort, out bUSB, out bWifi, out WifiAddress))
@@ -1489,7 +1556,7 @@ namespace AutomationTool
 					Thread.Sleep(200);
 
 					// sent start request again (won't do anything if already started)
-					adb.Shell(Device, StartCommand, false);
+					adb.Shell(Device, StartCommand);
 
 					GetListenStatus(Device, ServerPort, out bUSB, out bWifi, out WifiAddress);
 				}
@@ -1510,7 +1577,7 @@ namespace AutomationTool
 				}
 
 				// sent start request again (won't do anything if already started)
-				adb.Shell(Device, StartCommand, false);
+				adb.Shell(Device, StartCommand);
 			}
 
 			Log.TraceInformation("Timed out on connection attempts");
