Fall back to netstat if ss utility present but not showing connections during StartServer for Android File Server

#android

diff --git a/Engine/Source/Programs/AutomationTool/Android/AndroidFileClient.cs b/Engine/Source/Programs/AutomationTool/Android/AndroidFileClient.cs
--- a/Engine/Source/Programs/AutomationTool/Android/AndroidFileClient.cs
+++ b/Engine/Source/Programs/AutomationTool/Android/AndroidFileClient.cs
@@ -156,8 +156,6 @@ namespace AutomationTool
 					{
 						ClientSocket.Shutdown(SocketShutdown.Both);
 						ClientSocket.Close();
-						ClientSocket.Dispose();
-						ClientSocket = null;
 					}
 					catch (SocketException se)
 					{
@@ -167,6 +165,11 @@ namespace AutomationTool
 					{
 						Log.TraceError("Unexpected Exception: {0}", e.ToString());
 					}
+					finally
+					{
+						ClientSocket.Dispose();
+					}
+					ClientSocket = null;
 				}
 			}
 
@@ -255,6 +258,47 @@ namespace AutomationTool
 						Output = WorkBuffer;
 						return true;
 					}
+					else if (WorkBuffer.Substring(0, 4) == "FAIL")
+					{
+						WorkBuffer = WorkBuffer.Substring(4);
+						int remaining = 0;
+						while (true)
+						{
+							while (ClientSocket.Available > 0 || WorkBuffer.Length < 4)
+							{
+								bytesRecv = ClientSocket.Receive(buffer);
+								WorkBuffer += Encoding.UTF8.GetString(buffer, 0, bytesRecv);
+							}
+							if (WorkBuffer.Length >= 4)
+							{
+								try
+								{
+									remaining = int.Parse(WorkBuffer.Substring(0, 4), System.Globalization.NumberStyles.HexNumber);
+									WorkBuffer = WorkBuffer.Substring(4);
+									remaining -= WorkBuffer.Length;
+								}
+								catch (Exception)
+								{
+									while (ClientSocket.Available > 0)
+									{
+										bytesRecv = ClientSocket.Receive(buffer);
+										WorkBuffer += Encoding.UTF8.GetString(buffer, 0, bytesRecv);
+									}
+									Output = WorkBuffer;
+									return false;
+								}
+								break;
+							}
+						}
+						while (remaining > 0)
+						{
+							bytesRecv = ClientSocket.Receive(buffer);
+							WorkBuffer += Encoding.UTF8.GetString(buffer, 0, bytesRecv);
+							remaining -= bytesRecv;
+						}
+						Output = WorkBuffer;
+						return false;
+					}
 				}
 
 				// fail
@@ -1359,8 +1403,20 @@ namespace AutomationTool
 			{
 				return false;
 			}
-			foreach (string Line in ssResult.Split('\n'))
+			bool bConnectionsPresent = false;
+			bool bFirstLine = true;
+			foreach (string Line in ssResult.Split('\n', StringSplitOptions.RemoveEmptyEntries))
 			{
+				if (bFirstLine)
+				{
+					bFirstLine = false;
+					continue;
+				}
+				if (Line.Contains("Permission denied"))
+				{
+					continue;
+				}
+				bConnectionsPresent = true;
 				if (Line.Contains("LISTEN") && Line.Contains(Search))
 				{
 					String[] parts = Line.Split(new char[0], StringSplitOptions.RemoveEmptyEntries);
@@ -1395,7 +1451,7 @@ namespace AutomationTool
 				}
 			}
 
-			return true;
+			return bConnectionsPresent;
 		}
 
 		public static bool GetListenStatus(string Device, int Port, out bool bUSB, out bool bWifi, out string WifiAddress)
