WinPlatform automation can be specalized by Windows-based platform extensions.

#rnx

diff --git a/Engine/Source/Programs/AutomationTool/Win/WinPlatform.Automation.cs b/Engine/Source/Programs/AutomationTool/Win/WinPlatform.Automation.cs
--- a/Engine/Source/Programs/AutomationTool/Win/WinPlatform.Automation.cs
+++ b/Engine/Source/Programs/AutomationTool/Win/WinPlatform.Automation.cs
@@ -18,6 +18,11 @@ public class Win64Platform : Platform
 	{
 	}
 
+	protected Win64Platform(UnrealTargetPlatform PlatformType)
+		: base(PlatformType)
+	{
+	}
+
 	public override DeviceInfo[] GetDevices()
 	{
 		List<DeviceInfo> Devices = new List<DeviceInfo>();
@@ -48,7 +53,7 @@ public class Win64Platform : Platform
 
 		if (SC.bStageCrashReporter)
 		{
-			FileReference ReceiptFileName = TargetReceipt.GetDefaultPath(Unreal.EngineDirectory, "CrashReportClient", SC.StageTargetPlatform.PlatformType, UnrealTargetConfiguration.Shipping, null);
+			FileReference ReceiptFileName = TargetReceipt.GetDefaultPath(Unreal.EngineDirectory, "CrashReportClient", CrashReportPlatform ?? SC.StageTargetPlatform.PlatformType, UnrealTargetConfiguration.Shipping, null);
 			if(FileReference.Exists(ReceiptFileName))
 			{
 				TargetReceipt Receipt = TargetReceipt.Read(ReceiptFileName);
