UAT: Cleanup the CookCommand in BuildCookRun and share generic cook command line arguments between CBTB and COTF.
CBTB will now use the -zenstore parameter.
COTF will now use -AdditionalCookerOptions just as CBTB, and stop using the RunCommandline from -addcmdline.
COTF will now use the -buildmachine, -ddc, -unversioned, -fastcook, -skipeditorcontent, -outputdir and config override parameters.

#rnx

diff --git a/Engine/Source/Programs/AutomationTool/Scripts/CookCommand.Automation.cs b/Engine/Source/Programs/AutomationTool/Scripts/CookCommand.Automation.cs
--- a/Engine/Source/Programs/AutomationTool/Scripts/CookCommand.Automation.cs
+++ b/Engine/Source/Programs/AutomationTool/Scripts/CookCommand.Automation.cs
@@ -24,6 +24,72 @@ namespace AutomationScripts
 	/// </remarks>
 	public partial class Project : CommandUtils
 	{
+		private static string GetGenericCookCommandletParams(ProjectParams Params)
+		{
+			string CommandletParams = "";
+
+			if (IsBuildMachine)
+			{
+				CommandletParams += " -buildmachine";
+			}
+			if (Params.ZenStore)
+			{
+				CommandletParams += " -zenstore";
+			}
+			if (Params.HasDDCGraph)
+			{
+				CommandletParams += " -ddc=" + Params.DDCGraph;
+			}
+			if (Params.UnversionedCookedContent)
+			{
+				CommandletParams += " -unversioned";
+			}
+			if (Params.FastCook)
+			{
+				CommandletParams += " -FastCook";
+			}
+			if (Params.IterativeCooking)
+			{
+				CommandletParams += " -iterate";
+			}
+			if (Params.SkipCookingEditorContent)
+			{
+				CommandletParams += " -skipeditorcontent";
+			}
+			if (!String.IsNullOrEmpty(Params.CookOutputDir))
+			{
+				CommandletParams += " -outputdir=" + CommandUtils.MakePathSafeToUseWithCommandLine(Params.CookOutputDir);
+			}
+			if (!String.IsNullOrEmpty(Params.Trace))
+			{
+				CommandletParams += " " + Params.Trace;
+			}
+			if (!String.IsNullOrEmpty(Params.TraceHost))
+			{
+				CommandletParams += " " + Params.TraceHost;
+			}
+			if (!String.IsNullOrEmpty(Params.TraceFile))
+			{
+				CommandletParams += " " + Params.TraceFile;
+			}
+
+			// process additional cooker options
+			if (Params.HasAdditionalCookerOptions)
+			{
+				string FormatedAdditionalCookerParams = Params.AdditionalCookerOptions.TrimStart(new char[] { '\"', ' ' }).TrimEnd(new char[] { '\"', ' ' });
+				CommandletParams += " ";
+				CommandletParams += FormatedAdditionalCookerParams;
+			}
+
+			// process config overrides (-ini)
+			foreach (string ConfigOverrideParam in Params.ConfigOverrideParams)
+			{
+				CommandletParams += " -";
+				CommandletParams += ConfigOverrideParam;
+			}
+			return CommandletParams;
+		}
+
 		public static void GetCookByTheBookCommandletParams(ProjectParams Params, out string CommandletParams, out string[] Maps, out string[] DirectoriesToCook, out string InternationalizationPreset, out string[] CulturesToCook, out string PlatformsToCook)
 		{
 			var PlatformsToCookSet = new HashSet<string>();
@@ -89,28 +155,14 @@ namespace AutomationScripts
 				CulturesToCook = Params.CulturesToCook.ToArray();
 			}
 
-			CommandletParams = IsBuildMachine ? "-buildmachine -fileopenlog" : "-fileopenlog";
+			CommandletParams = GetGenericCookCommandletParams(Params);
+
+			CommandletParams += " -fileopenlog";
 
-			if (Params.HasDDCGraph)
-			{
-				CommandletParams += " -ddc=" + Params.DDCGraph;
-			}
-			if (Params.UnversionedCookedContent)
-			{
-				CommandletParams += " -unversioned";
-			}
-			if (Params.FastCook)
-			{
-				CommandletParams += " -FastCook";
-			}
 			if (Params.Manifests)
 			{
 				CommandletParams += " -manifests";
 			}
-			if (Params.IterativeCooking)
-			{
-				CommandletParams += " -iterate";
-			}
 			if (Params.HasIterateSharedCookedBuild)
 			{
 				new SharedCookedBuild(Params).CopySharedCookedBuilds();
@@ -129,10 +181,6 @@ namespace AutomationScripts
 			{
 				CommandletParams += " -createreleaseversion=" + Params.CreateReleaseVersion;
 			}
-			if (Params.SkipCookingEditorContent)
-			{
-				CommandletParams += " -skipeditorcontent";
-			}
 			if (Params.NumCookersToSpawn != 0)
 			{
 				CommandletParams += " -numcookerstospawn=" + Params.NumCookersToSpawn;
@@ -155,10 +203,6 @@ namespace AutomationScripts
 					CommandletParams += " -errorOnEngineContentUse";
 				}
 			}
-			if (!String.IsNullOrEmpty(Params.CookOutputDir))
-			{
-				CommandletParams += " -outputdir=" + CommandUtils.MakePathSafeToUseWithCommandLine(Params.CookOutputDir);
-			}
 			// don't include the based on release version unless we are cooking dlc or creating a new release version
 			// in this case the based on release version is used in packaging
 			if (Params.HasBasedOnReleaseVersion && (Params.HasDLCName || Params.HasCreateReleaseVersion))
@@ -173,18 +217,6 @@ namespace AutomationScripts
 			{
 				CommandletParams += " -basedonreleaseversionroot=" + Params.BasedOnReleaseVersionBasePath;
 			}
-			if (!String.IsNullOrEmpty(Params.Trace))
-			{
-				CommandletParams += " " + Params.Trace;
-			}
-			if (!String.IsNullOrEmpty(Params.TraceHost))
-			{
-				CommandletParams += " " + Params.TraceHost;
-			}
-			if (!String.IsNullOrEmpty(Params.TraceFile))
-			{
-				CommandletParams += " " + Params.TraceFile;
-			}
 
 			// if we are not going to pak but we specified compressed then compress in the cooker ;)
 			// otherwise compress the pak files
@@ -193,13 +225,6 @@ namespace AutomationScripts
 				CommandletParams += " -compressed";
 			}
 
-			if (Params.HasAdditionalCookerOptions)
-			{
-				string FormatedAdditionalCookerParams = Params.AdditionalCookerOptions.TrimStart(new char[] { '\"', ' ' }).TrimEnd(new char[] { '\"', ' ' });
-				CommandletParams += " ";
-				CommandletParams += FormatedAdditionalCookerParams;
-			}
-
 			if (!Params.NoClient)
 			{
 				var MapsList = Maps == null ? new List<string>() : Maps.ToList();
@@ -212,12 +237,6 @@ namespace AutomationScripts
 				Maps = MapsList.ToArray();
 			}
 
-			// Config overrides (-ini)
-			foreach (string ConfigOverrideParam in Params.ConfigOverrideParams)
-			{
-				CommandletParams += " -";
-				CommandletParams += ConfigOverrideParam;
-			}
 		}
 
 		public static void Cook(ProjectParams Params)
@@ -252,35 +271,10 @@ namespace AutomationScripts
 						CreateDirectory(LogFolderOutsideOfSandbox);
 					}
 
-					String COTFCommandLine = Params.RunCommandline;
-					if (Params.IterativeCooking)
-					{
-						COTFCommandLine += " -iterate";
-					}
-
-					if (Params.HasDDCGraph)
-					{
-						COTFCommandLine += " -ddc=" + Params.DDCGraph;
-					}
-
+					string COTFCommandLine = GetGenericCookCommandletParams(Params);
 					if (Params.ZenStore)
 					{
-						COTFCommandLine += " -zenstore -messaging";
-					}
-
-					if (!String.IsNullOrEmpty(Params.Trace))
-					{
-						COTFCommandLine += " " + Params.Trace;
-					}
-
-					if (!String.IsNullOrEmpty(Params.TraceHost))
-					{
-						COTFCommandLine += " " + Params.TraceHost;
-					}
-
-					if (!String.IsNullOrEmpty(Params.TraceFile))
-					{
-						COTFCommandLine += " " + Params.TraceFile;
+						COTFCommandLine += " -messaging";
 					}
 
 					var ServerLogFile = CombinePaths(LogFolderOutsideOfSandbox, "Server.log");
