UAT: Don't wait for 2 minutes for output handles to close for processes that does not redirect StdOut/StdErr at all (e.g. by calling Run with Options=ERunOptions.NoStdOutRedirect).
Give ERunOptions.NoStdOutCapture a unique value so that Options.HasFlag(ERunOptions.NoStdOutCapture) returns the expected result.

#rnx

diff --git a/Engine/Source/Programs/AutomationTool/AutomationUtils/ProcessUtils.cs b/Engine/Source/Programs/AutomationTool/AutomationUtils/ProcessUtils.cs
--- a/Engine/Source/Programs/AutomationTool/AutomationUtils/ProcessUtils.cs
+++ b/Engine/Source/Programs/AutomationTool/AutomationUtils/ProcessUtils.cs
@@ -234,7 +234,7 @@ namespace AutomationTool
 		public delegate string SpewFilterCallbackType(string Message);
 
 		private int ProcessExitCode = -1;
-		private StringBuilder ProcessOutput;
+		private StringBuilder ProcessOutput = null;
 		private bool AllowSpew = true;
 		private LogEventType SpewVerbosity = LogEventType.Console;
 		private SpewFilterCallbackType SpewFilterCallback = null;
@@ -250,11 +250,8 @@ namespace AutomationTool
 			ProcSyncObject = new object();
 			Proc = InProc;
 			AllowSpew = bAllowSpew;
-			if(bCaptureSpew)
-			{
-				ProcessOutput = new StringBuilder();
-			}
-			else
+			ProcessOutput = bCaptureSpew ? new StringBuilder() : null;
+			if (!AllowSpew && !bCaptureSpew)
 			{
 				OutputWaitHandle.Set();
 				ErrorWaitHandle.Set();
@@ -722,9 +719,17 @@ namespace AutomationTool
 		public enum ERunOptions
 		{
 			None = 0,
+            /// <summary>
+			/// If AllowSpew is set, then the redirected output from StdOut/StdErr is logged.
+			/// Not relevant when NoStdOutRedirect is set.
+            /// </summary>
 			AllowSpew = 1 << 0,
 			AppMustExist = 1 << 1,
 			NoWaitForExit = 1 << 2,
+            /// <summary>
+			/// If NoStdOutRedirect is set, then StdOut/StdErr output is not redirected, logged or captured.
+			/// Else, StdOut/StdErr is redirected and captured by default.
+            /// </summary>
 			NoStdOutRedirect = 1 << 3,
             NoLoggingOfRunCommand = 1 << 4,
             UTF8Output = 1 << 5,
@@ -739,11 +744,12 @@ namespace AutomationTool
 			/// If set, a window is allowed to be created
 			/// </summary>
 			NoHideWindow = 1 << 8,
-
 			/// <summary>
-			/// Do not capture stdout in the process result
+			/// If NoStdOutCapture is set, then the redirected output from StdOut/StdErr is not captured in the ProcessResult,
+			/// and ProcessResult.Output will return null.
+			/// Not relevant when NoStdOutRedirect is set.
 			/// </summary>
-			NoStdOutCapture = 1 << 8,
+			NoStdOutCapture = 1 << 9,
 
 			Default = AllowSpew | AppMustExist,
 		}
@@ -840,12 +846,14 @@ namespace AutomationTool
                 LogWithVerbosity(SpewVerbosity,"Running: " + App + " " + (String.IsNullOrEmpty(CommandLine) ? "" : CommandLine));
             }
 
-			IProcessResult Result = ProcessManager.CreateProcess(App, Options.HasFlag(ERunOptions.AllowSpew), !Options.HasFlag(ERunOptions.NoStdOutCapture), Env, SpewVerbosity: SpewVerbosity, SpewFilterCallback: SpewFilterCallback, WorkingDir: WorkingDir);
+			bool bRedirectStdOut = !Options.HasFlag(ERunOptions.NoStdOutRedirect);
+			bool bAllowSpew = bRedirectStdOut && Options.HasFlag(ERunOptions.AllowSpew);
+			bool bCaptureSpew = bRedirectStdOut && !Options.HasFlag(ERunOptions.NoStdOutCapture);
+			IProcessResult Result = ProcessManager.CreateProcess(App, bAllowSpew, bCaptureSpew, Env, SpewVerbosity: SpewVerbosity, SpewFilterCallback: SpewFilterCallback, WorkingDir: WorkingDir);
 			using (LogIndentScope Scope = Options.HasFlag(ERunOptions.AllowSpew) ? new LogIndentScope("  ") : null)
 			{
 				Process Proc = Result.ProcessObject;
 
-				bool bRedirectStdOut = (Options & ERunOptions.NoStdOutRedirect) != ERunOptions.NoStdOutRedirect;
 				Proc.StartInfo.FileName = App;
 
 				// Process Arguments follow windows conventions in .NET Core
