UAT: Cleanup the RunProjectCommand and CopyBuildToStagingDirectory in BuildCookRun to use the same GetFileHostCommandline logic.
Delete the RunFileServer logic since the UnrealFileServer.exe does not exist any more.

#rnx

diff --git a/Engine/Source/Programs/AutomationTool/Scripts/CopyBuildToStagingDirectory.Automation.cs b/Engine/Source/Programs/AutomationTool/Scripts/CopyBuildToStagingDirectory.Automation.cs
--- a/Engine/Source/Programs/AutomationTool/Scripts/CopyBuildToStagingDirectory.Automation.cs
+++ b/Engine/Source/Programs/AutomationTool/Scripts/CopyBuildToStagingDirectory.Automation.cs
@@ -4058,149 +4058,6 @@ namespace AutomationScripts
 			return DirectoryReference.Combine(SC.ProjectRoot, "Intermediate", "UAT", SC.FinalCookPlatform);
 		}
 
-		private static string GetFileHostCommandline(ProjectParams Params, DeploymentContext SC)
-		{
-			string FileHostParams = " ";
-			if (Params.CookOnTheFly || Params.FileServer)
-			{
-				if (Params.ZenStore)
-				{
-					if (Params.CookOnTheFly)
-					{
-						FileHostParams += "-cookonthefly ";
-					}
-					FileHostParams += "-zenstoreproject=" + ProjectUtils.GetProjectPathId(SC.RawProjectPath) + " ";
-
-					FileHostParams += "-zenstorehost=";
-				}
-				else
-				{
-					FileHostParams += "-filehostip=";
-				}
-
-				// add localhost first for platforms using redirection
-				const string LocalHost = "127.0.0.1";
-				if (!IsNullOrEmpty(Params.Port))
-				{
-					bool FirstParam = true;
-					foreach (var Port in Params.Port)
-					{
-						if (!FirstParam)
-						{
-							FileHostParams += "+";
-						}
-						FirstParam = false;
-						string[] PortProtocol = Port.Split(new char[] { ':' });
-						if (PortProtocol.Length > 1)
-						{
-							FileHostParams += String.Format("{0}://{1}:{2}", PortProtocol[0], LocalHost, PortProtocol[1]);
-						}
-						else
-						{
-							FileHostParams += LocalHost;
-							FileHostParams += ":";
-							FileHostParams += Params.Port;
-						}
-
-					}
-				}
-				else
-				{
-					// use default port
-					FileHostParams += LocalHost;
-				}
-
-				if (UnrealBuildTool.BuildHostPlatform.Current.Platform == UnrealTargetPlatform.Mac)
-				{
-					NetworkInterface[] Interfaces = NetworkInterface.GetAllNetworkInterfaces();
-					foreach (NetworkInterface adapter in Interfaces)
-					{
-						if (adapter.NetworkInterfaceType != NetworkInterfaceType.Loopback)
-						{
-							IPInterfaceProperties IP = adapter.GetIPProperties();
-							for (int Index = 0; Index < IP.UnicastAddresses.Count; ++Index)
-							{
-								if (InternalUtils.IsDnsEligible(IP.UnicastAddresses[Index]) && IP.UnicastAddresses[Index].Address.AddressFamily == System.Net.Sockets.AddressFamily.InterNetwork)
-								{
-									if (!IsNullOrEmpty(Params.Port))
-									{
-										foreach (var Port in Params.Port)
-										{
-											FileHostParams += "+";
-											string[] PortProtocol = Port.Split(new char[] { ':' });
-											if (PortProtocol.Length > 1)
-											{
-												FileHostParams += String.Format("{0}://{1}:{2}", PortProtocol[0], IP.UnicastAddresses[Index].Address.ToString(), PortProtocol[1]);
-											}
-											else
-											{
-												FileHostParams += IP.UnicastAddresses[Index].Address.ToString();
-												FileHostParams += ":";
-												FileHostParams += Params.Port;
-											}
-
-										}
-									}
-									else
-									{
-										FileHostParams += "+";
-										// use default port
-										FileHostParams += IP.UnicastAddresses[Index].Address.ToString();
-									}
-
-								}
-							}
-						}
-					}
-				}
-				else
-				{
-					NetworkInterface[] Interfaces = NetworkInterface.GetAllNetworkInterfaces();
-					foreach (NetworkInterface adapter in Interfaces)
-					{
-						if (adapter.OperationalStatus == OperationalStatus.Up)
-						{
-							IPInterfaceProperties IP = adapter.GetIPProperties();
-							for (int Index = 0; Index < IP.UnicastAddresses.Count; ++Index)
-							{
-								if (InternalUtils.IsDnsEligible(IP.UnicastAddresses[Index]))
-								{
-									if (!IsNullOrEmpty(Params.Port))
-									{
-										foreach (var Port in Params.Port)
-										{
-											FileHostParams += "+";
-											string[] PortProtocol = Port.Split(new char[] { ':' });
-											if (PortProtocol.Length > 1)
-											{
-												FileHostParams += String.Format("{0}://{1}:{2}", PortProtocol[0], IP.UnicastAddresses[Index].Address.ToString(), PortProtocol[1]);
-											}
-											else
-											{
-												FileHostParams += IP.UnicastAddresses[Index].Address.ToString();
-												FileHostParams += ":";
-												FileHostParams += Params.Port;
-											}
-										}
-									}
-									else
-									{
-										FileHostParams += "+";
-										// use default port
-										FileHostParams += IP.UnicastAddresses[Index].Address.ToString();
-									}
-								}
-							}
-						}
-					}
-				}
-
-				FileHostParams += " ";
-			}
-
-			return FileHostParams;
-		}
-
 		public static void WriteStageCommandline(FileReference IntermediateCmdLineFile, ProjectParams Params, DeploymentContext SC)
 		{
 			// this file needs to be treated as a UFS file for casing, but NonUFS for being put into the .pak file. 
diff --git a/Engine/Source/Programs/AutomationTool/Scripts/RunProjectCommand.Automation.cs b/Engine/Source/Programs/AutomationTool/Scripts/RunProjectCommand.Automation.cs
--- a/Engine/Source/Programs/AutomationTool/Scripts/RunProjectCommand.Automation.cs
+++ b/Engine/Source/Programs/AutomationTool/Scripts/RunProjectCommand.Automation.cs
@@ -151,10 +151,6 @@ namespace AutomationScripts
 					throw new AutomationException("Failed to run, server target platform not specified");
 				}
 			}
-			else if (Params.FileServer && !Params.SkipServer)
-			{
-				ServerProcess = RunFileServer(Params, ServerLogFile, Params.RunCommandline);
-			}
 
 			if (ServerProcess != null)
 			{
@@ -631,136 +627,8 @@ namespace AutomationScripts
 
 				if (Params.CookOnTheFly || Params.FileServer)
 				{
-					if (Params.ZenStore)
-					{
-						if (Params.CookOnTheFly)
-						{
-							TempCmdLine += "-cookonthefly ";
-						}
-						TempCmdLine += "-zenstoreproject=" + ProjectUtils.GetProjectPathId(SC.RawProjectPath) + " ";
-						TempCmdLine += "-zenstorehost=";
-					}
-					else
-					{
-						TempCmdLine += "-filehostip=";
-					}
-
-					// add localhost first for platforms using redirection
-					const string LocalHost = "127.0.0.1";
-					if (!IsNullOrEmpty(Params.Port))
-					{
-						bool FirstParam = true;
-						foreach (var Port in Params.Port)
-						{
-							if (!FirstParam)
-							{
-								TempCmdLine += "+";
-							}
-							FirstParam = false;
-							string[] PortProtocol = Port.Split(new char[] { ':' });
-							if (PortProtocol.Length > 1)
-							{
-								TempCmdLine += String.Format("{0}://{1}:{2}", PortProtocol[0], LocalHost, PortProtocol[1]);
-							}
-							else
-							{
-								TempCmdLine += LocalHost;
-								TempCmdLine += ":";
-								TempCmdLine += Params.Port;
-							}
-
-						}
-					}
-					else
-					{
-						// use default port
-						TempCmdLine += LocalHost;
-					}
-
-					if (UnrealBuildTool.BuildHostPlatform.Current.Platform == UnrealTargetPlatform.Mac)
-					{
-						NetworkInterface[] Interfaces = NetworkInterface.GetAllNetworkInterfaces();
-						foreach (NetworkInterface adapter in Interfaces)
-						{
-							if (adapter.NetworkInterfaceType != NetworkInterfaceType.Loopback)
-							{
-								IPInterfaceProperties IP = adapter.GetIPProperties();
-								for (int Index = 0; Index < IP.UnicastAddresses.Count; ++Index)
-								{
-									if (InternalUtils.IsDnsEligible(IP.UnicastAddresses[Index]) && IP.UnicastAddresses[Index].Address.AddressFamily == System.Net.Sockets.AddressFamily.InterNetwork)
-									{
-										if (!IsNullOrEmpty(Params.Port))
-										{
-											foreach (var Port in Params.Port)
-											{
-												TempCmdLine += "+";
-												string[] PortProtocol = Port.Split(new char[] { ':' });
-												if (PortProtocol.Length > 1)
-												{
-													TempCmdLine += String.Format("{0}://{1}:{2}", PortProtocol[0], IP.UnicastAddresses[Index].Address.ToString(), PortProtocol[1]);
-												}
-												else
-												{
-													TempCmdLine += IP.UnicastAddresses[Index].Address.ToString();
-													TempCmdLine += ":";
-													TempCmdLine += Params.Port;
-												}
-											}
-										}
-										else
-										{
-											TempCmdLine += "+";
-											// use default port
-											TempCmdLine += IP.UnicastAddresses[Index].Address.ToString();
-										}
-									}
-								}
-							}
-						}
-					}
-					else
-					{
-						NetworkInterface[] Interfaces = NetworkInterface.GetAllNetworkInterfaces();
-						foreach (NetworkInterface adapter in Interfaces)
-						{
-							if (adapter.OperationalStatus == OperationalStatus.Up)
-							{
-								IPInterfaceProperties IP = adapter.GetIPProperties();
-								for (int Index = 0; Index < IP.UnicastAddresses.Count; ++Index)
-								{
-									if (InternalUtils.IsDnsEligible(IP.UnicastAddresses[Index]))
-									{
-										if (!IsNullOrEmpty(Params.Port))
-										{
-											foreach (var Port in Params.Port)
-											{
-												TempCmdLine += "+";
-												string[] PortProtocol = Port.Split(new char[] { ':' });
-												if (PortProtocol.Length > 1)
-												{
-													TempCmdLine += String.Format("{0}://{1}:{2}", PortProtocol[0], IP.UnicastAddresses[Index].Address.ToString(), PortProtocol[1]);
-												}
-												else
-												{
-													TempCmdLine += IP.UnicastAddresses[Index].Address.ToString();
-													TempCmdLine += ":";
-													TempCmdLine += Params.Port;
-												}
-											}
-										}
-										else
-										{
-											TempCmdLine += "+";
-											// use default port
-											TempCmdLine += IP.UnicastAddresses[Index].Address.ToString();
-										}
-									}
-								}
-							}
-						}
-					}
-
-					TempCmdLine += " ";
+					string FileHostCommandline = GetFileHostCommandline(Params, SC);
+					TempCmdLine += FileHostCommandline + " ";
 
 					if (Params.CookOnTheFlyStreaming)
 					{
@@ -903,6 +771,147 @@ namespace AutomationScripts
 			}
 		}
 
+		private static string GetFileHostCommandline(ProjectParams Params, DeploymentContext SC)
+		{
+			string FileHostParams = "";
+			if (Params.CookOnTheFly || Params.FileServer)
+			{
+				if (Params.ZenStore)
+				{
+					if (Params.CookOnTheFly)
+					{
+						FileHostParams += "-cookonthefly ";
+					}
+					FileHostParams += "-zenstoreproject=" + ProjectUtils.GetProjectPathId(SC.RawProjectPath) + " ";
+					FileHostParams += "-zenstorehost=";
+				}
+				else
+				{
+					FileHostParams += "-filehostip=";
+				}
+
+				// add localhost first for platforms using redirection
+				const string LocalHost = "127.0.0.1";
+				if (!IsNullOrEmpty(Params.Port))
+				{
+					bool FirstParam = true;
+					foreach (var Port in Params.Port)
+					{
+						if (!FirstParam)
+						{
+							FileHostParams += "+";
+						}
+						FirstParam = false;
+						string[] PortProtocol = Port.Split(new char[] { ':' });
+						if (PortProtocol.Length > 1)
+						{
+							FileHostParams += String.Format("{0}://{1}:{2}", PortProtocol[0], LocalHost, PortProtocol[1]);
+						}
+						else
+						{
+							FileHostParams += LocalHost;
+							FileHostParams += ":";
+							FileHostParams += Params.Port;
+						}
+
+					}
+				}
+				else
+				{
+					// use default port
+					FileHostParams += LocalHost;
+				}
+
+				if (UnrealBuildTool.BuildHostPlatform.Current.Platform == UnrealTargetPlatform.Mac)
+				{
+					NetworkInterface[] Interfaces = NetworkInterface.GetAllNetworkInterfaces();
+					foreach (NetworkInterface adapter in Interfaces)
+					{
+						if (adapter.NetworkInterfaceType != NetworkInterfaceType.Loopback)
+						{
+							IPInterfaceProperties IP = adapter.GetIPProperties();
+							for (int Index = 0; Index < IP.UnicastAddresses.Count; ++Index)
+							{
+								if (InternalUtils.IsDnsEligible(IP.UnicastAddresses[Index]) && IP.UnicastAddresses[Index].Address.AddressFamily == System.Net.Sockets.AddressFamily.InterNetwork)
+								{
+									if (!IsNullOrEmpty(Params.Port))
+									{
+										foreach (var Port in Params.Port)
+										{
+											FileHostParams += "+";
+											string[] PortProtocol = Port.Split(new char[] { ':' });
+											if (PortProtocol.Length > 1)
+											{
+												FileHostParams += String.Format("{0}://{1}:{2}", PortProtocol[0], IP.UnicastAddresses[Index].Address.ToString(), PortProtocol[1]);
+											}
+											else
+											{
+												FileHostParams += IP.UnicastAddresses[Index].Address.ToString();
+												FileHostParams += ":";
+												FileHostParams += Params.Port;
+											}
+
+										}
+									}
+									else
+									{
+										FileHostParams += "+";
+										// use default port
+										FileHostParams += IP.UnicastAddresses[Index].Address.ToString();
+									}
+								}
+							}
+						}
+					}
+				}
+				else
+				{
+					NetworkInterface[] Interfaces = NetworkInterface.GetAllNetworkInterfaces();
+					foreach (NetworkInterface adapter in Interfaces)
+					{
+						if (adapter.OperationalStatus == OperationalStatus.Up)
+						{
+							IPInterfaceProperties IP = adapter.GetIPProperties();
+							for (int Index = 0; Index < IP.UnicastAddresses.Count; ++Index)
+							{
+								if (InternalUtils.IsDnsEligible(IP.UnicastAddresses[Index]))
+								{
+									if (!IsNullOrEmpty(Params.Port))
+									{
+										foreach (var Port in Params.Port)
+										{
+											FileHostParams += "+";
+											string[] PortProtocol = Port.Split(new char[] { ':' });
+											if (PortProtocol.Length > 1)
+											{
+												FileHostParams += String.Format("{0}://{1}:{2}", PortProtocol[0], IP.UnicastAddresses[Index].Address.ToString(), PortProtocol[1]);
+											}
+											else
+											{
+												FileHostParams += IP.UnicastAddresses[Index].Address.ToString();
+												FileHostParams += ":";
+												FileHostParams += Params.Port;
+											}
+										}
+									}
+									else
+									{
+										FileHostParams += "+";
+										// use default port
+										FileHostParams += IP.UnicastAddresses[Index].Address.ToString();
+									}
+								}
+							}
+						}
+					}
+				}
+
+				FileHostParams += " ";
+			}
+
+			return FileHostParams;
+		}
+
 		private static void ClientLogReaderProc(object ArgsContainer)
 		{
 			var Args = ArgsContainer as object[];
@@ -1060,46 +1069,6 @@ namespace AutomationScripts
 			return Result;
 		}
 
-		private static IProcessResult RunFileServer(ProjectParams Params, string ServerLogFile, string AdditionalCommandLine)
-		{
-#if false
-		// this section of code would provide UFS with a more accurate file mapping
-		var SC = new StagingContext(Params, false);
-		CreateStagingManifest(SC);
-		MaybeConvertToLowerCase(Params, SC);
-		var UnrealFileServerResponseFile = new List<string>();
-
-		foreach (var Pair in SC.UFSStagingFiles)
-		{
-			string Src = Pair.Key;
-			string Dest = Pair.Value;
-
-			Dest = CombinePaths(PathSeparator.Slash, SC.UnrealFileServerInternalRoot, Dest);
-
-			UnrealFileServerResponseFile.Add("\"" + Src + "\" \"" + Dest + "\"");
-		}
-
-
-		string UnrealFileServerResponseFileName = CombinePaths(CmdEnv.LogFolder, "UnrealFileServerList.txt");
-		File.WriteAllLines(UnrealFileServerResponseFileName, UnrealFileServerResponseFile);
-#endif
-			var UnrealFileServerExe = HostPlatform.Current.GetUnrealExePath("UnrealFileServer.exe");
-
-			LogInformation("Running UnrealFileServer *******");
-			var Args = String.Format("{0} -abslog={1} -unattended -CrashForUAT -log {2}",
-							CommandUtils.MakePathSafeToUseWithCommandLine(Params.RawProjectPath.FullName),
-							CommandUtils.MakePathSafeToUseWithCommandLine(ServerLogFile),
-							AdditionalCommandLine);
-			if (IsBuildMachine)
-			{
-				Args += " -buildmachine";
-			}
-			PushDir(Path.GetDirectoryName(UnrealFileServerExe));
-			var Result = Run(UnrealFileServerExe, Args, null, ERunOptions.AllowSpew | ERunOptions.NoWaitForExit | ERunOptions.AppMustExist | ERunOptions.NoStdOutRedirect);
-			PopDir();
-			return Result;
-		}
-
 		private static bool StartUnrealTrace()
 		{
 			// [TEMPORARY] - UnrealTrace server is currently only available on Windows
