Add #skipci to UploadDDC to AWS
#rnx

diff --git a/Engine/Source/Programs/AutomationTool/Scripts/UploadDDCToAWS.cs b/Engine/Source/Programs/AutomationTool/Scripts/UploadDDCToAWS.cs
--- a/Engine/Source/Programs/AutomationTool/Scripts/UploadDDCToAWS.cs
+++ b/Engine/Source/Programs/AutomationTool/Scripts/UploadDDCToAWS.cs
@@ -538,7 +538,7 @@ namespace AutomationTool
 								List<string> ExistingFiles = P4.Files(CommandUtils.MakePathSafeToUseWithCommandLine(RootManifestFile.FullName));
 
 								// Create a changelist containing the new manifest
-								ChangeNumber = P4.CreateChange(Description: "Updating DDC bundle manifest");
+								ChangeNumber = P4.CreateChange(Description: "Updating DDC bundle manifest\n#skipci");
 								if (ExistingFiles.Count > 0)
 								{
 									P4.Edit(ChangeNumber, CommandUtils.MakePathSafeToUseWithCommandLine(RootManifestFile.FullName));
