Add a retry to timestamp diff check, as networks might take a few seconds to settle the timestamp

diff --git a/Engine/Source/Programs/AutomationTool/AutomationUtils/Utils.cs b/Engine/Source/Programs/AutomationTool/AutomationUtils/Utils.cs
--- a/Engine/Source/Programs/AutomationTool/AutomationUtils/Utils.cs
+++ b/Engine/Source/Programs/AutomationTool/AutomationUtils/Utils.cs
@@ -424,8 +424,15 @@ namespace AutomationTool
 						// http://ntfs.com/exfat-time-stamp.htm
 						if (!((SourceInfo.LastWriteTimeUtc - TargetInfo.LastWriteTimeUtc).TotalSeconds < 2 && (SourceInfo.LastWriteTimeUtc - TargetInfo.LastWriteTimeUtc).TotalSeconds > -2))
 						{
-							Log.TraceInformation("Date mismatch {0} = {1} to {2} = {3}", SourceName, SourceInfo.LastWriteTimeUtc, TargetName, TargetInfo.LastWriteTimeUtc);
-							Retry = true;
+							// Copy on some networks have lag for updating timestamps, so sleep and retry the check
+							Thread.Sleep(2000);
+							SourceInfo.Refresh();
+							TargetInfo.Refresh();
+							if (!((SourceInfo.LastWriteTimeUtc - TargetInfo.LastWriteTimeUtc).TotalSeconds < 2 && (SourceInfo.LastWriteTimeUtc - TargetInfo.LastWriteTimeUtc).TotalSeconds > -2))
+							{
+								Log.TraceInformation("Date mismatch {0} = {1} to {2} = {3}", SourceName, SourceInfo.LastWriteTimeUtc, TargetName, TargetInfo.LastWriteTimeUtc);
+								Retry = true;
+							}
 						}
 					}
 				}
