Modified DOD test to elevate content beacon log warnings/errors to the test level.
Made DOD self-abort if it thinks there's an issue instead of waiting for the timeout

diff --git a/Engine/Source/Programs/AutomationTool/Gauntlet/Unreal/Base/Gauntlet.UnrealTestNode.cs b/Engine/Source/Programs/AutomationTool/Gauntlet/Unreal/Base/Gauntlet.UnrealTestNode.cs
--- a/Engine/Source/Programs/AutomationTool/Gauntlet/Unreal/Base/Gauntlet.UnrealTestNode.cs
+++ b/Engine/Source/Programs/AutomationTool/Gauntlet/Unreal/Base/Gauntlet.UnrealTestNode.cs
@@ -833,21 +833,21 @@ namespace Gauntlet
 		/// health, log updates etc. Base classes must call this or take all responsibility for
 		/// setting Status as necessary
 		/// </summary>
-		/// <returns></returns>
-		public override void TickTest()
+		/// <param name="InInstance">The Unrealnstance being tested</param>
+		public virtual void TickTest(UnrealSessionInstance InInstance)
 		{
 			IAppInstance App = null;
 			string AppInfoPrefix = "App";
-			if (TestInstance.ClientApps == null)
+			if (InInstance.ClientApps == null)
 			{
-				App = TestInstance.ServerApp;
+				App = InInstance.ServerApp;
 				AppInfoPrefix = "Server";
 			}
 			else
 			{
-				if (TestInstance.ClientApps.Length > 0)
+				if (InInstance.ClientApps.Length > 0)
 				{
-					App = TestInstance.ClientApps.First();
+					App = InInstance.ClientApps.First();
 					AppInfoPrefix = "Client";
 				}
 			}
@@ -868,7 +868,7 @@ namespace Gauntlet
 				List<string> TestLines = new List<string>();
 				// ONLY ADD RANGE ONCE. Ordering is important and will be skewed if multiple ranges are added which can skew how logs are pulled out.
 				TestLines.AddRange(Parser.GetLogChannels(LogCategories, true));
-								
+
 				for (int i = LastAppLogCount; i < TestLines.Count(); i++)
 				{
 					Log.Info(string.Format("{0}: {1}", AppInfoPrefix, TestLines[i]));
@@ -890,7 +890,7 @@ namespace Gauntlet
 				CheckHeartbeat();
 			}
 
-			IAppInstance EditorApp = TestInstance.EditorApp;
+			IAppInstance EditorApp = InInstance.EditorApp;
 			if (EditorApp != null)
 			{
 				UnrealLogParser Parser = new UnrealLogParser(EditorApp.StdOut);
@@ -904,8 +904,18 @@ namespace Gauntlet
 				}
 
 				LastEditorLogCount = TestLines.Count();
-			}
+			}			
+		}
 
+		/// <summary>
+		/// Periodically called while the test is running. A chance for tests to examine their
+		/// health, log updates etc. Base classes must call this or take all responsibility for
+		/// setting Status as necessary
+		/// </summary>
+		/// <returns></returns>
+		public override void TickTest()
+		{
+			TickTest(TestInstance);
 
 			// Check status and health after updating logs
 			if (GetTestStatus() == TestStatus.InProgress && IsTestRunning() == false)
@@ -1100,7 +1110,7 @@ namespace Gauntlet
 		/// <returns></returns>
 		public virtual ITestReport CreateReport(TestResult Result, UnrealTestContext Context, UnrealBuildSource Build, IEnumerable<UnrealRoleResult> RoleResults, string ArtifactPath)
 		{
-			if (GetConfiguration().WriteTestResultsForHorde)
+			if (GetCachedConfiguration().WriteTestResultsForHorde)
 			{
 				// write test report for Horde
 				HordeReport.SimpleTestReport HordeTestReport = CreateSimpleReportForHorde(Result);
@@ -1847,7 +1857,7 @@ namespace Gauntlet
 		}
 
 		/// <summary>
-		/// Returns the current Pass count for this node
+		/// Returns the current 0-based Pass count for this node.
 		/// </summary>
 		public int GetCurrentPass()
 		{
