- New CreateBuild Turnkey command to interactively create builds that can be run now, or saved for later (exposed to Editor Platforms menu)
- Started moving the code that turns ProjectPackagingSettings into BuildCookRun commandline from editor to C# so that we can run scripts without needing to bake out the commandline
- Updated ExecuteBuild to handle some new variables and such (like the new one for importing ProjectPackagingSettings into the commandline dynamically)
- Added ability to update a Default config file in C# (with limited functionality - set value and add array entry, with an enum we can extend later)
- Minor cleanup here and there

diff --git a/Engine/Source/Developer/DeveloperToolSettings/Classes/Settings/ProjectPackagingSettings.h b/Engine/Source/Developer/DeveloperToolSettings/Classes/Settings/ProjectPackagingSettings.h
--- a/Engine/Source/Developer/DeveloperToolSettings/Classes/Settings/ProjectPackagingSettings.h
+++ b/Engine/Source/Developer/DeveloperToolSettings/Classes/Settings/ProjectPackagingSettings.h
@@ -145,6 +145,10 @@ struct FProjectBuildSettings
 	UPROPERTY(EditAnywhere, Category="Packaging")
 	FString Name;
 
+	/** Any help that you would like to include in the ToolTip of the menu option (or shown in interactive mode Turnkey) */
+	UPROPERTY(EditAnywhere, Category="Packaging")
+	FString HelpText;
+
 	/** If this build step is specific to one or more platforms, add them here by name (note: use Windows, not Win64) */
 	UPROPERTY(EditAnywhere, Category="Packaging")
 	TArray<FString> SpecificPlatforms;
@@ -651,7 +655,7 @@ private:
 	TMap<FName, FName> PerPlatformTargetFlavorName;
 
 	/** Per platform build target */
-	UPROPERTY(config)
+	UPROPERTY(config, EditAnywhere, Category=Project)
 	TMap<FName, FString> PerPlatformBuildTarget;
 
 	/** Helper array used to mirror Blueprint asset selections across edits */
diff --git a/Engine/Source/Editor/TurnkeySupport/Private/TurnkeySupportModule.cpp b/Engine/Source/Editor/TurnkeySupport/Private/TurnkeySupportModule.cpp
--- a/Engine/Source/Editor/TurnkeySupport/Private/TurnkeySupportModule.cpp
+++ b/Engine/Source/Editor/TurnkeySupport/Private/TurnkeySupportModule.cpp
@@ -14,6 +14,7 @@
 #include "Widgets/Input/SEditableTextBox.h"
 #include "Misc/MessageDialog.h"
 #include "HAL/PlatformFileManager.h"
+#include "HAL/PlatformApplicationMisc.h"
 #include "ITargetDeviceServicesModule.h"
 #include "Misc/DataDrivenPlatformInfoRegistry.h"
 #include "Async/Async.h"
@@ -60,6 +61,9 @@
 DEFINE_LOG_CATEGORY(LogTurnkeySupport);
 #define LOCTEXT_NAMESPACE "FTurnkeySupportModule"
 
+
+#define ALLOW_CONTROL_TO_COPY_COMMANDLINE 0
+
 namespace 
 {
 	FCriticalSection GTurnkeySection;
@@ -104,6 +108,7 @@ protected:
 		return TEXT("-nocompileeditor -skipbuildeditor");
 	}
 
+
 	static bool ShowBadSDKDialog(FName IniPlatformName)
 	{
 		// Don't show the warning during automation testing; the dlg is modal and blocks
@@ -202,6 +207,16 @@ protected:
 public:
 
 
+	static FString GetLogAndReportCommandline(FString& LogFilename, FString& ReportFilename)
+	{
+		static int ReportIndex = 0;
+
+		LogFilename = FPaths::ConvertRelativePathToFull(FPaths::Combine(FPaths::ProjectIntermediateDir(), *FString::Printf(TEXT("TurnkeyLog_%d.log"), ReportIndex)));
+		ReportFilename = FPaths::ConvertRelativePathToFull(FPaths::Combine(FPaths::ProjectIntermediateDir(), *FString::Printf(TEXT("TurnkeyReport_%d.log"), ReportIndex++)));
+
+		return FString::Printf(TEXT("-ReportFilename=\"%s\" -log=\"%s\""), *ReportFilename, *LogFilename);
+	}
+
 	static void OpenProjectLauncher()
 	{
 		FGlobalTabmanager::Get()->TryInvokeTab(FTabId("ProjectLauncher"));
@@ -529,7 +544,7 @@ public:
 		FString CommandLine;
 		if (!ProjectPath.IsEmpty())
 		{
-			CommandLine.Appendf(TEXT("-ScriptsForProject=\"%s\" "), *ProjectPath);
+			CommandLine.Appendf(TEXT(" -ScriptsForProject=\"%s\" "), *ProjectPath);
 		}
 		CommandLine.Appendf(TEXT("Turnkey %s BuildCookRun %s"), *TurnkeyParams, *BuildCookRunParams);
 
@@ -546,25 +561,122 @@ public:
 		return true;
 	}
 
-	static void ExecuteCustomBuild(FName IniPlatformName, FProjectBuildSettings Build)
+	static FString GetCustomBuildCommandLine(FName IniPlatformName, const FString& DeviceId, const FProjectBuildSettings& Build)
 	{
-		const PlatformInfo::FTargetPlatformInfo* PlatformInfo = PlatformInfo::FindPlatformInfo(GetDefault<UProjectPackagingSettings>()->GetTargetFlavorForPlatform(IniPlatformName));
 		const FString ProjectPath = GetProjectPathForTurnkey();
 
 		FString CommandLine;
+		CommandLine.Appendf(TEXT("Turnkey -command=ExecuteBuild -build=\"%s\" -platform=%s"), *Build.Name, *ConvertToUATPlatform(IniPlatformName.ToString()));
 		if (!ProjectPath.IsEmpty())
 		{
-			CommandLine.Appendf(TEXT("-ScriptsForProject=\"%s\" "), *ProjectPath);
+			CommandLine.Appendf(TEXT(" -project=\"%s\""), *ProjectPath);
+		}
+		if (!DeviceId.IsEmpty())
+		{
+			CommandLine.Appendf(TEXT(" -device=\"%s\""), *ConvertToUATDeviceId(DeviceId));
+		}
+
+		// pass the editor setting down to UAT in case they aren't saved to the ini's that UAT would read
+		{
+			UProjectPackagingSettings* PackagingSettings = FTurnkeySupportCallbacks::GetPackagingSettingsForPlatform(IniPlatformName);
+			UProjectPackagingSettings* AllPlatformPackagingSettings = GetMutableDefault<UProjectPackagingSettings>();
+			bool bIsProjectBuildTarget = false;
+			const FTargetInfo* BuildTargetInfo = AllPlatformPackagingSettings->GetBuildTargetInfoForPlatform(IniPlatformName, bIsProjectBuildTarget);
+
+			CommandLine.Appendf(TEXT(" -overridetarget=%s"), *BuildTargetInfo->Name);
+
+			// distribution builds can set shipping 
+			EProjectPackagingBuildConfigurations BuildConfig = PackagingSettings->ForDistribution ? 
+				EProjectPackagingBuildConfigurations::PPBC_Shipping :
+				AllPlatformPackagingSettings->GetBuildConfigurationForPlatform(IniPlatformName);
+
+			// if PPBC_MAX is set, then the project default should be used instead of the per platform build config
+			if (BuildConfig == EProjectPackagingBuildConfigurations::PPBC_MAX)
+			{
+				BuildConfig = AllPlatformPackagingSettings->BuildConfiguration;
+			}
+			const UProjectPackagingSettings::FConfigurationInfo& ConfigurationInfo = UProjectPackagingSettings::ConfigurationInfo[(int)BuildConfig];
+			CommandLine.Appendf(TEXT(" -overrideconfiguration=%s"), LexToString(ConfigurationInfo.Configuration));
+
+
+			// get the chosen cook flavor (texture format, etc)
+			const PlatformInfo::FTargetPlatformInfo* PlatformInfo = PlatformInfo::FindPlatformInfo(AllPlatformPackagingSettings->GetTargetFlavorForPlatform(IniPlatformName));
+			if (PlatformInfo->IsFlavor())
+			{
+				CommandLine.Appendf(TEXT(" -overrideflavor=%s"), *PlatformInfo->PlatformFlavor.ToString());
+			}
 		}
-		CommandLine.Appendf(TEXT("Turnkey -command=ExecuteBuild -build=\"%s\" -platform=%s"),
-			*Build.Name, *IniPlatformName.ToString());
+
+		return CommandLine;
+	}
+
+	static void ExecuteCustomBuild(FName IniPlatformName, FString DeviceId, FProjectBuildSettings Build)
+	{
+		const PlatformInfo::FTargetPlatformInfo* PlatformInfo = PlatformInfo::FindPlatformInfo(GetDefault<UProjectPackagingSettings>()->GetTargetFlavorForPlatform(IniPlatformName));
+		const FString ProjectPath = GetProjectPathForTurnkey();
+
+		// throw this on the command before the actual automation name (Turnkey, BuildCOokRun, etc) because they are global-to-UAT options, not options for a single command
+		// and may not carry through when Turnkey passes along to the BCR command
+		FString CommandLine = TEXT("-utf8output");
 		if (!ProjectPath.IsEmpty())
 		{
-			CommandLine.Appendf(TEXT(" -project=\"%s\""), *ProjectPath);
+			CommandLine.Appendf(TEXT(" -ScriptsForProject=\"%s\" "), *ProjectPath);
+		}
+		CommandLine += GetCustomBuildCommandLine(IniPlatformName, DeviceId, Build);
+
+#if ALLOW_CONTROL_TO_COPY_COMMANDLINE
+		bool bIsControlHeld = FSlateApplication::Get().GetModifierKeys().IsControlDown();
+		if (bIsControlHeld)
+		{
+			CommandLine += TEXT(" -PrintOnly ");
+
+			FString LogFilename, ReportFilename;
+			CommandLine += GetLogAndReportCommandline(LogFilename, ReportFilename);
+
+			FTurnkeyEditorSupport::RunUAT(CommandLine, PlatformInfo->DisplayName, LOCTEXT("Turnkey_GettingCommandLine", "Copying Commandline"), LOCTEXT("Turnkey_CustomTaskName", "CustomCommandLine"), nullptr /* TaskIcon */,
+				[BuildName=Build.Name, ReportFilename](FString, double)
+				{
+					FString ReportContents;
+					FFileHelper::LoadFileToString(ReportContents, *ReportFilename);
+					UE_LOG(LogTurnkeySupport, Log, TEXT("Custom command '%s' gave commandline: %s"), *BuildName, *ReportContents);
+					if (ReportContents.Len())
+					{
+						FPlatformApplicationMisc::ClipboardCopy(*ReportContents);
+					}
+				});
 		}
+		else
+#endif
+		{
+			// append the stuff needed for running under editor
+			FString EditorSpecificCommandLine = FString::Printf(TEXT(" %s"), GetUATCompilationFlags());
+			if (FDerivedDataCacheInterface* DDC = GetDerivedDataCache())
+			{
+				EditorSpecificCommandLine += FString::Printf(TEXT(" -ddc=%s"), DDC->GetGraphName());
+			}
+
+			CommandLine += FString::Printf(TEXT(" -extraoptions=\"%s\""), *EditorSpecificCommandLine);
+
+			// if there is a Broese option, execute it now before passing to UAT to get the editor dialog
+			if (Build.BuildCookRunParams.Contains("{BrowseForDir}"))
+			{
+				// we reuse the already-presenbt StagingDirectory to save it, although it's more like "BuildOutpuDirectory" now 
+				UProjectPackagingSettings* AllPlatformPackagingSettings = GetMutableDefault<UProjectPackagingSettings>();
+				FString OutFolderName;
+				if (!FDesktopPlatformModule::Get()->OpenDirectoryDialog(FSlateApplication::Get().FindBestParentWindowHandleForDialogs(nullptr), LOCTEXT("PackageDirectoryDialogTitle", "Package project...").ToString(), AllPlatformPackagingSettings->StagingDirectory.Path, OutFolderName))
+				{
+					return;
+				}
+
+				AllPlatformPackagingSettings->StagingDirectory.Path = OutFolderName;
+				AllPlatformPackagingSettings->SaveConfig();
+
+				CommandLine += FString::Printf(TEXT(" -outputdir=\"%s\""), *OutFolderName);
+			}
 
-		FTurnkeyEditorSupport::RunUAT(CommandLine, PlatformInfo->DisplayName, LOCTEXT("Turnkey_CustomTaskNameVerbose", "Executing Custom Build"), LOCTEXT("Turnkey_CustomTaskName", "Custom"), FEditorStyle::GetBrush(TEXT("MainFrame.PackageProject")));
 
+			FTurnkeyEditorSupport::RunUAT(CommandLine, PlatformInfo->DisplayName, LOCTEXT("Turnkey_CustomTaskNameVerbose", "Executing Custom Build"), LOCTEXT("Turnkey_CustomTaskName", "Custom"), FEditorStyle::GetBrush(TEXT("MainFrame.PackageProject")));
+		}
 	}
 
 	static void SetPackageBuildConfiguration(const PlatformInfo::FTargetPlatformInfo* Info, EProjectPackagingBuildConfigurations BuildConfiguration)
@@ -747,7 +859,7 @@ static void TurnkeyInstallSdk(FString IniPlatformName, bool bPreferFull, bool bF
 	FString CommandLine;
 	if (!ProjectPath.IsEmpty())
 	{
-		CommandLine.Appendf(TEXT("-ScriptsForProject=\"%s\" "), *ProjectPath);
+		CommandLine.Appendf(TEXT(" -ScriptsForProject=\"%s\" "), *ProjectPath);
 	}
 	CommandLine.Appendf(TEXT("Turnkey -command=VerifySdk -UpdateIfNeeded -platform=%s %s %s -noturnkeyvariables -utf8output -WaitForUATMutex"), *IniPlatformName, *OptionalOptions, *ITurnkeyIOModule::Get().GetUATParams());
 
@@ -907,8 +1019,55 @@ static void FormatSdkInfo(const FString& PlatformOrDevice, const FTurnkeySdkInfo
 }
 
 
+static bool HasBuildForDeviceOrNot(const TArray<FProjectBuildSettings> Builds, bool bLookForDeviceBuilds)
+{
+	for (const FProjectBuildSettings& Build : Builds)
+	{
+		const bool bHasDeviceEntry = Build.BuildCookRunParams.Contains(TEXT("{DeviceId}"));
+		// if it has a device, and we are looking for device, or vice versa, return true
+		if (bHasDeviceEntry == bLookForDeviceBuilds)
+		{
+			return true;
+		}
+	}
+
+	return true;
+}
+
+static void MakeCustomBuildMenuEntries(FToolMenuSection& Section, const TArray<FProjectBuildSettings> Builds, FName IniPlatformName, const FString& DeviceId, const FString& DeviceName, const FText& ToolTipFormat)
+{
+	FString PlatformString = IniPlatformName.ToString();
+
+	for (FProjectBuildSettings Build : Builds)
+	{
+		if (Build.SpecificPlatforms.Num() == 0 || Build.SpecificPlatforms.Contains(PlatformString))
+		{
+			const bool bHasDeviceEntry = Build.BuildCookRunParams.Contains(TEXT("{DeviceId}"));
+			const bool bNeedsDeviceEntry = !DeviceId.IsEmpty();
+			if (bHasDeviceEntry == bNeedsDeviceEntry)
+			{
+				FString CommandLine = FTurnkeySupportCallbacks::GetCustomBuildCommandLine(IniPlatformName, DeviceId, Build);
+				FFormatNamedArguments Args;
+				Args.Add(TEXT("CommandLine"), FText::FromString(CommandLine));
+				Args.Add(TEXT("BuildName"), FText::FromString(Build.Name));
+				Args.Add(TEXT("BuildHelp"), FText::FromString(Build.HelpText));
+				Args.Add(TEXT("DeviceName"), FText::FromString(DeviceName));
 
+				Section.AddMenuEntry(
+					NAME_None,
+					FText::FromString(Build.Name),
+					FText::Format(ToolTipFormat, Args),
+					FSlateIcon(),
+					FUIAction(
+						FExecuteAction::CreateStatic(&FTurnkeySupportCallbacks::ExecuteCustomBuild, IniPlatformName, DeviceId, Build),
+						FCanExecuteAction::CreateStatic(&FTurnkeySupportCallbacks::CanExecuteCustomBuild, IniPlatformName, Build)
+					)
+				);
+			}
+		}
+	}
 
+}
 
 static void MakeTurnkeyPlatformMenu(UToolMenu* ToolMenu, FName IniPlatformName, ITargetDeviceServicesModule* TargetDeviceServicesModule)
 {
@@ -955,44 +1114,26 @@ static void MakeTurnkeyPlatformMenu(UToolMenu* ToolMenu, FName IniPlatformName,
 // 		);
 // 
 
-		FString PlatformString = IniPlatformName.ToString();
 		UProjectPackagingSettings* PackagingSettings = FTurnkeySupportCallbacks::GetPackagingSettingsForPlatform(IniPlatformName);
 
-		for (FProjectBuildSettings Build : PackagingSettings->EngineCustomBuilds)
-		{
-			if (Build.SpecificPlatforms.Num() == 0 || Build.SpecificPlatforms.Contains(PlatformString))
-			{
-				Section.AddMenuEntry(
-					NAME_None,
-					FText::FromString(Build.Name),
-					// @todo turnkey: add the build string to the tooltip
-					LOCTEXT("TurnkeyTooltip_EngineCustomBuild", "Execute a custom build"),
-					FSlateIcon(),
-					FUIAction(
-						FExecuteAction::CreateStatic(&FTurnkeySupportCallbacks::ExecuteCustomBuild, IniPlatformName, Build),
-						FCanExecuteAction::CreateStatic(&FTurnkeySupportCallbacks::CanExecuteCustomBuild, IniPlatformName, Build)
-					)
-				);
-			}
-		}
 
-		for (FProjectBuildSettings Build : PackagingSettings->ProjectCustomBuilds)
-		{
-			if (Build.SpecificPlatforms.Num() == 0 || Build.SpecificPlatforms.Contains(PlatformString))
-			{
-				Section.AddMenuEntry(
-					NAME_None,
-					FText::FromString(Build.Name),
-					// @todo turnkey: add the build string to the tooltip
-					LOCTEXT("TurnkeyTooltip_ProjectCustomBuild", "Execute a custom build (this comes from Packaging Settings)"),
-					FSlateIcon(),
-					FUIAction(
-						FExecuteAction::CreateStatic(&FTurnkeySupportCallbacks::ExecuteCustomBuild, IniPlatformName, Build),
-						FCanExecuteAction::CreateStatic(&FTurnkeySupportCallbacks::CanExecuteCustomBuild, IniPlatformName, Build)
-					)
-				);
-			}
-		}
+			
+		// these ToolTipFormats will be formatted inside the MakeCustomBuildMenuEntries function with some named args
+		// the empty strings passed in is the deviceId, which means to show builds that don't contain DeviceId components
+#if ALLOW_CONTROL_TO_COPY_COMMANDLINE
+		FText ToolTipFormat = LOCTEXT("TurnkeyTooltip_EngineCustomBuild_WithCopy", "Execute the '{BuildName}' command.\n{BuildHelp}\n--\nThis runs the following command:\nRunUAT {CommandLine}\nHold Control to copy the final underlying BuildCookRun commandline to the clipboard.");
+#else
+		FText ToolTipFormat = LOCTEXT("TurnkeyTooltip_EngineCustomBuild", "Execute the '{BuildName}' command.\n{BuildHelp}\n--\nThis runs the following command:\nRunUAT {CommandLine}");
+#endif
+		MakeCustomBuildMenuEntries(Section, PackagingSettings->EngineCustomBuilds, IniPlatformName, FString(), FString(), ToolTipFormat);
+		
+#if ALLOW_CONTROL_TO_COPY_COMMANDLINE
+		ToolTipFormat = LOCTEXT("TurnkeyTooltip_ProjectCustomBuild_WithCopy", "Execute this project's custom '{BuildName}' command.\n{BuildHelp}\n--\nThis runs the following command:\nRunUAT {CommandLine}\nHold Control to copy the final BuildCookRun commandline to the clipboard.\nThis custom command comes from Project Settings.");
+#else
+		ToolTipFormat = LOCTEXT("TurnkeyTooltip_ProjectCustomBuild", "Execute this project's custom '{BuildName}' command.\n{BuildHelp}\n--\nThis runs the following command:\nRunUAT {CommandLine}\nThis custom command comes from Project Settings.");
+#endif
+		MakeCustomBuildMenuEntries(Section, PackagingSettings->ProjectCustomBuilds, IniPlatformName, FString(), FString(), ToolTipFormat);
+
 
 		UProjectPackagingSettings* AllPlatformPackagingSettings = GetMutableDefault<UProjectPackagingSettings>();
 
@@ -1151,13 +1292,35 @@ static void MakeTurnkeyPlatformMenu(UToolMenu* ToolMenu, FName IniPlatformName,
 				NAME_None,
 				MakeSdkStatusAttribute(IniPlatformName, Proxy),
 				FText(),
-				FNewToolMenuDelegate::CreateLambda([IniPlatformName, DeviceName, DeviceId](UToolMenu* SubToolMenu)
+				FNewToolMenuDelegate::CreateLambda([IniPlatformName, DeviceName, DeviceId, PackagingSettings](UToolMenu* SubToolMenu)
 				{
+					if (HasBuildForDeviceOrNot(PackagingSettings->EngineCustomBuilds, true) || HasBuildForDeviceOrNot(PackagingSettings->ProjectCustomBuilds, true))
+					{
+						FToolMenuSection& CustomBuildSection = SubToolMenu->AddSection("DeviceCustomBuilds", LOCTEXT("TurnkeySection_DeviceCustomBuilds", "Builds For Devices"));
+
+						// these ToolTipFormats will be formatted inside the MakeCustomBuildMenuEntries function with some named args
+						// the empty strings passed in is the deviceId, which means to show builds that don't contain DeviceId components
+#if ALLOW_CONTROL_TO_COPY_COMMANDLINE
+						FText ToolTipFormat = LOCTEXT("TurnkeyTooltip_EngineCustomBuild_WithCopy", "Execute the '{BuildName}' command on {DeviceName}.\n{BuildHelp}\n--\nThis runs the following command:\nRunUAT {CommandLine}\nHold Control to copy the final underlying BuildCookRun commandline to the clipboard.");
+#else
+						FText ToolTipFormat = LOCTEXT("TurnkeyTooltip_EngineCustomBuild", "Execute the '{BuildName}' command on {DeviceName}.\n{BuildHelp}\n--\nThis runs the following command:\nRunUAT {CommandLine}");
+#endif
+						MakeCustomBuildMenuEntries(CustomBuildSection, PackagingSettings->EngineCustomBuilds, IniPlatformName, DeviceId, DeviceName, ToolTipFormat);
+
+#if ALLOW_CONTROL_TO_COPY_COMMANDLINE
+						ToolTipFormat = LOCTEXT("TurnkeyTooltip_ProjectCustomBuild", "Execute this project's custom '{BuildName}' command on {DeviceName}.\n{BuildHelp}\n--\nThis runs the following command:\nRunUAT {CommandLine}\nHold Control to copy the final BuildCookRun commandline to the clipboard.\nThis custom command comes from Project Settings.");
+#else
+						ToolTipFormat = LOCTEXT("TurnkeyTooltip_ProjectCustomBuild_WithCopy", "Execute this project's custom '{BuildName}' command on {DeviceName}.\n{BuildHelp}\n--\nThis runs the following command:\nRunUAT {CommandLine}\nThis custom command comes from Project Settings.");
+#endif
+						MakeCustomBuildMenuEntries(CustomBuildSection, PackagingSettings->ProjectCustomBuilds, IniPlatformName, DeviceId, DeviceName, ToolTipFormat);
+					}
+
+
 					FTurnkeySdkInfo SdkInfo = ITurnkeySupportModule::Get().GetSdkInfoForDeviceId(DeviceId);
 					FText SdkText, SdkToolTip;
 					FormatSdkInfo(DeviceId, SdkInfo, false, SdkText, SdkToolTip);
 
-					FToolMenuSection& Section = SubToolMenu->AddSection(NAME_None);
+					FToolMenuSection& Section = SubToolMenu->AddSection("DeviceSdkInfo", LOCTEXT("TurnkeySection_DeviceSdkInfo", "Sdk Info"));
 					Section.AddEntry(
 						FToolMenuEntry::InitWidget(
 							NAME_None,
@@ -1797,16 +1960,6 @@ void FTurnkeySupportModule::MakeTurnkeyMenu(FToolMenuSection& MenuSection) const
 	MenuSection.AddEntry(Entry);
 }
 
-static FString GetLogAndReportCommandline(FString& LogFilename, FString& ReportFilename)
-{
-	static int ReportIndex = 0;
-
-	LogFilename = FPaths::ConvertRelativePathToFull(FPaths::Combine(FPaths::ProjectIntermediateDir(), *FString::Printf(TEXT("TurnkeyLog_%d.log"), ReportIndex)));
-	ReportFilename = FPaths::ConvertRelativePathToFull(FPaths::Combine(FPaths::ProjectIntermediateDir(), *FString::Printf(TEXT("TurnkeyReport_%d.log"), ReportIndex++)));
-
-	return FString::Printf(TEXT("-ReportFilename=\"%s\" -log=\"%s\""), *ReportFilename, *LogFilename);
-}
-
 // some shared functionality
 static void PrepForTurnkeyReport(FString& BaseCommandline, FString& ReportFilename)
 {
@@ -1821,7 +1974,7 @@ static void PrepForTurnkeyReport(FString& BaseCommandline, FString& ReportFilena
 	}
 
 	FString LogFilename;
-	FString LogAndReportParams = GetLogAndReportCommandline(LogFilename, ReportFilename);
+	FString LogAndReportParams = FTurnkeySupportCallbacks::GetLogAndReportCommandline(LogFilename, ReportFilename);
 
 	BaseCommandline = BaseCommandline.Appendf(TEXT("Turnkey -utf8output -WaitForUATMutex -command=VerifySdk %s"), *LogAndReportParams);
 	// now pass a project to Turnkey
diff --git a/Engine/Source/Programs/AutomationTool/AutomationUtils/HostPlatform.cs b/Engine/Source/Programs/AutomationTool/AutomationUtils/HostPlatform.cs
--- a/Engine/Source/Programs/AutomationTool/AutomationUtils/HostPlatform.cs
+++ b/Engine/Source/Programs/AutomationTool/AutomationUtils/HostPlatform.cs
@@ -106,6 +106,11 @@ namespace AutomationTool
 		/// </summary>
 		abstract public UnrealBuildTool.UnrealTargetPlatform HostEditorPlatform { get; }
 
+		/// <summary>
+		/// Returns the type of the current running host platform
+		/// </summary>
+		public static UnrealBuildTool.UnrealTargetPlatform Platform { get => Current.HostEditorPlatform; }
+
 		/// <summary>
 		/// Returns the pdb file extenstion for the host platform.
 		/// </summary>
diff --git a/Engine/Source/Programs/AutomationTool/Turnkey/Commands/CreateBuild.cs b/Engine/Source/Programs/AutomationTool/Turnkey/Commands/CreateBuild.cs
new file mode 100644
--- /dev/null
+++ b/Engine/Source/Programs/AutomationTool/Turnkey/Commands/CreateBuild.cs
@@ -0,0 +1,807 @@
+// Copyright Epic Games, Inc. All Rights Reserved.
+
+using AutomationTool;
+using System;
+using System.Collections.Generic;
+using System.Linq;
+using System.IO;
+using System.Text.RegularExpressions;
+using EpicGames.Core;
+using UnrealBuildTool;
+using UnrealBuildBase;
+
+
+
+
+// Here's how we manage commandlines for builds/configs/targets. It has complexity so writing down here for reference
+//
+// Involved options:
+//		-client
+//		-server
+//			These denote Client or Server builds - both for executable and cooked data
+//		-target=Foo 
+//			Foo can be either say EngineTestClient, UnrealGame, etc. Note that UnrealGame is special for content only projects
+//			This option has historically not been used, instead being calculated from the above options. However with 
+//			more projects having multiple game/editor targets, it is useful to explicitly choose one of those targets
+//		-clientconfig=[Debug|Development|Test|Shipping]
+//			Used to choose a build configuration for client _and game_ executables
+//		-serverconfig
+//			Similar but for server executables
+//
+// In the case of conflicts (-client and -target=UnrealGame) UAT will fail, so we try to detect that earlier and disallow it
+//
+// Additionally, for code-based projects, 
+
+// However, given how builds are likely to be triggered by the editor, which can let user change the target/config/etc, when
+// ExecuteBuild handles the {ProjectPackagingSettings} replacement, it will first look at the existing commandline and not 
+// spceify a conflicting mode (commandline overrides PackagingSettings). 
+//  
+//
+
+namespace Turnkey.Commands
+{
+	internal class CreateBuild : TurnkeyCommand
+	{
+		protected override CommandGroup Group => CommandGroup.Builds;
+
+		const string SettingsSection = "/Script/UnrealEd.ProjectPackagingSettings";
+
+		internal static string[] SplitCommandLine(string CommandLine)
+		{
+			// handle -foo, -foo=bar, -foo="bar", and -foo="bar baz"
+			Regex CommandLineRegex = new Regex("(-\\w*\\s+)|(-\\w*=\\w*\\s+)|(-\\w*=\"[^\"]*\"\\s +)");
+			// split the args and clean up
+			return CommandLineRegex.Split(CommandLine).Select(x => x.Trim()).Where(x => x != "").ToArray();
+		}
+
+		string GetArchiveDirectory(string[] CommandOptions, FileReference ProjectFile)
+		{
+			ConfigHierarchy Config = ConfigCache.ReadHierarchy(ConfigHierarchyType.Game, ProjectFile.Directory, BuildHostPlatform.Current.Platform);
+
+			string StagingDirectory = ExecuteBuild.GetStructEntryForSetting(Config, SettingsSection, "StagingDirectory", "Path");
+			if (string.IsNullOrEmpty(StagingDirectory))
+			{
+				StagingDirectory = Path.Combine(ProjectFile.Directory.FullName, "Saved", "Archives");
+			}
+
+			string OutputDir = TurnkeyUtils.ParseParamValue("OutputDir", null, CommandOptions);
+			if (OutputDir == null)
+			{
+				//if (RuntimePlatform.IsWindows)
+				//{
+				//	string ChosenFile = null;
+				//	System.Threading.Thread t = new System.Threading.Thread(x =>
+				//	{
+				//		ChosenFile = UnrealWindowsForms.Utils.ShowOpenFileDialogAndReturnFilename("Project Files (*.uproject)|*.uproject");
+				//	});
+
+				//	t.SetApartmentState(System.Threading.ApartmentState.STA);
+				//	t.Start();
+				//	t.Join();
+
+				//	if (ChosenFile == null)
+				//	{
+				//		continue;
+				//	}
+				//}
+				//else
+				//{
+				//	while (true)
+				{
+					OutputDir = TurnkeyUtils.ReadInput("Enter archival directory:", StagingDirectory);
+				}
+				//}
+
+			}
+			return OutputDir;
+		}
+
+		private static string GetPlatformSetting(ConfigHierarchy Config, UnrealTargetPlatform Platform, string PerPlatformKey, string GenericKey, string Default)
+		{
+			string IniPlatformName = ConfigHierarchy.GetIniPlatformName(Platform);
+			string Value = ExecuteBuild.GetMapValueForSetting(Config, SettingsSection, PerPlatformKey, Platform.ToString());
+			if (Value == null && !string.IsNullOrEmpty(GenericKey))
+			{
+				Config.GetString(SettingsSection, GenericKey, out Value);
+			}
+			if (Value == null)
+			{
+				Value = Default;
+			}
+
+			return Value;
+		}
+
+		// hardcoded list of content-only, non-temp-target, targets
+		static Dictionary<TargetType, string> UETargets = new Dictionary<TargetType, string>()
+		{
+			{ TargetType.Game, "UnrealGame" },
+			{ TargetType.Client, "UnrealClient" },
+			{ TargetType.Server, "UnrealServer" },
+		};
+
+		public static string MakeCommandLineFromPackagingSettings(UnrealTargetPlatform Platform, FileReference ProjectFile, string ExistingCommandLine)
+		{
+			ConfigHierarchy Config = ConfigCache.ReadHierarchy(ConfigHierarchyType.Game, ProjectFile.Directory, Platform);
+
+			string[] ExistingOptions = SplitCommandLine(ExistingCommandLine);
+			bool bBuild = CommandUtils.ParseParam(ExistingOptions, "-build");
+			bool bCook = CommandUtils.ParseParam(ExistingOptions, "-cook") || CommandUtils.ParseParam(ExistingOptions, "-skipcook");
+			bool bStage = CommandUtils.ParseParam(ExistingOptions, "-stage") || CommandUtils.ParseParam(ExistingOptions, "-skipstage");
+			bool bPackage = CommandUtils.ParseParam(ExistingOptions, "-package") || CommandUtils.ParseParam(ExistingOptions, "-skippackage");
+			bool bArchive = CommandUtils.ParseParam(ExistingOptions, "-archive");
+			bool bClient = CommandUtils.ParseParam(ExistingOptions, "-client");
+			bool bServer = CommandUtils.ParseParam(ExistingOptions, "-server");
+
+
+			// for targetname and configuration, the priority is -overridetarget/config (coming from editor), then look in passed in
+			// commandline, and finally look in the ini
+			string IniTargetName = GetPlatformSetting(Config, Platform, "PerPlatformBuildTarget", "BuildTarget", null);
+			string TargetName = CommandUtils.ParseParamValue(ExistingOptions, "-target", IniTargetName);
+			TargetName = TurnkeyUtils.ParseParamValue("OverrideTarget", TargetName, new string[] { });
+
+			ProjectProperties Props = ProjectUtils.GetProjectProperties(ProjectFile);
+			List<SingleTargetProperties> DetectedTargets = Props.Targets;
+
+
+
+			//if (CommandUtils.IsNullOrEmpty(DetectedTargets))
+			//{
+			//	DetectedTargets = new List<SingleTargetProperties>
+			//	{
+			//		new SingleTargetProperties { TargetName = "UnrealGame", Rules = new (new TargetInfo("UnrealGame", Platform, UnrealTargetConfiguration.Development, null)) }
+			//	}
+			//	ProjectParams Params = new ProjectParams(Command: TurnkeyUtils.CommandUtilHelper, RawProjectPath: ProjectFile, Client: true, DedicatedServer: true);
+			//	DetectedTargets = Params.ProjectTargets;
+			//}
+
+			// figure out the target/targetype to use 
+			TargetRules ChosenTarget = null;
+
+			// first, handle the special case where the editor passed down a UETarget but UBT knows it needs a generated target - swap for the generated one
+			if (Props.bWasGenerated && !string.IsNullOrEmpty(TargetName) && 
+				UETargets.Values.Contains(TargetName, StringComparer.InvariantCultureIgnoreCase))
+			{
+				// generated projects likely only have one game target, so use it
+				if (Props.Targets.Count == 1)
+				{
+					ChosenTarget = Props.Targets[0].Rules;
+				}
+				else
+				{
+					// get the type of the UE target
+					TargetType UETargetType = UETargets.FirstOrDefault(x => x.Value.Equals(TargetName, StringComparison.InvariantCultureIgnoreCase)).Key;
+					ChosenTarget = Props.Targets.FirstOrDefault(x => x.Rules.Type == UETargetType)?.Rules;
+				}
+			}
+
+			if (ChosenTarget == null && !CommandUtils.IsNullOrEmpty(Props.Targets))
+			{
+				if (!string.IsNullOrEmpty(TargetName))
+				{
+					ChosenTarget = DetectedTargets.FirstOrDefault(x => x.TargetName.Equals(TargetName, StringComparison.InvariantCultureIgnoreCase))?.Rules;
+					if (ChosenTarget == null)
+					{
+						throw new AutomationException($"The target specified on the commandline, {TargetName}, could not be found in project {ProjectFile}" +
+							(Props.bWasGenerated ? $" Note that this project requires a temporary target, so {ProjectFile.GetFileNameWithoutAnyExtensions()} is probably the target you want." : ""));
+					}
+				}
+				// if we don't have one now, find one based on params
+				if (ChosenTarget == null)
+				{
+					ChosenTarget = DetectedTargets.FirstOrDefault(x => (bClient && x.Rules.Type == TargetType.Client) || (bServer && x.Rules.Type == TargetType.Server) || (!bClient && !bServer && x.Rules.Type == TargetType.Game))?.Rules;
+				}
+			}
+
+			// now figure out the final type to use, with specified target, specified flags, and detected targets
+			TargetType FinalTargetType = TargetType.Game;
+			string FinalTargetName = "UnrealGame";
+			if (ChosenTarget != null)
+			{
+				FinalTargetType = ChosenTarget.Type;
+				FinalTargetName = ChosenTarget.Name;
+			}
+			else if (bClient || TargetName.Equals("UnrealClient", StringComparison.InvariantCultureIgnoreCase))
+			{
+				FinalTargetType = TargetType.Client;
+				FinalTargetName = "UnrealClient";
+			}
+			else if (bServer || TargetName.Equals("UnrealServer", StringComparison.InvariantCultureIgnoreCase))
+			{
+				FinalTargetType = TargetType.Server;
+				FinalTargetName = "UnrealServer";
+			}
+
+			// now that we have a target, we can use it to look on the commandline for server or clientconfig
+			string IniBuildConfig = GetPlatformSetting(Config, Platform, "PerPlatformBuildConfig", "BuildConfiguration", "Development");
+			string BuildConfig = CommandUtils.ParseParamValue(ExistingOptions, FinalTargetType == TargetType.Server ? "serverconfig" : "clientconfig", IniBuildConfig);
+			BuildConfig = TurnkeyUtils.ParseParamValue("OverrideConfiguration", BuildConfig, new string[] { });
+			BuildConfig = BuildConfig.Replace("PPBC_", "");
+
+			string CommandLine = "";
+
+			Func<string, bool> GetBool = (string Setting) =>
+			{
+				bool bValue;
+				if (Config.GetBool(SettingsSection, Setting, out bValue))
+				{
+					return bValue;
+				}
+				return false;
+			};
+
+			Action<string, string> AddIfTrue = (string Setting, string ToAdd) =>
+			{
+				if (GetBool(Setting))
+				{
+					CommandLine += $" {ToAdd}";
+				}
+			};
+
+			Action<string> AddIfNotAlready = (string ToAdd) =>
+			{
+				string[] Tokens = ToAdd.Split('=');
+				// if we are a -option, look for mach, if we are a -foo= option, look for a -foo= option in the 
+				if ((Tokens.Count() == 1 && !ExistingOptions.Contains(Tokens[0], StringComparer.InvariantCultureIgnoreCase)) ||
+					(Tokens.Count() == 2 && !ExistingOptions.Any(x => x.StartsWith(Tokens[0] + "=", StringComparison.InvariantCultureIgnoreCase))))
+				{
+					CommandLine += $" {ToAdd}";
+				}
+			};
+
+
+			switch (FinalTargetType)
+			{
+				case TargetType.Client:
+					AddIfNotAlready($"-clientconfig={BuildConfig}");
+					break;
+				case TargetType.Server:
+					AddIfNotAlready("-noclient");
+					AddIfNotAlready($"-serverconfig={BuildConfig}");
+					break;
+				case TargetType.Game:
+					AddIfNotAlready($"-clientconfig={BuildConfig}");
+					break;
+				default: throw new AutomationException($"TypeType {FinalTargetType} is not allowed for packaging at this time");
+			}
+
+			AddIfNotAlready($"-target={FinalTargetName}");
+
+			if (bBuild)
+			{
+				AddIfTrue("FullRebuild", "-clean");
+			}
+
+			if (bCook)
+			{
+				AddIfTrue("bUseZenStore", "-zenstore");
+				AddIfTrue("bSkipEditorContent", "-SkipCookingEditorContent");
+			}
+			
+			if (bCook || bStage)
+			{
+				AddIfTrue("bGenerateChunks", "-manifests");
+			}
+
+			if (bCook || bStage || bPackage)
+			{
+				string CookFlavor = GetPlatformSetting(Config, Platform, "PerPlatformTargetFlavorName", null, null);
+				CookFlavor = TurnkeyUtils.ParseParamValue("OverrideFlavor", CookFlavor, new string[] { });
+				if (!string.IsNullOrEmpty(CookFlavor))
+				{
+					string IniPlatformName = ConfigHierarchy.GetIniPlatformName(Platform);
+					// remove the Platform_ from the flavor if there, so Android_ASTC would become just Android
+					CookFlavor = CookFlavor.Replace($"{IniPlatformName}_", "");
+					AddIfNotAlready($"-cookflavor={CookFlavor}");
+				}
+			}
+
+			if (bStage)
+			{
+				// zenstore precludes iostore when staging (it has already done the work)
+				if (!GetBool("bUseZenStore"))
+				{
+					AddIfTrue("bUseIoStore", "-iostore -pak");
+				}
+				AddIfTrue("UsePakFile", "-pak");
+				
+				// more complex settings below
+				if (Platform == UnrealTargetPlatform.Win64)
+				{
+					AddIfTrue("IncludePrerequisites", "-prereqs");
+
+					//string AppLocalDirectory = ExecuteBuild.GetStructEntryForSetting(Config, SettingsSection, "ApplocalPrerequisitesDirectory", "Path");
+					//if (!string.IsNullOrEmpty(AppLocalDirectory))
+					//{
+					//	CommandLine += $" -applocaldirectory=\"{AppLocalDirectory}";
+					//}
+					//else
+					//{
+					//	AddIfTrue("IncludeAppLocalPrerequisites", $" -applocaldirectory=\"$(EngineDir)/Binaries/ThirdParty/AppLocalDependencies\"");
+					//}
+				}
+
+				DataDrivenPlatformInfo.ConfigDataDrivenPlatformInfo DDPI = DataDrivenPlatformInfo.GetDataDrivenInfoForPlatform(Platform);
+				if (DDPI.bCanUseCrashReporter)
+				{
+					AddIfTrue("IncludeCrashReporter", "-CrashReporter");
+				}
+
+				bool bBuildHttpChunkInstallData;
+				Config.GetBool(SettingsSection, "bBuildHttpChunkInstallData", out bBuildHttpChunkInstallData);
+				if (bBuildHttpChunkInstallData)
+				{
+					string HttpChunkInstallDataDirectory = ExecuteBuild.GetStructEntryForSetting(Config, SettingsSection, "HttpChunkInstallDataDirectory", "Path");
+					string HttpChunkInstallDataVersion;
+					Config.GetString(SettingsSection, "HttpChunkInstallDataVersion", out HttpChunkInstallDataVersion);
+
+					CommandLine += $" -manifests -createchunkinstall -chunkinstalldirectory=\"{HttpChunkInstallDataDirectory}\" -chunkinstallversion={HttpChunkInstallDataVersion}";
+				}
+			}
+
+			if (bStage || bPackage)
+			{
+				AddIfTrue("ForDistribution", "-distribution");
+			}
+
+			return CommandLine;
+		}
+
+		static string MakeCommandLineInteractively(List<UnrealTargetPlatform> Platforms, FileReference ProjectFile, string ExistingCommandLine, bool bBuild, bool bCook, bool bStage, bool bPackage, bool bArchive)
+		{
+			string CommandLine = "";
+
+			if (bBuild || bStage || bPackage)
+			{
+				List<string> Configs = new List<string> { "Debug", "Development", "Test", "Shipping" };
+				int BinaryConfigChoice = TurnkeyUtils.ReadInputInt("Select the compiled game configuration:", Configs, false, 1);
+				string BuildConfig = Configs[BinaryConfigChoice];
+
+				//switch (TargetType)
+				//{
+				//	case TargetType.Client: CommandLine += $" -client -clientconfig={BuildConfig}"; break;
+				//	case TargetType.Server: CommandLine += $" -server -noclient -serverconfig={BuildConfig}"; break;
+				//	case TargetType.Game: CommandLine += $" -clientconfig={BuildConfig}"; break;
+				//	default: throw new AutomationException($"TypeType {TargetType} is not allowed for packaging at this time");
+				//}
+			}
+
+			if (bBuild)
+			{
+				if (TurnkeyUtils.GetUserConfirmation("Would you like to perform a full rebuild of all programs? [-clean]", false))
+				{
+					CommandLine += " -clean";
+				}
+			}
+
+			if (bCook)
+			{
+				if (TurnkeyUtils.GetUserConfirmation("Would you like to skip cooking editor content? [-SkipCookingEditorContent]", false))
+				{
+					CommandLine += " -SkipCookingEditorContent";
+				}
+			}
+
+			if (bStage || bCook)
+			{
+				if (TurnkeyUtils.GetUserConfirmation("Would you like to split your data up into chunks? This requires a UAssetManager (???) subclass to manage chunk assignments [-manifests]", false))
+				{
+					CommandLine += " -manifests";
+
+					//					if (TurnkeyUtils.GetUserConfirmation("Would you like to generate Http chunk install data? [-createchunkinstall]", false))
+					//					{
+					//						ConfigHierarchy Config = ConfigCache.ReadHierarchy(ConfigHierarchyType.Game, ProjectFile.Directory, BuildHostPlatform.Current.Platform);
+
+					//						CommandLine += " -createchunkinstall";
+					//						TurnkeyUtils.ReadInput("Entier the HttpChunkInstallDirectory")
+					////						BuildCookRunParams += FString::Printf(TEXT(" -manifests -createchunkinstall -chunkinstalldirectory=\"%s\" -chunkinstallversion=%s"), *(PackagingSettings->HttpChunkInstallDataDirectory.Path), *(PackagingSettings->HttpChunkInstallDataVersion));
+					//					}
+
+				}
+
+
+				int ContainerModeChoice = TurnkeyUtils.ReadInputInt("How would you like your files collected?", new List<string>
+						{ "Optimized containers when staging [-iostore -pak]", "[Experimental] Optimized containers when cooking [-zenstore -pak]", "Left as fully loose files"}, false, 0);
+				switch (ContainerModeChoice)
+				{
+					case 0: CommandLine += " -iostore -pak"; break;
+					case 1: CommandLine += " -zenstore -pak"; break;
+					case 2: CommandLine += ""; break;
+				}
+			}
+
+			if (bStage)
+			{
+				bool bCanUseCrashReporter = false;
+				foreach (UnrealTargetPlatform Platform in Platforms)
+				{
+					DataDrivenPlatformInfo.ConfigDataDrivenPlatformInfo DDPI = DataDrivenPlatformInfo.GetDataDrivenInfoForPlatform(Platform);
+					bCanUseCrashReporter = bCanUseCrashReporter || DDPI.bCanUseCrashReporter;
+				}
+				if (bCanUseCrashReporter)
+				{
+					if (TurnkeyUtils.GetUserConfirmation("Would you like to include the crash reporter? [-CrashReporter]", true))
+					{
+						CommandLine += " -CrashReporter";
+					}
+				}
+
+				if (Platforms.Contains(UnrealTargetPlatform.Win64))
+				{
+					if (TurnkeyUtils.GetUserConfirmation("Would you like to include the Windows Prerequisites installer? [-prereqs]", false))
+					{
+						CommandLine += " -prereqs";
+					}
+					//int ContainerModeChoice = TurnkeyUtils.ReadInputInt("How would you like your files collected?", new List<string>
+					//	{ "Optimized containers when staging [-iostore -pak]", "[Experimental] Optimized containers when cooking [-zenstore -pak]", "Left as fully loose files"}, false, 0);
+				}
+			}
+
+			if (bStage || bPackage)
+			{
+				if (TurnkeyUtils.GetUserConfirmation("Would you create build appropriate for distribution to end users? [-distribution]", false))
+				{
+					CommandLine += " -distribution";
+				}
+			}
+
+			return CommandLine;
+		}
+
+		const string DividerLine = "================================================\n";
+
+		protected override void Execute(string[] CommandOptions)
+		{
+			int BuildChoice = TurnkeyUtils.ReadInputInt("What kind of build would you like to create?", new List<string> { "One-off build to execute now", "Save a build to project settings to execute in Editor or Turnkey" }, true, 1);
+			if (BuildChoice == 0)
+			{
+				return;
+			}
+			bool bMakingSavedBuild = BuildChoice == 2;
+
+			int UserMode = TurnkeyUtils.ReadInputInt("Choose an interactive mode to use:", new List<string> { "Simple Mode: Running a game (locally, remote device, etc)", "Simple Mode: Making a packaged up build (install later, share with others, etc)", "Advanced mode: Walk through all questions" }, true, 2);
+			if (UserMode == 0)
+			{
+				return;
+			}
+			bool bAskForDevice = UserMode == 1 || UserMode == 3;
+			bool bAskForPackage = UserMode == 2 || UserMode == 3;
+
+			// get the project, always needed
+			FileReference ProjectFile = TurnkeyUtils.GetProjectFromCommandLineOrUser(CommandOptions);
+
+			// we need a platform to execute
+			List<UnrealTargetPlatform> Platforms = null;
+			FileReference DefaultConfig = null;
+
+			if (bMakingSavedBuild)
+			{
+				Platforms = new List<UnrealTargetPlatform> { HostPlatform.Platform };
+
+				DefaultConfig = ConfigCache.GetDefaultConfigFileReference(ConfigHierarchyType.Game, ProjectFile.Directory);
+				if (!FileReference.Exists(DefaultConfig))
+				{
+					TurnkeyUtils.Log($"Unable to find DefaultGame.ini for {ProjectFile.GetFileNameWithoutExtension()}. We will continue, but just print out the info at the end.");
+					DefaultConfig = null;
+				}
+				else
+				{
+					// make sure the file is writable
+					while (DefaultConfig != null && File.GetAttributes(DefaultConfig.FullName).HasFlag(FileAttributes.ReadOnly))
+					{
+						string ReadOnlyPrompt = "The project's DefaultGame.ini is read-only, unable to save. What would you like to do?";
+						List<string> Options = new List<string> { "Try again", "Continue anyway, printout info at the end", "Make file writeable" };
+						if (CommandUtils.P4Enabled)
+						{
+							Options.Add("Checkout with Perforce");
+						}
+						else
+						{
+							ReadOnlyPrompt += "\n(Try re-running UAT with -p4 to enable Perforce to allow Turnkey to checkout for you)";
+						}
+
+						int Choice = TurnkeyUtils.ReadInputInt(ReadOnlyPrompt, Options, true, 2);
+						if (Choice == 0)
+						{
+							return;
+						}
+						else if (Choice == 2)
+						{
+							DefaultConfig = null;
+							break;
+						}
+						else if (Choice == 3)
+						{
+							// clear read-only flag
+							File.SetAttributes(DefaultConfig.FullName, File.GetAttributes(DefaultConfig.FullName) & ~FileAttributes.ReadOnly);
+						}
+						else if (Choice == 4)
+						{
+							TurnkeyUtils.Log("Attempting to create a p4 changelist and check out DefaultGame.ini...");
+							try
+							{
+								int CL = CommandUtils.P4.CreateChange(CommandUtils.P4Env.Client, $"[AUTO-GENERATED] - Adding a custom project build to DefaultGame.ini");
+								CommandUtils.P4.Edit(CL, CommandUtils.MakePathSafeToUseWithCommandLine(DefaultConfig.FullName));
+								TurnkeyUtils.Log($"Created changelist {CL}. Use P4 to check in when ready!");
+							}
+							catch  (Exception Ex)
+							{
+								TurnkeyUtils.Log($"P4 error: {Ex.Message}");
+							}
+						}
+					}
+				}
+			}
+			else
+			{
+				Platforms = TurnkeyUtils.GetPlatformsFromCommandLineOrUser(CommandOptions, null);
+			}
+
+			if (Platforms.Count == 0 || ProjectFile == null)
+			{
+				return;
+			}
+
+			//if (TurnkeyUtils.ParseParam("testconfig", CommandOptions))
+			//{
+			//	ConfigCache.WriteSettingToDefaultConfig(ConfigHierarchyType.Game, ProjectFile.Directory, ConfigCache.ConfigDefaultUpdateType.SetValue, "SectionTestA", "InsertedValue", "True");
+			//	ConfigCache.WriteSettingToDefaultConfig(ConfigHierarchyType.Game, ProjectFile.Directory, ConfigCache.ConfigDefaultUpdateType.SetValue, "SectionTestA", "Foo", "SetValue");
+			//	ConfigCache.WriteSettingToDefaultConfig(ConfigHierarchyType.Game, ProjectFile.Directory, ConfigCache.ConfigDefaultUpdateType.AddArrayEntry, "SectionTestB", "Builds", "(Blah=123123)");
+			//	ConfigCache.WriteSettingToDefaultConfig(ConfigHierarchyType.Game, ProjectFile.Directory, ConfigCache.ConfigDefaultUpdateType.SetValue, "SectionTest2", "InsertedValue", "True");
+			//	ConfigCache.WriteSettingToDefaultConfig(ConfigHierarchyType.Game, ProjectFile.Directory, ConfigCache.ConfigDefaultUpdateType.AddArrayEntry, "SectionTest1", "Builds", "(Blah=123123)");
+			//	ConfigCache.WriteSettingToDefaultConfig(ConfigHierarchyType.Game, ProjectFile.Directory, ConfigCache.ConfigDefaultUpdateType.SetValue, "SectionTest1", "Foo2", "SetValue2");
+			//	return;
+			//}
+			 
+			bool bBuild = false, bCook = false, bStage = false, bDeploy = false, bRun = false, bPackage = false, bArchive = false;
+			string TargetType = "Game";
+			string TargetName = null;
+//			string BuildConfig = "Development";
+
+			if (UserMode == 1)
+			{
+				bBuild = true;
+				bCook = true;
+				bStage = true;
+				bDeploy = true;
+				bRun = true;
+			}
+			else if (UserMode == 2)
+			{
+				bBuild = true;
+				bCook = true;
+				bStage = true;
+				bPackage = true;
+				bArchive = TurnkeyUtils.GetUserConfirmation("Would you like to copy the output to an archival location? [-archive]", false);
+			}
+			else
+			{
+				ProjectProperties Props = ProjectUtils.GetProjectProperties(ProjectFile);
+//				int TargetTypeChoice = TurnkeyUtils.ReadInputInt("Select your target type:", new List<string> { "Game (Game + Server)", "Client (No Server)", "Server (No Game)" }, false, 0);
+//				TargetType = TargetTypeChoice == 0 ? "Game" : TargetTypeChoice == 1 ? "Client" : "Server";
+
+				int TargetTypeChoice = TurnkeyUtils.ReadInputInt("Select your target:", Props.Targets.Select(x => x.TargetName).ToList(), true, Props.Targets.IndexOf(Props.Targets.First(x => x.Rules.Type == UnrealBuildTool.TargetType.Game)) + 1);
+				if (TargetTypeChoice == 0)
+				{
+					return;
+				}
+				TargetName = 
+				TargetType = Props.Targets[TargetTypeChoice - 1].Rules.Type.ToString();
+
+				bBuild = TurnkeyUtils.GetUserConfirmation("Would you like to compile necessary programs? [-build]", true);
+				bCook = TurnkeyUtils.GetUserConfirmation("Would you like to cook your project's assets? [-cook]", true);
+				bStage = TurnkeyUtils.GetUserConfirmation("Would you like gather all files into a staging directory? [-stage]", true);
+
+				bPackage = TurnkeyUtils.GetUserConfirmation("Would you like to generate a final platform-specific package? [-package]", true);
+				bArchive = TurnkeyUtils.GetUserConfirmation("Would you like to copy the output to an archival location? [-archive]", false);
+
+				bDeploy = TurnkeyUtils.GetUserConfirmation("Would you like to copy the gathered files to a device? [-deploy]", !bPackage);
+				bRun = TurnkeyUtils.GetUserConfirmation("Would you like to run the game? [-run]", bDeploy);
+			}
+
+			string BCRParams = "";
+			if (!bMakingSavedBuild)
+			{
+				BCRParams = $"-project=\"{ProjectFile}\"";
+			}
+
+			BCRParams +=
+				  (bBuild ? " -build" : "")
+				+ (bCook ? " -cook" : " -skipcook")
+				+ (bStage ? " -stage" : " -skipstage")
+				+ (bPackage ? " -package" : " -skippackage")
+				+ (bArchive ? " -archive" : "")
+				+ (bDeploy ? " -deploy" : "")
+				+ (bRun ? " -run" : "")
+				;
+
+			if (bArchive)
+			{
+				// when saving, use a placeholder to ask later
+				BCRParams += " -archivedirectory=" + (bMakingSavedBuild ? "{BrowseForDir}" : GetArchiveDirectory(CommandOptions, ProjectFile));
+			}
+
+			string PPSettings = MakeCommandLineFromPackagingSettings(HostPlatform.Platform, ProjectFile, BCRParams);
+			string PPPrompt = DividerLine + $"The options so far are:\n  {BCRParams}\nThe project's packaging settings (for {Platforms[0]}) are:\n  {PPSettings}\n" +
+				"How would you like to use the packaging settings? Not all may apply to this build and will be ignored.";
+			if (bMakingSavedBuild)
+			{
+//				PPPrompt += "\nNote that different platforms "
+			}
+			List<string> PPOptions = new List<string>
+			{
+				"Do not add any of the packaging settings",
+				"Interactively choose packaging settings",
+				"Use the packaging settings shown above",
+			};
+			if (bMakingSavedBuild)
+			{
+				PPOptions.Add("Use what the settings will be when this command is run later");
+			}
+			int ExtendedParamsChoice = TurnkeyUtils.ReadInputInt(PPPrompt, PPOptions, false, PPOptions.Count - 1);
+			//bool bUsePackagingSettings = TurnkeyUtils.GetUserConfirmation("Would you like to further control the build with your project's packaging settings? [-xyz]", true);
+
+			string ExtendedParams = "";
+			if (ExtendedParamsChoice == 1)
+			{
+				ExtendedParams = MakeCommandLineInteractively(Platforms, ProjectFile, BCRParams, bBuild, bCook, bStage, bPackage, bArchive);
+			}
+			else  if (ExtendedParamsChoice == 3)
+			{
+				// encode the targettype into the tag for later use
+				ExtendedParams = " {ProjectPackagingSettings}";
+			}
+
+			string LastPlatformParams = "";
+			string AddCmdLine = "";
+			foreach (UnrealTargetPlatform Platform in Platforms)
+			{
+				string PlatformParams = BCRParams;
+
+				// if we are making a saved build, then leave the {Platform} tag in the string, otherwise, put in the actual platform
+				if (bMakingSavedBuild)
+				{
+					PlatformParams += " -platform={Platform}";
+				}
+				else
+				{
+					PlatformParams += " -platform=" + Platform;
+				}
+
+				if (ExtendedParamsChoice == 2)
+				{
+					ExtendedParams = MakeCommandLineFromPackagingSettings(Platform, ProjectFile, BCRParams + ExtendedParams + PlatformParams);
+				}
+
+
+				PlatformParams += ExtendedParams;
+				if (TurnkeyUtils.GetUserConfirmation($"Would you like to add anymore options to the build commandline for {Platform}?\nCurrent commandline is:\n   BuildCookRun {PlatformParams}", false))
+				{
+					LastPlatformParams = TurnkeyUtils.ReadInput("Enter extra options:", LastPlatformParams);
+					PlatformParams += " " + LastPlatformParams;
+				}
+
+				// the commandline may be saved to a commandline.txt, or we may be running with it drectly (See jira UE-99467 for upcoming work that will affect this)
+				if (bStage || bRun)
+				{
+					AddCmdLine = TurnkeyUtils.ReadInput($"Enter any commandline options you want to add to the game [-addcmdline]", AddCmdLine);
+					if (!string.IsNullOrEmpty(AddCmdLine))
+					{
+						PlatformParams += $" -addcmdline=\"{AddCmdLine}\"";
+					}
+				}
+
+				if (bDeploy || bRun)
+				{
+					if (bMakingSavedBuild)
+					{
+						PlatformParams += " -device={DeviceId}";
+					}
+					else
+					{
+						List<DeviceInfo> Devices = TurnkeyUtils.GetDevicesFromCommandLineOrUser(CommandOptions, Platform);
+						if (Devices == null || Devices.Count == 0)
+						{
+							// we need a device here, if canceled, quit out
+							return;
+						}
+						PlatformParams += " -device=" + string.Join("+", Devices);
+					}
+				}
+
+				TurnkeyUtils.Log($"Final Command:\nRunUAT BuildCookRun {PlatformParams}");
+
+//				if (TurnkeyUtils.GetUserConfirmation("Do you want to save this to DefaultGame.ini for use in the editor?", false))
+				if (bMakingSavedBuild)
+				{
+					TurnkeyUtils.Log(DividerLine + $"We are now ready to save this build to DefaultGame.ini!");
+
+					List<string> ExistingBuilds;
+					ConfigHierarchy Config = ConfigCache.ReadHierarchy(ConfigHierarchyType.Game, ProjectFile.Directory, BuildHostPlatform.Current.Platform);
+					Config.GetArray(SettingsSection, "ProjectCustomBuilds", out ExistingBuilds);
+
+					string Name;
+					if (ExistingBuilds != null && ExistingBuilds.Count > 0)
+					{
+						TurnkeyUtils.Log($"Here are the existing builds your project already has (and their descrptions):");
+						foreach (string Build in ExistingBuilds)
+						{
+							string BuildName = ExecuteBuild.GetStructEntry(Build, "Name", false);
+							string BuildHelp = ExecuteBuild.GetStructEntry(Build, "HelpText", false);
+							TurnkeyUtils.Log($"  {BuildName} - {BuildHelp}");
+						}
+						Name = TurnkeyUtils.ReadInput("Enter unique name for this build, it must be different than the above names", "");
+					}
+					else
+					{
+						Name = TurnkeyUtils.ReadInput("Enter a name for this build", "");
+					}
+
+					if (!string.IsNullOrEmpty(Name))
+					{
+						string HelpText = TurnkeyUtils.ReadInput("Enter the help/description for this build", "");
+						string RestrictedPlatforms = TurnkeyUtils.ReadInput("If you want to restrict this to certain plaforms, enter a comma separated list of platform names:", "");
+
+						if (!string.IsNullOrEmpty(RestrictedPlatforms))
+						{
+							IEnumerable<string> PlatformNames = RestrictedPlatforms.Split(", +".ToCharArray()).ToList()
+								// fixup windows, and wrap in quotes
+								.Select(x => (x.Equals("Win64", StringComparison.InvariantCultureIgnoreCase) ? "Windows" : x))
+								.Select(x => $"\"{x}\"");
+							RestrictedPlatforms = "(" + string.Join(",", PlatformNames) + ")";
+						}
+
+						//// fix up the commandline to be appropriate (remove certain options, like project, platform, and archive directory)
+						//Regex CommandLineRegex = new Regex("(-\\w*\\s+)|(-\\w*=\\w*\\s+)|(-\\w*=\"[^\"]*\"\\s +)");
+						//// split the args and clean up
+						//List<string> IndividualArgs = CommandLineRegex.Split(PlatformParams).Select(x => x.Trim()).Where(x => x != "").ToList();
+
+						//// remove some known bad params
+						//IndividualArgs = IndividualArgs.Where(x => !x.StartsWith("-project=", StringComparison.InvariantCultureIgnoreCase)).ToList();
+
+						//// update some params to have variables
+						//for (int Index = 0; Index < IndividualArgs.Count; Index++)
+						//{
+						//	if (IndividualArgs[Index].StartsWith("-platform=", StringComparison.InvariantCultureIgnoreCase))
+						//	{
+						//		IndividualArgs[Index] = "-platform={Platform}";
+						//	}
+						//	else if (IndividualArgs[Index].StartsWith("-archivedirectory=", StringComparison.InvariantCultureIgnoreCase))
+						//	{
+						//		IndividualArgs[Index] = "-archivedirectory={BrowseForDir}";
+						//	}
+						//}
+
+
+						// generate struct string for ini file
+						string ValueText = $"(Name=\"{Name}\",HelpText=\"{HelpText}\",SpecificPlatforms={RestrictedPlatforms},BuildCookRunParams=\"{PlatformParams}\")";
+
+						// just print out the text to put into the DefaultGame.ini later
+						if (DefaultConfig == null)
+						{
+							TurnkeyUtils.Log($"Since DefaultGame.ini isn't writeable, here is the line you can manually add to DefaultGame.ini in the [{SettingsSection}] section to add the build to editor and Turnkey's ExecuteBuild command:\n  {ValueText}");
+						}
+						else
+						{
+							//string SavedCommandLine = string.Join(" ", IndividualArgs);
+							TurnkeyUtils.Log($"Saving reusable commandline to DefaultGame.ini:\n  {PlatformParams}");
+							TurnkeyUtils.Log($"You can now run the command via RunUAT Turnkey -command=ExecuteBuild -build=\"{Name}\" or run the editor and find it in the Platforms menu.");
+							TurnkeyUtils.Log("Further, you can edit this in the Packaging Settings in the editor as needed");
+
+							// add this setting in to the ProjectSettings section
+							ConfigCache.WriteSettingToDefaultConfig(ConfigHierarchyType.Game, ProjectFile.Directory, ConfigCache.ConfigDefaultUpdateType.AddArrayEntry, SettingsSection, "ProjectCustomBuilds", ValueText);
+							ConfigCache.InvalidateCaches();
+						}
+					}
+				}
+				else
+				{
+					if (TurnkeyUtils.GetUserConfirmation($"Ready to execute the follwing command?\n  BuildCookRun {PlatformParams}", true))
+					{
+						ExecuteBuild.ExecuteBuildCookRun(PlatformParams);
+					}
+				}
+			}
+		}
+	}
+}
diff --git a/Engine/Source/Programs/AutomationTool/Turnkey/Commands/ExecuteBuild.cs b/Engine/Source/Programs/AutomationTool/Turnkey/Commands/ExecuteBuild.cs
--- a/Engine/Source/Programs/AutomationTool/Turnkey/Commands/ExecuteBuild.cs
+++ b/Engine/Source/Programs/AutomationTool/Turnkey/Commands/ExecuteBuild.cs
@@ -16,7 +16,17 @@ namespace Turnkey.Commands
 		protected override CommandGroup Group => CommandGroup.Builds;
 
 		// could move this to TurnkeyUtils!
-		static string GetStructEntry(string Input, string Property, bool bIsArrayProperty)
+		public static string GetStructEntryForSetting(ConfigHierarchy Config, string Section, string Setting, string Property)
+		{
+			string ConfigEntry;
+			if (Config.GetString(Section, Setting, out ConfigEntry))
+			{
+				return GetStructEntry(ConfigEntry, Property, false);
+			}
+			return null;
+		}
+
+		public static string GetStructEntry(string Input, string Property, bool bIsArrayProperty)
 		{
 			string PrimaryRegex;
 			string AltRegex = null;
@@ -26,9 +36,10 @@ namespace Turnkey.Commands
 			}
 			else
 			{
-				// handle quoted strings, allowing for quoited quotation marks (basically doing " followed by whatever, until we see a quote that was not proceeded by a \, and gather the whole mess in an outer group)
+				// handle quoted strings, allowing for escaped quotation marks (basically doing " followed by whatever, until we see a quote that was not proceeded by a \, and gather the whole mess in an outer group)
 				PrimaryRegex = string.Format("{0}\\s*=\\s*\"((.*?)[^\\\\])\"", Property);
-				AltRegex = string.Format("{0}\\s*=\\s*(.*?)\\,", Property);
+				// when no quotes, we skip over whitespace, and we end when we see whitespace, a comma or a ). This will handle (Ip = 192.168.0.1 , Name=....) , and return only '192.168.0.1'
+				AltRegex = string.Format("{0}\\s*=\\s*(.*?)[\\s,\\)]", Property);
 			}
 
 			// attempt to match it!
@@ -47,11 +58,51 @@ namespace Turnkey.Commands
 			return null;
 		}
 
+		public static string GetMapValueForSetting(ConfigHierarchy Config, string Section, string Setting, string Key)
+		{
+			string ConfigEntry;
+			if (Config.GetString(Section, Setting, out ConfigEntry))
+			{
+				return GetMapValue(ConfigEntry, Key);
+			}
+			return null;
+		}
+
+		// Key cannot have escaped quotes or commas
+		public static string GetMapValue(string Input, string Key)
+		{
+			string PrimaryRegex;
+			string AltRegex = null;
+			// handle quoted strings, allowing for escaped quotation marks (and possibly the key in quotes as well)
+			PrimaryRegex = string.Format("{0}\"?\\s*,\\s*\"((.*?)[^\\\\])\"", Key);
+			AltRegex = string.Format("{0}\"?\\s*,\\s*(.*?)[\\s,\\)]", Key);
+
+			// attempt to match it!
+			Match Result = Regex.Match(Input, PrimaryRegex);
+			if (!Result.Success && AltRegex != null)
+			{
+				Result = Regex.Match(Input, AltRegex);
+			}
+
+			// if we got a success, return the main match value
+			if (Result.Success)
+			{
+				return Result.Groups[1].Value.ToString();
+			}
+
+			return null;
+		}
+
 		protected override void Execute(string[] CommandOptions)
 		{
 			// we need a platform to execute
-			List<UnrealTargetPlatform> Platforms = TurnkeyUtils.GetPlatformsFromCommandLineOrUser(CommandOptions, null);
 			FileReference ProjectFile = TurnkeyUtils.GetProjectFromCommandLineOrUser(CommandOptions);
+			List<UnrealTargetPlatform> Platforms = TurnkeyUtils.GetPlatformsFromCommandLineOrUser(CommandOptions, null);
+			string DesiredBuild = TurnkeyUtils.ParseParamValue("Build", null, CommandOptions);
+			string ExtraOptions = TurnkeyUtils.ParseParamValue("ExtraOptions", "", CommandOptions);
+			string OutputDir = TurnkeyUtils.ParseParamValue("OutputDir", null, CommandOptions);
+			bool bPrintCommandOnly = TurnkeyUtils.ParseParam("PrintOnly", CommandOptions);
+			bool bInteractive = TurnkeyUtils.ParseParam("Interactive", CommandOptions);
 
 			// we need a project file, so if canceled, abore this command
 			if (ProjectFile == null)
@@ -59,8 +110,6 @@ namespace Turnkey.Commands
 				return;
 			}
 
-			string DesiredBuild = TurnkeyUtils.ParseParamValue("Build", null, CommandOptions);
-
 			// get a list of builds from config
 			foreach (UnrealTargetPlatform Platform in Platforms)
 			{
@@ -81,12 +130,13 @@ namespace Turnkey.Commands
 					Builds.AddRange(ProjectBuilds);
 				}
 
-				Dictionary<string, string> BuildCommands = new Dictionary<string, string>(StringComparer.InvariantCultureIgnoreCase);
+				Dictionary<string, Tuple<string, string>> BuildCommands = new Dictionary<string, Tuple<string, string>>(StringComparer.InvariantCultureIgnoreCase);
 				if (Builds != null)
 				{
 					foreach (string Build in Builds)
 					{
 						string Name = GetStructEntry(Build, "Name", false);
+						string Help = GetStructEntry(Build, "HelpText", false);
 						string SpecificPlatforms = GetStructEntry(Build, "SpecificPlatforms", true);
 						string Params = GetStructEntry(Build, "BuildCookRunParams", false);
 
@@ -99,16 +149,17 @@ namespace Turnkey.Commands
 						// if platforms are specified, and this platform isn't one of them, skip it
 						if (!string.IsNullOrEmpty(SpecificPlatforms))
 						{
+							string IniPlatformName = ConfigHierarchy.GetIniPlatformName(Platform);
 							string[] PlatformList = SpecificPlatforms.Split(",\"".ToCharArray(), StringSplitOptions.RemoveEmptyEntries);
 							// case insensitive Contains
-							if (PlatformList.Length > 0 && !PlatformList.Any(x => x.Equals(Platform.ToString(), StringComparison.OrdinalIgnoreCase)))
+							if (PlatformList.Length > 0 && !PlatformList.Any(x => x.Equals(IniPlatformName, StringComparison.OrdinalIgnoreCase)))
 							{
 								continue;
 							}
 						}
 
 						// add to list of commands
-						BuildCommands.Add(Name, Params);
+						BuildCommands.Add(Name, new Tuple<string, string>(Help, Params));
 					}
 				}
 
@@ -121,31 +172,100 @@ namespace Turnkey.Commands
 				string FinalParams;
 				if (!string.IsNullOrEmpty(DesiredBuild) && BuildCommands.ContainsKey(DesiredBuild))
 				{
-					FinalParams = BuildCommands[DesiredBuild];
+					FinalParams = BuildCommands[DesiredBuild].Item2;
 				}
 				else
 				{
 					List<string> BuildNames = BuildCommands.Keys.ToList();
-					int Choice = TurnkeyUtils.ReadInputInt("Choose a build to execute", BuildNames, true);
+					List<string> BuildItems = BuildNames.Select(x => x + (string.IsNullOrEmpty(BuildCommands[x].Item1) ? "" : $" [{BuildCommands[x].Item1}]")).ToList();
+					int Choice = TurnkeyUtils.ReadInputInt("Choose a build to execute", BuildItems, true);
 					if (Choice == 0)
 					{
 						continue;
 					}
 
-					FinalParams = BuildCommands[BuildNames[Choice - 1]];
+					FinalParams = BuildCommands[BuildNames[Choice - 1]].Item2;
 				}
 
-				FinalParams = FinalParams.Replace("{Project}", "\"" + ProjectFile.FullName + "\"");
-				FinalParams = FinalParams.Replace("{Platform}", Platform.ToString());
+				// make sure a project option is specified
+				if (!FinalParams.Contains("-project=", StringComparison.InvariantCultureIgnoreCase))
+				{
+					FinalParams = "-project={Project} " + FinalParams;
+				}
+
+				FinalParams = FinalParams.Replace("{Project}", "\"" + ProjectFile.FullName + "\"", StringComparison.InvariantCultureIgnoreCase);
+				FinalParams = FinalParams.Replace("{Platform}", Platform.ToString(), StringComparison.InvariantCultureIgnoreCase);
+				FinalParams = FinalParams.Replace("{ProjectPackagingSettings}", CreateBuild.MakeCommandLineFromPackagingSettings(Platform, ProjectFile, FinalParams));
 				FinalParams = PerformIniReplacements(FinalParams, ProjectFile.Directory, Platform);
 
-				TurnkeyUtils.Log("Executing '{0}'...", FinalParams);
+				// if there is a {DeviceId} param, get it from the commandline or user
+				if (FinalParams.Contains("{DeviceId}", StringComparison.InvariantCultureIgnoreCase))
+				{
+					List<DeviceInfo> Devices = TurnkeyUtils.GetDevicesFromCommandLineOrUser(CommandOptions, Platform);
+					string DeviceIds = string.Join("+", Devices.Select(x => x.Id));
+					FinalParams = FinalParams.Replace("{DeviceId}", DeviceIds, StringComparison.InvariantCultureIgnoreCase);
+				}
+
+				FinalParams += " " + ExtraOptions;
 
+				if (bPrintCommandOnly)
+				{
+					TurnkeyUtils.Log("To execute this build manually, run:");
+					TurnkeyUtils.Log("");
+					TurnkeyUtils.Report($"RunUAT BuildCookRun {FinalParams}");
+					TurnkeyUtils.Log("");
+					return;
+				}
+
+				// handle browsing for BrowseForDir after printing out, as it doesn't make much sense to ask for a directory or print out one
+				if (FinalParams.Contains("{BrowseForDir}"))
+				{
+					// if the -outputdir option was specified, use that, otherwise ask user for directory
+					if (OutputDir == null)
+					{
+						//if (RuntimePlatform.IsWindows)
+						//{
+						//	string ChosenFile = null;
+						//	System.Threading.Thread t = new System.Threading.Thread(x =>
+						//	{
+						//		ChosenFile = UnrealWindowsForms.Utils.ShowOpenFileDialogAndReturnFilename("Project Files (*.uproject)|*.uproject");
+						//	});
+
+						//	t.SetApartmentState(System.Threading.ApartmentState.STA);
+						//	t.Start();
+						//	t.Join();
+
+						//	if (ChosenFile == null)
+						//	{
+						//		continue;
+						//	}
+						//}
+						//else
+						//{
+						//	while (true)
+							{
+								OutputDir = TurnkeyUtils.ReadInput("Enter output directory:");
+							}
+						//}
+
+					}
+					if (string.IsNullOrEmpty(OutputDir))
+					{
+						TurnkeyUtils.Log("Cancelling...");
+						return;
+					}
+					// handle {BrowseForDir} with or without quotes already around it (if already has quotes, then replace with the path, and if
+					// without quotes, replace with path wrapped in quotes
+					FinalParams = FinalParams.Replace("\"{BrowseForDir}\"", OutputDir);
+					FinalParams = FinalParams.Replace("{BrowseForDir}", $"\"{OutputDir}\"");
+				}
+
+				TurnkeyUtils.Log("Executing '{0}'...", FinalParams);
 				ExecuteBuildCookRun(FinalParams);
 			}
 		}
 
-		private static void ExecuteBuildCookRun(string Params)
+		internal static void ExecuteBuildCookRun(string Params)
 		{
 			// split the params on whitespace not inside quotes (see https://stackoverflow.com/questions/4780728/regex-split-string-preserving-quotes/4780801#4780801 to explain the regex)
 			Regex Matcher = new Regex("(?<=^[^\"]*(?:\"[^\"]*\"[^\"]*)*)\\s(?=(?:[^\"]*\"[^\"]*\")*[^\"]*$)");
diff --git a/Engine/Source/Programs/AutomationTool/Turnkey/Providers/ConsoleIOProvider.cs b/Engine/Source/Programs/AutomationTool/Turnkey/Providers/ConsoleIOProvider.cs
--- a/Engine/Source/Programs/AutomationTool/Turnkey/Providers/ConsoleIOProvider.cs
+++ b/Engine/Source/Programs/AutomationTool/Turnkey/Providers/ConsoleIOProvider.cs
@@ -101,7 +101,7 @@ namespace Turnkey
 
 		public override bool GetUserConfirmation(string Message, bool bDefaultValue, bool bAppendNewLine)
 		{
-			string Prompt = string.Format($"{Message} [{(bDefaultValue ? 'Y' : 'y')}/{(bDefaultValue ? 'n' : 'N')}]");
+			string Prompt = $"{Message} [{(bDefaultValue ? 'Y' : 'y')}/{(bDefaultValue ? 'n' : 'N')}]";
 			string Result = ReadInput(Prompt, bDefaultValue ? "Y" : "N", bAppendNewLine);
 			if (string.IsNullOrEmpty(Result))
 			{
diff --git a/Engine/Source/Programs/AutomationTool/Turnkey/TurnkeyUtils.cs b/Engine/Source/Programs/AutomationTool/Turnkey/TurnkeyUtils.cs
--- a/Engine/Source/Programs/AutomationTool/Turnkey/TurnkeyUtils.cs
+++ b/Engine/Source/Programs/AutomationTool/Turnkey/TurnkeyUtils.cs
@@ -96,7 +96,7 @@ namespace Turnkey
 		public static string ParseParamValue(string Param, string Default, string[] ExtraOptions)
 		{
 			// our internal extraoptions still have - in front, but CommandUtilHelper won't have the dashes
-			string Value = CommandUtils.ParseParamValue(ExtraOptions, "-" + Param, Default);
+			string Value = CommandUtils.ParseParamValue(ExtraOptions, "-" + Param, null);
 			if (Value == null)
 			{
 				Value = CommandUtilHelper.ParseParamValue(Param, Default);
diff --git a/Engine/Source/Programs/UnrealBuildTool/System/ConfigCache.cs b/Engine/Source/Programs/UnrealBuildTool/System/ConfigCache.cs
--- a/Engine/Source/Programs/UnrealBuildTool/System/ConfigCache.cs
+++ b/Engine/Source/Programs/UnrealBuildTool/System/ConfigCache.cs
@@ -5,6 +5,7 @@ using System.Collections.Generic;
 using System.Diagnostics.CodeAnalysis;
 using System.Linq;
 using System.Reflection;
+using System.IO;
 using System.Text;
 using System.Threading.Tasks;
 using EpicGames.Core;
@@ -61,6 +62,21 @@ namespace UnrealBuildTool
 			}
 		}
 
+		/// <summary>
+		/// Allowed modification types allowed for default config files
+		/// </summary>
+		public enum ConfigDefaultUpdateType
+		{
+			/// <summary>
+			/// Used for non-array types, this will replace a setting, or it will add a setting if it doesn't exist
+			/// </summary>
+			SetValue,
+			/// <summary>
+			/// Used to add an array entry to the end of any existing array entries, or will add to the end of the section
+			/// </summary>
+			AddArrayEntry,
+		}
+
 		/// <summary>
 		/// Stores information identifying a unique config hierarchy
 		/// </summary>
@@ -457,5 +473,158 @@ namespace UnrealBuildTool
 				throw new Exception("Unsupported type for [ConfigFile] attribute");
 			}
 		}
+
+
+		#region Updating Default Config file support
+
+		/// <summary>
+		/// Calculates the path to where the project's Default config of the type given (ie DefaultEngine.ini)
+		/// </summary>
+		/// <param name="ConfigType">Game, Engine, etc</param>
+		/// <param name="ProjectDir">Project directory, used to find Config/Default[Type].ini</param>
+		public static FileReference GetDefaultConfigFileReference(ConfigHierarchyType ConfigType, DirectoryReference ProjectDir)
+		{
+			return FileReference.Combine(ProjectDir, "Config", $"Default{ConfigType}.ini");
+		}
+
+		/// <summary>
+		/// Updates a section in a Default***.ini, and will write it out. If the file is not writable, p4 can attempt to check it out.
+		/// </summary>
+		/// <param name="ConfigType">Game, Engine, etc</param>
+		/// <param name="ProjectDir">Project directory, used to find Config/Default[Type].ini</param>
+		/// <param name="UpdateType">How to modify the secion</param>
+		/// <param name="Section">Name of the section with the Key in it</param>
+		/// <param name="Key">Key to update</param>
+		/// <param name="Value">Value to write for te Key</param>
+		/// <returns></returns>		
+		public static bool WriteSettingToDefaultConfig(ConfigHierarchyType ConfigType, DirectoryReference ProjectDir, ConfigDefaultUpdateType UpdateType, string Section, string Key, string Value)
+		{
+			FileReference DefaultConfigFile = GetDefaultConfigFileReference(ConfigType, ProjectDir);
+
+			if (!FileReference.Exists(DefaultConfigFile))
+			{
+				Log.TraceWarning($"Failed to find config file '{DefaultConfigFile.FullName}' to update");
+				return false;
+			}
+
+			if (File.GetAttributes(DefaultConfigFile.FullName).HasFlag(FileAttributes.ReadOnly))
+			{
+				Log.TraceWarning($"Config file '{DefaultConfigFile.FullName}' is read-only, unable to write setting {Key}");
+				return false;
+			}
+
+			// generate the section header
+			string SectionString = $"[{Section}]";
+
+			// genrate the line we are going to write to the config
+			string KeyWithEquals = $"{Key}=";
+			string LineToWrite = (UpdateType == ConfigDefaultUpdateType.AddArrayEntry ? "+" : "");
+			LineToWrite += KeyWithEquals + Value;
+
+
+			// read in all the lines so we can insert or replace one
+			List<string> Lines = File.ReadAllLines(DefaultConfigFile.FullName).ToList();
+
+			// look for the section
+			int SectionIndex = -1;
+			for (int Index = 0; Index < Lines.Count; Index++)
+			{
+				if (Lines[Index].Trim().Equals(SectionString, StringComparison.InvariantCultureIgnoreCase))
+				{
+					SectionIndex = Index;
+					break;
+				}
+			}
+
+			// if section not found, just append to the end
+			if (SectionIndex == -1)
+			{
+				Lines.Add(SectionString);
+				Lines.Add(LineToWrite);
+
+				File.WriteAllLines(DefaultConfigFile.FullName, Lines);
+				return true;
+			}
+
+
+			// find the last line in the section with the prefix
+			int LastIndexOfPrefix = -1;
+			int NextSectionIndex = -1;
+			for (int Index = SectionIndex + 1; Index < Lines.Count; Index++)
+			{
+				string Line = Lines[Index];
+				if (Line.StartsWith('+') || Line.StartsWith('-') || Line.StartsWith('.') || Line.StartsWith('!'))
+				{
+					Line = Line.Substring(1);
+				}
+
+				// look for last array entry in case of multiples (or the only line for non-array type)
+				if (Line.StartsWith(KeyWithEquals, StringComparison.InvariantCultureIgnoreCase))
+				{
+					LastIndexOfPrefix = Index;
+				}
+				else if (Lines[Index].StartsWith("["))
+				{
+					NextSectionIndex = Index;
+					// we found another section, so break out
+					break;
+				}
+			}
+
+			// now we know enough to either insert or replace a line
+
+			// if we never found the key, we will insert at the end of the section
+			if (LastIndexOfPrefix == -1)
+			{
+				// if we didn't find a next section, thjen we will insert at the end of the file
+				if (NextSectionIndex == -1)
+				{
+					NextSectionIndex = Lines.Count;
+				}
+
+				// move past blank lines between sections
+				while (string.IsNullOrWhiteSpace(Lines[NextSectionIndex - 1]))
+				{
+					NextSectionIndex--;
+				}
+				// insert before the next section (or end of file)
+				Lines.Insert(NextSectionIndex, LineToWrite);
+			}
+			// otherwise, insert after, or replace, a line, depending on type
+			else
+			{
+				if (UpdateType == ConfigDefaultUpdateType.AddArrayEntry)
+				{
+					Lines.Insert(LastIndexOfPrefix + 1, LineToWrite);
+				}
+				else
+				{
+					Lines[LastIndexOfPrefix] = LineToWrite;
+				}
+			}
+
+			// now the lines are updated, we can overwrite the file
+			File.WriteAllLines(DefaultConfigFile.FullName, Lines);
+			return true;
+		}
+
+		/// <summary>
+		/// Invalidates the hierarchy internal caches so that next call to ReadHierarchy will re-read from disk
+		/// but existing ones will still be valid with old values
+		/// </summary>
+		public static void InvalidateCaches()
+		{
+			lock (LocationToConfigFile)
+			{
+				LocationToConfigFile.Clear();
+			}
+
+			lock (HierarchyKeyToHierarchy)
+			{
+				HierarchyKeyToHierarchy.Clear();
+			}
+		}
+
+		#endregion
 	}
 }
diff --git a/Engine/Source/Programs/UnrealBuildTool/System/DataDrivenPlatformInfo.cs b/Engine/Source/Programs/UnrealBuildTool/System/DataDrivenPlatformInfo.cs
--- a/Engine/Source/Programs/UnrealBuildTool/System/DataDrivenPlatformInfo.cs
+++ b/Engine/Source/Programs/UnrealBuildTool/System/DataDrivenPlatformInfo.cs
@@ -30,6 +30,11 @@ namespace UnrealBuildTool
 			/// </summary>
 			public bool bIsEnabled = false;
 
+			/// <summary>
+			/// If true, this platform can stage the crashreporter
+			/// </summary>
+			public bool bCanUseCrashReporter = false;
+
 			/// <summary>
 			/// Additional restricted folders for this platform.
 			/// </summary>
@@ -115,6 +120,12 @@ namespace UnrealBuildTool
 							{
 								ConfigHierarchy.TryParse(Temp, out NewInfo.bIsConfidential);
 							}
+							// unspecified means true for this property
+							NewInfo.bCanUseCrashReporter = true;
+							if (ParsedSection.TryGetValue("bCanUseCrashReporter", out Temp))
+							{
+								ConfigHierarchy.TryParse(Temp, out NewInfo.bCanUseCrashReporter);
+							}
 
 							string HostKey = ConfigHierarchy.GetIniPlatformName(BuildHostPlatform.Current.Platform) + ":bIsEnabled";
 							string NormalKey = "bIsEnabled";
