Fixed null access when no devices are available for a test

diff --git a/Engine/Source/Programs/AutomationTool/Gauntlet/Framework/Devices/Gauntlet.UnrealDeviceReservation.cs b/Engine/Source/Programs/AutomationTool/Gauntlet/Framework/Devices/Gauntlet.UnrealDeviceReservation.cs
--- a/Engine/Source/Programs/AutomationTool/Gauntlet/Framework/Devices/Gauntlet.UnrealDeviceReservation.cs
+++ b/Engine/Source/Programs/AutomationTool/Gauntlet/Framework/Devices/Gauntlet.UnrealDeviceReservation.cs
@@ -127,6 +127,7 @@ namespace Gauntlet
 			{
 				Log.Info("Failed to resolve all devices. Releasing the ones we have ");
 				DevicePool.Instance.ReleaseDevices(AcquiredDevices);
+				ReservedDevices = new List<ITargetDevice>();
 			}
 			else
 			{
