Gauntlet - Horde Automated Test Session report
* Introduce Automated Test Session data model for Horde TestData
* Introduce Suite and Type to UnrealTestNode to have better control over Session naming
* Add PreflightChange option to discriminate Session that are not official run
* Add test Duration and DateTime to UE json report model
* Accomodate for posting New and Old report models

diff --git a/Engine/Source/Developer/AutomationController/Private/AutomationControllerManager.h b/Engine/Source/Developer/AutomationController/Private/AutomationControllerManager.h
--- a/Engine/Source/Developer/AutomationController/Private/AutomationControllerManager.h
+++ b/Engine/Source/Developer/AutomationController/Private/AutomationControllerManager.h
@@ -39,11 +39,19 @@ public:
 	UPROPERTY()
 	FString DeviceInstance;
 
+	UPROPERTY()
+	float Duration;
+
+	UPROPERTY()
+	FDateTime DateTime;
+
 	FAutomatedTestResult()
 	{
 		Warnings = 0;
 		Errors = 0;
 		State = EAutomationState::NotRun;
+		Duration = 0;
+		DateTime = 0;
 	}
 
 	void SetEvents(const TArray<FAutomationExecutionEntry>& InEntries, int32 InWarnings, int32 InErrors)
diff --git a/Engine/Source/Developer/AutomationController/Private/AutomationControllerManger.cpp b/Engine/Source/Developer/AutomationController/Private/AutomationControllerManger.cpp
--- a/Engine/Source/Developer/AutomationController/Private/AutomationControllerManger.cpp
+++ b/Engine/Source/Developer/AutomationController/Private/AutomationControllerManger.cpp
@@ -132,6 +132,7 @@ void FAutomatedTestPassResults::UpdateTestResultStatus(const IAutomationReportPt
 		Failed++;
 		break;
 	case EAutomationState::InProcess:
+		TestResult.DateTime = FDateTime::Now();
 		InProcess++;
 		break;
 	case EAutomationState::NotRun:
@@ -534,6 +535,7 @@ void FAutomationControllerManager::CollectTestResults(TSharedPtr<IAutomationRepo
 		TestResult.SetEvents(Results.GetEntries(), Results.GetWarningTotal(), Results.GetErrorTotal());
 		TestResult.SetArtifacts(Results.Artifacts);
 		TestResult.DeviceInstance = Results.GameInstance;
+		TestResult.Duration = Results.Duration;
 
 		JsonTestPassResults.TotalDuration += Results.Duration;
 
diff --git a/Engine/Source/Programs/AutomationTool/Gauntlet/Framework/Base/Gauntlet.TestReport.cs b/Engine/Source/Programs/AutomationTool/Gauntlet/Framework/Base/Gauntlet.TestReport.cs
--- a/Engine/Source/Programs/AutomationTool/Gauntlet/Framework/Base/Gauntlet.TestReport.cs
+++ b/Engine/Source/Programs/AutomationTool/Gauntlet/Framework/Base/Gauntlet.TestReport.cs
@@ -78,6 +78,12 @@ namespace Gauntlet
 		/// <param name="ArtifactPath"></param>
 		/// <returns></returns>
 		bool AttachArtifact(string ArtifactPath, string Name = null);
+
+		/// <summary>
+		/// Return a dictionary of key/object intended to be a data collection of dependencies that comes with the report but not part of the report. ie: an artifact manifest
+		/// </summary>
+		/// <returns></returns>
+		Dictionary<string, object> GetReportDependencies();
 	}
 
 	/// <summary>
@@ -231,6 +237,15 @@ namespace Gauntlet
 		{
 			return TelemetryDataList;
 		}
+
+		/// <summary>
+		/// Return a dictionary of key/object intended to be a data collection of dependencies that comes with the report but not part of the report. ie: an artifact manifest
+		/// </summary>
+		/// <returns></returns>
+		public virtual Dictionary<string, object> GetReportDependencies()
+		{
+			return new Dictionary<string, object>();
+		}
 	}
 
 	class JsonTryParseEnumConverter : JsonConverterFactory
diff --git a/Engine/Source/Programs/AutomationTool/Gauntlet/Framework/Gauntlet.Utils.cs b/Engine/Source/Programs/AutomationTool/Gauntlet/Framework/Gauntlet.Utils.cs
--- a/Engine/Source/Programs/AutomationTool/Gauntlet/Framework/Gauntlet.Utils.cs
+++ b/Engine/Source/Programs/AutomationTool/Gauntlet/Framework/Gauntlet.Utils.cs
@@ -448,6 +448,29 @@ namespace Gauntlet
 
 			return HashString;
 		}
+
+		public static string ComputeHash(string Input, HashAlgorithm Algo, int MaxLength = 0)
+		{
+			if (string.IsNullOrEmpty(Input))
+			{
+				return "0";
+			}
+
+			byte[] Hash = Algo.ComputeHash(Encoding.UTF8.GetBytes(Input));
+			var SBuilder = new StringBuilder();
+
+			for (int i = 0; i < Hash.Length; i++)
+			{
+				SBuilder.Append(Hash[i].ToString("x2"));
+				if (MaxLength > 0 && SBuilder.Length >= MaxLength)
+				{
+					if (SBuilder.Length > MaxLength) { SBuilder.Remove(MaxLength, 1); }
+					break;
+				}
+			}
+
+			return SBuilder.ToString();
+		}
 	}
 
 	/*
diff --git a/Engine/Source/Programs/AutomationTool/Gauntlet/Framework/Utils/Gauntlet.HordeReport.cs b/Engine/Source/Programs/AutomationTool/Gauntlet/Framework/Utils/Gauntlet.HordeReport.cs
--- a/Engine/Source/Programs/AutomationTool/Gauntlet/Framework/Utils/Gauntlet.HordeReport.cs
+++ b/Engine/Source/Programs/AutomationTool/Gauntlet/Framework/Utils/Gauntlet.HordeReport.cs
@@ -3,9 +3,13 @@
 using System;
 using System.Collections.Generic;
 using System.IO;
+using System.Linq;
+using System.Text;
 using System.Text.Json;
+using System.Text.RegularExpressions;
 using System.Text.Json.Serialization;
 using AutomationTool;
+using System.Security.Cryptography;
 
 namespace Gauntlet
 {
@@ -82,6 +86,8 @@ namespace Gauntlet
 		public abstract class BaseHordeReport : BaseTestReport
 		{
 			protected string OutputArtifactPath;
+			protected HashSet<string> ArtifactProcessedHashes;
+			protected Dictionary<string, object> ExtraReports;
 
 			/// <summary>
 			/// Attach Artifact to the Test Report
@@ -90,28 +96,67 @@ namespace Gauntlet
 			/// <param name="Name"></param>
 			/// <returns>return true if the file was successfully attached</returns>
 			public override bool AttachArtifact(string ArtifactPath, string Name = null)
+			{
+				return AttachArtifact(ArtifactPath, Name, false);
+			}
+			/// <summary>
+			/// Attach Artifact to the Test Report
+			/// </summary>
+			/// <param name="ArtifactPath"></param>
+			/// <param name="Name"></param>
+			/// <param name="Overwrite"></param>
+			/// <returns>return true if the file was successfully attached</returns>
+			public bool AttachArtifact(string ArtifactPath, string Name = null, bool Overwrite = false)
 			{
 				if (string.IsNullOrEmpty(OutputArtifactPath))
 				{
 					throw new InvalidOperationException("OutputArtifactPath must be set before attaching any artifact");
 				}
 
-				if (File.Exists(ArtifactPath))
+				string ArtifactHash;
+				{
+					// Generate a hash from the artifact path
+					using (SHA1 Sha1 = SHA1.Create())
+					{
+						ArtifactHash = Hasher.ComputeHash(ArtifactPath, Sha1, 8);
+					}
+				}
+
+				if (ArtifactProcessedHashes == null)
+				{
+					ArtifactProcessedHashes = new HashSet<string>();
+				}
+				else
+				{
+					if(ArtifactProcessedHashes.Contains(ArtifactHash))
+					{
+						// already processed
+						return true;
+					}
+				}
+
+				// Mark as processed. Even in case of failure, we don't want to try over again.
+				ArtifactProcessedHashes.Add(ArtifactHash);
+
+				string TargetPath = Utils.SystemHelpers.GetFullyQualifiedPath(Path.Combine(OutputArtifactPath, Name ?? Path.GetFileName(ArtifactPath)));
+				ArtifactPath = Utils.SystemHelpers.GetFullyQualifiedPath(ArtifactPath);
+				bool isFileExist = File.Exists(ArtifactPath);
+				if (isFileExist && (!File.Exists(TargetPath) || Overwrite))
 				{
 					try
 					{
-						string TargetPath = Path.Combine(OutputArtifactPath, Name ?? Path.GetFileName(ArtifactPath));
 						string TargetDirectry = Path.GetDirectoryName(TargetPath);
 						if (!Directory.Exists(TargetDirectry)) { Directory.CreateDirectory(TargetDirectry); }
-						File.Copy(Utils.SystemHelpers.GetFullyQualifiedPath(ArtifactPath), Utils.SystemHelpers.GetFullyQualifiedPath(TargetPath));
+						File.Copy(ArtifactPath, TargetPath, true);
 						return true;
 					}
 					catch (Exception Ex)
 					{
 						Log.Error("Failed to copy artifact '{0}'. {1}", Path.GetFileName(ArtifactPath), Ex);
 					}
+					return false;
 				}
-				return false;
+				return isFileExist;
 			}
 
 			/// <summary>
@@ -129,6 +174,38 @@ namespace Gauntlet
 				}
 				Log.Verbose(string.Format("Test Report output artifact path is set to: {0}", OutputArtifactPath));
 			}
+
+			public void AttachDependencyReport(object InReport, string Key = null)
+			{
+				if(ExtraReports == null)
+				{
+					ExtraReports = new Dictionary<string, object>();
+				}
+				if(InReport is BaseHordeReport InHordeReport)
+				{
+					Key = InHordeReport.GetTestDataKey(Key);
+				}
+				ExtraReports[Key] = InReport;
+			}
+
+			public override Dictionary<string, object> GetReportDependencies()
+			{
+				var Reports = base.GetReportDependencies();
+				if (ExtraReports != null)
+				{
+					ExtraReports.ToList().ForEach(Item => Reports.Add(Item.Key, Item.Value));
+				}
+				return Reports;
+			}
+
+			public virtual string GetTestDataKey(string BaseKey = null)
+			{
+				if(BaseKey == null)
+				{
+					return Type;
+				}
+				return string.Format("{0}::{1}", Type, BaseKey);
+			}
 		}
 
 		/// <summary>
@@ -470,6 +547,450 @@ namespace Gauntlet
 			}
 		}
 
+		/// <summary>
+		/// Contains information about a test session 
+		/// </summary>
+		public class AutomatedTestSessionData : BaseHordeReport
+		{
+			public override string Type
+			{
+				get { return "Automated Test Session"; }
+			}
+
+			public class TestResult
+			{
+				public string Name { get; set; }
+				public string TestUID { get; set; }
+				public string Suite { get; set; }
+				public TestStateType State { get; set; }
+				public string DeviceAppInstanceName { get; set; }
+				public uint ErrorCount { get; set; }
+				public uint WarningCount { get; set; }
+				public string ErrorHashAggregate { get; set; }
+				public string DateTime { get; set; }
+				public float TimeElapseSec { get; set; }
+				// not part of json output
+				public List<TestEvent> Events;
+				public SortedSet<string> ErrorHashes;
+				public TestResult(string InName, string InSuite, TestStateType InState, string InDevice, string InDateTime)
+				{
+					Name = InName;
+					TestUID = GenerateHash(Name);
+					Suite = InSuite;
+					State = InState;
+					DeviceAppInstanceName = InDevice;
+					ErrorHashes = new SortedSet<string>();
+					ErrorHashAggregate = "";
+					DateTime = InDateTime;
+					Events = new List<TestEvent>();
+				}
+				private string GenerateHash(string InName)
+				{
+					using (SHA1 Sha1 = SHA1.Create())
+					{
+						return Hasher.ComputeHash(InName, Sha1, 8);
+					}
+				}
+				public TestEvent NewEvent(EventType InType, string InMessage, string InTag, string InContext, string InDateTime)
+				{
+					TestEvent NewItem = new TestEvent(InType, InMessage, InTag, InContext ?? "", InDateTime);
+					Events.Add(NewItem);
+					switch (InType)
+					{
+						case EventType.Error:
+							ErrorCount++;
+							if (ErrorHashes.Count == 0 || !ErrorHashes.Contains(NewItem.Hash))
+							{
+								// Avoid error duplication and order sensitivity
+								ErrorHashes.Add(NewItem.Hash);
+								if (ErrorHashes.Count > 1)
+								{
+									ErrorHashAggregate = GenerateHash(string.Join("", ErrorHashes.ToArray()));
+								}
+								else
+								{
+									ErrorHashAggregate = NewItem.Hash;
+								}
+							}
+							break;
+						case EventType.Warning:
+							WarningCount++;
+							break;
+					}
+
+					return NewItem;
+				}
+				public TestEvent GetLastEvent()
+				{
+					if (Events.Count == 0)
+					{
+						throw new AutomationException("A test event must be added first to the report!");
+					}
+					return Events.Last();
+				}
+			}
+			public class TestSession
+			{
+				public string DateTime { get; set; }
+				public float TimeElapseSec { get; set; }
+				public Dictionary<string, TestResult> Tests { get; set; }
+				public string TestResultsTestDataUID { get; set; }
+				public TestSession()
+				{
+					Tests = new Dictionary<string, TestResult>();
+					TestResultsTestDataUID = Guid.NewGuid().ToString();
+				}
+				public TestResult NewTestResult(string InName, string InSuite, TestStateType InState, string InDevice, string InDateTime)
+				{
+					TestResult NewItem = new TestResult(InName, InSuite, InState, InDevice, InDateTime);
+					Tests[NewItem.TestUID] = NewItem;
+					return NewItem;
+				}
+			}
+			public class Device
+			{
+				public string Name { get; set; }
+				public string AppInstanceName { get; set; }
+				public Dictionary<string, string> Metadata { get; set; }
+				public Device(string InName, string InInstance)
+				{
+					Name = InName;
+					AppInstanceName = InInstance;
+					Metadata = new Dictionary<string, string>();
+				}
+				public void SetMetaData(string Key, string Value)
+				{
+					Metadata.Add(Key, Value);
+				}
+			}
+			public class TestArtifact
+			{
+				public string Tag { get; set; }
+				public string ReferencePath { get; set; }
+			}
+			public class TestEvent
+			{
+				public string Message { get; set; }
+				public string Context { get; set; }
+				public EventType Type { get; set; }
+				public string Tag { get; set; }
+				public string Hash { get; set; }
+				public string DateTime { get; set; }
+				public List<TestArtifact> Artifacts { get; set; }
+				public TestEvent(EventType InType, string InMessage, string InTag, string InContext, string InDateTime)
+				{
+					Type = InType;
+					Message = InMessage;
+					Context = InContext;
+					Tag = InTag;
+					Hash = InType != EventType.Info ? GenerateEventHash() : "";
+					DateTime = InDateTime;
+					Artifacts = new List<TestArtifact>();
+				}
+				public TestArtifact NewArtifacts(string InTag, string InReferencePath)
+				{
+					TestArtifact NewItem = new TestArtifact();
+					NewItem.Tag = InTag;
+					NewItem.ReferencePath = InReferencePath;
+					Artifacts.Add(NewItem);
+					return NewItem;
+				}
+				private string FilterEvent(string Text)
+				{
+					// Filter out time stamp in message event
+					// [2021.05.20-09.55.28:474][  3]
+					Regex Regex_timestamp = new Regex(@"\[[0-9]+\.[0-9]+\.[0-9]+-[0-9]+\.[0-9]+\.[0-9]+(:[0-9]+)?\](\[ *[0-9]+\])?");
+					Text = Regex_timestamp.Replace(Text, "");
+
+					// Filter out worker number in message event
+					// Worker #13
+					Regex Regex_workernumber = new Regex(@"Worker #[0-9]+");
+					Text = Regex_workernumber.Replace(Text, "");
+
+					// Filter out path
+					// D:\build\U5+R5.0+Inc\Sync\Engine\Source\Runtime\...
+					Regex Regex_pathstring = new Regex(@"([A-Z]:)?([\\/][^\\/]+)+[\\/]");
+					Text = Regex_pathstring.Replace(Text, "");
+
+					// Filter out hexadecimal string in message event
+					// 0x00007ffa7b04e533
+					// Mesh 305f21682cf3a231a0947ffb35c51567
+					// <lambda_2a907a23b64c2d53ad869d04fdc6d423>
+					Regex Regex_hexstring = new Regex(@"(0x)?[0-9a-fA-F]{6,32}");
+					Text = Regex_hexstring.Replace(Text, "");
+
+					return Text;
+				}
+				private string GenerateEventHash()
+				{
+					string FilteredEvent = FilterEvent(string.Format("{0}{1}{2}{3}", Type.ToString(), Message, Context, Tag));
+					using (SHA1 Sha1 = SHA1.Create())
+					{
+						return Hasher.ComputeHash(FilteredEvent, Sha1, 8);
+					}
+				}
+			}
+			public class TestResultData
+			{
+ 				public List<TestEvent> Events
+				{
+					get { return Result.Events; }
+					set { Result.Events = value; }
+				}
+				private TestResult Result;
+				public TestResultData(TestResult InTestResult)
+				{
+					Result = InTestResult;
+				}
+			}
+			public class IndexedError
+			{
+				public string Message { get; set; }
+				public string Tag { get; set; }
+				public List<string> TestUIDs { get; set; }
+				public IndexedError(string InMessage, string InTag)
+				{
+					Message = InMessage;
+					Tag = InTag;
+					TestUIDs = new List<string>();
+				}
+			}
+
+			public AutomatedTestSessionData(string InName) : base()
+			{
+				Name = InName;
+				PreFlightChange = "";
+				TestSessionInfo = new TestSession();
+				Devices = new List<Device>();
+				IndexedErrors = new Dictionary<string, IndexedError>();
+				TestResults = new Dictionary<string, TestResultData>();
+			}
+
+			public string Name { get; set; }
+			public string PreFlightChange { get; set; }
+			public TestSession TestSessionInfo { get; set; }
+			public List<Device> Devices { get; set; }
+			/// <summary>
+			/// Database side optimization: Index the TestUID by TestError hash.
+			/// Key is meant to be the TestError hash and Value the list of related TestUID.
+			/// </summary>
+			public Dictionary<string, IndexedError> IndexedErrors { get; set; }
+			/// <summary>
+			/// Allow the Database to pull only one test result from the TestData.
+			/// Key is meant to be the TestUID and Value the test result detailed events.
+			/// </summary>
+			private Dictionary<string, TestResultData> TestResults { get; set; }
+
+			private string CurrentTestUID;
+
+			private Device NewDevice(string InName, string InInstance)
+			{
+				Device NewItem = new Device(InName, InInstance);
+				Devices.Add(NewItem);
+				return NewItem;
+			}
+
+			private void IndexTestError(string ErrorHash, string TestUID, string Message, string Tag)
+			{
+				if (!IndexedErrors.ContainsKey(ErrorHash))
+				{
+					IndexedErrors[ErrorHash] = new IndexedError(Message, Tag);
+				}
+				else if (IndexedErrors[ErrorHash].TestUIDs.Contains(TestUID))
+				{
+					// already stored
+					return;
+				}
+				IndexedErrors[ErrorHash].TestUIDs.Add(TestUID);
+			}
+
+			private TestResult GetCurrentTest()
+			{
+				if (string.IsNullOrEmpty(CurrentTestUID))
+				{
+					throw new AutomationException("A test must be set to the report!");
+				}
+				return TestSessionInfo.Tests[CurrentTestUID];
+			}
+
+			public void SetTest(string InName, string InSuite, string InDevice, string InDateTime = null)
+			{
+				TestResult Test = TestSessionInfo.NewTestResult(InName, InSuite, TestStateType.Unknown, InDevice, InDateTime ?? DateTime.Now.ToString("yyyy.MM.dd-HH.mm.ss"));
+				CurrentTestUID = Test.TestUID;
+				// Index the test result data.
+				TestResults[Test.TestUID] = new TestResultData(Test);
+			}
+			public void SetTestState(TestStateType InState)
+			{
+				TestResult CurrentTest = GetCurrentTest();
+				CurrentTest.State = InState;
+			}
+			public void SetTestTimeElapseSec(float InTimeElapseSec)
+			{
+				TestResult CurrentTest = GetCurrentTest();
+				CurrentTest.TimeElapseSec = InTimeElapseSec;
+			}
+			public override void AddEvent(EventType InType, string InMessage, object InContext = null)
+			{
+				AddEvent(InType, InMessage, "gauntlet", InContext == null ? null : InContext.ToString());
+			}
+			public void AddEvent(EventType InType, string InMessage, string InTag, string InContext = null, string InDateTime = null)
+			{
+				TestResult CurrentTest = GetCurrentTest();
+				TestEvent Event = CurrentTest.NewEvent(InType, InMessage, InTag, InContext, InDateTime ?? DateTime.Now.ToString("yyyy.MM.dd-HH.mm.ss"));
+				if(InType == EventType.Error)
+				{
+					IndexTestError(Event.Hash, CurrentTest.TestUID, InMessage, InTag);
+				}
+			}
+			public void AddArtifactToLastEvent(string InTag, string InFilePath, string InReferencePath = null)
+			{
+				TestEvent LastEvent = GetCurrentTest().GetLastEvent();
+				if (AttachArtifact(InFilePath, InReferencePath))
+				{
+					if (InReferencePath == null)
+					{
+						InReferencePath = Path.GetFileName(InFilePath);
+					}
+					LastEvent.NewArtifacts(InTag, InReferencePath);
+				}
+				else
+				{
+					AddEvent(EventType.Error, string.Format("Failed to attached Artifact {0}.", InReferencePath ?? Path.GetFileName(InFilePath)));
+				}
+			}
+			public override string GetTestDataKey(string BaseKey = null)
+			{
+				return base.GetTestDataKey(); // Ignore BaseKey
+			}
+			public override Dictionary<string, object> GetReportDependencies()
+			{
+
+				var Reports = base.GetReportDependencies();
+				// Test Result Details
+				var Key = string.Format("{0} Result Details::{1}", Type, TestSessionInfo.TestResultsTestDataUID);
+				Reports[Key] = TestResults;
+
+				return Reports;
+			}
+
+			/// <summary>
+			/// Convert UnrealAutomatedTestPassResults to Horde data model
+			/// </summary>
+			/// <param name="InTestPassResults"></param>
+			/// <param name="InName"></param>
+			/// <param name="InSuite"></param>
+			/// <param name="InReportPath"></param>
+			/// <param name="InHordeArtifactPath"></param>
+			/// <returns></returns>
+			public static AutomatedTestSessionData FromUnrealAutomatedTests(UnrealAutomatedTestPassResults InTestPassResults, string InName, string InSuite, string InReportPath, string InHordeArtifactPath)
+			{
+				AutomatedTestSessionData OutTestPassResults = new AutomatedTestSessionData(InName);
+				OutTestPassResults.SetOutputArtifactPath(InHordeArtifactPath);
+				if (InTestPassResults.Devices != null)
+				{
+					foreach (UnrealAutomationDevice InDevice in InTestPassResults.Devices)
+					{
+						Device ConvertedDevice = OutTestPassResults.NewDevice(InDevice.DeviceName, InDevice.Instance);
+						ConvertedDevice.SetMetaData("platform", InDevice.Platform);
+						ConvertedDevice.SetMetaData("os_version", InDevice.OSVersion);
+						ConvertedDevice.SetMetaData("model", InDevice.Model);
+						ConvertedDevice.SetMetaData("gpu", InDevice.GPU);
+						ConvertedDevice.SetMetaData("cpumodel", InDevice.CPUModel);
+						ConvertedDevice.SetMetaData("ram_in_gb", InDevice.RAMInGB.ToString());
+						ConvertedDevice.SetMetaData("render_mode", InDevice.RenderMode);
+						ConvertedDevice.SetMetaData("rhi", InDevice.RHI);
+					}
+				}
+				OutTestPassResults.TestSessionInfo.DateTime = InTestPassResults.ReportCreatedOn;
+				OutTestPassResults.TestSessionInfo.TimeElapseSec = InTestPassResults.TotalDuration;
+				if (InTestPassResults.Tests != null)
+				{
+					foreach (UnrealAutomatedTestResult InTestResult in InTestPassResults.Tests)
+					{
+						OutTestPassResults.SetTest(InTestResult.FullTestPath, InSuite, InTestResult.DeviceInstance, InTestResult.DateTime);
+						foreach (UnrealAutomationEntry InTestEntry in InTestResult.Entries)
+						{
+							string Tag = "entry";
+							string Context = InTestEntry.Event.Context;
+							string Artifact = InTestEntry.Event.Artifact;
+							UnrealAutomationArtifact IntArtifactEntry = null;
+							// If Artifact values is not null nor a bunch on 0, then we have a file attachment.
+							if (!string.IsNullOrEmpty(Artifact) && Artifact.Substring(0, 4) != "0000")
+							{
+								IntArtifactEntry = InTestResult.Artifacts.Find(A => A.Id == Artifact);
+								if (IntArtifactEntry.Type == "Comparison")
+								{
+									// UE for now only produces one type of artifact that can be attached to an event and that's image comparison
+									Tag = "image comparison";
+									Context = IntArtifactEntry.Name;
+								}
+							}
+							if (string.IsNullOrEmpty(Context) && !string.IsNullOrEmpty(InTestEntry.Filename))
+							{
+								Context = string.Format("{0} @line:{1}", InTestEntry.Filename, InTestEntry.LineNumber.ToString());
+							}
+							OutTestPassResults.AddEvent(InTestEntry.Event.Type, InTestEntry.Event.Message, Tag, Context, InTestEntry.Timestamp);
+							// Add Artifacts
+							if (IntArtifactEntry != null)
+							{
+								if (!string.IsNullOrEmpty(IntArtifactEntry.Files.Difference))
+								{
+									OutTestPassResults.AddArtifactToLastEvent(
+										"difference",
+										Path.Combine(InReportPath, IntArtifactEntry.Files.Difference),
+										IntArtifactEntry.Files.Difference
+									);
+								}
+								if (!string.IsNullOrEmpty(IntArtifactEntry.Files.Approved))
+								{
+									OutTestPassResults.AddArtifactToLastEvent(
+										"approved",
+										Path.Combine(InReportPath, IntArtifactEntry.Files.Approved),
+										IntArtifactEntry.Files.Approved
+									);
+								}
+								if (!string.IsNullOrEmpty(IntArtifactEntry.Files.Unapproved))
+								{
+									string AbsoluteLocation = Path.Combine(InReportPath, IntArtifactEntry.Files.Unapproved);
+									OutTestPassResults.AddArtifactToLastEvent(
+										"unapproved",
+										AbsoluteLocation,
+										IntArtifactEntry.Files.Unapproved
+									);
+									// Add Json meta data if any
+									string MetadataLocation = Utils.SystemHelpers.GetFullyQualifiedPath(Path.GetDirectoryName(AbsoluteLocation));
+									if (Directory.Exists(MetadataLocation))
+									{
+										string[] JsonMetadataFiles = System.IO.Directory.GetFiles(MetadataLocation, "*.json");
+										if (JsonMetadataFiles.Length > 0)
+										{
+											int LastSlash = IntArtifactEntry.Files.Unapproved.LastIndexOf("/");
+											string RelativeLocation = IntArtifactEntry.Files.Unapproved.Substring(0, LastSlash);
+											OutTestPassResults.AddEvent(
+												EventType.Info,
+												"The image reference can be updated by pointing the Screen Comparison tab from the Test Automation window to the artifacts from this test.",
+												"image comparison metadata", null,
+												InTestEntry.Timestamp
+											);
+											foreach (string JsonFile in JsonMetadataFiles)
+											{
+												string JsonArtifactName = RelativeLocation + "/" + Path.GetFileName(JsonFile);
+												OutTestPassResults.AddArtifactToLastEvent("json metadata", JsonFile, JsonArtifactName);
+											}
+										}
+									}
+								}
+							}
+						}
+						OutTestPassResults.SetTestState(InTestResult.State);
+						OutTestPassResults.SetTestTimeElapseSec(InTestResult.Duration);
+					}
+				}
+				return OutTestPassResults;
+			}
+		}
+
 		/// <summary>
 		/// Contains Test success/failed status and a list of errors and warnings
 		/// </summary>
@@ -514,7 +1035,7 @@ namespace Gauntlet
 			{
 				if (base.AttachArtifact(ArtifactPath, Name))
 				{
-					Logs.Add(Name ?? Path.GetFileName(ArtifactPath));
+					Logs.Add(string.Format("Attached artifact: {0}", Name ?? Path.GetFileName(ArtifactPath)));
 					return true;
 				}
 				return false;
@@ -536,16 +1057,29 @@ namespace Gauntlet
 				Items = new List<DataItem>();
 			}
 
-			public DataItem AddNewTestReport(string InKey, ITestReport InData)
+			public DataItem AddNewTestReport(ITestReport InData, string InKey = null)
 			{
-				if (string.IsNullOrEmpty(InKey))
+				if(InData is BaseHordeReport InHordeReport)
 				{
-					throw new System.ArgumentException("Test Data key can't be an empty string.");
+					InKey = InHordeReport.GetTestDataKey(InKey);
 				}
 				DataItem NewDataItem = new DataItem();
-				NewDataItem.Key = InData.Type +"::"+ InKey;
+				NewDataItem.Key = string.IsNullOrEmpty(InKey) ? InData.Type : InKey;
 				NewDataItem.Data = InData;
 				Items.Add(NewDataItem);
+
+				var ExtraItems = InData.GetReportDependencies();
+				if (ExtraItems.Count() > 0)
+				{
+					foreach (string Key in ExtraItems.Keys)
+					{
+						DataItem ExtraDataItem = new DataItem();
+						ExtraDataItem.Key = Key;
+						ExtraDataItem.Data = ExtraItems[Key];
+						Items.Add(ExtraDataItem);
+					}
+				}
+
 				return NewDataItem;
 			}
 
diff --git a/Engine/Source/Programs/AutomationTool/Gauntlet/Unreal/Automation/UE.Automation.cs b/Engine/Source/Programs/AutomationTool/Gauntlet/Unreal/Automation/UE.Automation.cs
--- a/Engine/Source/Programs/AutomationTool/Gauntlet/Unreal/Automation/UE.Automation.cs
+++ b/Engine/Source/Programs/AutomationTool/Gauntlet/Unreal/Automation/UE.Automation.cs
@@ -301,6 +301,11 @@ namespace UE
 		/// <returns></returns>
 		public override AutomationTestConfig GetConfiguration()
 		{
+			if (CachedConfig != null)
+			{
+				return CachedConfig;
+			}
+
 			AutomationTestConfig Config = base.GetConfiguration();
 
 			// Tests in the editor only require a single role
@@ -328,6 +333,11 @@ namespace UE
 		/// <returns></returns>
 		public override AutomationTestConfig GetConfiguration()
 		{
+			if (CachedConfig != null)
+			{
+				return CachedConfig;
+			}
+
 			AutomationTestConfig Config = base.GetConfiguration();
 
 			// Target tests require an editor which hosts the process
@@ -388,6 +398,33 @@ namespace UE
 			}
 		}
 
+		/// <summary>
+		/// Override the Suite by RunTest if it is a Group
+		/// </summary>
+		public override string Suite
+		{
+			get
+			{
+				string BaseSuite = base.Suite;
+				var Config = GetConfiguration();
+				if (Config is AutomationTestConfig)
+				{
+					var AutomationConfig = Config as AutomationTestConfig;
+					if (!string.IsNullOrEmpty(AutomationConfig.RunTest) && AutomationConfig.RunTest.ToLower().StartsWith("group:"))
+					{
+						BaseSuite  = AutomationConfig.RunTest.Substring(6);
+					}
+				}
+
+				return BaseSuite;
+			}
+		}
+
+		/// <summary>
+		/// Override the Type
+		/// </summary>
+		public override string Type { get { return "UE.Automation"; } }
+
 		/// <summary>
 		/// Called when a test is starting
 		/// </summary>
diff --git a/Engine/Source/Programs/AutomationTool/Gauntlet/Unreal/Base/Gauntlet.UnrealTestConfiguration.cs b/Engine/Source/Programs/AutomationTool/Gauntlet/Unreal/Base/Gauntlet.UnrealTestConfiguration.cs
--- a/Engine/Source/Programs/AutomationTool/Gauntlet/Unreal/Base/Gauntlet.UnrealTestConfiguration.cs
+++ b/Engine/Source/Programs/AutomationTool/Gauntlet/Unreal/Base/Gauntlet.UnrealTestConfiguration.cs
@@ -581,6 +581,12 @@ namespace Gauntlet
 		[AutoParam]
 		public string HordeArtifactPath = "";
 
+		/// <summary>
+		/// PreFlight change id
+		/// </summary>
+		[AutoParam]
+		public string PreFlightChange = "";
+
 		/// <summary>
 		/// Telemetry Database config to use
 		/// </summary>
@@ -905,7 +911,7 @@ namespace Gauntlet
 			}
 
 			// we write results to Horde test data if we run under Horde agent
-			if (HordeReport.IsUnderHordeAgent)
+			if (HordeReport.IsUnderHordeAgent && Globals.Params.ParseValues("WriteTestResultsForHorde").Count() == 0)
 			{
 				WriteTestResultsForHorde = true;
 			}
diff --git a/Engine/Source/Programs/AutomationTool/Gauntlet/Unreal/Base/Gauntlet.UnrealTestNode.cs b/Engine/Source/Programs/AutomationTool/Gauntlet/Unreal/Base/Gauntlet.UnrealTestNode.cs
--- a/Engine/Source/Programs/AutomationTool/Gauntlet/Unreal/Base/Gauntlet.UnrealTestNode.cs
+++ b/Engine/Source/Programs/AutomationTool/Gauntlet/Unreal/Base/Gauntlet.UnrealTestNode.cs
@@ -37,6 +37,16 @@ namespace Gauntlet
 		/// </summary>
 		public override string Name { get { return this.GetType().FullName; } }
 
+		/// <summary>
+		/// Returns the test suite. Default to the project name
+		/// </summary>
+		public virtual string Suite { get { return Context.BuildInfo.ProjectName; } }
+
+		/// <summary>
+		///Returns an identifier for this type of test
+		/// </summary>
+		public virtual string Type { get { return this.GetType().FullName; } }
+
 		/// <summary>
 		/// This class will log its own warnings and errors as part of its summary
 		/// </summary>
@@ -1127,7 +1137,7 @@ namespace Gauntlet
 		/// <returns>ITestReport</returns>
 		public virtual ITestReport CreateReport(TestResult Result)
 		{
-			if (GetConfiguration().WriteTestResultsForHorde)
+			if (GetCachedConfiguration().WriteTestResultsForHorde)
 			{
 				// write test report for Horde
 				HordeReport.SimpleTestReport HordeTestReport = CreateSimpleReportForHorde(Result);
@@ -1143,6 +1153,10 @@ namespace Gauntlet
 		/// <param name="Result"></param>
 		protected virtual HordeReport.SimpleTestReport CreateSimpleReportForHorde(TestResult Result)
 		{
+			if (string.IsNullOrEmpty(GetCachedConfiguration().HordeTestDataKey))
+			{
+				GetCachedConfiguration().HordeTestDataKey = Name + " " + Context.ToString();
+			}
 			HordeReport.SimpleTestReport HordeTestReport = new HordeReport.SimpleTestReport();
 			HordeTestReport.ReportCreatedOn = DateTime.Now.ToString();
 			HordeTestReport.TotalDurationSeconds = (float) (DateTime.Now - SessionStartTime).TotalSeconds;
@@ -1160,7 +1174,7 @@ namespace Gauntlet
 				SetUnrealTestResult(TestResult.Failed);
 			}
 			HordeTestReport.Status = GetTestResult().ToString();
-			string HordeArtifactPath = string.IsNullOrEmpty(GetConfiguration().HordeArtifactPath) ? HordeReport.DefaultArtifactsDir : GetConfiguration().HordeArtifactPath;
+			string HordeArtifactPath = string.IsNullOrEmpty(GetCachedConfiguration().HordeArtifactPath) ? HordeReport.DefaultArtifactsDir : GetCachedConfiguration().HordeArtifactPath;
 			HordeTestReport.SetOutputArtifactPath(HordeArtifactPath);
 			if (SessionArtifacts != null)
 			{
@@ -1182,7 +1196,7 @@ namespace Gauntlet
 				}
 			}
 			// Metadata
-			SetReportMetadata(HordeTestReport, GetConfiguration().RequiredRoles.Values.SelectMany(V => V));
+			SetReportMetadata(HordeTestReport, GetCachedConfiguration().RequiredRoles.Values.SelectMany(V => V));
 
 			return HordeTestReport;
 		}
@@ -1198,7 +1212,7 @@ namespace Gauntlet
 			string JsonReportPath = Path.Combine(UnrealAutomatedTestReportPath, "index.json");
 			if (File.Exists(JsonReportPath))
 			{
-				string HordeArtifactPath = GetConfiguration().HordeArtifactPath;
+				string HordeArtifactPath = GetCachedConfiguration().HordeArtifactPath;
 				Log.Verbose("Reading json Unreal Automated test report from {0}", JsonReportPath);
 				UnrealAutomatedTestPassResults JsonTestPassResults = UnrealAutomatedTestPassResults.LoadFromJson(JsonReportPath);
 				if (JsonTestPassResults.InProcess > 0)
@@ -1233,7 +1247,7 @@ namespace Gauntlet
 					// The test pass did not run at all
 					Log.Verbose("Found not-run tests: {0}", JsonTestPassResults.NotRun);
 					bool HasTimeout = RoleResults != null && RoleResults.Where(R => R.ProcessResult == UnrealProcessResult.TimeOut).Any();
-					if (GetConfiguration().ResumeOnCriticalFailure && !HasTimeout)
+					if (GetCachedConfiguration().ResumeOnCriticalFailure && !HasTimeout)
 					{
 						if (Retries < MaxRetries)
 						{
@@ -1260,13 +1274,23 @@ namespace Gauntlet
 						}
 					}
 				}
+				string HordeTestDataKey = string.IsNullOrEmpty(GetCachedConfiguration().HordeTestDataKey) ? Name + " " + Context.ToString() : GetCachedConfiguration().HordeTestDataKey;
+				GetCachedConfiguration().HordeTestDataKey = HordeTestDataKey;
 				// Convert test results for Horde
-				HordeReport.UnrealEngineTestPassResults HordeTestPassResults = HordeReport.UnrealEngineTestPassResults.FromUnrealAutomatedTests(JsonTestPassResults, ReportURL);
-				HordeTestPassResults.CopyTestResultsArtifacts(UnrealAutomatedTestReportPath, HordeArtifactPath);
+				HordeReport.AutomatedTestSessionData HordeTestPassResults = HordeReport.AutomatedTestSessionData.FromUnrealAutomatedTests(
+					JsonTestPassResults, string.Format("{0}({1})", Type, Suite), Suite, UnrealAutomatedTestReportPath, HordeArtifactPath
+				);
+				// Make a copy of the report in the old way - until we decide the transition is over
+				HordeReport.UnrealEngineTestPassResults CopyTestPassResults = HordeReport.UnrealEngineTestPassResults.FromUnrealAutomatedTests(JsonTestPassResults, ReportURL);
+				CopyTestPassResults.CopyTestResultsArtifacts(UnrealAutomatedTestReportPath, HordeArtifactPath);
+				HordeTestPassResults.AttachDependencyReport(CopyTestPassResults, HordeTestDataKey);
+				// Pre Flight information
+				HordeTestPassResults.PreFlightChange = GetCachedConfiguration().PreFlightChange;
 				// Metadata
 				// With UE Test Automation, we care only for one role.
-				var MainRole = new List<UnrealTestRole>() { GetConfiguration().GetMainRequiredRole() };
+				var MainRole = new List<UnrealTestRole>() { GetCachedConfiguration().GetMainRequiredRole() };
 				SetReportMetadata(HordeTestPassResults, MainRole);
+				SetReportMetadata(CopyTestPassResults, MainRole); // Set the metadata to the old report too
 				// Attached test Artifacts
 				if (SessionArtifacts != null)
 				{
@@ -1326,27 +1350,26 @@ namespace Gauntlet
 		/// <param name="Report"></param>
 		public virtual void SubmitToDashboard(ITestReport Report)
 		{
-			if (GetConfiguration().WriteTestResultsForHorde)
+			if (GetCachedConfiguration().WriteTestResultsForHorde)
 			{
 				// write test data collection for Horde
-				string HordeTestDataKey = string.IsNullOrEmpty(GetConfiguration().HordeTestDataKey) ? Name + " " + Context.ToString() : GetConfiguration().HordeTestDataKey;
 				string HordeTestDataFilePath = Path.Combine(
-					string.IsNullOrEmpty(GetConfiguration().HordeTestDataPath) ? HordeReport.DefaultTestDataDir : GetConfiguration().HordeTestDataPath,
+					string.IsNullOrEmpty(GetCachedConfiguration().HordeTestDataPath) ? HordeReport.DefaultTestDataDir : GetCachedConfiguration().HordeTestDataPath,
 					FileUtils.SanitizeFilename(Name) + ".TestData.json"
 				);
 				HordeReport.TestDataCollection HordeTestDataCollection = new HordeReport.TestDataCollection();
-				HordeTestDataCollection.AddNewTestReport(HordeTestDataKey, Report);
+				HordeTestDataCollection.AddNewTestReport(Report, GetCachedConfiguration().HordeTestDataKey);
 				HordeTestDataCollection.WriteToJson(HordeTestDataFilePath, true);
 			}
-			if (!string.IsNullOrEmpty(GetConfiguration().PublishTelemetryTo) && Report is ITelemetryReport Telemetry)
+			if (!string.IsNullOrEmpty(GetCachedConfiguration().PublishTelemetryTo) && Report is ITelemetryReport Telemetry)
 			{
 				IEnumerable<TelemetryData> DataRows = Telemetry.GetAllTelemetryData();
 				if (DataRows != null)
 				{
-					IDatabaseConfig<TelemetryData> DBConfig = DatabaseConfigManager<TelemetryData>.GetConfigByName(GetConfiguration().PublishTelemetryTo);
+					IDatabaseConfig<TelemetryData> DBConfig = DatabaseConfigManager<TelemetryData>.GetConfigByName(GetCachedConfiguration().PublishTelemetryTo);
 					if (DBConfig != null)
 					{
-						DBConfig.LoadConfig(GetConfiguration().DatabaseConfigPath);
+						DBConfig.LoadConfig(GetCachedConfiguration().DatabaseConfigPath);
 						IDatabaseDriver<TelemetryData> DB = DBConfig.GetDriver();
 						Log.Verbose("Submitting telemetry data to {0}", DB.ToString());
 
@@ -1354,7 +1377,7 @@ namespace Gauntlet
 						TestContext.SetProperty("ProjectName", Context.BuildInfo.ProjectName);
 						TestContext.SetProperty("Branch", Context.BuildInfo.Branch);
 						TestContext.SetProperty("Changelist", Context.BuildInfo.Changelist);
-						var RoleType = GetConfiguration().GetMainRequiredRole().Type;
+						var RoleType = GetCachedConfiguration().GetMainRequiredRole().Type;
 						var Role = Context.GetRoleContext(RoleType);
 						TestContext.SetProperty("Platform", Role.Platform);
 						TestContext.SetProperty("Configuration", string.Format("{0} {1}", RoleType, Role.Configuration));
@@ -1363,7 +1386,7 @@ namespace Gauntlet
 					}
 					else
 					{
-						Log.Warning("Got telemetry data, but database configuration is unknown '{0}'.", GetConfiguration().PublishTelemetryTo);
+						Log.Warning("Got telemetry data, but database configuration is unknown '{0}'.", GetCachedConfiguration().PublishTelemetryTo);
 					}
 				}
 			}
diff --git a/Engine/Source/Programs/AutomationTool/Gauntlet/Unreal/Utils/Gauntlet.UnrealTestResults.cs b/Engine/Source/Programs/AutomationTool/Gauntlet/Unreal/Utils/Gauntlet.UnrealTestResults.cs
--- a/Engine/Source/Programs/AutomationTool/Gauntlet/Unreal/Utils/Gauntlet.UnrealTestResults.cs
+++ b/Engine/Source/Programs/AutomationTool/Gauntlet/Unreal/Utils/Gauntlet.UnrealTestResults.cs
@@ -83,7 +83,9 @@ namespace Gauntlet
 		public string FullTestPath { get; set; }
 		public string ArtifactName { get; set; }
 		public TestStateType State { get; set; }
-		public String DeviceInstance { get; set; }
+		public string DeviceInstance { get; set; }
+		public float Duration { get; set; }
+		public string DateTime { get; set; }
 		public int Warnings { get; set; }
 		public int Errors { get; set; }
 		public List<UnrealAutomationArtifact> Artifacts { get; set; }
@@ -100,7 +102,7 @@ namespace Gauntlet
 			var Event = new UnrealAutomationEvent(EventType, Message);
 			var Entry = new UnrealAutomationEntry();
 			Entry.Event = Event;
-			Entry.Timestamp = DateTime.Now.ToString("yyyy.MM.dd-HH.mm.ss");
+			Entry.Timestamp = System.DateTime.Now.ToString("yyyy.MM.dd-HH.mm.ss");
 			Entries.Add(Entry);
 
 			switch(EventType)
