UAT: Fix problem of waiting forever to kill client or cook server process when launching with -cookonthefly.
Optimize StopProcess() when called with KillDescendants=true by using the standard Process.Kill(bool entireProcessTree) instead of our own slow KillAllDescendants() implementation.
Kill entire process trees when killing the launched processes from KillAll on uat exit.
Build machines will still call KillAllDescendants as the last fallback just as before.

#rnx

diff --git a/Engine/Source/Programs/AutomationTool/AutomationUtils/ProcessUtils.cs b/Engine/Source/Programs/AutomationTool/AutomationUtils/ProcessUtils.cs
--- a/Engine/Source/Programs/AutomationTool/AutomationUtils/ProcessUtils.cs
+++ b/Engine/Source/Programs/AutomationTool/AutomationUtils/ProcessUtils.cs
@@ -188,7 +188,7 @@ namespace AutomationTool
 						if (!Proc.HasExited)
 						{
 							CommandUtils.LogLog("Killing process: {0}", ProcName);
-							Proc.StopProcess(false);
+							Proc.StopProcess();
 						}
 					}
 					catch (Exception Ex)
@@ -655,13 +655,9 @@ namespace AutomationTool
 					Proc = null;
 				}
 				// Now actually kill the process and all its descendants if requested
-				if (KillDescendants)
-				{
-					KillAllDescendants(ProcToKill);
-				}
 				try
 				{
-					ProcToKill.Kill();
+					ProcToKill.Kill(KillDescendants);
 					ProcToKill.WaitForExit(60000);
 					if (!ProcToKill.HasExited)
 					{
