minor fix for CreatePlatformExtension if the .cs class name is does not have whitespace before the parent declaration.

#rnx

diff --git a/Engine/Source/Programs/AutomationTool/Scripts/CreatePlatformExtension.cs b/Engine/Source/Programs/AutomationTool/Scripts/CreatePlatformExtension.cs
--- a/Engine/Source/Programs/AutomationTool/Scripts/CreatePlatformExtension.cs
+++ b/Engine/Source/Programs/AutomationTool/Scripts/CreatePlatformExtension.cs
@@ -451,6 +451,7 @@ public class CreatePlatformExtension : BuildCommand
 		}
 
 		// load module file & find module class name, and optional class namespace
+		char[] ClassNameSeparators = new char[] { ' ', '\t', ':' };
 		const string ClassDeclaration = "public class ";
 		const string NamespaceDeclaration = "namespace ";
 		string[] ModuleContents = File.ReadAllLines(ModulePath.FullName);
@@ -461,10 +462,10 @@ public class CreatePlatformExtension : BuildCommand
 			Log.TraceError($"Cannot find class declaration in ${ModulePath}");
 			return;
 		}
-		string ParentModuleName = ModuleClassDeclaration.Trim().Remove(0, ClassDeclaration.Length).Split(' ', StringSplitOptions.None).Last();
+		string ParentModuleName = ModuleClassDeclaration.Trim().Remove(0, ClassDeclaration.Length).Split(ClassNameSeparators, StringSplitOptions.None).Last();
 		if (bAllowPlatformExtensionsAsParents || ParentModuleName.Equals("ModuleRules") || ParentModuleName.Equals("TargetRules") )
 		{
-			ParentModuleName = ModuleClassDeclaration.Trim().Remove(0, ClassDeclaration.Length ).Split(' ', StringSplitOptions.None ).First();
+			ParentModuleName = ModuleClassDeclaration.Trim().Remove(0, ClassDeclaration.Length ).Split(ClassNameSeparators, StringSplitOptions.None ).First();
 		}
 		if (string.IsNullOrEmpty(ParentModuleName))
 		{
