Add initial support for Deploying Windows games to the SteamDeck through .ini entries

diff --git a/Engine/Source/Developer/Windows/WindowsTargetPlatform/Private/GenericWindowsTargetPlatform.h b/Engine/Source/Developer/Windows/WindowsTargetPlatform/Private/GenericWindowsTargetPlatform.h
--- a/Engine/Source/Developer/Windows/WindowsTargetPlatform/Private/GenericWindowsTargetPlatform.h
+++ b/Engine/Source/Developer/Windows/WindowsTargetPlatform/Private/GenericWindowsTargetPlatform.h
@@ -9,6 +9,7 @@
 #include "LocalPcTargetDevice.h"
 #endif
 #include "Serialization/MemoryLayout.h"
+#include "SteamDeckDevice.h"
 
 #if WITH_ENGINE
 	#include "Sound/SoundWave.h"
@@ -45,6 +46,9 @@ public:
 #if PLATFORM_WINDOWS
 		// only add local device if actually running on Windows
 		LocalDevice = MakeShareable(new TTargetDevice(*this));
+
+		// Check if we have any SteamDeck devices around
+		SteamDevices = FSteamDeckDevice::DiscoverDevices(*this);
 #endif
 
 	#if WITH_ENGINE
@@ -91,6 +95,7 @@ public:
 		bRequiresEncodedHDRReflectionCaptures =	TargetedShaderFormats.Contains(NAME_SF_VULKAN_ES31)
 												|| TargetedShaderFormats.Contains(NAME_OPENGL_150_ES3_1)
 												|| TargetedShaderFormats.Contains(NAME_PCD3D_ES3_1);
+
 	#endif
 	}
 
@@ -107,6 +112,14 @@ public:
 		{
 			OutDevices.Add(LocalDevice);
 		}
+
+		for (const ITargetDevicePtr& SteamDeck : SteamDevices)
+		{
+			if (SteamDeck.IsValid())
+			{
+				OutDevices.Add(SteamDeck);
+			}
+		}
 	}
 
 	virtual bool GenerateStreamingInstallManifest(const TMultiMap<FString, int32>& PakchunkMap, const TSet<int32>& PakchunkIndicesInUse) const override
@@ -131,6 +144,14 @@ public:
 			return LocalDevice;
 		}
 
+		for (const ITargetDevicePtr& SteamDeck : SteamDevices)
+		{
+			if (SteamDeck.IsValid() && DeviceId == SteamDeck->GetId())
+			{
+				return SteamDeck;
+			}
+		}
+
 		return nullptr;
 	}
 
@@ -362,6 +383,8 @@ private:
 	// Holds the local device.
 	ITargetDevicePtr LocalDevice;
 
+	TArray<ITargetDevicePtr> SteamDevices;
+
 #if WITH_ENGINE
 	// Holds the texture LOD settings.
 	const UTextureLODSettings* TextureLODSettings;
diff --git a/Engine/Source/Developer/Windows/WindowsTargetPlatform/Private/SteamDeckDevice.h b/Engine/Source/Developer/Windows/WindowsTargetPlatform/Private/SteamDeckDevice.h
new file mode 100644
--- /dev/null
+++ b/Engine/Source/Developer/Windows/WindowsTargetPlatform/Private/SteamDeckDevice.h
@@ -0,0 +1,90 @@
+// Copyright Epic Games, Inc. All Rights Reserved.
+
+#pragma once
+
+#include "CoreMinimal.h"
+#include "Common/TargetPlatformBase.h"
+#include "Misc/ConfigCacheIni.h"
+#if PLATFORM_WINDOWS
+#include "LocalPcTargetDevice.h"
+#endif
+#include "Serialization/MemoryLayout.h"
+
+// Extend SteamDevice from LocalPcDevice as we are leveraging the Windows Device for Steam Deck when building
+// for the Win64 platform. As SteamDeck consumes the exact same build it would for a normal Win64 device
+#if PLATFORM_WINDOWS
+class FSteamDeckDevice : public TLocalPcTargetDevice<true>
+{
+public:
+	FSteamDeckDevice(FString InIpAddr, FString InDeviceName, FString InUserName, const ITargetPlatform& InTargetPlatform)
+		: TLocalPcTargetDevice<true>(InTargetPlatform)
+		, IpAddr(InIpAddr)
+		, DeviceName(InDeviceName)
+		, UserName(InUserName)
+	{
+	}
+
+	virtual FString GetName() const override
+	{
+		return DeviceName;
+	}
+
+	virtual FTargetDeviceId GetId() const override
+	{
+		return FTargetDeviceId(TargetPlatform.PlatformName(), IpAddr);
+	}
+
+	virtual FString GetOperatingSystemName() override
+	{
+		return TEXT("SteamOS");
+	}
+
+	virtual bool GetUserCredentials(FString& OutUserName, FString& OutUserPassword) override
+	{
+		OutUserName = UserName;
+
+		// No need for a password, as it will use an rsa key that is part of the SteamOS Devkit Client
+		OutUserPassword = FString();
+
+		return true;
+	}
+
+	static TArray<ITargetDevicePtr> DiscoverDevices(const ITargetPlatform& GenericWindowsTP)
+	{
+		TArray<FString> EngineIniSteamDeckDevices;
+
+		// Expected ini format: +SteamDeckDevice=(IpAddr=10.1.33.19,Name=MySteamDeck,UserName=deck)
+		GConfig->GetArray(TEXT("/Script/WindowsTargetPlatform.WindowsTargetSettings"), TEXT("SteamDeckDevice"), EngineIniSteamDeckDevices, GEngineIni);
+
+		TArray<ITargetDevicePtr> SteamDevices;
+		for (const FString& Device : EngineIniSteamDeckDevices)
+		{
+			FString IpAddr;
+			FString Name;
+			FString UserName;
+			if (!FParse::Value(*Device, TEXT("IpAddr="), IpAddr))
+			{
+				continue;
+			}
+			// Name is not required, if not set use the IpAddr. This is what is displayed in the Editor
+			if (!FParse::Value(*Device, TEXT("Name="), Name))
+			{
+				Name = IpAddr;
+			}
+			if (!FParse::Value(*Device, TEXT("UserName="), UserName))
+			{
+				continue;
+			}
+
+			SteamDevices.Add(MakeShareable(new FSteamDeckDevice(IpAddr, Name, UserName, GenericWindowsTP)));
+		}
+
+		return SteamDevices;
+	}
+
+private:
+	FString IpAddr;
+	FString DeviceName;
+	FString UserName;
+};
+#endif // PLATFORM_WINDOWS
diff --git a/Engine/Source/Programs/AutomationTool/Win/WinPlatform.Automation.cs b/Engine/Source/Programs/AutomationTool/Win/WinPlatform.Automation.cs
--- a/Engine/Source/Programs/AutomationTool/Win/WinPlatform.Automation.cs
+++ b/Engine/Source/Programs/AutomationTool/Win/WinPlatform.Automation.cs
@@ -4,6 +4,7 @@ using System.Collections.Generic;
 using System.Linq;
 using System.Text;
 using System.IO;
+using System.Text.RegularExpressions;
 using AutomationTool;
 using UnrealBuildTool;
 using Microsoft.Win32;
@@ -11,6 +12,162 @@ using System.Diagnostics;
 using EpicGames.Core;
 using UnrealBuildBase;
 
+public static class SteamDeckSupport
+{
+	public static string RSyncPath = Path.Combine(Unreal.RootDirectory.FullName, "Engine\\Extras\\ThirdPartyNotUE\\cwrsync\\bin\\rsync.exe");
+	public static string SSHPath   = Path.Combine(Unreal.RootDirectory.FullName, "Engine\\Extras\\ThirdPartyNotUE\\cwrsync\\bin\\ssh.exe");
+
+	public static string GetRegisterGameScript(string GameId, string GameExePath, string GameFolderPath, string GameRunArgs)
+	{
+		// create a script that will be copied over to the SteamDeck and ran to register the game
+		// TODO make these easier to customize, vs hard coding the settings. Assume debugging for now, requires the user to have uploaded the required msvsmom/remote debugging stuff
+		// which is done through uploading any game with debugging enabled through the SteamOS Devkit Client
+		string GameIdArgs    = String.Format("\"gameid\":\"{0}\"", GameId);
+		string DirectoryArgs = String.Format("\"directory\":\"{0}\"", GameFolderPath);
+		string ArgvArgs      = String.Format("\"argv\":[\"{0} {1}\"]", GameExePath, GameRunArgs);
+		string SettingsArgs  = String.Format("\"settings\": {{\"steam_play\": \"1\", \"steam_play_debug\": \"1\", \"steam_play_debug_version\": \"2019\"}}");
+
+		return String.Format("#!/bin/bash\npython3 ~/devkit-utils/steam-client-create-shortcut --parms '{{{0}, {1}, {2}, {3}}}'", GameIdArgs, DirectoryArgs, ArgvArgs, SettingsArgs).Replace("\r\n", "\n");
+	}
+
+	// This is a bit nasty, due to rsync needing to use cygdrive path for its local location over Windows paths.
+	// This will not work with UNC paths
+	public static string ConvertWindowsPathToCygdrive(string WindowsPath)
+	{
+		string CygdrivePath = "";
+
+		if (!string.IsNullOrEmpty(WindowsPath))
+		{
+			string FullPath = Path.GetFullPath(WindowsPath);
+			string RootPath = Path.GetPathRoot(FullPath);
+
+			System.Console.WriteLine("{0}", RootPath);
+			CygdrivePath = Path.Combine("/cygdrive", Char.ToLower(FullPath[0]).ToString(), FullPath.Substring(RootPath.Length));
+
+			return CygdrivePath.Replace('\\','/');
+		}
+
+		return CygdrivePath;
+	}
+
+	public static List<DeviceInfo> GetDevices()
+	{
+		List<DeviceInfo> Devices = new List<DeviceInfo>();
+
+		// Look for any Steam Deck devices that are in the Engine ini files. If so lets add them as valid devices
+		// This will be required for matching devices passed into the BuildCookRun command
+		ConfigHierarchy EngineConfig = ConfigCache.ReadHierarchy(ConfigHierarchyType.Engine, null, UnrealTargetPlatform.Win64);
+
+		List<string> SteamDeckDevices;
+		if (EngineConfig.GetArray("/Script/WindowsTargetPlatform.WindowsTargetSettings", "SteamDeckDevice", out SteamDeckDevices))
+		{
+			// Expected ini format: +SteamDeckDevice=(IpAddr=10.1.33.19,Name=MySteamDeck,UserName=deck)
+			foreach (string DeckDevice in SteamDeckDevices)
+			{
+				string IpAddr     = GetStructEntry(DeckDevice, "IpAddr", false);
+				string DeviceName = GetStructEntry(DeckDevice, "Name", false);
+				// Skipping the UserName as its unused here
+
+				// Name is optional, if its empty/not found lets just use the IpAddr for the Name
+				if (string.IsNullOrEmpty(DeviceName))
+				{
+					DeviceName = IpAddr;
+				}
+
+				if (!string.IsNullOrEmpty(IpAddr))
+				{
+					// TODO Fix the usage of OSVersion here. We are abusing this and using MS OSVersion to allow Turnkey to be happy
+					DeviceInfo SteamDeck = new DeviceInfo(UnrealTargetPlatform.Win64, DeviceName, IpAddr,
+						Environment.OSVersion.Version.ToString(), "SteamDeck", true, true);
+
+					Devices.Add(SteamDeck);
+				}
+			}
+		}
+
+		return Devices;
+	}
+
+	// GetStructEntry copied from ExecuteBuild.cs TODO move to a better to share this
+	private static string GetStructEntry(string Input, string Property, bool bIsArrayProperty)
+	{
+		string PrimaryRegex;
+		string AltRegex = null;
+		if (bIsArrayProperty)
+		{
+			PrimaryRegex = string.Format("{0}\\s*=\\s*\\((.*?)\\)", Property);
+		}
+		else
+		{
+			// handle quoted strings, allowing for escaped quotation marks (basically doing " followed by whatever, until we see a quote that was not proceeded by a \, and gather the whole mess in an outer group)
+			PrimaryRegex = string.Format("{0}\\s*=\\s*\"((.*?)[^\\\\])\"", Property);
+			// when no quotes, we skip over whitespace, and we end when we see whitespace, a comma or a ). This will handle (Ip = 192.168.0.1 , Name=....) , and return only '192.168.0.1'
+			AltRegex = string.Format("{0}\\s*=\\s*(.*?)[\\s,\\)]", Property);
+		}
+
+		// attempt to match it!
+		Match Result = Regex.Match(Input, PrimaryRegex);
+		if (!Result.Success && AltRegex != null)
+		{
+			Result = Regex.Match(Input, AltRegex);
+		}
+
+		// if we got a success, return the main match value
+		if (Result.Success)
+		{
+			return Result.Groups[1].Value.ToString();
+		}
+
+		return null;
+	}
+
+	/* Deploying to a steam deck currently does 2 things
+	 *
+	 * 1) Generates a script CreateShortcutHelper.sh that will register the game on the SteamDeck once uploaded
+	 * 2) Uploads the build using rsync to the devkit-game location. Once uploaded it runs the CreateShortcutHelper.sh generated before.
+	 */
+	public static void Deploy(ProjectParams Params, DeploymentContext SC)
+	{
+		string DevKitRSAPath  = Path.Combine(CommandUtils.GetEnvVar("LOCALAPPDATA"), "steamos-devkit\\steamos-devkit\\devkit_rsa");
+		string SSHArgs        = String.Format("-i {0} {1}@{2}", DevKitRSAPath, Params.DeviceUsername, Params.DeviceNames[0]);
+		string GameFolderPath = Path.Combine("/home", Params.DeviceUsername, "devkit-game", Params.ShortProjectName).Replace('\\', '/');
+		string GameRunArgs    = String.Format("{0} {1} {2}", SC.ProjectArgForCommandLines, Params.StageCommandline, Params.RunCommandline).Replace("\"", "\\\"");
+
+		FileReference ExePath = Params.GetProjectExeForPlatform(UnrealTargetPlatform.Win64);
+		string RelGameExePath = ExePath.MakeRelativeTo(DirectoryReference.Combine(ExePath.Directory, "../../..")).Replace('\\', '/');
+
+		if (!File.Exists(DevKitRSAPath))
+		{
+			CommandUtils.LogWarning("Unable to find '{0}' rsa key needed to deploy to the steam deck. Make sure you've installed the SteamOS Devkit client", DevKitRSAPath);
+			return;
+		}
+
+		string ScriptFileName = "CreateShortcutHelper.sh";
+		string ScriptFile = Path.Combine(SC.StageDirectory.FullName, ScriptFileName);
+		File.WriteAllText(ScriptFile, SteamDeckSupport.GetRegisterGameScript(Params.ShortProjectName, RelGameExePath, GameFolderPath, GameRunArgs));
+
+		// Exclude removing the Saved folders to preserve logs and crash data. Though note these will keep filling up with data
+		IProcessResult Result = CommandUtils.Run(SteamDeckSupport.RSyncPath,
+			String.Format("-avh --delete --exclude=\"Saved/\" --chmod=Du=rwx,Dgo=rx,Fu=rwx,Fog=rx -e \"{0} -o StrictHostKeyChecking=no -i '{1}'\" --update \"{2}/\" {3}@{4}:{5}",
+			SteamDeckSupport.SSHPath, DevKitRSAPath, SteamDeckSupport.ConvertWindowsPathToCygdrive(SC.StageDirectory.FullName), Params.DeviceUsername, Params.DeviceNames[0], GameFolderPath), "");
+
+		if (Result.ExitCode > 0)
+		{
+			CommandUtils.LogWarning("Failed to rsync files to the SteamDeck. Check connection on ip {0}@{1}", Params.DeviceUsername, Params.DeviceNames[0]);
+			return;
+		}
+
+		// Run the script to register the game with the Deck
+		Result = CommandUtils.Run(SteamDeckSupport.SSHPath, String.Format("{0} \"chmod +x {1}/{2} && {1}/{2}\"", SSHArgs, GameFolderPath, ScriptFileName), "");
+
+		if (Result.ExitCode > 0)
+		{
+			CommandUtils.LogWarning("Failed to run the {0}.sh script. Check connection on ip {1}@{2}", ScriptFileName, Params.DeviceUsername, Params.DeviceNames[0]);
+			return;
+		}
+	}
+}
+
 public class Win64Platform : Platform
 {
 	public Win64Platform()
@@ -33,12 +190,34 @@ public class Win64Platform : Platform
 				Environment.OSVersion.Version.ToString(), "Computer", true, true);
 
 			Devices.Add(LocalMachine);
+
+			Devices.AddRange(SteamDeckSupport.GetDevices());
 		}
 
 		return Devices.ToArray();
 	}
 
+	public override void Deploy(ProjectParams Params, DeploymentContext SC)
+	{
+		// We only care about deploying for SteamDeck
+		if (Params.Devices.Count == 1 && GetDevices().FirstOrDefault(x => x.Id == Params.DeviceNames[0])?.Type == "SteamDeck")
+		{
+			SteamDeckSupport.Deploy(Params, SC);
+		}
+	}
 
+	public override IProcessResult RunClient(ERunOptions ClientRunFlags, string ClientApp, string ClientCmdLine, ProjectParams Params)
+	{
+		if (Params.Devices.Count == 1 && GetDevices().FirstOrDefault(x => x.Id == Params.DeviceNames[0])?.Type == "SteamDeck")
+		{
+			// TODO figure out how to get the steam app num id. Then can run like this:
+			// steam steam://rungameid/<GameNumId>
+			// TODO would be great if we could tail the log while running. Figure out how to cancel/exit app
+			return null;
+		}
+
+		return base.RunClient(ClientRunFlags, ClientApp, ClientCmdLine, Params);
+	}
 
 	protected override string GetPlatformExeExtension()
 	{
