Clean up unused source files when AFS disabled

#android

diff --git a/Engine/Plugins/Runtime/AndroidFileServer/Source/AndroidFileServer/AndroidFileServer_UPL.xml b/Engine/Plugins/Runtime/AndroidFileServer/Source/AndroidFileServer/AndroidFileServer_UPL.xml
--- a/Engine/Plugins/Runtime/AndroidFileServer/Source/AndroidFileServer/AndroidFileServer_UPL.xml
+++ b/Engine/Plugins/Runtime/AndroidFileServer/Source/AndroidFileServer/AndroidFileServer_UPL.xml
@@ -127,9 +127,26 @@
 					<true>
 						<copyFile src="$S(PluginDir)/Android/java/com/epicgames/unreal/RemoteFileManagerReceiver.java" dst="$S(BuildDir)/src/com/epicgames/unreal/RemoteFileManagerReceiver.java" />
 					</true>
+					<false>
+						<isArch arch="arm64-v8a">
+							<deleteFiles filespec="arm64/src/com/epicgames/unreal/RemoteFileManagerReceiver.java" />
+						</isArch>
+						<isArch arch="x86_64">
+							<deleteFiles filespec="x64/src/com/epicgames/unreal/RemoteFileManagerReceiver.java" />
+						</isArch>
+					</false>
 				</if>
-<!--				<copyDir src="$S(PluginDir)/Android/libs" dst="$S(BuildDir)/libs" /> -->
 			</true>
+			<false>
+				<isArch arch="arm64-v8a">
+					<deleteFiles filespec="arm64/src/com/epicgames/unreal/RemoteFileManager.java" />
+					<deleteFiles filespec="arm64/src/com/epicgames/unreal/RemoteFileManagerReceiver.java" />
+				</isArch>
+				<isArch arch="x86_64">
+					<deleteFiles filespec="x64/src/com/epicgames/unreal/RemoteFileManager.java" />
+					<deleteFiles filespec="x64/src/com/epicgames/unreal/RemoteFileManagerReceiver.java" />
+				</isArch>
+			</false>
 		</if>
 	</resourceCopies>
 
diff --git a/Engine/Source/Programs/AutomationTool/Android/AndroidPlatform.Automation.cs b/Engine/Source/Programs/AutomationTool/Android/AndroidPlatform.Automation.cs
--- a/Engine/Source/Programs/AutomationTool/Android/AndroidPlatform.Automation.cs
+++ b/Engine/Source/Programs/AutomationTool/Android/AndroidPlatform.Automation.cs
@@ -553,14 +553,22 @@ public class AndroidPlatform : Platform
 		bIsShipping = TargetConfiguration == UnrealTargetConfiguration.Shipping;
 
 		ConfigHierarchy Ini = ConfigCache.ReadHierarchy(ConfigHierarchyType.Engine, DirectoryReference.FromFile(SC.RawProjectPath), SC.StageTargetPlatform.PlatformType);
-		bEnablePlugin = true;
-		AFSToken = "";
-		bIncludeInShipping = false;
-		bAllowExternalStartInShipping = false;
-		Ini.GetBool("/Script/AndroidFileServerEditor.AndroidFileServerRuntimeSettings", "bEnablePlugin", out bEnablePlugin);
-		Ini.GetString("/Script/AndroidFileServerEditor.AndroidFileServerRuntimeSettings", "SecurityToken", out AFSToken);
-		Ini.GetBool("/Script/AndroidFileServerEditor.AndroidFileServerRuntimeSettings", "bIncludeInShipping", out bIncludeInShipping);
-		Ini.GetBool("/Script/AndroidFileServerEditor.AndroidFileServerRuntimeSettings", "bAllowExternalStartInShipping", out bAllowExternalStartInShipping);
+		if (!Ini.GetBool("/Script/AndroidFileServerEditor.AndroidFileServerRuntimeSettings", "bEnablePlugin", out bEnablePlugin))
+		{
+			bEnablePlugin = true;
+		}
+		if (!Ini.GetString("/Script/AndroidFileServerEditor.AndroidFileServerRuntimeSettings", "SecurityToken", out AFSToken))
+		{
+			AFSToken = "";
+		}
+		if (!Ini.GetBool("/Script/AndroidFileServerEditor.AndroidFileServerRuntimeSettings", "bIncludeInShipping", out bIncludeInShipping))
+		{
+			bIncludeInShipping = false;
+		}
+		if (!Ini.GetBool("/Script/AndroidFileServerEditor.AndroidFileServerRuntimeSettings", "bAllowExternalStartInShipping", out bAllowExternalStartInShipping))
+		{
+			bAllowExternalStartInShipping = false;
+		}
 
 		if (bIsShipping && !(bIncludeInShipping && bAllowExternalStartInShipping))
 		{
@@ -601,23 +609,32 @@ public class AndroidPlatform : Platform
 				break;
 		}
 
-		bUseCompression = false;
-		bLogFiles = false;
-		bReportStats = false;
-		bUseManualIPAddress = false;
-		ManualIPAddress = "";
-		bool bAllowNetworkConnection = true;
-
-		Ini.GetBool("/Script/AndroidFileServerEditor.AndroidFileServerRuntimeSettings", "bUseCompression", out bUseCompression);
-		Ini.GetBool("/Script/AndroidFileServerEditor.AndroidFileServerRuntimeSettings", "bLogFiles", out bLogFiles);
-		Ini.GetBool("/Script/AndroidFileServerEditor.AndroidFileServerRuntimeSettings", "bReportStats", out bReportStats);
-		Ini.GetBool("/Script/AndroidFileServerEditor.AndroidFileServerRuntimeSettings", "bUseManualIPAddress", out bUseManualIPAddress);
+		if (!Ini.GetBool("/Script/AndroidFileServerEditor.AndroidFileServerRuntimeSettings", "bUseCompression", out bUseCompression))
+		{
+			bUseCompression = false;
+		}
+		if (!Ini.GetBool("/Script/AndroidFileServerEditor.AndroidFileServerRuntimeSettings", "bLogFiles", out bLogFiles))
+		{
+			bLogFiles = false;
+		}
+		if (!Ini.GetBool("/Script/AndroidFileServerEditor.AndroidFileServerRuntimeSettings", "bReportStats", out bReportStats))
+		{
+			bReportStats = false;
+		}
+		if (!Ini.GetBool("/Script/AndroidFileServerEditor.AndroidFileServerRuntimeSettings", "bUseManualIPAddress", out bUseManualIPAddress))
+		{
+			bUseManualIPAddress = false;
+		}
 		if (!Ini.GetString("/Script/AndroidFileServerEditor.AndroidFileServerRuntimeSettings", "ManualIPAddress", out ManualIPAddress))
 		{
 			ManualIPAddress = "127.0.0.1";
 		}
 
-		Ini.GetBool("/Script/AndroidFileServerEditor.AndroidFileServerRuntimeSettings", "bAllowNetworkConnection", out bAllowNetworkConnection);
+		bool bAllowNetworkConnection = true;
+		if (!Ini.GetBool("/Script/AndroidFileServerEditor.AndroidFileServerRuntimeSettings", "bAllowNetworkConnection", out bAllowNetworkConnection))
+		{
+			bAllowNetworkConnection = true;
+		}
 		if (!bAllowNetworkConnection && ConnectionType != EConnectionType.USBOnly)
 		{
 			Log.TraceWarning("AFS will only use USB connection due to network connection disabled");
@@ -2088,7 +2105,7 @@ public class AndroidPlatform : Platform
 			// start up AFS connection (allowed to fail here.. may not be a server installed yet)
 			AndroidFileClient client = new AndroidFileClient(DeviceName);
 			AndroidFileClient client2 = null;
-			string IPAddress = bUseManualIPAddress ? ManualIPAddress : "127.0.0.1";
+			string IPAddress = (ConnectionType == EConnectionType.NetworkOnly && bUseManualIPAddress) ? ManualIPAddress : "127.0.0.1";
 			if (!client.OpenConnection(IPAddress))
 			{
 				if (!client.StartServer(PackageName, AFSToken, IPAddress))
