Fix minor issue on devkit-game folder may not exist on first upload

diff --git a/Engine/Source/Programs/AutomationTool/Win/WinPlatform.Automation.cs b/Engine/Source/Programs/AutomationTool/Win/WinPlatform.Automation.cs
--- a/Engine/Source/Programs/AutomationTool/Win/WinPlatform.Automation.cs
+++ b/Engine/Source/Programs/AutomationTool/Win/WinPlatform.Automation.cs
@@ -148,7 +148,7 @@ public static class SteamDeckSupport
 
 		// Exclude removing the Saved folders to preserve logs and crash data. Though note these will keep filling up with data
 		IProcessResult Result = CommandUtils.Run(SteamDeckSupport.RSyncPath,
-			String.Format("-avh --delete --exclude=\"Saved/\" --chmod=Du=rwx,Dgo=rx,Fu=rwx,Fog=rx -e \"{0} -o StrictHostKeyChecking=no -i '{1}'\" --update \"{2}/\" {3}@{4}:{5}",
+			String.Format("-avh --delete --exclude=\"Saved/\" --rsync-path=\"mkdir -p {5} && rsync\" --chmod=Du=rwx,Dgo=rx,Fu=rwx,Fog=rx -e \"{0} -o StrictHostKeyChecking=no -i '{1}'\" --update \"{2}/\" {3}@{4}:{5}",
 			SteamDeckSupport.SSHPath, DevKitRSAPath, SteamDeckSupport.ConvertWindowsPathToCygdrive(SC.StageDirectory.FullName), Params.DeviceUsername, Params.DeviceNames[0], GameFolderPath), "");
 
 		if (Result.ExitCode > 0)
