Added explicit check and exeption for the case where a zip file doesn't exist.It should make sume build failures easier to read.

diff --git a/Engine/Source/Programs/AutomationTool/AutomationUtils/CommandUtils.cs b/Engine/Source/Programs/AutomationTool/AutomationUtils/CommandUtils.cs
--- a/Engine/Source/Programs/AutomationTool/AutomationUtils/CommandUtils.cs
+++ b/Engine/Source/Programs/AutomationTool/AutomationUtils/CommandUtils.cs
@@ -2247,6 +2247,12 @@ namespace AutomationTool
 		public static IEnumerable<FileReference> UnzipFiles(FileReference ZipFileName, DirectoryReference BaseDirectory)
 		{
 			List<FileReference> OutputFiles = new List<FileReference>();
+
+			if (!FileReference.Exists(ZipFileName))
+			{
+				throw new AutomationException("Cannot unzip {0}. File not found", ZipFileName);
+			}
+
 			using (ZipArchive ZipArchive = ZipFile.Open(ZipFileName.FullName, ZipArchiveMode.Read))
 			{
 				foreach (ZipArchiveEntry Entry in ZipArchive.Entries)
