- Added PlatformValues to DeviceInfo in UAT land. Removed the Subtype field and changed it's very limited use to PlatformValues
- Added support for copying remote debugging files to SteamDeck when deploying
- Username is now pulled out of the Device info, or put on the commandline via -deviceuser to override

#p4v-preflight-copy 19392535

diff --git a/Engine/Source/Programs/AutomationTool/AutomationUtils/Platform.cs b/Engine/Source/Programs/AutomationTool/AutomationUtils/Platform.cs
--- a/Engine/Source/Programs/AutomationTool/AutomationUtils/Platform.cs
+++ b/Engine/Source/Programs/AutomationTool/AutomationUtils/Platform.cs
@@ -35,16 +35,19 @@ namespace AutomationTool
 			this.Platform = Platform;
 		}
 
-		public DeviceInfo(UnrealTargetPlatform Platform, string Name, string Id, string SoftwareVersion, string Type, bool bIsDefault, bool bCanConnect, string SubType = "")
+		public DeviceInfo(UnrealTargetPlatform Platform, string Name, string Id, string SoftwareVersion, string Type, bool bIsDefault, bool bCanConnect, Dictionary<string, string> PlatformValues=null)
 		{
 			this.Platform = Platform;
 			this.Name = Name;
 			this.Id = Id;
 			this.SoftwareVersion = SoftwareVersion;
 			this.Type = Type;
-			this.SubType = SubType;
 			this.bIsDefault = bIsDefault;
 			this.bCanConnect = bCanConnect;
+			if (PlatformValues != null)
+			{
+				this.PlatformValues = new Dictionary<string, string>(PlatformValues);
+			}
 		}
 
 		public UnrealTargetPlatform Platform;
@@ -52,11 +55,13 @@ namespace AutomationTool
 		public string Id;
 		public string SoftwareVersion;
 		public string Type;
-		public string SubType;
 		public bool bIsDefault = false;
 		// is the device able to be connected to (this is more about able to flash SDK or run, not about matching SDK version)
 		// if false, any of the above fields are suspect, especually SoftwareVersion
 		public bool bCanConnect = true;
+
+		// case insesitive platform value dictionary. turnkey doesn't use this, but the platform can look up the device during deployment, etc to get this out
+		public Dictionary<string, string> PlatformValues = new Dictionary<string, string>(StringComparer.InvariantCultureIgnoreCase);
 	}
 
 	/// <summary>
@@ -80,7 +85,7 @@ namespace AutomationTool
 
 		public override string ToString()
 		{
-			return Type.ToString();
+			return Type.ToString() + CookFlavor;
 		}
 	}
 
diff --git a/Engine/Source/Programs/AutomationTool/IOS/IOSPlatform.Automation.cs b/Engine/Source/Programs/AutomationTool/IOS/IOSPlatform.Automation.cs
--- a/Engine/Source/Programs/AutomationTool/IOS/IOSPlatform.Automation.cs
+++ b/Engine/Source/Programs/AutomationTool/IOS/IOSPlatform.Automation.cs
@@ -537,12 +537,12 @@ public class IOSPlatform : Platform
 
 			if (UnparsedUDID.Contains("Network"))
 			{
-				CurrentDevice.SubType = "Network";
+				CurrentDevice.PlatformValues["Connection"] = "Network";
 				IdeviceInfoArgs = "-n " + IdeviceInfoArgs;
 			}
 			else
 			{
-				CurrentDevice.SubType = "USB";
+				CurrentDevice.PlatformValues["Connection"] = "USB";
 			}
 
 			string OutputInfo = Utils.RunLocalProcessAndReturnStdOut(IdeviceInfoPath, IdeviceInfoArgs);
@@ -1739,7 +1739,7 @@ public class IOSPlatform : Platform
 	{
 		DeviceInfo[] CachedDevices = GetDevices();
 
-		if (CachedDevices.Where(CachedDevice => CachedDevice.Id == UDID && CachedDevice.SubType == "Network").Count() > 0)
+		if (CachedDevices.Where(CachedDevice => CachedDevice.Id == UDID && CachedDevice.PlatformValues["Connection"] == "Network").Count() > 0)
 		{
 			return EntryArguments + " -n";
 		}
diff --git a/Engine/Source/Programs/AutomationTool/Turnkey/Commands/ListPlatforms.cs b/Engine/Source/Programs/AutomationTool/Turnkey/Commands/ListPlatforms.cs
--- a/Engine/Source/Programs/AutomationTool/Turnkey/Commands/ListPlatforms.cs
+++ b/Engine/Source/Programs/AutomationTool/Turnkey/Commands/ListPlatforms.cs
@@ -78,6 +78,10 @@ namespace Turnkey.Commands
 						TurnkeyUtils.Log("    Name: {0}{1}", Device.Name, Device.bIsDefault ? "*" : "");
 						TurnkeyUtils.Log("      Id: {0}", Device.Id);
 						TurnkeyUtils.Log("      Type: {0}", Device.Type);
+						foreach (KeyValuePair<string, string> Pair in Device.PlatformValues)
+						{
+							TurnkeyUtils.Log("      {0}: {1}", Pair.Key, Pair.Value);
+						}
 						TurnkeyUtils.Log("      Installed Software Version: {0}", Device.SoftwareVersion);
 						TurnkeyUtils.Log("      Valid Software Installed?: {0}", bIsSoftwareValid);
 
diff --git a/Engine/Source/Programs/AutomationTool/Win/WinPlatform.Automation.cs b/Engine/Source/Programs/AutomationTool/Win/WinPlatform.Automation.cs
--- a/Engine/Source/Programs/AutomationTool/Win/WinPlatform.Automation.cs
+++ b/Engine/Source/Programs/AutomationTool/Win/WinPlatform.Automation.cs
@@ -17,17 +17,30 @@ public static class SteamDeckSupport
 	public static string RSyncPath = Path.Combine(Unreal.RootDirectory.FullName, "Engine\\Extras\\ThirdPartyNotUE\\cwrsync\\bin\\rsync.exe");
 	public static string SSHPath   = Path.Combine(Unreal.RootDirectory.FullName, "Engine\\Extras\\ThirdPartyNotUE\\cwrsync\\bin\\ssh.exe");
 
-	public static string GetRegisterGameScript(string GameId, string GameExePath, string GameFolderPath, string GameRunArgs)
+	public static string GetRegisterGameScript(string GameId, string GameExePath, string GameFolderPath, string GameRunArgs, string MsvsMonVersion)
 	{
 		// create a script that will be copied over to the SteamDeck and ran to register the game
 		// TODO make these easier to customize, vs hard coding the settings. Assume debugging for now, requires the user to have uploaded the required msvsmom/remote debugging stuff
 		// which is done through uploading any game with debugging enabled through the SteamOS Devkit Client
-		string GameIdArgs    = String.Format("\"gameid\":\"{0}\"", GameId);
-		string DirectoryArgs = String.Format("\"directory\":\"{0}\"", GameFolderPath);
-		string ArgvArgs      = String.Format("\"argv\":[\"{0} {1}\"]", GameExePath, GameRunArgs);
-		string SettingsArgs  = String.Format("\"settings\": {{\"steam_play\": \"1\", \"steam_play_debug\": \"1\", \"steam_play_debug_version\": \"2019\"}}");
+		string[] Settings =
+		{
+			$"\"steam_play\": \"1\"",
+			$"\"steam_play_debug\": \"1\"",
+			$"\"steam_play_debug_version\": \"{MsvsMonVersion}\"",
+		};
+		string JoinedSettings = string.Join(", ", Settings);
+
+		string[] Parms =
+		{
+			$"\"gameid\":\"{GameId}\"",
+			$"\"directory\":\"{GameFolderPath}\"",
+			$"\"argv\":[\"{GameExePath} {GameRunArgs}\"]",
+			$"\"settings\": {{{JoinedSettings}}}",
+		};
+		string JoinedParams = string.Join(", ", Parms);
 
-		return String.Format("#!/bin/bash\npython3 ~/devkit-utils/steam-client-create-shortcut --parms '{{{0}, {1}, {2}, {3}}}'", GameIdArgs, DirectoryArgs, ArgvArgs, SettingsArgs).Replace("\r\n", "\n");
+		// combine all the settings
+		return $"#!/bin/bash\npython3 ~/devkit-utils/steam-client-create-shortcut --parms '{{{JoinedParams}}}'";
 	}
 
 	// This is a bit nasty, due to rsync needing to use cygdrive path for its local location over Windows paths.
@@ -66,7 +79,7 @@ public static class SteamDeckSupport
 			{
 				string IpAddr     = GetStructEntry(DeckDevice, "IpAddr", false);
 				string DeviceName = GetStructEntry(DeckDevice, "Name", false);
-				// Skipping the UserName as its unused here
+				string UserName   = GetStructEntry(DeckDevice, "UserName", false);
 
 				// Name is optional, if its empty/not found lets just use the IpAddr for the Name
 				if (string.IsNullOrEmpty(DeviceName))
@@ -78,7 +91,8 @@ public static class SteamDeckSupport
 				{
 					// TODO Fix the usage of OSVersion here. We are abusing this and using MS OSVersion to allow Turnkey to be happy
 					DeviceInfo SteamDeck = new DeviceInfo(UnrealTargetPlatform.Win64, DeviceName, IpAddr,
-						Environment.OSVersion.Version.ToString(), "SteamDeck", true, true);
+						Environment.OSVersion.Version.ToString(), "SteamDeck", true, true, new Dictionary<string, string>() { { "UserName", UserName } });
+//					SteamDeck.PlatformValues["UserName"] = UserName;
 
 					Devices.Add(SteamDeck);
 				}
@@ -121,6 +135,49 @@ public static class SteamDeckSupport
 		return null;
 	}
 
+	/**
+	 * Get the ipaddr and username of the device from the params and/or device registry
+	 */
+	static bool GetDeviceInfo(ProjectParams Params, out string IpAddr, out string UserName)
+	{
+		IpAddr = UserName = null;
+
+		if (Params.DeviceNames.Count != 1)
+		{
+			CommandUtils.LogWarning("SteamDeck deployment requires 1, and only, device");
+			return false;
+		}
+
+		// look up in the devices in the registry based on the params (commandline -device=)
+		List<DeviceInfo> Devices = GetDevices();
+		// assume DeviceName is an Id (ie IpAddr)
+		string DeviceId = Params.DeviceNames[0];
+		DeviceInfo Device = Devices.FirstOrDefault(x => x.Id.Equals(DeviceId, StringComparison.InvariantCultureIgnoreCase));
+		// if the id didn't match, fallback to trying name match
+		if (Device == null)
+		{
+			Device = Devices.FirstOrDefault(x => x.Name.Equals(DeviceId, StringComparison.InvariantCultureIgnoreCase));
+		}
+		// if the device wasn't found, just use the DeviceId (it may be a DNS name, so we can't really verify it here)
+		IpAddr = Device == null ? DeviceId : Device.Id;
+
+
+		// if -deviceuser was specified, prefer that, otherwise use what's in the registry
+		UserName = Params.DeviceUsername;
+		if (string.IsNullOrEmpty(UserName))
+		{
+			UserName = Device?.PlatformValues["UserName"];
+		}
+
+		if (string.IsNullOrEmpty(UserName) || string.IsNullOrEmpty(IpAddr))
+		{
+			CommandUtils.LogWarning("Unable to discover Ip Adress and Username for your Steamdeck. Use -device to specify device, and if that device is not in your Engine.ini, then you will need to specify -deviceuser= to set the username.");
+			return false;
+		}
+
+		return true;
+	}
+
 	/* Deploying to a steam deck currently does 2 things
 	 *
 	 * 1) Generates a script CreateShortcutHelper.sh that will register the game on the SteamDeck once uploaded
@@ -128,13 +185,47 @@ public static class SteamDeckSupport
 	 */
 	public static void Deploy(ProjectParams Params, DeploymentContext SC)
 	{
-		string DevKitRSAPath  = Path.Combine(CommandUtils.GetEnvVar("LOCALAPPDATA"), "steamos-devkit\\steamos-devkit\\devkit_rsa");
-		string SSHArgs        = String.Format("-i {0} {1}@{2}", DevKitRSAPath, Params.DeviceUsername, Params.DeviceNames[0]);
-		string GameFolderPath = Path.Combine("/home", Params.DeviceUsername, "devkit-game", Params.ShortProjectName).Replace('\\', '/');
-		string GameRunArgs    = String.Format("{0} {1} {2}", SC.ProjectArgForCommandLines, Params.StageCommandline, Params.RunCommandline).Replace("\"", "\\\"");
+		string IpAddr, UserName;
+		if (GetDeviceInfo(Params, out IpAddr, out UserName) == false)
+		{
+			return;
+		}
 
+		string DevKitRSAPath  = Path.Combine(Environment.GetFolderPath(Environment.SpecialFolder.LocalApplicationData), @"steamos-devkit\steamos-devkit\devkit_rsa");
+		string KnownHostsPath = Path.Combine(Environment.GetFolderPath(Environment.SpecialFolder.UserProfile), @".ssh\known_hosts");
+		string SSHArgs        = $"-i {DevKitRSAPath} {UserName}@{IpAddr}";
+		string GameFolderPath = $"[REDACTED_USER_HOME]/devkit-game/{Params.ShortProjectName}";
+		string GameRunArgs    = $"{SC.ProjectArgForCommandLines} {Params.StageCommandline} {Params.RunCommandline}".Replace("\"", "\\\"");
 		FileReference ExePath = Params.GetProjectExeForPlatform(UnrealTargetPlatform.Win64);
 		string RelGameExePath = ExePath.MakeRelativeTo(DirectoryReference.Combine(ExePath.Directory, "../../..")).Replace('\\', '/');
+		// string DevkitUtilPath = Microsoft.Win32.Registry.GetValue(@"HKEY_LOCAL_MACHINE\SOFTWARE\Microsoft\Windows\CurrentVersion\Uninstall\Steam App 943760", "InstallLocation", null) as string;
+
+		// get the path to some compiler's Msvsmon/RemoteDebugger support files to copy over
+		string MsvsmonPath = null;
+		string MsvsmonVersion = "2022";
+		IEnumerable<DirectoryReference> InstallDirs = WindowsExports.TryGetVSInstallDirs(WindowsCompiler.VisualStudio2022);
+		if (InstallDirs == null)
+		{
+			InstallDirs = WindowsExports.TryGetVSInstallDirs(WindowsCompiler.VisualStudio2019);
+			MsvsmonVersion = "2019";
+		}
+		if (InstallDirs != null)
+		{
+			MsvsmonPath = Path.Combine(InstallDirs.First().FullName, "Common7/IDE/Remote Debugger");
+		}
+		string TargetMsvsmonPath = $"[REDACTED_USER_HOME]/devkit-msvsmon/msvsmon{MsvsmonVersion}";
+
+
+		// make a standard set of options to pass to the rsync for auth's -e option
+		string[] AuthOpts =
+		{
+			SteamDeckSupport.SSHPath,
+			$"-o UserKnownHostsFile='{KnownHostsPath}'",
+			$"-o StrictHostKeyChecking=no",
+			$"-i '{DevKitRSAPath}'",
+		};
+		string AuthOptions = $"-e \"{string.Join(" ", AuthOpts)}\"";
+
 
 		if (!File.Exists(DevKitRSAPath))
 		{
@@ -142,27 +233,78 @@ public static class SteamDeckSupport
 			return;
 		}
 
+
+		// rsync remote debugger support files
+		if (string.IsNullOrEmpty(MsvsmonPath) || !Directory.Exists(MsvsmonPath))
+		{
+			CommandUtils.LogWarning("Unable to find MSVC remote debugger source files, remote debugging may not work correctly.");
+		}
+		else
+		{
+			string[] DebuggerOpts =
+			{
+				// ?
+				$"-avh",
+				// ?
+				$"--chmod=Du=rwx,Dgo=rx,Fu=rwx,Fog=rx",
+				// standard authorization options
+				AuthOptions,
+				// this is a trick to make sure the target path exists before running the remote rsync command
+				$"--rsync-path=\"mkdir -p {TargetMsvsmonPath} && rsync\"",
+				// copy/update files as needed
+				$"--update",
+				// source path, in cygwin-speak (/cygdrive/c/Program Files.....)
+				$"'{ConvertWindowsPathToCygdrive(MsvsmonPath)}/'",
+				// destrination in standard 'user@ipaddr:destpath' format
+				$"{UserName}@{IpAddr}:{TargetMsvsmonPath}",
+			};
+
+			// run rsync command
+			IProcessResult DebuggerSyncResult = CommandUtils.Run(SteamDeckSupport.RSyncPath, string.Join(" ", DebuggerOpts), "");
+			if (DebuggerSyncResult.ExitCode > 0)
+			{
+				CommandUtils.LogWarning($"Failed to rsync debugger support files to the SteamDeck. Check connection on ip {UserName}@{IpAddr}");
+				return;
+			}
+		}
+
+
 		string ScriptFileName = "CreateShortcutHelper.sh";
 		string ScriptFile = Path.Combine(SC.StageDirectory.FullName, ScriptFileName);
-		File.WriteAllText(ScriptFile, SteamDeckSupport.GetRegisterGameScript(Params.ShortProjectName, RelGameExePath, GameFolderPath, GameRunArgs));
+		File.WriteAllText(ScriptFile, SteamDeckSupport.GetRegisterGameScript(Params.ShortProjectName, RelGameExePath, GameFolderPath, GameRunArgs, MsvsmonVersion));
+
+		string[] Opts =
+		{
+			// ?
+			$"-avh",
+			// ?
+			$"--chmod=Du=rwx,Dgo=rx,Fu=rwx,Fog=rx",
+			// standard authorization options
+			AuthOptions,
+			// this is a trick to make sure the target path exists before running the remote rsync command
+			$"--rsync-path=\"mkdir -p {GameFolderPath} && rsync\"",
+			// copy/update files as needed
+			$"--update",
+			// source path, in cygwin-speak (/cygdrive/c/Program Files.....)
+			$"'{ConvertWindowsPathToCygdrive(SC.StageDirectory.FullName)}/'",
+			// destrination in standard 'user@ipaddr:destpath' format
+			$"{UserName}@{IpAddr}:{GameFolderPath}",
+		};
 
 		// Exclude removing the Saved folders to preserve logs and crash data. Though note these will keep filling up with data
-		IProcessResult Result = CommandUtils.Run(SteamDeckSupport.RSyncPath,
-			String.Format("-avh --delete --exclude=\"Saved/\" --rsync-path=\"mkdir -p {5} && rsync\" --chmod=Du=rwx,Dgo=rx,Fu=rwx,Fog=rx -e \"{0} -o StrictHostKeyChecking=no -i '{1}'\" --update \"{2}/\" {3}@{4}:{5}",
-			SteamDeckSupport.SSHPath, DevKitRSAPath, SteamDeckSupport.ConvertWindowsPathToCygdrive(SC.StageDirectory.FullName), Params.DeviceUsername, Params.DeviceNames[0], GameFolderPath), "");
-
+		IProcessResult Result = CommandUtils.Run(SteamDeckSupport.RSyncPath, string.Join(" ", Opts), "");
 		if (Result.ExitCode > 0)
 		{
-			CommandUtils.LogWarning("Failed to rsync files to the SteamDeck. Check connection on ip {0}@{1}", Params.DeviceUsername, Params.DeviceNames[0]);
+			CommandUtils.LogWarning("Failed to rsync files to the SteamDeck. Check connection on ip {UserName}@{IpAddr}");
 			return;
 		}
 
 		// Run the script to register the game with the Deck
-		Result = CommandUtils.Run(SteamDeckSupport.SSHPath, String.Format("{0} \"chmod +x {1}/{2} && {1}/{2}\"", SSHArgs, GameFolderPath, ScriptFileName), "");
+		Result = CommandUtils.Run(SteamDeckSupport.SSHPath, $"{SSHArgs} \"chmod +x {GameFolderPath}/{ScriptFileName} && {GameFolderPath}/{ScriptFileName}\"", "");
 
 		if (Result.ExitCode > 0)
 		{
-			CommandUtils.LogWarning("Failed to run the {0}.sh script. Check connection on ip {1}@{2}", ScriptFileName, Params.DeviceUsername, Params.DeviceNames[0]);
+			CommandUtils.LogWarning($"Failed to run the {ScriptFileName}.sh script. Check connection on ip Check connection on ip {UserName}@{IpAddr}");
 			return;
 		}
 	}
