Change TVOS CFBundleName during plist generation to match that of iOS fixing problems packaging Lyra for TVOS from Windows

diff --git a/Engine/Source/Programs/AutomationTool/TVOS/TVOSPlatform.Automation.cs b/Engine/Source/Programs/AutomationTool/TVOS/TVOSPlatform.Automation.cs
--- a/Engine/Source/Programs/AutomationTool/TVOS/TVOSPlatform.Automation.cs
+++ b/Engine/Source/Programs/AutomationTool/TVOS/TVOSPlatform.Automation.cs
@@ -91,8 +91,10 @@ public class TVOSPlatform : IOSPlatform
             // copy the plist (only if code signing, as it's protected by the code sign blob in the executable and can't be modified independently)
             if (GetCodeSignDesirability(Params))
             {
-                DirectoryReference SourcePath = DirectoryReference.Combine((SC.IsCodeBasedProject ? SC.ProjectRoot : SC.EngineRoot), "Intermediate", "TVOS");
-                FileReference TargetPListFile = FileReference.Combine(SourcePath, (SC.IsCodeBasedProject ? SC.ShortProjectName : "UnrealGame") + "-Info.plist");
+				// this would be FooClient when making a client-only build
+				string TargetName = SC.StageExecutables[0].Split("-".ToCharArray())[0];
+				DirectoryReference SourcePath = DirectoryReference.Combine((SC.IsCodeBasedProject ? SC.ProjectRoot : SC.EngineRoot), "Intermediate", "TVOS");
+                FileReference TargetPListFile = FileReference.Combine(SourcePath, (SC.IsCodeBasedProject ? TargetName : "UnrealGame") + "-Info.plist");
                 //				if (!File.Exists(TargetPListFile))
                 {
                     // ensure the plist, entitlements, and provision files are properly copied
@@ -111,7 +113,23 @@ public class TVOSPlatform : IOSPlatform
 
                     bool bSupportsPortrait = false;
                     bool bSupportsLandscape = false;
-                    DeployGeneratePList(SC.RawProjectPath, TargetConfiguration, (SC.IsCodeBasedProject ? SC.ProjectRoot : SC.EngineRoot), !SC.IsCodeBasedProject, (SC.IsCodeBasedProject ? SC.ShortProjectName : "UnrealGame"), Params.Client, SC.ShortProjectName, SC.EngineRoot, DirectoryReference.Combine((SC.IsCodeBasedProject ? SC.ProjectRoot : SC.EngineRoot), "Binaries", "TVOS", "Payload", (SC.IsCodeBasedProject ? SC.ShortProjectName : "UnrealGame") + ".app"), SC.StageExecutables[0], out bSupportsPortrait, out bSupportsLandscape);
+                    DeployGeneratePList(
+						SC.RawProjectPath,
+						TargetConfiguration,
+						(SC.IsCodeBasedProject ? SC.ProjectRoot : SC.EngineRoot),
+						!SC.IsCodeBasedProject,
+						(SC.IsCodeBasedProject ? SC.StageExecutables[0] : "UnrealGame"),
+						Params.Client,
+						SC.ShortProjectName,
+						SC.EngineRoot,
+						DirectoryReference.Combine((SC.IsCodeBasedProject ? SC.ProjectRoot : SC.EngineRoot),
+						"Binaries",
+						"TVOS",
+						"Payload",
+						(SC.IsCodeBasedProject ? SC.ShortProjectName : "UnrealGame") + ".app"),
+						SC.StageExecutables[0],
+						out bSupportsPortrait,
+						out bSupportsLandscape);
                 }
 
 
