Add a Summary Notes system for QA to be able to look at for more user-friendly messaging about various errors that can pop up.

diff --git a/Engine/Source/Programs/AutomationTool/Gauntlet/Framework/Base/Gauntlet.BaseTestNode.cs b/Engine/Source/Programs/AutomationTool/Gauntlet/Framework/Base/Gauntlet.BaseTestNode.cs
--- a/Engine/Source/Programs/AutomationTool/Gauntlet/Framework/Base/Gauntlet.BaseTestNode.cs
+++ b/Engine/Source/Programs/AutomationTool/Gauntlet/Framework/Base/Gauntlet.BaseTestNode.cs
@@ -79,6 +79,13 @@ namespace Gauntlet
 		/// <param name="testResult">New result that the Test should have.</param>
 		public abstract void SetTestResult(TestResult testResult);
 
+		/// <summary>
+		/// Add a new test event to be rolled up into the summary at the end of this test.
+		/// </summary>
+		public virtual void AddTestEvent(UnrealTestEvent InEvent)
+		{
+		}
+
 		/// <summary>
 		/// Summarize the result of the test
 		/// </summary>
diff --git a/Engine/Source/Programs/AutomationTool/Gauntlet/Framework/Base/Gauntlet.TestNode.cs b/Engine/Source/Programs/AutomationTool/Gauntlet/Framework/Base/Gauntlet.TestNode.cs
--- a/Engine/Source/Programs/AutomationTool/Gauntlet/Framework/Base/Gauntlet.TestNode.cs
+++ b/Engine/Source/Programs/AutomationTool/Gauntlet/Framework/Base/Gauntlet.TestNode.cs
@@ -189,12 +189,17 @@ namespace Gauntlet
 		/// times.
 		/// </summary>
 		TestResult GetTestResult();
-
+		
 		/// <summary>
 		/// Set the Test Result value
 		/// </summary>
 		void SetTestResult(TestResult testResult);
 
+		/// <summary>
+		/// Add a new Test Event to be rolled up and summarized at the end of this test.
+		/// </summary>
+		void AddTestEvent(UnrealTestEvent InEvent);
+
 		/// <summary>
 		/// Summary of the test. Only called once GetTestStatus() returns complete, but may be called multiple
 		/// times.
diff --git a/Engine/Source/Programs/AutomationTool/Gauntlet/Framework/Gauntlet.TestExecutor.cs b/Engine/Source/Programs/AutomationTool/Gauntlet/Framework/Gauntlet.TestExecutor.cs
--- a/Engine/Source/Programs/AutomationTool/Gauntlet/Framework/Gauntlet.TestExecutor.cs
+++ b/Engine/Source/Programs/AutomationTool/Gauntlet/Framework/Gauntlet.TestExecutor.cs
@@ -226,7 +226,7 @@ namespace Gauntlet
 							catch (System.Exception ex)
 							{
 								Log.Error("Test {0} threw an exception during ready check. Ex: {1}", Node, ex);
-
+								Node.AddTestEvent(new UnrealTestEvent(EventSeverity.Error, "Test Failed to Start", new List<string> {ex.Message}));
 								PendingTests[i] = null;
 								NodeInfo.TimeSetupBegan = NodeInfo.TimeSetupEnded = NodeInfo.TimeTestEnded = DateTime.Now;
 								CompletedTests.Add(NodeInfo);
@@ -261,6 +261,7 @@ namespace Gauntlet
 									if (TimeWaiting >= Options.Wait)
 									{
 										Log.Warning("Test {0} has been waiting to run resource-free for {1:00} seconds. Removing from wait list", Node, TimeWaiting);
+										Node.AddTestEvent(new UnrealTestEvent(EventSeverity.Error, "Insufficient devices found", new List<string> {string.Format("Test {0} was unable to find enough devices after trying for {1:00} seconds.", Node, TimeWaiting), "This is not a test-related failure."}));
 										PendingTests[i] = null;
 										NodeInfo.TimeSetupBegan = NodeInfo.TimeSetupEnded = NodeInfo.TimeTestEnded = DateTime.Now;
 										NodeInfo.Result = TestExecutionInfo.ExecutionResult.TimedOut;
diff --git a/Engine/Source/Programs/AutomationTool/Gauntlet/SelfTest/Framework/Gauntlet.SelfTest.OrderOfOpsTest.cs b/Engine/Source/Programs/AutomationTool/Gauntlet/SelfTest/Framework/Gauntlet.SelfTest.OrderOfOpsTest.cs
--- a/Engine/Source/Programs/AutomationTool/Gauntlet/SelfTest/Framework/Gauntlet.SelfTest.OrderOfOpsTest.cs
+++ b/Engine/Source/Programs/AutomationTool/Gauntlet/SelfTest/Framework/Gauntlet.SelfTest.OrderOfOpsTest.cs
@@ -104,6 +104,10 @@ namespace Gauntlet.SelfTest
 			return "";
 		}
 
+		public virtual void AddTestEvent(UnrealTestEvent InEvent)
+		{
+
+		}
 
 		public TestStatus GetTestStatus()
 		{
diff --git a/Engine/Source/Programs/AutomationTool/Gauntlet/SelfTest/Gauntlet.SelfTest.BaseNode.cs b/Engine/Source/Programs/AutomationTool/Gauntlet/SelfTest/Gauntlet.SelfTest.BaseNode.cs
--- a/Engine/Source/Programs/AutomationTool/Gauntlet/SelfTest/Gauntlet.SelfTest.BaseNode.cs
+++ b/Engine/Source/Programs/AutomationTool/Gauntlet/SelfTest/Gauntlet.SelfTest.BaseNode.cs
@@ -127,6 +127,11 @@ namespace Gauntlet.SelfTest
 		{
 		}
 
+		public virtual void AddTestEvent(UnrealTestEvent InEvent)
+		{ 
+		
+		}
+
 		public BaseTestNode()
 		{
 			TestFailures = new List<string>();
diff --git a/Engine/Source/Programs/AutomationTool/Gauntlet/Unreal/Base/Gauntlet.UnrealTestNode.cs b/Engine/Source/Programs/AutomationTool/Gauntlet/Unreal/Base/Gauntlet.UnrealTestNode.cs
--- a/Engine/Source/Programs/AutomationTool/Gauntlet/Unreal/Base/Gauntlet.UnrealTestNode.cs
+++ b/Engine/Source/Programs/AutomationTool/Gauntlet/Unreal/Base/Gauntlet.UnrealTestNode.cs
@@ -232,6 +232,17 @@ namespace Gauntlet
 		/// </summary>
 		protected List<UnrealAutomationEvent> Events { get; private set; } = new List<UnrealAutomationEvent>();
 
+
+		/// <summary>
+		/// Collection of events thrown by gauntlet itself during the test run.
+		/// </summary>
+		protected List<UnrealTestEvent> TestNodeEvents { get; private set; } = new List<UnrealTestEvent>();
+
+		public override void AddTestEvent(UnrealTestEvent InEvent)
+		{
+			TestNodeEvents.Add(InEvent);
+		}
+
 		/// <summary>
 		/// Whether we submit to the dashboard
 		/// </summary>
@@ -1058,6 +1069,44 @@ namespace Gauntlet
 		/// </summary>
 		public bool CreateReportFailed { get; protected set; }
 
+		/// <summary>
+		/// Add all of our process results to the gauntlet event summary list.
+		/// </summary>
+		public virtual void AddProcessResultEventsFromTestNode()
+		{
+
+			if (RoleResults != null)
+			{
+				// First, make sure we have all of our process exit events.
+				// Iterate through ProcessResults to see if we need to add any messages to summary.
+				HashSet<UnrealProcessResult> DistinctResults = new HashSet<UnrealProcessResult>();
+				foreach (UnrealRoleResult roleResult in RoleResults)
+				{
+					DistinctResults.Add(roleResult.ProcessResult);
+				}
+				if (DistinctResults.Contains(UnrealProcessResult.InitializationFailure))
+				{
+					TestNodeEvents.Add(new UnrealTestEvent(EventSeverity.Fatal, "Engine Initialization Failed", new List<string> { "Engine failed to initialize in one or more roles. This is likely a bad build." }));
+				}
+				if (DistinctResults.Contains(UnrealProcessResult.EncounteredFatalError))
+				{
+					TestNodeEvents.Add(new UnrealTestEvent(EventSeverity.Fatal, "Fatal Error Encountered", new List<string> { "Test encountered a fatal error. Check ClientOutput.log and ServerOutput.log for details." }));
+				}
+				if (DistinctResults.Contains(UnrealProcessResult.EncounteredEnsure))
+				{
+					TestNodeEvents.Add(new UnrealTestEvent(EventSeverity.Error, "Ensure Found", new List<string> { "Test encountered one or more ensures. See above for details." }));
+				}
+				if (DistinctResults.Contains(UnrealProcessResult.TestFailure))
+				{
+					TestNodeEvents.Add(new UnrealTestEvent(EventSeverity.Error, "Test Failure Encountered", new List<string> { "Test encountered a failure. See above for more info." }));
+				}
+				if (DistinctResults.Contains(UnrealProcessResult.TimeOut))
+				{
+					TestNodeEvents.Add(new UnrealTestEvent(EventSeverity.Error, "Test Timed Out", new List<string> { "Test terminated due to timeout. Check ClientOutput.log and ServerOutput.log for more info." }));
+				}
+			}
+		}
+
 		/// <summary>
 		/// Display all of the defined commandline information for this test.
 		/// Will display generic gauntlet params as well if -help=full is passed in.
@@ -1726,7 +1775,6 @@ namespace Gauntlet
 			return CombinedHash;
 		}
 
-		
 		/// <summary>
 		/// Returns a formatted summary of the role that's suitable for displaying
 		/// </summary>
@@ -1739,7 +1787,6 @@ namespace Gauntlet
 			const int MaxCallstackLines = 20;
 
 			UnrealLog LogSummary = InRoleResult.LogSummary;
-						
 			MarkdownBuilder MB = new MarkdownBuilder();
 
 			UnrealRoleArtifacts RoleArtifacts = InRoleResult.Artifacts;
@@ -1812,7 +1859,6 @@ namespace Gauntlet
 					MB.Paragraph("Could not parse callstack. See log for full callstack");
 				}
 			}
-		
 
 			if (Errors.Any())
 			{
@@ -1876,7 +1922,110 @@ namespace Gauntlet
 				}
 			}
 
-			return MB.ToString(); 		
+			return MB.ToString();
+		}
+
+		/// <summary>
+		/// Returns a formatted summary of the events that were thrown by the test node itself during run.
+		/// </summary>
+		/// <returns>An HTML string containing the properly formatted list of results thrown by the Test Node.</returns>
+		protected virtual string CreateFormattedEventListFromTestNode()
+		{
+
+			AddProcessResultEventsFromTestNode();
+			
+
+			const int MaxCallstackLines = 20;
+			MarkdownBuilder MB = new MarkdownBuilder();
+
+			if (TestNodeEvents.Count == 0)
+			{
+				MB.H4("No Gauntlet Events were fired during this test!");
+				return MB.ToString();
+			}
+
+			// Separate the events we want to report on
+			IEnumerable<UnrealTestEvent> Fatals = TestNodeEvents.Where(E => E.Severity == EventSeverity.Fatal);
+			IEnumerable<UnrealTestEvent> Errors = TestNodeEvents.Where(E => E.Severity == EventSeverity.Error);
+			IEnumerable<UnrealTestEvent> Ensures = TestNodeEvents.Where(E => E.IsEnsure);
+			IEnumerable<UnrealTestEvent> Warnings = TestNodeEvents.Where(E => E.Severity == EventSeverity.Warning && !E.IsEnsure);
+			IEnumerable<UnrealTestEvent> Infos = TestNodeEvents.Where(E => E.Severity == EventSeverity.Info );
+
+			foreach (UnrealTestEvent Event in Fatals)
+			{
+				MB.H4(string.Format("Fatal Error: {0}", Event.Summary));
+
+				if (Event.Callstack.Any())
+				{
+					MB.UnorderedList(Event.Callstack.Take(MaxCallstackLines));
+
+					if (Event.Callstack.Count() > MaxCallstackLines)
+					{
+						MB.Paragraph("See log for full callstack");
+					}
+				}
+				else
+				{
+					MB.Paragraph("Could not parse callstack. See log for full callstack");
+				}
+			}
+
+			foreach (UnrealTestEvent Event in Ensures.Distinct())
+			{
+				MB.H4(string.Format("Warning: Ensure: {0}", Event.Summary));
+
+				if (Event.Callstack.Any())
+				{
+					MB.UnorderedList(Event.Callstack.Take(MaxCallstackLines));
+
+					if (Event.Callstack.Count() > MaxCallstackLines)
+					{
+						MB.Paragraph("See log for full callstack");
+					}
+				}
+				else
+				{
+					MB.Paragraph("Could not parse callstack. See log for full callstack");
+				}
+			}
+
+
+			if (Errors.Any())
+			{
+				MB.H3("Errors:");
+				MB.HorizontalLine();
+				foreach (UnrealTestEvent errorEvent in Errors)
+				{
+					MB.H5(errorEvent.Summary);
+					MB.UnorderedList(errorEvent.Details);
+					MB.NewLine();
+				}
+			}
+
+			if (Warnings.Any())
+			{
+				MB.H3("Warnings:");
+				MB.HorizontalLine();
+				foreach (UnrealTestEvent warnEvent in Warnings)
+				{
+					MB.H5(warnEvent.Summary);
+					MB.UnorderedList(warnEvent.Details);
+					MB.NewLine();
+				}
+			}
+			if (Infos.Any())
+			{
+				MB.H3("Info:");
+				MB.HorizontalLine();
+				MB.NewLine();
+				foreach (UnrealTestEvent infoEvent in Infos)
+				{
+					MB.H5(infoEvent.Summary);
+					MB.UnorderedList(infoEvent.Details);
+					MB.NewLine();
+				}
+			}
+			return MB.ToString();
 		}
 
 		/// <summary>
@@ -2061,14 +2210,16 @@ namespace Gauntlet
 		/// <returns></returns>
 		public override string GetTestSummary()
 		{
+			MarkdownBuilder ReportBuilder = new MarkdownBuilder();
+
 			// Handle case where there aren't any session artifacts, for example with device starvation
 			if (SessionArtifacts == null)
 			{
-				return "NoSummary";
+				ReportBuilder.H1(string.Format("Gauntlet summary events:"));
+				ReportBuilder.HorizontalLine();
+				ReportBuilder.Append(CreateFormattedEventListFromTestNode());
+				return ReportBuilder.ToString();
 			}
-
-			MarkdownBuilder ReportBuilder = new MarkdownBuilder();			
-
 			// Sort roles so problem ones are at the bottom, just above the summary
 			var SortedRoles = RoleResults.OrderBy(R =>
 			{
@@ -2101,7 +2252,9 @@ namespace Gauntlet
 				string Summary = GetFormattedRoleSummary(Role);
 				ReportBuilder.Append(Summary);
 			}
-
+			ReportBuilder.H1(string.Format("Gauntlet summary events:"));
+			ReportBuilder.HorizontalLine();
+			ReportBuilder.Append(CreateFormattedEventListFromTestNode());
 			return ReportBuilder.ToString();
 		}
 	}
