Better Manage UE screenshot artifacts with Horde Automated Test Session
* avoid attaching reference and delta of passed image comparison
* add warning instead of error when attaching image compare failed to work
* fix an issue with image artifacts being attached to the wrong event when attaching artifact failed...

#rnx

diff --git a/Engine/Source/Programs/AutomationTool/Gauntlet/Framework/Utils/Gauntlet.HordeReport.cs b/Engine/Source/Programs/AutomationTool/Gauntlet/Framework/Utils/Gauntlet.HordeReport.cs
--- a/Engine/Source/Programs/AutomationTool/Gauntlet/Framework/Utils/Gauntlet.HordeReport.cs
+++ b/Engine/Source/Programs/AutomationTool/Gauntlet/Framework/Utils/Gauntlet.HordeReport.cs
@@ -843,7 +843,7 @@ namespace Gauntlet
 					IndexTestError(Event.Hash, CurrentTest.TestUID, InMessage, InTag);
 				}
 			}
-			public void AddArtifactToLastEvent(string InTag, string InFilePath, string InReferencePath = null)
+			public bool AddArtifactToLastEvent(string InTag, string InFilePath, string InReferencePath = null)
 			{
 				TestEvent LastEvent = GetCurrentTest().GetLastEvent();
 				if (AttachArtifact(InFilePath, InReferencePath))
@@ -853,11 +853,11 @@ namespace Gauntlet
 						InReferencePath = Path.GetFileName(InFilePath);
 					}
 					LastEvent.NewArtifacts(InTag, InReferencePath);
+
+					return true;
 				}
-				else
-				{
-					AddEvent(EventType.Error, string.Format("Failed to attached Artifact {0}.", InReferencePath ?? Path.GetFileName(InFilePath)));
-				}
+
+				return false;
 			}
 			public override string GetTestDataKey(string BaseKey = null)
 			{
@@ -931,36 +931,47 @@ namespace Gauntlet
 								Context = string.Format("{0} @line:{1}", InTestEntry.Filename, InTestEntry.LineNumber.ToString());
 							}
 							OutTestPassResults.AddEvent(InTestEntry.Event.Type, InTestEntry.Event.Message, Tag, Context, InTestEntry.Timestamp);
+							bool IsInfo = InTestEntry.Event.Type == EventType.Info;
 							// Add Artifacts
 							if (IntArtifactEntry != null)
 							{
-								if (!string.IsNullOrEmpty(IntArtifactEntry.Files.Difference))
+								List<string> FailedToAttached = new List<string>();
+								if (!IsInfo && !string.IsNullOrEmpty(IntArtifactEntry.Files.Difference))
 								{
-									OutTestPassResults.AddArtifactToLastEvent(
+									if(!OutTestPassResults.AddArtifactToLastEvent(
 										"difference",
 										Path.Combine(InReportPath, IntArtifactEntry.Files.Difference),
 										IntArtifactEntry.Files.Difference
-									);
+									))
+									{
+										FailedToAttached.Add(IntArtifactEntry.Files.Difference);
+									}
 								}
-								if (!string.IsNullOrEmpty(IntArtifactEntry.Files.Approved))
+								if (!IsInfo && !string.IsNullOrEmpty(IntArtifactEntry.Files.Approved))
 								{
-									OutTestPassResults.AddArtifactToLastEvent(
+									if(!OutTestPassResults.AddArtifactToLastEvent(
 										"approved",
 										Path.Combine(InReportPath, IntArtifactEntry.Files.Approved),
 										IntArtifactEntry.Files.Approved
-									);
+									))
+									{
+										FailedToAttached.Add(IntArtifactEntry.Files.Approved);
+									}
 								}
 								if (!string.IsNullOrEmpty(IntArtifactEntry.Files.Unapproved))
 								{
 									string AbsoluteLocation = Path.Combine(InReportPath, IntArtifactEntry.Files.Unapproved);
-									OutTestPassResults.AddArtifactToLastEvent(
+									if(!OutTestPassResults.AddArtifactToLastEvent(
 										"unapproved",
 										AbsoluteLocation,
 										IntArtifactEntry.Files.Unapproved
-									);
+									))
+									{
+										FailedToAttached.Add(IntArtifactEntry.Files.Unapproved);
+									}
 									// Add Json meta data if any
 									string MetadataLocation = Utils.SystemHelpers.GetFullyQualifiedPath(Path.GetDirectoryName(AbsoluteLocation));
-									if (Directory.Exists(MetadataLocation))
+									if (!IsInfo && Directory.Exists(MetadataLocation))
 									{
 										string[] JsonMetadataFiles = System.IO.Directory.GetFiles(MetadataLocation, "*.json");
 										if (JsonMetadataFiles.Length > 0)
@@ -976,11 +987,21 @@ namespace Gauntlet
 											foreach (string JsonFile in JsonMetadataFiles)
 											{
 												string JsonArtifactName = RelativeLocation + "/" + Path.GetFileName(JsonFile);
-												OutTestPassResults.AddArtifactToLastEvent("json metadata", JsonFile, JsonArtifactName);
+												if(!OutTestPassResults.AddArtifactToLastEvent("json metadata", JsonFile, JsonArtifactName))
+												{
+													FailedToAttached.Add(JsonArtifactName);
+												}
 											}
 										}
 									}
 								}
+								if(FailedToAttached.Count() > 0)
+								{
+									foreach(var Item in FailedToAttached)
+									{
+										OutTestPassResults.AddEvent(EventType.Warning, string.Format("Failed to attached Artifact {0}.", FailedToAttached));
+									}
+								}
 							}
 						}
 						OutTestPassResults.SetTestState(InTestResult.State);
