Improve Session name for Horde AutomatedTestSession TestData

#rnx

diff --git a/Engine/Source/Programs/AutomationTool/Gauntlet/Unreal/Automation/UE.Automation.cs b/Engine/Source/Programs/AutomationTool/Gauntlet/Unreal/Automation/UE.Automation.cs
--- a/Engine/Source/Programs/AutomationTool/Gauntlet/Unreal/Automation/UE.Automation.cs
+++ b/Engine/Source/Programs/AutomationTool/Gauntlet/Unreal/Automation/UE.Automation.cs
@@ -423,7 +423,24 @@ namespace UE
 		/// <summary>
 		/// Override the Type
 		/// </summary>
-		public override string Type { get { return "UE.Automation"; } }
+		public override string Type
+		{
+			get
+			{
+				string Base = "UE.Automation";
+				var Config = GetConfiguration();
+				if (Config is AutomationTestConfig)
+				{
+					var AutomationConfig = Config as AutomationTestConfig;
+					if (!string.IsNullOrEmpty(AutomationConfig.RunTest))
+					{
+						return string.Format("{0}({1}) {2}", Base, AutomationConfig.RunTest, Context.BuildInfo.ProjectName);
+					}
+				}
+
+				return string.Format("{0}({1})", Base, Suite);
+			}
+		}
 
 		/// <summary>
 		/// Called when a test is starting
diff --git a/Engine/Source/Programs/AutomationTool/Gauntlet/Unreal/Base/Gauntlet.UnrealTestNode.cs b/Engine/Source/Programs/AutomationTool/Gauntlet/Unreal/Base/Gauntlet.UnrealTestNode.cs
--- a/Engine/Source/Programs/AutomationTool/Gauntlet/Unreal/Base/Gauntlet.UnrealTestNode.cs
+++ b/Engine/Source/Programs/AutomationTool/Gauntlet/Unreal/Base/Gauntlet.UnrealTestNode.cs
@@ -45,7 +45,7 @@ namespace Gauntlet
 		/// <summary>
 		///Returns an identifier for this type of test
 		/// </summary>
-		public virtual string Type { get { return this.GetType().FullName; } }
+		public virtual string Type { get { return string.Format("{0}({1})", this.GetType().FullName, Suite); } }
 
 		/// <summary>
 		/// This class will log its own warnings and errors as part of its summary
@@ -1327,7 +1327,7 @@ namespace Gauntlet
 				GetCachedConfiguration().HordeTestDataKey = HordeTestDataKey;
 				// Convert test results for Horde
 				HordeReport.AutomatedTestSessionData HordeTestPassResults = HordeReport.AutomatedTestSessionData.FromUnrealAutomatedTests(
-					JsonTestPassResults, string.Format("{0}({1})", Type, Suite), Suite, UnrealAutomatedTestReportPath, HordeArtifactPath
+					JsonTestPassResults, Type, Suite, UnrealAutomatedTestReportPath, HordeArtifactPath
 				);
 				// Make a copy of the report in the old way - until we decide the transition is over
 				HordeReport.UnrealEngineTestPassResults CopyTestPassResults = HordeReport.UnrealEngineTestPassResults.FromUnrealAutomatedTests(JsonTestPassResults, ReportURL);
