- Linux on SteamDeck support (from Windows host)

diff --git a/Engine/Source/Programs/AutomationTool/Linux/Linux.Automation.csproj b/Engine/Source/Programs/AutomationTool/Linux/Linux.Automation.csproj
--- a/Engine/Source/Programs/AutomationTool/Linux/Linux.Automation.csproj
+++ b/Engine/Source/Programs/AutomationTool/Linux/Linux.Automation.csproj
@@ -37,6 +37,7 @@
     <ProjectReference Include="..\..\Shared\EpicGames.Core\EpicGames.Core.csproj" PrivateAssets="All"><Private>false</Private></ProjectReference>
     <ProjectReference Include="..\..\UnrealBuildTool\UnrealBuildTool.csproj" PrivateAssets="All"><Private>false</Private></ProjectReference>
     <ProjectReference Include="..\AutomationUtils\AutomationUtils.Automation.csproj" PrivateAssets="All"><Private>false</Private></ProjectReference>
+    <ProjectReference Include="..\SteamDeck\SteamDeck.Automation.csproj" />
   </ItemGroup>
   <ItemGroup>
     <PackageReference Include="Microsoft.CSharp" Version="4.7.0" />
diff --git a/Engine/Source/Programs/AutomationTool/Linux/LinuxPlatform.Automation.cs b/Engine/Source/Programs/AutomationTool/Linux/LinuxPlatform.Automation.cs
--- a/Engine/Source/Programs/AutomationTool/Linux/LinuxPlatform.Automation.cs
+++ b/Engine/Source/Programs/AutomationTool/Linux/LinuxPlatform.Automation.cs
@@ -35,6 +35,8 @@ public abstract class BaseLinuxPlatform : Platform
 			Devices.Add(LocalMachine);
 		}
 
+		Devices.AddRange(SteamDeckSupport.GetDevices(UnrealTargetPlatform.Linux));
+
 		return Devices.ToArray();
 	}
 
@@ -207,6 +209,13 @@ public abstract class BaseLinuxPlatform : Platform
 	/// <param name="SC"></param>
 	public override void Deploy(ProjectParams Params, DeploymentContext SC)
 	{
+		// We only care about deploying for SteamDeck
+		if (Params.Devices.Count == 1 && GetDevices().FirstOrDefault(x => x.Id == Params.DeviceNames[0])?.Type == "SteamDeck")
+		{
+			SteamDeckSupport.Deploy(UnrealTargetPlatform.Linux, Params, SC);
+			return;
+		}
+
 		if (BuildHostPlatform.Current.Platform != UnrealTargetPlatform.Linux)
 		{
 			foreach (string DeviceAddress in Params.DeviceNames)
@@ -308,50 +317,68 @@ chmod +x {0}
 	public override void PlatformSetupParams(ref ProjectParams ProjParams)
 	{
 		if ((ProjParams.Deploy || ProjParams.Run) && BuildHostPlatform.Current.Platform != UnrealTargetPlatform.Linux)
-		{ 
-			// Prompt for username if not already set
-			while (String.IsNullOrEmpty(ProjParams.DeviceUsername))
+		{
+			// we don't need username/password if we are only targeting steamdeck devices
+			bool bNeedsUsernameAndPassword = false;
+			foreach (string DeviceId in ProjParams.DeviceNames)
 			{
-				Console.Write("Username: ");
-				ProjParams.DeviceUsername = Console.ReadLine();
+				if (SteamDeckSupport.GetDeviceInfo(UnrealTargetPlatform.Linux, ProjParams, out _, out _) == false)
+				{
+					bNeedsUsernameAndPassword = true;
+				}
 			}
 
-			// Prompty for password if not already set
-			while (String.IsNullOrEmpty(ProjParams.DevicePassword))
+			if (bNeedsUsernameAndPassword)
 			{
-				ProjParams.DevicePassword = String.Empty;
-				Console.Write("Password: ");
-				ConsoleKeyInfo key;
-				do
+				// Prompt for username if not already set
+				while (String.IsNullOrEmpty(ProjParams.DeviceUsername))
 				{
-					key = Console.ReadKey(true);
-					if (key.Key != ConsoleKey.Backspace && key.Key != ConsoleKey.Enter)
-					{
-						ProjParams.DevicePassword += key.KeyChar;
-						Console.Write("*");
-					}
-					else
+					Console.Write("Username: ");
+					ProjParams.DeviceUsername = Console.ReadLine();
+				}
+
+				// Prompty for password if not already set
+				while (String.IsNullOrEmpty(ProjParams.DevicePassword))
+				{
+					ProjParams.DevicePassword = String.Empty;
+					Console.Write("Password: ");
+					ConsoleKeyInfo key;
+					do
 					{
-						if (key.Key == ConsoleKey.Backspace && ProjParams.DevicePassword.Length > 0)
+						key = Console.ReadKey(true);
+						if (key.Key != ConsoleKey.Backspace && key.Key != ConsoleKey.Enter)
 						{
-							ProjParams.DevicePassword = ProjParams.DevicePassword.Substring(0, (ProjParams.DevicePassword.Length - 1));
-							Console.Write("\b \b");
+							ProjParams.DevicePassword += key.KeyChar;
+							Console.Write("*");
+						}
+						else
+						{
+							if (key.Key == ConsoleKey.Backspace && ProjParams.DevicePassword.Length > 0)
+							{
+								ProjParams.DevicePassword = ProjParams.DevicePassword.Substring(0, (ProjParams.DevicePassword.Length - 1));
+								Console.Write("\b \b");
+							}
 						}
-					}
 
-				} while (key.Key != ConsoleKey.Enter);
-				Console.WriteLine();
-			}
+					} while (key.Key != ConsoleKey.Enter);
+					Console.WriteLine();
+				}
 
-			// try contacting the device(s) and cache the key(s)
-			foreach(string DeviceAddress in ProjParams.DeviceNames)
-			{
-				RunAndLog(CmdEnv, "cmd.exe", String.Format("/c \"echo y | {0} -P 22 -ssh -t -l {1} -pw {2} {3} echo All Ok\"", PlinkPath, ProjParams.DeviceUsername, ProjParams.DevicePassword, DeviceAddress));
+				// try contacting the device(s) and cache the key(s)
+				foreach (string DeviceAddress in ProjParams.DeviceNames)
+				{
+					RunAndLog(CmdEnv, "cmd.exe", String.Format("/c \"echo y | {0} -P 22 -ssh -t -l {1} -pw {2} {3} echo All Ok\"", PlinkPath, ProjParams.DeviceUsername, ProjParams.DevicePassword, DeviceAddress));
+				}
 			}
 		}
 	}
 	public override IProcessResult RunClient(ERunOptions ClientRunFlags, string ClientApp, string ClientCmdLine, ProjectParams Params)
 	{
+		if (Params.Devices.Count == 1 && GetDevices().FirstOrDefault(x => x.Id == Params.DeviceNames[0])?.Type == "SteamDeck")
+		{
+			return SteamDeckSupport.RunClient(UnrealTargetPlatform.Linux, ClientRunFlags, ClientApp, ClientCmdLine, Params);
+		}
+
 		if (BuildHostPlatform.Current.Platform != UnrealTargetPlatform.Linux)
 		{
 			IProcessResult Result = null;
diff --git a/Engine/Source/Programs/AutomationTool/SteamDeck/SteamDeck.Automation.csproj b/Engine/Source/Programs/AutomationTool/SteamDeck/SteamDeck.Automation.csproj
new file mode 100644
--- /dev/null
+++ b/Engine/Source/Programs/AutomationTool/SteamDeck/SteamDeck.Automation.csproj
@@ -0,0 +1,47 @@
+<Project Sdk="Microsoft.NET.Sdk">
+  <PropertyGroup>
+    <TargetFramework>netcoreapp3.1</TargetFramework>
+    <Configuration Condition=" '$(Configuration)' == '' ">Development</Configuration>
+    <OutputType>Library</OutputType>
+    <AppendTargetFrameworkToOutputPath>false</AppendTargetFrameworkToOutputPath>
+    <GenerateAssemblyInfo>false</GenerateAssemblyInfo>
+    <GenerateTargetFrameworkAttribute>false</GenerateTargetFrameworkAttribute>
+    <Configurations>Debug;Release;Development</Configurations>
+    <RootNamespace>SteamDeck.Automation</RootNamespace>
+    <AssemblyName>SteamDeck.Automation</AssemblyName>
+    <WarningsNotAsErrors>612,618</WarningsNotAsErrors>
+    <OutputPath>..\..\..\..\Binaries\DotNET\AutomationTool\AutomationScripts\Platforms\SteamDeck\</OutputPath>
+    <TreatWarningsAsErrors>true</TreatWarningsAsErrors>
+    <DebugType>pdbonly</DebugType>
+  </PropertyGroup>
+  <PropertyGroup Condition=" '$(Configuration)|$(Platform)' == 'Development|AnyCPU' ">
+    <DefineConstants>$(DefineConstants);TRACE</DefineConstants>
+    <Optimize>true</Optimize>
+  </PropertyGroup>
+  <PropertyGroup Condition=" '$(Configuration)|$(Platform)' == 'Release|AnyCPU' ">
+    <DefineConstants>$(DefineConstants);TRACE</DefineConstants>
+    <Optimize>true</Optimize>
+  </PropertyGroup>
+  <PropertyGroup Condition=" '$(Configuration)|$(Platform)' == 'Debug|AnyCPU' ">
+    <DefineConstants>$(DefineConstants);DEBUG;TRACE</DefineConstants>
+    <DebugSymbols>true</DebugSymbols>
+    <DebugType>full</DebugType>
+  </PropertyGroup>
+
+  <ItemGroup>
+    <Compile Include="..\..\Shared\MetaData.cs">
+      <Link>Properties\MetaData.cs</Link>
+    </Compile>
+  </ItemGroup>
+  <ItemGroup>
+    <ProjectReference Include="..\..\Shared\EpicGames.Core\EpicGames.Core.csproj" PrivateAssets="All"><Private>false</Private></ProjectReference>
+    <ProjectReference Include="..\..\UnrealBuildTool\UnrealBuildTool.csproj" PrivateAssets="All"><Private>false</Private></ProjectReference>
+    <ProjectReference Include="..\AutomationUtils\AutomationUtils.Automation.csproj" PrivateAssets="All"><Private>false</Private></ProjectReference>
+  </ItemGroup>
+  <ItemGroup>
+    <PackageReference Include="Microsoft.CSharp" Version="4.7.0" />
+    <PackageReference Include="System.Data.DataSetExtensions" Version="4.5.0" />
+    <PackageReference Include="System.Net.Http" Version="4.3.4" />
+    <PackageReference Include="System.Drawing.Common" Version="4.7.0" />
+  </ItemGroup>
+</Project>
\ No newline at end of file
diff --git a/Engine/Source/Programs/AutomationTool/SteamDeck/SteamDeckSupport.cs b/Engine/Source/Programs/AutomationTool/SteamDeck/SteamDeckSupport.cs
new file mode 100644
--- /dev/null
+++ b/Engine/Source/Programs/AutomationTool/SteamDeck/SteamDeckSupport.cs
@@ -0,0 +1,313 @@
+// Copyright Epic Games, Inc. All Rights Reserved.
+using System;
+using System.Collections.Generic;
+using System.Linq;
+using System.IO;
+using AutomationTool;
+using UnrealBuildTool;
+using EpicGames.Core;
+using UnrealBuildBase;
+
+public static class SteamDeckSupport
+{
+	public static string RSyncPath = Path.Combine(Unreal.RootDirectory.FullName, "Engine\\Extras\\ThirdPartyNotUE\\cwrsync\\bin\\rsync.exe");
+	public static string SSHPath = Path.Combine(Unreal.RootDirectory.FullName, "Engine\\Extras\\ThirdPartyNotUE\\cwrsync\\bin\\ssh.exe");
+	static string DevKitRSAPath = Path.Combine(Environment.GetFolderPath(Environment.SpecialFolder.LocalApplicationData), @"steamos-devkit\steamos-devkit\devkit_rsa");
+	static string KnownHostsPath = Path.Combine(Environment.GetFolderPath(Environment.SpecialFolder.UserProfile), @".ssh\known_hosts");
+	
+	#region Devices
+
+	public static List<DeviceInfo> GetDevices(UnrealTargetPlatform RuntimePlatform)
+	{
+		List<DeviceInfo> Devices = new List<DeviceInfo>();
+
+		// Look for any Steam Deck devices that are in the Engine ini files. If so lets add them as valid devices
+		// This will be required for matching devices passed into the BuildCookRun command
+		ConfigHierarchy EngineConfig = ConfigCache.ReadHierarchy(ConfigHierarchyType.Engine, null, UnrealTargetPlatform.Win64);
+
+		List<string> SteamDeckDevices;
+		if (EngineConfig.GetArray("/Script/WindowsTargetPlatform.WindowsTargetSettings", "SteamDeckDevice", out SteamDeckDevices))
+		{
+			// Expected ini format: +SteamDeckDevice=(IpAddr=10.1.33.19,Name=MySteamDeck,UserName=deck)
+			foreach (string DeckDevice in SteamDeckDevices)
+			{
+				string IpAddr = ConfigHierarchy.GetStructEntry(DeckDevice, "IpAddr", false);
+				string DeviceName = ConfigHierarchy.GetStructEntry(DeckDevice, "Name", false);
+				string UserName = ConfigHierarchy.GetStructEntry(DeckDevice, "UserName", false);
+
+				if (RuntimePlatform == UnrealTargetPlatform.Win64)
+				{
+					DeviceName += " (Proton)";
+				}
+				else if (RuntimePlatform == UnrealTargetPlatform.Linux)
+				{
+					DeviceName += " (Native Linux)";
+				}
+
+				// Name is optional, if its empty/not found lets just use the IpAddr for the Name
+				if (string.IsNullOrEmpty(DeviceName))
+				{
+					DeviceName = IpAddr;
+				}
+
+				if (!string.IsNullOrEmpty(IpAddr))
+				{
+					// TODO Fix the usage of OSVersion here. We are abusing this and using MS OSVersion to allow Turnkey to be happy
+					DeviceInfo SteamDeck = new DeviceInfo(RuntimePlatform, DeviceName, IpAddr,
+						Environment.OSVersion.Version.ToString(), "SteamDeck", true, true);
+					SteamDeck.PlatformValues["UserName"] = UserName;
+
+					Devices.Add(SteamDeck);
+				}
+			}
+		}
+
+		return Devices;
+	}
+
+	/**
+	 * Get the ipaddr and username of the device from the params and/or device registry
+	 */
+	public static bool GetDeviceInfo(UnrealTargetPlatform RuntimePlatform, ProjectParams Params, out string IpAddr, out string UserName)
+	{
+		IpAddr = UserName = null;
+
+		if (Params.DeviceNames.Count != 1)
+		{
+			CommandUtils.LogWarning("SteamDeck deployment requires 1, and only, device");
+			return false;
+		}
+
+		// look up in the devices in the registry based on the params (commandline -device=)
+		List<DeviceInfo> Devices = GetDevices(RuntimePlatform);
+		// assume DeviceName is an Id (ie IpAddr)
+		string DeviceId = Params.DeviceNames[0];
+		DeviceInfo Device = Devices.FirstOrDefault(x => x.Id.Equals(DeviceId, StringComparison.InvariantCultureIgnoreCase));
+		// if the id didn't match, fallback to trying name match
+		if (Device == null)
+		{
+			Device = Devices.FirstOrDefault(x => x.Name.Equals(DeviceId, StringComparison.InvariantCultureIgnoreCase));
+		}
+		// if the device wasn't found, just use the DeviceId (it may be a DNS name, so we can't really verify it here)
+		IpAddr = Device == null ? DeviceId : Device.Id;
+
+
+		// if -deviceuser was specified, prefer that, otherwise use what's in the registry
+		UserName = Params.DeviceUsername;
+		if (string.IsNullOrEmpty(UserName))
+		{
+			UserName = Device?.PlatformValues["UserName"];
+		}
+
+		if (string.IsNullOrEmpty(UserName) || string.IsNullOrEmpty(IpAddr))
+		{
+			CommandUtils.LogWarning("Unable to discover Ip Adress and Username for your Steamdeck. Use -device to specify device, and if that device is not in your Engine.ini, then you will need to specify -deviceuser= to set the username.");
+			return false;
+		}
+
+		return true;
+	}
+
+	#endregion
+
+	#region Deploying
+
+	public static string GetRegisterGameScript(ProjectParams Params, UnrealTargetPlatform RuntimePlatform, string GameFolderPath, string GameRunArgs, string MsvsmonVersion)
+	{
+		FileReference ExePath = Params.GetProjectExeForPlatform(RuntimePlatform);
+		string RelGameExePath = ExePath.MakeRelativeTo(DirectoryReference.Combine(ExePath.Directory, "../../..")).Replace('\\', '/');
+		string GameId = $"{Params.ShortProjectName}_{RuntimePlatform}";
+
+		// create a script that will be copied over to the SteamDeck and ran to register the game
+		// TODO make these easier to customize, vs hard coding the settings. Assume debugging for now, requires the user to have uploaded the required msvsmom/remote debugging stuff
+		// which is done through uploading any game with debugging enabled through the SteamOS Devkit Client
+		string[] LinuxSettings =
+		{
+			$"\"steam_play\": \"0\"",
+		};
+		string[] ProtonSettings =
+		{
+			$"\"steam_play\": \"1\"",
+			$"\"steam_play_debug\": \"1\"",
+			$"\"steam_play_debug_version\": \"{MsvsmonVersion}\"",
+		};
+		string JoinedSettings = string.Join(", ", RuntimePlatform == UnrealTargetPlatform.Linux ? LinuxSettings : ProtonSettings);
+
+		string[] Parms =
+		{
+			$"\"gameid\":\"{GameId}\"",
+			$"\"directory\":\"{GameFolderPath}\"",
+			$"\"argv\":[\"{RelGameExePath} {GameRunArgs}\"]",
+			$"\"settings\": {{{JoinedSettings}}}",
+		};
+		string JoinedParams = string.Join(", ", Parms);
+
+		// combine all the settings
+		return $"#!/bin/bash\nrm ~/devkit-game/{GameId}-settings.json\npython3 ~/devkit-utils/steam-client-create-shortcut --parms '{{{JoinedParams}}}'";
+	}
+
+	// Rsync is a Cygwin utility, so convert windows paths to a path that Cygwin Rsync will understand
+	// This will not work with UNC paths
+	public static string FixupHostPath(string HostPath)
+	{
+		if (HostPlatform.Platform == UnrealTargetPlatform.Win64)
+		{
+			string CygdrivePath = "";
+
+			if (!string.IsNullOrEmpty(HostPath))
+			{
+				string FullPath = Path.GetFullPath(HostPath);
+				string RootPath = Path.GetPathRoot(FullPath);
+
+				CygdrivePath = Path.Combine("/cygdrive", Char.ToLower(FullPath[0]).ToString(), FullPath.Substring(RootPath.Length));
+
+				return CygdrivePath.Replace('\\', '/');
+			}
+
+			return CygdrivePath;
+		}
+
+		return HostPath;
+	}
+
+
+	private static IProcessResult Rsync(string SourceFolder, string DestFolder, string UserName, string IpAddr, string[] ExcludeList)
+	{
+		// make a standard set of options to pass to the rsync for auth's -e option
+		string[] AuthOpts =
+		{
+			SSHPath,
+			$"-o UserKnownHostsFile='{KnownHostsPath}'",
+			$"-o StrictHostKeyChecking=no",
+			$"-i '{DevKitRSAPath}'",
+		};
+		string AuthOptions = $"-e \"{string.Join(" ", AuthOpts)}\"";
+
+		// make a set of --exclude options for anything we want to exclude
+		string ExcludeOptions = "";
+		if (ExcludeList != null)
+		{
+			ExcludeOptions = string.Join(" ", ExcludeList.Select(x => $"--exclude={x}"));
+		}
+
+		string[] RsyncOpts =
+		{
+			// ?
+			$"-avh",
+			// ?
+			$"--chmod=Du=rwx,Dgo=rx,Fu=rwx,Fog=rx",
+			// standard authorization options
+			AuthOptions,
+			// this is a trick to make sure the target path exists before running the remote rsync command
+			$"--rsync-path=\"mkdir -p {DestFolder} && rsync\"",
+			// copy/update files as needed
+			$"--update",
+			// delete destination files not in source
+			$"--delete",
+			// exclude some files
+			ExcludeOptions,
+			// source path, in cygwin-speak (/cygdrive/c/Program Files.....)
+			$"'{FixupHostPath(SourceFolder)}/'",
+			// destrination in standard 'user@ipaddr:destpath' format
+			$"{UserName}@{IpAddr}:{DestFolder}",
+		};
+
+		// Exclude removing the Saved folders to preserve logs and crash data. Though note these will keep filling up with data
+		return CommandUtils.Run(RSyncPath, string.Join(" ", RsyncOpts), "");
+	}
+
+	/* Deploying to a steam deck currently does 2 things
+	 *
+	 * 1) Generates a script CreateShortcutHelper.sh that will register the game on the SteamDeck once uploaded
+	 * 2) Uploads the build using rsync to the devkit-game location. Once uploaded it runs the CreateShortcutHelper.sh generated before.
+	 */
+	public static void Deploy(UnrealTargetPlatform RuntimePlatform, ProjectParams Params, DeploymentContext SC)
+	{
+		string IpAddr, UserName;
+		if (GetDeviceInfo(RuntimePlatform, Params, out IpAddr, out UserName) == false)
+		{
+			return;
+		}
+
+		string GameFolderPath = $"[REDACTED_USER_HOME]/devkit-game/{Params.ShortProjectName}_{RuntimePlatform}";
+		string GameRunArgs = $"{SC.ProjectArgForCommandLines} {Params.StageCommandline} {Params.RunCommandline}".Replace("\"", "\\\"");
+		// string DevkitUtilPath = Microsoft.Win32.Registry.GetValue(@"HKEY_LOCAL_MACHINE\SOFTWARE\Microsoft\Windows\CurrentVersion\Uninstall\Steam App 943760", "InstallLocation", null) as string;
+
+
+		if (!File.Exists(DevKitRSAPath))
+		{
+			CommandUtils.LogWarning("Unable to find '{0}' rsa key needed to deploy to the steam deck. Make sure you've installed the SteamOS Devkit client", DevKitRSAPath);
+			return;
+		}
+
+		// copy debugger support files if needed
+		string MsvsmonVersion = "";
+		if (RuntimePlatform == UnrealTargetPlatform.Win64)
+		{
+			if (DeployMsvsmon(IpAddr, UserName, out MsvsmonVersion) == false)
+			{
+				return;
+			}	
+		}
+
+		// write a file into the staged directory to be copied with the rest
+		string ScriptFileName = "CreateShortcutHelper.sh";
+		string ScriptFile = Path.Combine(SC.StageDirectory.FullName, ScriptFileName);
+		File.WriteAllText(ScriptFile, GetRegisterGameScript(Params, RuntimePlatform, GameFolderPath, GameRunArgs, MsvsmonVersion));
+
+		// copy the staged directory, skipping over huge debug files that we don't need on the target
+		IProcessResult Result = Rsync(SC.StageDirectory.FullName, GameFolderPath, UserName, IpAddr, new string[] { /*"*.debug", */ "*.pdb", "Saved/" });
+
+		// Run the script to register the game with the Deck
+		string SSHArgs = $"-i {DevKitRSAPath} {UserName}@{IpAddr}";
+		Result = CommandUtils.Run(SSHPath, $"{SSHArgs} \"chmod +x {GameFolderPath}/{ScriptFileName} && {GameFolderPath}/{ScriptFileName}\"", "");
+
+		if (Result.ExitCode > 0)
+		{
+			CommandUtils.LogWarning($"Failed to run the {ScriptFileName}.sh script. Check connection on ip Check connection on ip {UserName}@{IpAddr}");
+			return;
+		}
+	}
+
+	private static bool DeployMsvsmon(string IpAddr, string UserName, out string MsvsmonVersion)
+	{
+		// get the path to some compiler's Msvsmon/RemoteDebugger support files to copy over
+		string MsvsmonPath = null;
+		MsvsmonVersion = "2022";
+		IEnumerable<DirectoryReference> InstallDirs = WindowsExports.TryGetVSInstallDirs(WindowsCompiler.VisualStudio2022);
+		if (InstallDirs == null)
+		{
+			InstallDirs = WindowsExports.TryGetVSInstallDirs(WindowsCompiler.VisualStudio2019);
+			MsvsmonVersion = "2019";
+		}
+		if (InstallDirs != null)
+		{
+			MsvsmonPath = Path.Combine(InstallDirs.First().FullName, "Common7/IDE/Remote Debugger");
+		}
+		string TargetMsvsmonPath = $"[REDACTED_USER_HOME]/devkit-msvsmon/msvsmon{MsvsmonVersion}";
+
+		IProcessResult Result = Rsync(MsvsmonPath, TargetMsvsmonPath, UserName, IpAddr, null);
+		if (Result.ExitCode > 0)
+		{
+			CommandUtils.LogWarning($"Failed to rsync debugger support files to the SteamDeck. Check connection on ip {UserName}@{IpAddr}");
+			return false;
+		}
+
+		return true;
+	}
+
+	#endregion
+
+	#region Running
+
+	public static IProcessResult RunClient(UnrealTargetPlatform RuntimePlatform, CommandUtils.ERunOptions ClientRunFlags, string ClientApp, string ClientCmdLine, ProjectParams Params)
+	{
+		// TODO figure out how to get the steam app num id. Then can run like this:
+		// steam steam://rungameid/<GameNumId>
+		// TODO would be great if we could tail the log while running. Figure out how to cancel/exit app
+
+		return null;
+	}
+
+	#endregion
+}
\ No newline at end of file
diff --git a/Engine/Source/Programs/AutomationTool/Win/Win.Automation.csproj b/Engine/Source/Programs/AutomationTool/Win/Win.Automation.csproj
--- a/Engine/Source/Programs/AutomationTool/Win/Win.Automation.csproj
+++ b/Engine/Source/Programs/AutomationTool/Win/Win.Automation.csproj
@@ -44,6 +44,7 @@
     <ProjectReference Include="..\..\UnrealBuildTool\UnrealBuildTool.csproj" PrivateAssets="All"><Private>false</Private></ProjectReference>
     <ProjectReference Include="..\AutomationUtils\AutomationUtils.Automation.csproj" PrivateAssets="All"><Private>false</Private></ProjectReference>
     <ProjectReference Include="..\BuildGraph\BuildGraph.Automation.csproj" PrivateAssets="All"><Private>false</Private></ProjectReference>
+    <ProjectReference Include="..\SteamDeck\SteamDeck.Automation.csproj" />
   </ItemGroup>
   <ItemGroup>
     <PackageReference Include="Microsoft.CSharp" Version="4.7.0" />
diff --git a/Engine/Source/Programs/AutomationTool/Win/WinPlatform.Automation.cs b/Engine/Source/Programs/AutomationTool/Win/WinPlatform.Automation.cs
--- a/Engine/Source/Programs/AutomationTool/Win/WinPlatform.Automation.cs
+++ b/Engine/Source/Programs/AutomationTool/Win/WinPlatform.Automation.cs
@@ -12,270 +12,6 @@ using System.Diagnostics;
 using EpicGames.Core;
 using UnrealBuildBase;
 
-public static class SteamDeckSupport
-{
-	public static string RSyncPath = Path.Combine(Unreal.RootDirectory.FullName, "Engine\\Extras\\ThirdPartyNotUE\\cwrsync\\bin\\rsync.exe");
-	public static string SSHPath   = Path.Combine(Unreal.RootDirectory.FullName, "Engine\\Extras\\ThirdPartyNotUE\\cwrsync\\bin\\ssh.exe");
-
-	public static string GetRegisterGameScript(string GameId, string GameExePath, string GameFolderPath, string GameRunArgs, string MsvsMonVersion)
-	{
-		// create a script that will be copied over to the SteamDeck and ran to register the game
-		// TODO make these easier to customize, vs hard coding the settings. Assume debugging for now, requires the user to have uploaded the required msvsmom/remote debugging stuff
-		// which is done through uploading any game with debugging enabled through the SteamOS Devkit Client
-		string[] Settings =
-		{
-			$"\"steam_play\": \"1\"",
-			$"\"steam_play_debug\": \"1\"",
-			$"\"steam_play_debug_version\": \"{MsvsMonVersion}\"",
-		};
-		string JoinedSettings = string.Join(", ", Settings);
-
-		string[] Parms =
-		{
-			$"\"gameid\":\"{GameId}\"",
-			$"\"directory\":\"{GameFolderPath}\"",
-			$"\"argv\":[\"{GameExePath} {GameRunArgs}\"]",
-			$"\"settings\": {{{JoinedSettings}}}",
-		};
-		string JoinedParams = string.Join(", ", Parms);
-
-		// combine all the settings
-		return $"#!/bin/bash\npython3 ~/devkit-utils/steam-client-create-shortcut --parms '{{{JoinedParams}}}'";
-	}
-
-	// This is a bit nasty, due to rsync needing to use cygdrive path for its local location over Windows paths.
-	// This will not work with UNC paths
-	public static string ConvertWindowsPathToCygdrive(string WindowsPath)
-	{
-		string CygdrivePath = "";
-
-		if (!string.IsNullOrEmpty(WindowsPath))
-		{
-			string FullPath = Path.GetFullPath(WindowsPath);
-			string RootPath = Path.GetPathRoot(FullPath);
-
-			System.Console.WriteLine("{0}", RootPath);
-			CygdrivePath = Path.Combine("/cygdrive", Char.ToLower(FullPath[0]).ToString(), FullPath.Substring(RootPath.Length));
-
-			return CygdrivePath.Replace('\\','/');
-		}
-
-		return CygdrivePath;
-	}
-
-	public static List<DeviceInfo> GetDevices()
-	{
-		List<DeviceInfo> Devices = new List<DeviceInfo>();
-
-		// Look for any Steam Deck devices that are in the Engine ini files. If so lets add them as valid devices
-		// This will be required for matching devices passed into the BuildCookRun command
-		ConfigHierarchy EngineConfig = ConfigCache.ReadHierarchy(ConfigHierarchyType.Engine, null, UnrealTargetPlatform.Win64);
-
-		List<string> SteamDeckDevices;
-		if (EngineConfig.GetArray("/Script/WindowsTargetPlatform.WindowsTargetSettings", "SteamDeckDevice", out SteamDeckDevices))
-		{
-			// Expected ini format: +SteamDeckDevice=(IpAddr=10.1.33.19,Name=MySteamDeck,UserName=deck)
-			foreach (string DeckDevice in SteamDeckDevices)
-			{
-				string IpAddr     = ConfigHierarchy.GetStructEntry(DeckDevice, "IpAddr", false);
-				string DeviceName = ConfigHierarchy.GetStructEntry(DeckDevice, "Name", false);
-				string UserName   = ConfigHierarchy.GetStructEntry(DeckDevice, "UserName", false);
-
-				// Name is optional, if its empty/not found lets just use the IpAddr for the Name
-				if (string.IsNullOrEmpty(DeviceName))
-				{
-					DeviceName = IpAddr;
-				}
-
-				if (!string.IsNullOrEmpty(IpAddr))
-				{
-					// TODO Fix the usage of OSVersion here. We are abusing this and using MS OSVersion to allow Turnkey to be happy
-					DeviceInfo SteamDeck = new DeviceInfo(UnrealTargetPlatform.Win64, DeviceName, IpAddr,
-						Environment.OSVersion.Version.ToString(), "SteamDeck", true, true, new Dictionary<string, string>() { { "UserName", UserName } });
-//					SteamDeck.PlatformValues["UserName"] = UserName;
-
-					Devices.Add(SteamDeck);
-				}
-			}
-		}
-
-		return Devices;
-	}
-
-	/**
-	 * Get the ipaddr and username of the device from the params and/or device registry
-	 */
-	static bool GetDeviceInfo(ProjectParams Params, out string IpAddr, out string UserName)
-	{
-		IpAddr = UserName = null;
-
-		if (Params.DeviceNames.Count != 1)
-		{
-			CommandUtils.LogWarning("SteamDeck deployment requires 1, and only, device");
-			return false;
-		}
-
-		// look up in the devices in the registry based on the params (commandline -device=)
-		List<DeviceInfo> Devices = GetDevices();
-		// assume DeviceName is an Id (ie IpAddr)
-		string DeviceId = Params.DeviceNames[0];
-		DeviceInfo Device = Devices.FirstOrDefault(x => x.Id.Equals(DeviceId, StringComparison.InvariantCultureIgnoreCase));
-		// if the id didn't match, fallback to trying name match
-		if (Device == null)
-		{
-			Device = Devices.FirstOrDefault(x => x.Name.Equals(DeviceId, StringComparison.InvariantCultureIgnoreCase));
-		}
-		// if the device wasn't found, just use the DeviceId (it may be a DNS name, so we can't really verify it here)
-		IpAddr = Device == null ? DeviceId : Device.Id;
-
-
-		// if -deviceuser was specified, prefer that, otherwise use what's in the registry
-		UserName = Params.DeviceUsername;
-		if (string.IsNullOrEmpty(UserName))
-		{
-			UserName = Device?.PlatformValues["UserName"];
-		}
-
-		if (string.IsNullOrEmpty(UserName) || string.IsNullOrEmpty(IpAddr))
-		{
-			CommandUtils.LogWarning("Unable to discover Ip Adress and Username for your Steamdeck. Use -device to specify device, and if that device is not in your Engine.ini, then you will need to specify -deviceuser= to set the username.");
-			return false;
-		}
-
-		return true;
-	}
-
-	/* Deploying to a steam deck currently does 2 things
-	 *
-	 * 1) Generates a script CreateShortcutHelper.sh that will register the game on the SteamDeck once uploaded
-	 * 2) Uploads the build using rsync to the devkit-game location. Once uploaded it runs the CreateShortcutHelper.sh generated before.
-	 */
-	public static void Deploy(ProjectParams Params, DeploymentContext SC)
-	{
-		string IpAddr, UserName;
-		if (GetDeviceInfo(Params, out IpAddr, out UserName) == false)
-		{
-			return;
-		}
-
-		string DevKitRSAPath  = Path.Combine(Environment.GetFolderPath(Environment.SpecialFolder.LocalApplicationData), @"steamos-devkit\steamos-devkit\devkit_rsa");
-		string KnownHostsPath = Path.Combine(Environment.GetFolderPath(Environment.SpecialFolder.UserProfile), @".ssh\known_hosts");
-		string SSHArgs        = $"-i {DevKitRSAPath} {UserName}@{IpAddr}";
-		string GameFolderPath = $"[REDACTED_USER_HOME]/devkit-game/{Params.ShortProjectName}";
-		string GameRunArgs    = $"{SC.ProjectArgForCommandLines} {Params.StageCommandline} {Params.RunCommandline}".Replace("\"", "\\\"");
-		FileReference ExePath = Params.GetProjectExeForPlatform(UnrealTargetPlatform.Win64);
-		string RelGameExePath = ExePath.MakeRelativeTo(DirectoryReference.Combine(ExePath.Directory, "../../..")).Replace('\\', '/');
-		// string DevkitUtilPath = Microsoft.Win32.Registry.GetValue(@"HKEY_LOCAL_MACHINE\SOFTWARE\Microsoft\Windows\CurrentVersion\Uninstall\Steam App 943760", "InstallLocation", null) as string;
-
-		// get the path to some compiler's Msvsmon/RemoteDebugger support files to copy over
-		string MsvsmonPath = null;
-		string MsvsmonVersion = "2022";
-		IEnumerable<DirectoryReference> InstallDirs = WindowsExports.TryGetVSInstallDirs(WindowsCompiler.VisualStudio2022);
-		if (InstallDirs == null)
-		{
-			InstallDirs = WindowsExports.TryGetVSInstallDirs(WindowsCompiler.VisualStudio2019);
-			MsvsmonVersion = "2019";
-		}
-		if (InstallDirs != null)
-		{
-			MsvsmonPath = Path.Combine(InstallDirs.First().FullName, "Common7/IDE/Remote Debugger");
-		}
-		string TargetMsvsmonPath = $"[REDACTED_USER_HOME]/devkit-msvsmon/msvsmon{MsvsmonVersion}";
-
-
-		// make a standard set of options to pass to the rsync for auth's -e option
-		string[] AuthOpts =
-		{
-			SteamDeckSupport.SSHPath,
-			$"-o UserKnownHostsFile='{KnownHostsPath}'",
-			$"-o StrictHostKeyChecking=no",
-			$"-i '{DevKitRSAPath}'",
-		};
-		string AuthOptions = $"-e \"{string.Join(" ", AuthOpts)}\"";
-
-
-		if (!File.Exists(DevKitRSAPath))
-		{
-			CommandUtils.LogWarning("Unable to find '{0}' rsa key needed to deploy to the steam deck. Make sure you've installed the SteamOS Devkit client", DevKitRSAPath);
-			return;
-		}
-
-
-		// rsync remote debugger support files
-		if (string.IsNullOrEmpty(MsvsmonPath) || !Directory.Exists(MsvsmonPath))
-		{
-			CommandUtils.LogWarning("Unable to find MSVC remote debugger source files, remote debugging may not work correctly.");
-		}
-		else
-		{
-			string[] DebuggerOpts =
-			{
-				// ?
-				$"-avh",
-				// ?
-				$"--chmod=Du=rwx,Dgo=rx,Fu=rwx,Fog=rx",
-				// standard authorization options
-				AuthOptions,
-				// this is a trick to make sure the target path exists before running the remote rsync command
-				$"--rsync-path=\"mkdir -p {TargetMsvsmonPath} && rsync\"",
-				// copy/update files as needed
-				$"--update",
-				// source path, in cygwin-speak (/cygdrive/c/Program Files.....)
-				$"'{ConvertWindowsPathToCygdrive(MsvsmonPath)}/'",
-				// destrination in standard 'user@ipaddr:destpath' format
-				$"{UserName}@{IpAddr}:{TargetMsvsmonPath}",
-			};
-
-			// run rsync command
-			IProcessResult DebuggerSyncResult = CommandUtils.Run(SteamDeckSupport.RSyncPath, string.Join(" ", DebuggerOpts), "");
-			if (DebuggerSyncResult.ExitCode > 0)
-			{
-				CommandUtils.LogWarning($"Failed to rsync debugger support files to the SteamDeck. Check connection on ip {UserName}@{IpAddr}");
-				return;
-			}
-		}
-
-
-		string ScriptFileName = "CreateShortcutHelper.sh";
-		string ScriptFile = Path.Combine(SC.StageDirectory.FullName, ScriptFileName);
-		File.WriteAllText(ScriptFile, SteamDeckSupport.GetRegisterGameScript(Params.ShortProjectName, RelGameExePath, GameFolderPath, GameRunArgs, MsvsmonVersion));
-
-		string[] Opts =
-		{
-			// ?
-			$"-avh",
-			// ?
-			$"--chmod=Du=rwx,Dgo=rx,Fu=rwx,Fog=rx",
-			// standard authorization options
-			AuthOptions,
-			// this is a trick to make sure the target path exists before running the remote rsync command
-			$"--rsync-path=\"mkdir -p {GameFolderPath} && rsync\"",
-			// copy/update files as needed
-			$"--update",
-			// source path, in cygwin-speak (/cygdrive/c/Program Files.....)
-			$"'{ConvertWindowsPathToCygdrive(SC.StageDirectory.FullName)}/'",
-			// destrination in standard 'user@ipaddr:destpath' format
-			$"{UserName}@{IpAddr}:{GameFolderPath}",
-		};
-
-		// Exclude removing the Saved folders to preserve logs and crash data. Though note these will keep filling up with data
-		IProcessResult Result = CommandUtils.Run(SteamDeckSupport.RSyncPath, string.Join(" ", Opts), "");
-		if (Result.ExitCode > 0)
-		{
-			CommandUtils.LogWarning("Failed to rsync files to the SteamDeck. Check connection on ip {UserName}@{IpAddr}");
-			return;
-		}
-
-		// Run the script to register the game with the Deck
-		Result = CommandUtils.Run(SteamDeckSupport.SSHPath, $"{SSHArgs} \"chmod +x {GameFolderPath}/{ScriptFileName} && {GameFolderPath}/{ScriptFileName}\"", "");
-
-		if (Result.ExitCode > 0)
-		{
-			CommandUtils.LogWarning($"Failed to run the {ScriptFileName}.sh script. Check connection on ip Check connection on ip {UserName}@{IpAddr}");
-			return;
-		}
-	}
-}
 
 public class Win64Platform : Platform
 {
@@ -300,7 +36,7 @@ public class Win64Platform : Platform
 
 			Devices.Add(LocalMachine);
 
-			Devices.AddRange(SteamDeckSupport.GetDevices());
+			Devices.AddRange(SteamDeckSupport.GetDevices(UnrealTargetPlatform.Win64));
 		}
 
 		return Devices.ToArray();
@@ -311,7 +47,7 @@ public class Win64Platform : Platform
 		// We only care about deploying for SteamDeck
 		if (Params.Devices.Count == 1 && GetDevices().FirstOrDefault(x => x.Id == Params.DeviceNames[0])?.Type == "SteamDeck")
 		{
-			SteamDeckSupport.Deploy(Params, SC);
+			SteamDeckSupport.Deploy(UnrealTargetPlatform.Win64, Params, SC);
 		}
 	}
 
@@ -319,10 +55,7 @@ public class Win64Platform : Platform
 	{
 		if (Params.Devices.Count == 1 && GetDevices().FirstOrDefault(x => x.Id == Params.DeviceNames[0])?.Type == "SteamDeck")
 		{
-			// TODO figure out how to get the steam app num id. Then can run like this:
-			// steam steam://rungameid/<GameNumId>
-			// TODO would be great if we could tail the log while running. Figure out how to cancel/exit app
-			return null;
+			return SteamDeckSupport.RunClient(UnrealTargetPlatform.Win64, ClientRunFlags, ClientApp, ClientCmdLine, Params);
 		}
 
 		return base.RunClient(ClientRunFlags, ClientApp, ClientCmdLine, Params);
