- Fixed some dependency checking in RunUAT's CompileScripts - the code that checks for all automation modules to be up to date was checking against a list of valid modules that included more than the automation modules, so the numeric comparison was incorrect - usually would force a compile even when nothing changed, and sometimes would fail to compile when code did change.

diff --git a/Engine/Source/Programs/AutomationTool/CompileScriptModules.cs b/Engine/Source/Programs/AutomationTool/CompileScriptModules.cs
--- a/Engine/Source/Programs/AutomationTool/CompileScriptModules.cs
+++ b/Engine/Source/Programs/AutomationTool/CompileScriptModules.cs
@@ -156,7 +156,7 @@ namespace AutomationToolDriver
 			if (!bForceCompile && bUseBuildRecords)
 			{
 				// If all found records are valid, we can return their targets directly
-				if (ValidBuildRecords.Count == FoundAutomationProjects.Count)
+				if (FoundAutomationProjects.All(x => ValidBuildRecords.ContainsKey(x)))
                 {
 					bBuildSuccess = true;
 					return new HashSet<FileReference>(ValidBuildRecords.Select(x => FileReference.Combine(x.Key.Directory, x.Value.TargetPath)));
