- Removed code that was accidentally checked in, that broke Android .apk names

diff --git a/Engine/Source/Programs/AutomationTool/AutomationUtils/Platform.cs b/Engine/Source/Programs/AutomationTool/AutomationUtils/Platform.cs
--- a/Engine/Source/Programs/AutomationTool/AutomationUtils/Platform.cs
+++ b/Engine/Source/Programs/AutomationTool/AutomationUtils/Platform.cs
@@ -85,7 +85,7 @@ namespace AutomationTool
 
 		public override string ToString()
 		{
-			return Type.ToString() + CookFlavor;
+			return Type.ToString();
 		}
 	}
 
