Fix issue with packaging with data inside APK with AFS for Android

#android

diff --git a/Engine/Source/Programs/AutomationTool/Android/AndroidPlatform.Automation.cs b/Engine/Source/Programs/AutomationTool/Android/AndroidPlatform.Automation.cs
--- a/Engine/Source/Programs/AutomationTool/Android/AndroidPlatform.Automation.cs
+++ b/Engine/Source/Programs/AutomationTool/Android/AndroidPlatform.Automation.cs
@@ -791,7 +791,7 @@ public class AndroidPlatform : Platform
 		bool bAFSAllowExternalStartInShipping;
 		UsingAndroidFileServer(SC, out bAFSEnablePlugin, out AFSToken, out bIsShipping, out bAFSIncludeInShipping, out bAFSAllowExternalStartInShipping);
 
-		if (bAFSEnablePlugin)
+		if (bAFSEnablePlugin && !bPackageDataInsideApk)
 		{
 			bUseAFS = true;
 			// AFSProject APK should be used if shipping and AFS wasn't included
@@ -1355,7 +1355,7 @@ public class AndroidPlatform : Platform
 						"echo There was an error installing the game or the obb file. Look above for more info.",
 						"echo",
 						"echo Things to try:",
-						"echo Check that the device (and only the device) is listed with \\\"$ADB devices\\\" from a command prompt.",
+						"echo 'Check that the device (and only the device) is listed with \\\"$ADB devices\\\" from a command prompt.'",
 						"echo Make sure all Developer options look normal on the device",
 						"echo Check that the device has an SD card.",
 						"exit 1"
