UnrealTraceServer should start with "fork" argument.

diff --git a/Engine/Source/Programs/AutomationTool/Scripts/RunProjectCommand.Automation.cs b/Engine/Source/Programs/AutomationTool/Scripts/RunProjectCommand.Automation.cs
--- a/Engine/Source/Programs/AutomationTool/Scripts/RunProjectCommand.Automation.cs
+++ b/Engine/Source/Programs/AutomationTool/Scripts/RunProjectCommand.Automation.cs
@@ -1089,7 +1089,7 @@ namespace AutomationScripts
 			}
 
 			// Launch UnrealTrace and wait for it to fork and return
-			Process Proc = Process.Start(UnrealTracePath);
+			Process Proc = Process.Start(UnrealTracePath, " fork");
 			Proc.WaitForExit();
 			if (Proc.ExitCode != 0)
 			{
