* Fix spelling on Retrieve
* Copy manifest ZIP files locally before unzipping them.
* Use retries when copying to harden against network blips.

#rnx

diff --git a/Engine/Source/Programs/AutomationTool/BuildGraph/BuildGraph.cs b/Engine/Source/Programs/AutomationTool/BuildGraph/BuildGraph.cs
--- a/Engine/Source/Programs/AutomationTool/BuildGraph/BuildGraph.cs
+++ b/Engine/Source/Programs/AutomationTool/BuildGraph/BuildGraph.cs
@@ -1154,7 +1154,7 @@ namespace AutomationTool
 				Scope.Span.SetTag("blocks", InputStorageBlocks.Count);
 				foreach (TempStorageBlock InputStorageBlock in InputStorageBlocks)
 				{
-					TempStorageManifest Manifest = Storage.Retreive(InputStorageBlock.NodeName, InputStorageBlock.OutputName);
+					TempStorageManifest Manifest = Storage.Retrieve(InputStorageBlock.NodeName, InputStorageBlock.OutputName);
 					InputManifests[InputStorageBlock] = Manifest;
 				}
 				Scope.Span.SetTag("size", InputManifests.Sum(x => x.Value.GetTotalSize()));
diff --git a/Engine/Source/Programs/AutomationTool/BuildGraph/TempStorage.cs b/Engine/Source/Programs/AutomationTool/BuildGraph/TempStorage.cs
--- a/Engine/Source/Programs/AutomationTool/BuildGraph/TempStorage.cs
+++ b/Engine/Source/Programs/AutomationTool/BuildGraph/TempStorage.cs
@@ -873,7 +873,7 @@ namespace AutomationTool
 		/// <param name="NodeName">The node which created the storage block</param>
 		/// <param name="OutputName">Name of the block to retrieve. May be null or empty.</param>
 		/// <returns>Manifest of the files retrieved</returns>
-		public TempStorageManifest Retreive(string NodeName, string OutputName)
+		public TempStorageManifest Retrieve(string NodeName, string OutputName)
 		{
 			using (IScope Scope = GlobalTracer.Instance.BuildSpan("RetrieveFromTempStorage").StartActive())
 			{
@@ -997,7 +997,7 @@ namespace AutomationTool
 						try
 						{
 							// Create one zip per thread using the given basename
-							using (var ZipArchive = System.IO.Compression.ZipFile.Open(ZipFileName.FullName, System.IO.Compression.ZipArchiveMode.Create))
+							using (var ZipArchive = ZipFile.Open(ZipFileName.FullName, ZipArchiveMode.Create))
 							{
 								// pull from the queue until we are out of files.
 								do
@@ -1005,11 +1005,10 @@ namespace AutomationTool
 									// use fastest compression. In our best case we are CPU bound, so this is a good tradeoff,
 									// cutting overall time by 2/3 while only modestly increasing the compression ratio (22.7% -> 23.8% for RootEditor PDBs).
 									// This is in cases of a super hot cache, so the operation was largely CPU bound.
-									ZipArchiveExtensions.CreateEntryFromFile_CrossPlatform(ZipArchive, File.FullName, CommandUtils.ConvertSeparators(PathSeparator.Slash, File.MakeRelativeTo(RootDir)), System.IO.Compression.CompressionLevel.Fastest);
+									ZipArchiveExtensions.CreateEntryFromFile_CrossPlatform(ZipArchive, File.FullName, CommandUtils.ConvertSeparators(PathSeparator.Slash, File.MakeRelativeTo(RootDir)), CompressionLevel.Fastest);
 								} while (FilesToZip.TryDequeue(out File));
 							}
 						}
-
 						catch (IOException)
 						{
 							CommandUtils.LogError("Unable to open file for TempStorage zip: \"{0}\"", ZipFileName.FullName);
@@ -1019,19 +1018,18 @@ namespace AutomationTool
 						try
 						{
 							// if we are using a staging dir, copy to the final location and delete the staged copy.
-							FileInfo ZipFile = new FileInfo(ZipFileName.FullName);
+							string FinalZipFile = ZipFileName.FullName;
 							if (StagingDir != null)
 							{
-								FileInfo NewZipFile = ZipFile.CopyTo(CommandUtils.MakeRerootedFilePath(ZipFile.FullName, StagingDir.FullName, OutputDir.FullName));
-								ZipFile.Delete();
-								ZipFile = NewZipFile;
+								FinalZipFile = CommandUtils.CombinePaths(OutputDir.FullName, ZipFileName.GetFileName());
+								CommandUtils.CopyFile(ZipFileName.FullName, FinalZipFile, bQuiet: true, bRetry: true);
+								FileReference.Delete(ZipFileName);
 							}
-							ZipFiles.Add(ZipFile);					
+							ZipFiles.Add(new FileInfo(FinalZipFile));
 						}
-
-						catch (IOException)
+						catch (IOException Ex)
 						{
-							CommandUtils.LogError("Unable to copy file for TempStorage zip: \"{0}\"", ZipFileName.FullName);
+							CommandUtils.LogError("Unable to copy file \"{0}\" to TempStorage: {1}", ZipFileName.FullName, Ex);
 							throw new AutomationException("Unable to copy file {0}", ZipFileName.FullName);
 						}
 					}
@@ -1062,6 +1060,14 @@ namespace AutomationTool
 			Parallel.ForEach(ZipFiles,
 				(ZipFile) =>
 				{
+					// Copy the zip file locally before unzipped to harden against network hiccups.
+					string LocalZipPath = CommandUtils.CombinePaths(RootDir.FullName, "Engine/Intermediate/Manifests", ZipFile.Name);
+					CommandUtils.CopyFile(ZipFile.FullName, LocalZipPath, bQuiet: true, bRetry: true);
+					if (!File.Exists(LocalZipPath))
+					{
+						throw new AutomationException("Failed to copy manifest {0} to {1}", ZipFile, LocalZipPath);
+					}
+
 					// unzip the files manually instead of caling ZipFile.ExtractToDirectory() because we need to overwrite readonly files. Because of this, creating the directories is up to us as well.
 					List<string> ExtractedPaths = new List<string>();
 					int UnzipFileAttempts = 3;
@@ -1069,7 +1075,7 @@ namespace AutomationTool
 					{
 						try
 						{
-							using (ZipArchive ZipArchive = System.IO.Compression.ZipFile.OpenRead(ZipFile.FullName))
+							using (ZipArchive ZipArchive = System.IO.Compression.ZipFile.OpenRead(LocalZipPath))
 							{
 								foreach (ZipArchiveEntry Entry in ZipArchive.Entries)
 								{
@@ -1117,7 +1123,7 @@ namespace AutomationTool
 													throw;
 												}
 
-												Log.TraceWarning("Failed to unzip '{0}' from '{1}' to '{2}', retrying.. (Error: {3})", Entry.FullName, ZipFile.FullName, ExtractedFilename, IOEx.Message);
+												Log.TraceWarning("Failed to unzip '{0}' from '{1}' to '{2}', retrying.. (Error: {3})", Entry.FullName, LocalZipPath, ExtractedFilename, IOEx.Message);
 											}
 										}
 									}
@@ -1130,16 +1136,20 @@ namespace AutomationTool
 						{
 							if (UnzipFileAttempts == 0)
 							{
-								Log.TraceError("All retries exhausted attempting to unzip entries from '{0}'. Terminating.", ZipFile.FullName);
+								Log.TraceError("All retries exhausted attempting to unzip entries from '{0}'. Terminating.", LocalZipPath);
 								throw;
 							}
 
 							// Some exceptions may be caused by networking hiccups. We want to retry in those cases.
 							if ((Ex is IOException || Ex is InvalidDataException))
 							{
-								Log.TraceWarning("Failed to unzip entries from '{0}' to '{1}', retrying.. (Error: {2})", ZipFile.FullName, RootDir.FullName, Ex.Message);
+								Log.TraceWarning("Failed to unzip entries from '{0}' to '{1}', retrying.. (Error: {2})", LocalZipPath, RootDir.FullName, Ex.Message);
 							}
 						}
+						finally
+						{
+							File.Delete(LocalZipPath);
+						}
 					}
 				});
 		}
@@ -1284,10 +1294,10 @@ namespace AutomationTool
 			TempStore.Archive("TestNode", "NamedOutput", NamedOutput.Keys.ToArray(), true);
 			
 			// Check both outputs are still ok
-			TempStorageManifest DefaultManifest = TempStore.Retreive("TestNode", null);
+			TempStorageManifest DefaultManifest = TempStore.Retrieve("TestNode", null);
 			CheckManifest(WorkingDir, DefaultManifest, DefaultOutput);
 
-			TempStorageManifest NamedManifest = TempStore.Retreive("TestNode", "NamedOutput");
+			TempStorageManifest NamedManifest = TempStore.Retrieve("TestNode", "NamedOutput");
 			CheckManifest(WorkingDir, NamedManifest, NamedOutput);
 
 			// Delete local temp storage and the working directory and try again
@@ -1300,7 +1310,7 @@ namespace AutomationTool
 			bool bGotManifest = false;
 			try
 			{
-				TempStore.Retreive("TestNode", null);
+				TempStore.Retrieve("TestNode", null);
 				bGotManifest = true;
 			}
 			catch
@@ -1313,7 +1323,7 @@ namespace AutomationTool
 			}
 
 			// Second one should be fine
-			TempStorageManifest NamedManifestFromShared = TempStore.Retreive("TestNode", "NamedOutput");
+			TempStorageManifest NamedManifestFromShared = TempStore.Retrieve("TestNode", "NamedOutput");
 			CheckManifest(WorkingDir, NamedManifestFromShared, NamedOutput);
 		}
 
