Fixed issue where compling script modules in UAT where all modules were up-to-date would result in extra assemblies being returned.

diff --git a/Engine/Source/Programs/AutomationTool/CompileScriptModules.cs b/Engine/Source/Programs/AutomationTool/CompileScriptModules.cs
--- a/Engine/Source/Programs/AutomationTool/CompileScriptModules.cs
+++ b/Engine/Source/Programs/AutomationTool/CompileScriptModules.cs
@@ -159,8 +159,8 @@ namespace AutomationToolDriver
 				if (FoundAutomationProjects.All(x => ValidBuildRecords.ContainsKey(x)))
                 {
 					bBuildSuccess = true;
-					return new HashSet<FileReference>(ValidBuildRecords.Select(x => FileReference.Combine(x.Key.Directory, x.Value.TargetPath)));
-                }
+					return new HashSet<FileReference>(FoundAutomationProjects.Select(x => FileReference.Combine(x.Directory, ValidBuildRecords[x].TargetPath!)));
+				}
 			}
 
 			// Fall back to the slower approach: use msbuild to load csproj files & build as necessary
