Tests converted from 18848115 made to work multi-platform.
Not all were converted because they were designed for packaged applications whereas low level tests are designed to run as "native" non-packaged applications.
- reporting support for non-desktop platforms, for the moment Catch2 report type "console" is used that is sent to a .out file
- most number of tests possible running multi platform, slower tests excluded on incremental builds
- slow tests are moved to run on the monolithic build
- Catch2 report failure event listener such that Horde detects them as build errors. Must add new Horde matcher for Catch2 failures

diff --git a/Engine/Source/Developer/LowLevelTestsRunner/Private/TestCommon/ApplicationCoreUtilities.cpp b/Engine/Source/Developer/LowLevelTestsRunner/Private/TestCommon/ApplicationCoreUtilities.cpp
--- a/Engine/Source/Developer/LowLevelTestsRunner/Private/TestCommon/ApplicationCoreUtilities.cpp
+++ b/Engine/Source/Developer/LowLevelTestsRunner/Private/TestCommon/ApplicationCoreUtilities.cpp
@@ -3,6 +3,7 @@
 #if WITH_APPLICATION_CORE
 
 #include "TestCommon/ApplicationCoreUtilities.h"
+#include "HAL/PlatformApplicationMisc.h"
 
 void InitOutputDevicesAppCore()
 {
diff --git a/Engine/Source/Developer/LowLevelTestsRunner/Private/TestCommon/CoreUObjectUtilities.cpp b/Engine/Source/Developer/LowLevelTestsRunner/Private/TestCommon/CoreUObjectUtilities.cpp
--- a/Engine/Source/Developer/LowLevelTestsRunner/Private/TestCommon/CoreUObjectUtilities.cpp
+++ b/Engine/Source/Developer/LowLevelTestsRunner/Private/TestCommon/CoreUObjectUtilities.cpp
@@ -3,6 +3,12 @@
 #if WITH_COREUOBJECT
 
 #include "TestCommon/CoreUObjectUtilities.h"
+#include "HAL/PlatformFileManager.h"
+#include "Misc/ConfigCacheIni.h"
+#include "Misc/CoreDelegates.h"
+#include "Modules/ModuleManager.h"
+#include "UObject/Package.h"
+#include "UObject/PackageResourceManager.h"
 
 void InitCoreUObject()
 {
diff --git a/Engine/Source/Developer/LowLevelTestsRunner/Private/TestCommon/CoreUtilities.cpp b/Engine/Source/Developer/LowLevelTestsRunner/Private/TestCommon/CoreUtilities.cpp
--- a/Engine/Source/Developer/LowLevelTestsRunner/Private/TestCommon/CoreUtilities.cpp
+++ b/Engine/Source/Developer/LowLevelTestsRunner/Private/TestCommon/CoreUtilities.cpp
@@ -3,6 +3,10 @@
 #include "CoreMinimal.h"
 #include "Misc/CommandLine.h"
 #include "TestCommon/CoreUtilities.h"
+#include "HAL/PlatformOutputDevices.h"
+#include "HAL/PlatformProcess.h"
+#include "Logging/LogSuppressionInterface.h"
+#include "Misc/QueuedThreadPool.h"
 
 void InitCommandLine(bool AllowLogging)
 {
@@ -100,16 +104,14 @@ void CleanupTaskGraph()
 
 void InitTaskGraphAndDependencies(bool MultiThreaded)
 {
-	InitCommandLine();
-	InitAllThreadPools(MultiThreaded);
 	InitTaskGraph(MultiThreaded);
+	InitAllThreadPools(MultiThreaded);
 }
 
 void CleanupTaskGraphAndDependencies()
 {
-	CleanupTaskGraph();
 	CleanupAllThreadPools();
-	CleanupCommandLine();
+	CleanupTaskGraph();
 }
 
 void CleanupPlatform()
diff --git a/Engine/Source/Developer/LowLevelTestsRunner/Private/TestCommon/EngineUtilities.cpp b/Engine/Source/Developer/LowLevelTestsRunner/Private/TestCommon/EngineUtilities.cpp
--- a/Engine/Source/Developer/LowLevelTestsRunner/Private/TestCommon/EngineUtilities.cpp
+++ b/Engine/Source/Developer/LowLevelTestsRunner/Private/TestCommon/EngineUtilities.cpp
@@ -4,6 +4,13 @@
 
 #include "TestCommon/EngineUtilities.h"
 
+#include "CoreMinimal.h"
+#include "DistanceFieldAtlas.h"
+#include "MeshCardRepresentation.h"
+#include "RenderUtils.h"
+#include "ShaderParameterMetadata.h"
+
+
 void InitAsyncQueues()
 {
 	check(!GDistanceFieldAsyncQueue);
diff --git a/Engine/Source/Developer/LowLevelTestsRunner/Private/TestCommon/Initialization.cpp b/Engine/Source/Developer/LowLevelTestsRunner/Private/TestCommon/Initialization.cpp
--- a/Engine/Source/Developer/LowLevelTestsRunner/Private/TestCommon/Initialization.cpp
+++ b/Engine/Source/Developer/LowLevelTestsRunner/Private/TestCommon/Initialization.cpp
@@ -1,6 +1,12 @@
 // Copyright Epic Games, Inc. All Rights Reserved.
 
 #include "TestCommon/Initialization.h"
+#include "TestCommon/ApplicationCoreUtilities.h"
+#include "TestCommon/CoreUtilities.h"
+#include "TestCommon/CoreUObjectUtilities.h"
+#include "TestCommon/EditorUtilities.h"
+#include "TestCommon/EngineUtilities.h"
+
 
 
 void InitAllThreadPoolsEditorEx(bool MultiThreaded)
diff --git a/Engine/Source/Developer/LowLevelTestsRunner/Private/TestRunner.cpp b/Engine/Source/Developer/LowLevelTestsRunner/Private/TestRunner.cpp
--- a/Engine/Source/Developer/LowLevelTestsRunner/Private/TestRunner.cpp
+++ b/Engine/Source/Developer/LowLevelTestsRunner/Private/TestRunner.cpp
@@ -10,19 +10,41 @@
 #include "LowLevelTestModule.h"
 #include "TestCommon/CoreUtilities.h"
 
-// Test run interceptor
-struct TestRunListener : public Catch::TestEventListenerBase {
-	using TestEventListenerBase::TestEventListenerBase; // inherit constructor
-public:
-	void testCaseStarting(Catch::TestCaseInfo  const& TestInfo) override {
-		if (bGDebug)
-		{
-			std::cout << TestInfo.lineInfo.file << ":" << TestInfo.lineInfo.line << " with tags " << TestInfo.tagsAsString() << std::endl;
+#include <set>
+
+namespace Catch
+{
+	// Test run interceptor
+	struct TestRunListener : public TestEventListenerBase {
+		using TestEventListenerBase::TestEventListenerBase; // inherit constructor
+	private:
+		std::ostream& catchOut = getCurrentContext().getConfig()->stream();
+	public:
+		void testCaseStarting(Catch::TestCaseInfo  const& TestInfo) override {
+			if (bGDebug)
+			{
+				catchOut << TestInfo.lineInfo.file << ":" << TestInfo.lineInfo.line << " with tags " << TestInfo.tagsAsString() << " \n";
+			}
 		}
-	}
-};
 
-CATCH_REGISTER_LISTENER(TestRunListener);
+		void testCaseEnded(Catch::TestCaseStats const& testCaseStats) override {
+			if (testCaseStats.totals.testCases.failed > 0)
+			{
+				catchOut << "* Error: Test case \"" << testCaseStats.testInfo.name << "\" failed \n";
+			}
+		}
+
+		bool assertionEnded(Catch::AssertionStats const& assertionStats) override {
+			if (!assertionStats.assertionResult.succeeded()) {
+				catchOut << "* Error: Assertion \"" << assertionStats.assertionResult.getExpression() << "\" failed at " << assertionStats.assertionResult.getSourceInfo().file << ": " << assertionStats.assertionResult.getSourceInfo().line << "\n";
+			}
+			// true == clear message buffer
+			return true;
+		}
+	};
+
+	CATCH_REGISTER_LISTENER(TestRunListener);
+}
 
 void LoadBaseTestModule(FString BaseModuleName)
 {
@@ -108,56 +130,75 @@ int RunTests(int argc, const char* argv[])
 	}
 #endif
 
-	int CatchArgc;
+	int CatchArgc = 0;
 	TUniquePtr<const char* []> CatchArgv = MakeUnique<const char* []>(argc);
 
-	// Everything past a "--" argument, if present, is not sent to the catch test runner.
-	for (CatchArgc = 0; CatchArgc < argc; ++CatchArgc)
-	{
-		if (std::strcmp(argv[CatchArgc], "--") == 0)
-		{
-			break;
-		}
-		CatchArgv.Get()[CatchArgc] = argv[CatchArgc];
-	}
+	std::set<std::string> KnownLLTArgs;
+	KnownLLTArgs.insert("--wait");
+	KnownLLTArgs.insert("--no-wait");
+	KnownLLTArgs.insert("--log");
+	KnownLLTArgs.insert("--no-log");
+	KnownLLTArgs.insert("--mt");
+	KnownLLTArgs.insert("--no-mt");
+	KnownLLTArgs.insert("--debug");
+	KnownLLTArgs.insert("--base-global-module");
 
 	// By default don't wait for input
 	bool bWaitForInputToTerminate = false;
 
-	
 	FString LoadModuleName = TEXT("");
 
-	for (int i = CatchArgc; i < argc; ++i)
+	// Every argument that is not in the list of known low level test options and is considered to be a catch test runner argument.
+	for (int i = 0; i < argc; ++i)
 	{
-		if (std::strcmp(argv[i], "--wait") == 0)
-		{
-			bWaitForInputToTerminate = true;
-		}
-		else if (std::strcmp(argv[i], "--no-wait") == 0)
-		{
-			bWaitForInputToTerminate = false;
-		}
-		if (std::strcmp(argv[i], "--no-log") == 0)
+		if (KnownLLTArgs.find(argv[i]) != KnownLLTArgs.end())
 		{
-			bGAllowLogging = false;
-		}
-		if (std::strcmp(argv[i], "--no-mt") == 0)
-		{
-			bGMultithreaded = false;
-		}
-		if (std::strcmp(argv[i], "--debug") == 0)
-		{
-			bGDebug = true;
+			if (std::strcmp(argv[i], "--wait") == 0)
+			{
+				bWaitForInputToTerminate = true;
+			}
+			else if (std::strcmp(argv[i], "--no-wait") == 0)
+			{
+				bWaitForInputToTerminate = false;
+			}
+			if (std::strcmp(argv[i], "--log") == 0)
+			{
+				bGAllowLogging = true;
+			}
+			if (std::strcmp(argv[i], "--no-log") == 0)
+			{
+				bGAllowLogging = false;
+			}
+			if (std::strcmp(argv[i], "--mt") == 0)
+			{
+				bGMultithreaded = true;
+			}
+			if (std::strcmp(argv[i], "--no-mt") == 0)
+			{
+				bGMultithreaded = false;
+			}
+			if (std::strcmp(argv[i], "--debug") == 0)
+			{
+				bGDebug = true;
+			}
+			// If we have --base-global-module parse proceeding argument as its option-value
+			if (std::strcmp(argv[i], "--base-global-module") == 0 &&
+				i + 1 < argc)
+			{
+				LoadModuleName = argv[i + 1];
+				++i;
+			}
 		}
-		// If we have --base-global-module parse proceeding argument as its option-value
-		if (std::strcmp(argv[i], "--base-global-module") == 0 &&
-			i + 1 < argc)
+		else
 		{
-			LoadModuleName = argv[i + 1];
-			++i;
+			// Passing to catch2
+			CatchArgv.Get()[CatchArgc++] = argv[i];
 		}
 	}
 
+	// Global command line initialization
+	InitCommandLine(bGAllowLogging);
+
 	LoadBaseTestModule(LoadModuleName);
 	GlobalModuleSetup();
 
@@ -171,9 +212,12 @@ int RunTests(int argc, const char* argv[])
 	GlobalModuleTeardown();
 	UnloadBaseTestModule(LoadModuleName);
 
-	// Required, will crash on exit otherwise...
+	// Required for platform cleanup, program will crash on exit otherwise
 	CleanupPlatform();
 
+	// Command line cleanup
+	CleanupCommandLine();
+
 #if PLATFORM_DESKTOP
 	if (bWaitForInputToTerminate)
 	{
diff --git a/Engine/Source/Developer/LowLevelTestsRunner/Public/TestCommon/ApplicationCoreUtilities.h b/Engine/Source/Developer/LowLevelTestsRunner/Public/TestCommon/ApplicationCoreUtilities.h
--- a/Engine/Source/Developer/LowLevelTestsRunner/Public/TestCommon/ApplicationCoreUtilities.h
+++ b/Engine/Source/Developer/LowLevelTestsRunner/Public/TestCommon/ApplicationCoreUtilities.h
@@ -4,8 +4,6 @@
 
 #if WITH_APPLICATION_CORE
 
-#include "HAL/PlatformApplicationMisc.h"
-
 void InitOutputDevicesAppCore();
 
 #endif // WITH_APPLICATION_CORE
diff --git a/Engine/Source/Developer/LowLevelTestsRunner/Public/TestCommon/CoreUObjectUtilities.h b/Engine/Source/Developer/LowLevelTestsRunner/Public/TestCommon/CoreUObjectUtilities.h
--- a/Engine/Source/Developer/LowLevelTestsRunner/Public/TestCommon/CoreUObjectUtilities.h
+++ b/Engine/Source/Developer/LowLevelTestsRunner/Public/TestCommon/CoreUObjectUtilities.h
@@ -4,14 +4,6 @@
 
 #if WITH_COREUOBJECT
 
-#include "HAL/PlatformFileManager.h"
-#include "Misc/ConfigCacheIni.h"
-#include "Misc/CoreDelegates.h"
-#include "Modules/ModuleManager.h"
-#include "UObject/Package.h"
-#include "UObject/PackageResourceManager.h"
-
-
 void InitCoreUObject();
 
 void CleanupCoreUObject();
diff --git a/Engine/Source/Developer/LowLevelTestsRunner/Public/TestCommon/CoreUtilities.h b/Engine/Source/Developer/LowLevelTestsRunner/Public/TestCommon/CoreUtilities.h
--- a/Engine/Source/Developer/LowLevelTestsRunner/Public/TestCommon/CoreUtilities.h
+++ b/Engine/Source/Developer/LowLevelTestsRunner/Public/TestCommon/CoreUtilities.h
@@ -3,12 +3,7 @@
 #pragma once
 
 #include "CoreMinimal.h"
-
 #include "Async/TaskGraphInterfaces.h"
-#include "HAL/PlatformOutputDevices.h"
-#include "HAL/PlatformProcess.h"
-#include "Logging/LogSuppressionInterface.h"
-#include "Misc/QueuedThreadPool.h"
 
 void InitCommandLine(bool AllowLogging = false);
 
diff --git a/Engine/Source/Developer/LowLevelTestsRunner/Public/TestCommon/EngineUtilities.h b/Engine/Source/Developer/LowLevelTestsRunner/Public/TestCommon/EngineUtilities.h
--- a/Engine/Source/Developer/LowLevelTestsRunner/Public/TestCommon/EngineUtilities.h
+++ b/Engine/Source/Developer/LowLevelTestsRunner/Public/TestCommon/EngineUtilities.h
@@ -4,13 +4,6 @@
 
 #if WITH_ENGINE
 
-#include "CoreMinimal.h"
-
-#include "DistanceFieldAtlas.h"
-#include "MeshCardRepresentation.h"
-#include "RenderUtils.h"
-#include "ShaderParameterMetadata.h"
-
 void InitAsyncQueues();
 void InitRendering();
 
diff --git a/Engine/Source/Developer/LowLevelTestsRunner/Public/TestCommon/Initialization.h b/Engine/Source/Developer/LowLevelTestsRunner/Public/TestCommon/Initialization.h
--- a/Engine/Source/Developer/LowLevelTestsRunner/Public/TestCommon/Initialization.h
+++ b/Engine/Source/Developer/LowLevelTestsRunner/Public/TestCommon/Initialization.h
@@ -4,13 +4,6 @@
 
 #include "CoreMinimal.h"
 
-#include "TestCommon/ApplicationCoreUtilities.h"
-#include "TestCommon/CoreUtilities.h"
-#include "TestCommon/CoreUObjectUtilities.h"
-#include "TestCommon/EditorUtilities.h"
-#include "TestCommon/EngineUtilities.h"
-
-
 void InitAllThreadPoolsEditorEx(bool MultiThreaded);
 
 void InitOutputDevicesEx();
diff --git a/Engine/Source/Developer/LowLevelTestsRunner/Public/TestHarness.h b/Engine/Source/Developer/LowLevelTestsRunner/Public/TestHarness.h
--- a/Engine/Source/Developer/LowLevelTestsRunner/Public/TestHarness.h
+++ b/Engine/Source/Developer/LowLevelTestsRunner/Public/TestHarness.h
@@ -2,11 +2,11 @@
 
 #pragma once
 
-#ifndef UE5_ENABLE_TESTHARNESS_ENGINE_SUPPORT
-#define UE5_ENABLE_TESTHARNESS_ENGINE_SUPPORT 1
+#ifndef UE_ENABLE_TESTHARNESS_ENGINE_SUPPORT
+#define UE_ENABLE_TESTHARNESS_ENGINE_SUPPORT 1
 #endif
 
-#if UE5_ENABLE_TESTHARNESS_ENGINE_SUPPORT
+#if UE_ENABLE_TESTHARNESS_ENGINE_SUPPORT
 #include "CoreTypes.h"
 #include "Containers/ContainersFwd.h"
 #include "Containers/StringFwd.h"
@@ -58,8 +58,6 @@ inline bool operator==(const TMap<KeyT, ValueT>& Left, const TMap<KeyT, ValueT>&
 	return bIsEqual;
 }
 
-#endif // UE5_ENABLE_TESTHARNESS_ENGINE_SUPPORT
-
 #if defined(THIRD_PARTY_INCLUDES_START) && defined(THIRD_PARTY_INCLUDES_END)
 THIRD_PARTY_INCLUDES_START
 #endif // defined(THIRD_PARTY_INCLUDES_START) && defined(THIRD_PARTY_INCLUDES_END)
@@ -81,8 +79,6 @@ THIRD_PARTY_INCLUDES_START
 THIRD_PARTY_INCLUDES_END
 #endif // defined(THIRD_PARTY_INCLUDES_START) && defined(THIRD_PARTY_INCLUDES_END)
 
-#if UE5_ENABLE_TESTHARNESS_ENGINE_SUPPORT
-
 // Tell Catch how to print TTuple<KeyType, ValueType>
 template <typename... Types> struct TTuple;
 
@@ -95,4 +91,68 @@ struct Catch::StringMaker<TTuple<KeyType, ValueType>>
 	}
 };
 
-#endif // #if UE5_ENABLE_TESTHARNESS_ENGINE_SUPPORT
+// Tests if true and prints message otherwise
+#define TestTrue(a, b) \
+	INFO((a)); \
+	CHECK((b) == true); \
+
+// Quiet version of TestTrue
+#define TestTrue_Q(a) \
+	CHECK((a) == true); \
+
+// Tests if false and prints message otherwise
+#define TestFalse(a, b) \
+	INFO((a)); \
+	CHECK((b) != true); \
+
+// Quiet version of TestFalse
+#define TestFalse_Q(a) \
+	CHECK((a) != true); \
+
+// Tests if b == c and prints message otherwise
+#define TestEqual(a, b, c) \
+	INFO((a)); \
+	CHECK((b) == (c)); \
+
+#define TestEqual_Conditional(a, b, c) \
+	INFO((a)); \
+	CHECKED_IF((b) == (c)) \
+
+// Quiet version of TestEqual
+#define TestEqual_Q(a, b) \
+	CHECK((a) == (b)); \
+
+#define TestNotEqual(a, b, c) \
+	INFO((a)); \
+	CHECK((b) != (c)); \
+
+#define TestNull(a, b) \
+	INFO((a)); \
+	CHECK((b) == nullptr); \
+
+#define TestNotNull(a, b) \
+	INFO((a)); \
+	CHECK((b) != nullptr); \
+
+#define TestValid(a, b) \
+	INFO((a)); \
+	CHECK((b).IsValid() == true); \
+
+#define TestInvalid(a, b) \
+	INFO((a)); \
+	CHECK((b).IsValid() != true); \
+
+#define TestEqualString(a, b, c) \
+	INFO((a)); \
+	CHECK(FCString::Strcmp(ToCStr((b)), ToCStr((c))) == 0); \
+
+#define TestAddInfo(a) \
+	INFO((a)); \
+
+#define TestAddWarning(a) \
+	WARN((a)); \
+
+#define TestAddError(a) \
+	FAIL_CHECK((a)); \
+
+#endif // UE_ENABLE_TESTHARNESS_ENGINE_SUPPORT
\ No newline at end of file
diff --git a/Engine/Source/Developer/LowLevelTestsRunner/Public/catch_debugapi.h b/Engine/Source/Developer/LowLevelTestsRunner/Public/catch_debugapi.h
--- a/Engine/Source/Developer/LowLevelTestsRunner/Public/catch_debugapi.h
+++ b/Engine/Source/Developer/LowLevelTestsRunner/Public/catch_debugapi.h
@@ -2,13 +2,13 @@
 
 #pragma once
 
-#if UE5_ENABLE_TESTHARNESS_ENGINE_SUPPORT
+#if UE_ENABLE_TESTHARNESS_ENGINE_SUPPORT
 #include "Windows/WindowsHWrapper.h"
 #include "Windows/AllowWindowsPlatformTypes.h"
 #endif
 
 #include <debugapi.h>
 
-#if UE5_ENABLE_TESTHARNESS_ENGINE_SUPPORT
+#if UE_ENABLE_TESTHARNESS_ENGINE_SUPPORT
 #include "Windows/HideWindowsPlatformTypes.h"
 #endif
\ No newline at end of file
diff --git a/Engine/Source/Developer/LowLevelTestsRunner/Public/catch_windows.h b/Engine/Source/Developer/LowLevelTestsRunner/Public/catch_windows.h
--- a/Engine/Source/Developer/LowLevelTestsRunner/Public/catch_windows.h
+++ b/Engine/Source/Developer/LowLevelTestsRunner/Public/catch_windows.h
@@ -2,7 +2,7 @@
 
 #pragma once
 
-#if UE5_ENABLE_TESTHARNESS_ENGINE_SUPPORT
+#if UE_ENABLE_TESTHARNESS_ENGINE_SUPPORT
 #include "Windows/WindowsHWrapper.h"
 #else
 #include <windows.h>
diff --git a/Engine/Source/Programs/AutomationTool/Gauntlet/Unreal/Base/Gauntlet.UnrealTestConfiguration.cs b/Engine/Source/Programs/AutomationTool/Gauntlet/Unreal/Base/Gauntlet.UnrealTestConfiguration.cs
--- a/Engine/Source/Programs/AutomationTool/Gauntlet/Unreal/Base/Gauntlet.UnrealTestConfiguration.cs
+++ b/Engine/Source/Programs/AutomationTool/Gauntlet/Unreal/Base/Gauntlet.UnrealTestConfiguration.cs
@@ -15,6 +15,7 @@ namespace Gauntlet
 		public string Project;
 		public string GameMap;
 		private Dictionary<string, object> Params;
+		private HashSet<string> NonOptionParams;
 
 		// Give external people read-only access
 		public IReadOnlyDictionary<string, object> Arguments {  get { return Params;  } }
@@ -31,6 +32,7 @@ namespace Gauntlet
 			GameMap = string.Empty;
 			AdditionalExplicitCommandLineArgs = string.Empty;
 			Params = new Dictionary<string, object>(StringComparer.OrdinalIgnoreCase);
+			NonOptionParams = new HashSet<string>(StringComparer.OrdinalIgnoreCase);
 		}
 
 		// copy constructor
@@ -40,6 +42,7 @@ namespace Gauntlet
 			GameMap = InCopy.GameMap;
 			AdditionalExplicitCommandLineArgs = InCopy.AdditionalExplicitCommandLineArgs;
 			Params = new Dictionary<string, object>(InCopy.Params);
+			NonOptionParams = new HashSet<string>(StringComparer.OrdinalIgnoreCase);
 		}
 
 		/// <summary>
@@ -119,8 +122,13 @@ namespace Gauntlet
 		/// </summary>
 		/// <param name="ParamName"></param>
 		/// <param name="ParamVal"></param>
-		public void Add(string ParamName, object ParamVal = null)
+		public void Add(string ParamName, object ParamVal = null, bool IsNonOption = false)
 		{
+			if (IsNonOption)
+			{
+				NonOptionParams.Add(ParamName);
+			}
+
 			if (Params.ContainsKey(ParamName))
 			{
 				if (ParamName.ToLower() == "execcmds" && ParamVal != null)
@@ -233,13 +241,13 @@ namespace Gauntlet
 				string CurrentArgument;
 				if (Params[Key] != null && !string.IsNullOrWhiteSpace(Params[Key].ToString()))
 				{
-					CurrentArgument = string.Format("-{0}={1}", Key,
+					CurrentArgument = string.Format("{0}{1}={2}", NonOptionParams.Contains(Key) ? "" : "-", Key,
 						(Params[Key].ToString().Contains(' ') && !Params[Key].ToString().Contains('\"'))
 						? string.Format("\"{0}\"", Params[Key]) : Params[Key]);
 				}
 				else
 				{
-					CurrentArgument = string.Format("-{0}", Key);
+					CurrentArgument = string.Format("{0}{1}", NonOptionParams.Contains(Key) ? "" : "-", Key);
 				}
 				FinalCommandline = string.Format("{0} {1} ", FinalCommandline, CurrentArgument);
 			}
diff --git a/Engine/Source/Programs/AutomationTool/LowLevelTests/RunLowLevelTests.cs b/Engine/Source/Programs/AutomationTool/LowLevelTests/RunLowLevelTests.cs
--- a/Engine/Source/Programs/AutomationTool/LowLevelTests/RunLowLevelTests.cs
+++ b/Engine/Source/Programs/AutomationTool/LowLevelTests/RunLowLevelTests.cs
@@ -153,6 +153,8 @@ namespace LowLevelTests
 
 		public string TestApp;
 
+		public string Tags;
+
 		public string Build;
 
 		[AutoParam("")]
@@ -198,6 +200,8 @@ namespace LowLevelTests
 			Build = Params.ParseValue("build=", null);
 			TestApp = Globals.Params.ParseValue("testapp=", "");
 
+			Tags = Params.ParseValue("tags=", null);
+
 			string PlatformArgString = Params.ParseValue("platform=", null);
 			Platform = string.IsNullOrEmpty(PlatformArgString) ? BuildHostPlatform.Current.Platform : UnrealTargetPlatform.Parse(PlatformArgString);
 
@@ -217,14 +221,17 @@ namespace LowLevelTests
 	{
 		private static int QUERY_STATE_INTERVAL = 1;
 
+		public IAppInstall Install { get; protected set; }
 		public IAppInstance Instance { get; protected set; }
 		private LowLevelTestsBuildSource BuildSource { get; set; }
+		private string Tags { get; set; }
 
 		public UnrealDeviceReservation UnrealDeviceReservation { get; private set; }
 
-		public LowLevelTestsSession(LowLevelTestsBuildSource InBuildSource)
+		public LowLevelTestsSession(LowLevelTestsBuildSource InBuildSource, string InTags)
 		{
 			BuildSource = InBuildSource;
+			Tags = InTags;
 			UnrealDeviceReservation = new UnrealDeviceReservation();
 		}
 
@@ -248,13 +255,11 @@ namespace LowLevelTests
 
 			// TargetDevice<Platform> classes have a hard dependency on UnrealAppConfig instead of IAppConfig.
 			// More refactoring needed to support non-packaged applications that can be run natively from a path on the device.
-			UnrealAppConfig AppConfig = BuildSource.GetUnrealAppConfig();
+			UnrealAppConfig AppConfig = BuildSource.GetUnrealAppConfig(Tags);
 
 			IEnumerable<ITargetDevice> DevicesToInstallOn = UnrealDeviceReservation.ReservedDevices.ToArray();
 			ITargetDevice Device = DevicesToInstallOn.Where(D => D.IsConnected && D.Platform == BuildSource.Platform).First();
 
-			IAppInstall Install = null;
-
 			IDeviceUsageReporter.RecordStart(Device.Name, (UnrealTargetPlatform)Device.Platform, IDeviceUsageReporter.EventType.Device, IDeviceUsageReporter.EventState.Success);
 			IDeviceUsageReporter.RecordStart(Device.Name, (UnrealTargetPlatform)Device.Platform, IDeviceUsageReporter.EventType.Install, IDeviceUsageReporter.EventState.Success, BuildSource.BuildName);
 			try
@@ -393,6 +398,27 @@ namespace LowLevelTests
 		}
 	}
 
+	/// <summary>
+	/// Platform-specific reporting utility for defining Catch2 report path and means to copy it to a saved storage either from its local directory or from a target device.
+	/// </summary>
+	public interface ILowLevelTestsReporting
+	{
+		/// <summary>
+		/// Use this to implement platform-specific specializations of ILowLevelTestsReporting.
+		/// </summary>
+		bool CanSupportPlatform(UnrealTargetPlatform InPlatform);
+
+		/// <summary>
+		/// Specify Catch2 target report path for this platform.
+		/// </summary>
+		string GetTargetReportPath(UnrealTargetPlatform InPlatform, string InTestApp, string InBuildPath);
+
+		/// <summary>
+		/// Copy generated report  and return copied path
+		/// </summary>
+		string CopyDeviceReportTo(IAppInstall InAppInstall, UnrealTargetPlatform InPlatform, string InTestApp, string InBuildPath, string InTargetDirectory);
+	}
+
 	public interface ILowLevelTestsBuildFactory
 	{
 		bool CanSupportPlatform(UnrealTargetPlatform InPlatform);
@@ -462,11 +488,37 @@ namespace LowLevelTests
 		}
 	}
 
+	public class DesktopLowLevelTestsReporting : ILowLevelTestsReporting
+	{
+		public bool CanSupportPlatform(UnrealTargetPlatform InPlatform)
+		{
+			return InPlatform.IsInGroup(UnrealPlatformGroup.Desktop);
+		}
+
+		public string GetTargetReportPath(UnrealTargetPlatform InPlatform, string InTestApp, string InBuildPath)
+		{
+			return string.Format("{0}LLTResults.out", InPlatform.ToString());
+		}
+
+		public string CopyDeviceReportTo(IAppInstall InAppInstall, UnrealTargetPlatform InPlatform, string InTestApp, string InBuildPath, string InTargetDirectory)
+		{
+			string ReportRelativePath = GetTargetReportPath(InPlatform, InTestApp, InBuildPath);
+			string ReportPath = Path.Combine(InBuildPath, ReportRelativePath);
+			string ExpectedLocalPath = Path.Combine(InTargetDirectory, ReportRelativePath);
+			File.Copy(ReportPath, InTargetDirectory);
+			return ExpectedLocalPath;
+		}
+	}
+
 	public class LowLevelTestsBuildSource : IBuildSource
 	{
 		private string TestApp;
+		private string BuildPath;
 		private UnrealAppConfig CachedConfig = null;
 
+		private ILowLevelTestsBuildFactory LowLevelTestsBuildFactory;
+		private ILowLevelTestsReporting LowLevelTestsReporting;
+
 		public UnrealTargetPlatform Platform { get; protected set; }
 		public LowLevelTestsBuild DiscoveredBuild { get; protected set; }
 
@@ -474,12 +526,13 @@ namespace LowLevelTests
 		{
 			TestApp = InTestApp;
 			Platform = InTargetPlatform;
+			BuildPath = InBuildPath;
 			InitBuildSource(InTestApp, InBuildPath, InTargetPlatform);
 		}
 
 		protected void InitBuildSource(string InTestApp, string InBuildPath, UnrealTargetPlatform InTargetPlatform)
 		{
-			ILowLevelTestsBuildFactory LowLevelTestsBuildFactory = Gauntlet.Utils.InterfaceHelpers.FindImplementations<ILowLevelTestsBuildFactory>(true)
+			LowLevelTestsBuildFactory = Gauntlet.Utils.InterfaceHelpers.FindImplementations<ILowLevelTestsBuildFactory>(true)
 				.Where(B => B.CanSupportPlatform(InTargetPlatform))
 				.First();
 			DiscoveredBuild = LowLevelTestsBuildFactory.CreateBuild(InTargetPlatform, InTestApp, InBuildPath);
@@ -487,9 +540,13 @@ namespace LowLevelTests
 			{
 				throw new AutomationException("No builds were discovered at path {0} matching test app name {1} and target platform {2}", InBuildPath, InTestApp, InTargetPlatform);
 			}
+
+			LowLevelTestsReporting = Gauntlet.Utils.InterfaceHelpers.FindImplementations<ILowLevelTestsReporting>(true)
+				.Where(B => B.CanSupportPlatform(InTargetPlatform))
+				.First();
 		}
 
-		public UnrealAppConfig GetUnrealAppConfig()
+		public UnrealAppConfig GetUnrealAppConfig(string InTags)
 		{
 			if (CachedConfig == null)
 			{
@@ -502,6 +559,19 @@ namespace LowLevelTests
 				CachedConfig.Build = DiscoveredBuild;
 				CachedConfig.Sandbox = "LowLevelTests";
 				CachedConfig.FilesToCopy = new List<UnrealFileToCopy>();
+				// Set reporting options, filters etc
+				if (TestApp == "LowLevelTests" || TestApp == "LowLevelTestsTests")
+				{
+					CachedConfig.CommandLineParams.Add("d", "yes");
+				}
+				CachedConfig.CommandLineParams.Add("r", "console");
+				CachedConfig.CommandLineParams.Add("o", LowLevelTestsReporting.GetTargetReportPath(Platform, TestApp, BuildPath));
+				CachedConfig.CommandLineParams.Add("-#", null, true);
+				if (!string.IsNullOrEmpty(InTags))
+				{
+					CachedConfig.CommandLineParams.Add(InTags, null, true);
+				}
+				CachedConfig.CommandLineParams.Add("--debug", null, true);
 			}
 			return CachedConfig;
 		}
diff --git a/Engine/Source/Programs/AutomationTool/LowLevelTests/Tests/Gauntlet.LowLevelTestNode.cs b/Engine/Source/Programs/AutomationTool/LowLevelTests/Tests/Gauntlet.LowLevelTestNode.cs
--- a/Engine/Source/Programs/AutomationTool/LowLevelTests/Tests/Gauntlet.LowLevelTestNode.cs
+++ b/Engine/Source/Programs/AutomationTool/LowLevelTests/Tests/Gauntlet.LowLevelTestNode.cs
@@ -26,7 +26,7 @@ namespace LowLevelTests
 		{
 			Context = InContext;
 
-			MaxDuration = 60 * 5;
+			MaxDuration = 60 * 30;
 			LowLevelTestResult = TestResult.Invalid;
 		}
 
@@ -41,7 +41,7 @@ namespace LowLevelTests
 		{
 			if (LowLevelTestsApp == null)
 			{
-				LowLevelTestsApp = new LowLevelTestsSession(Context.BuildInfo);
+				LowLevelTestsApp = new LowLevelTestsSession(Context.BuildInfo, Context.Options.Tags);
 			}
 
 			return LowLevelTestsApp.TryReserveDevices();
@@ -124,14 +124,14 @@ namespace LowLevelTests
 				StdOut = TestInstance.StdOut;
 			}
 
+			string LogDir = Path.Combine(Unreal.EngineDirectory.FullName, "Programs", "AutomationTool", "Saved", "Logs");
+
 			if (StdOut == null || string.IsNullOrEmpty(StdOut.Trim()))
 			{
-				Log.Error("No StdOut returned from low level test app.");
+				Log.Warning("No StdOut returned from low level test app.");
 			}
-			else // Save artifact
-			{
-				string LogDir = Path.Combine(Unreal.EngineDirectory.FullName, "Programs", "AutomationTool", "Saved", "Logs");
-				
+			else // Save log artifact
+			{	
 				const string ClientLogFile = "ClientOutput.log";
 				string ClientOutputLog = Path.Combine(ArtifactPath, ClientLogFile);
 				
@@ -142,12 +142,35 @@ namespace LowLevelTests
 				File.Copy(ClientOutputLog, Path.Combine(LogDir, ClientLogFile));
 			}
 
-			LowLevelTestsLogParser LowLevelTestsLogParser = new LowLevelTestsLogParser(StdOut);
+			ILowLevelTestsReporting LowLevelTestsReporting = Gauntlet.Utils.InterfaceHelpers.FindImplementations<ILowLevelTestsReporting>(true)
+				.Where(B => B.CanSupportPlatform(Context.Options.Platform))
+				.First();
+
+			string ReportPath = null;
+			bool ReportCopied = false;
+			try
+			{
+				ReportPath = LowLevelTestsReporting.CopyDeviceReportTo(LowLevelTestsApp.Install, Context.Options.Platform, Context.Options.TestApp, Context.Options.Build, LogDir);
+				ReportCopied = true;
+			}
+			catch (Exception ex)
+			{
+				Log.Error("Failed to copy report: {0}", ex.ToString());
+			}
+
+			bool ReportResult = false;
+			if (ReportCopied)
+			{
+				string ReportContents = File.ReadAllText(ReportPath);
+				Console.WriteLine(ReportContents);
+				LowLevelTestsLogParser LowLevelTestsLogParser = new LowLevelTestsLogParser(ReportContents);
+				ReportResult = LowLevelTestsLogParser.GetCatchTestResults().Passed;
+			}
 
 			UnrealProcessResult Result = GetExitCodeAndReason(
 				InReason,
-				LowLevelTestsLogParser.GetSummary(),
-				LowLevelTestsLogParser.GetCatchTestResults(),
+				ReportResult,
+				StdOut,
 				out string ExitReason,
 				out int ExitCode);
 
@@ -178,8 +201,10 @@ namespace LowLevelTests
 			}
 		}
 
-		protected virtual UnrealProcessResult GetExitCodeAndReason(StopReason InReason, UnrealLog InLog, LowLevelTestsLogParser.CatchTestResults InCatchResults, out string ExitReason, out int ExitCode)
+		protected virtual UnrealProcessResult GetExitCodeAndReason(StopReason InReason, bool ReportResult, string StdOut, out string ExitReason, out int ExitCode)
 		{
+			UnrealLog UnrealLog = new UnrealLogParser(StdOut).GetSummary();
+			
 			if (TestInstance.WasKilled)
 			{
 				if (InReason == StopReason.MaxDuration)
@@ -198,29 +223,29 @@ namespace LowLevelTests
 
 			// First we check for unreal specific issues.
 			// A successful test run must be free of these.
-			if (InLog.FatalError != null)
+			if (UnrealLog.FatalError != null)
 			{
 				ExitReason = "Process encountered fatal error";
 				ExitCode = -1;
 				return UnrealProcessResult.EncounteredFatalError;
 			}
 
-			if (InLog.Ensures.Count() > 0)
+			if (UnrealLog.Ensures.Count() > 0)
 			{
-				ExitReason = string.Format("Process encountered {0} Ensures", InLog.Ensures.Count());
+				ExitReason = string.Format("Process encountered {0} Ensures", UnrealLog.Ensures.Count());
 				ExitCode = -1;
 				return UnrealProcessResult.EncounteredEnsure;
 			}
 
-			if (InLog.RequestedExit)
+			if (UnrealLog.RequestedExit)
 			{
-				ExitReason = string.Format("Exit was requested: {0}", InLog.RequestedExitReason);
+				ExitReason = string.Format("Exit was requested: {0}", UnrealLog.RequestedExitReason);
 				ExitCode = 0;
 				return UnrealProcessResult.ExitOk;
 			}
 
 			// Then we check for actual test results.
-			if (InCatchResults.Passed)
+			if (ReportResult)
 			{
 				ExitReason = "All Catch2 tests passed.";
 				ExitCode = 0;
diff --git a/Engine/Source/Programs/LowLevelTests/LowLevelTests.Target.cs b/Engine/Source/Programs/LowLevelTests/LowLevelTests.Target.cs
--- a/Engine/Source/Programs/LowLevelTests/LowLevelTests.Target.cs
+++ b/Engine/Source/Programs/LowLevelTests/LowLevelTests.Target.cs
@@ -21,7 +21,7 @@ public class LowLevelTestsTarget : TargetRules
 			bIncludeAllTestsOverride = true;
 		}
 
-		bDeployAfterCompile = false;
+		bDeployAfterCompile = Target.Platform != UnrealTargetPlatform.Android;
 		bIsBuildingConsoleApplication = true;
 
 		// Disabling default true flags that aren't necessary for tests
@@ -41,6 +41,7 @@ public class LowLevelTestsTarget : TargetRules
 		bCompileAgainstCoreUObject = false;
 		bCompileAgainstApplicationCore = false;
 		bUseLoggingInShipping = true;
+		bForceEnableExceptions = true;
 
 		bool bDebugOrDevelopment = Target.Configuration == UnrealTargetConfiguration.Debug || Target.Configuration == UnrealTargetConfiguration.Development;
 		bBuildWithEditorOnlyData = Target.Platform.IsInGroup(UnrealPlatformGroup.Desktop) && bDebugOrDevelopment;
@@ -65,6 +66,7 @@ public class LowLevelTestsTarget : TargetRules
 
 	protected virtual void SetupPreprocessorDefinitions(TargetInfo Target)
 	{
+		GlobalDefinitions.Add("UE_LOW_LEVEL_TESTS=1");
 		GlobalDefinitions.Add("STATS=0");
 
 		// Platform specific setup
diff --git a/Engine/Source/Programs/UnrealBuildTool/Configuration/TestTargetRules.cs b/Engine/Source/Programs/UnrealBuildTool/Configuration/TestTargetRules.cs
--- a/Engine/Source/Programs/UnrealBuildTool/Configuration/TestTargetRules.cs
+++ b/Engine/Source/Programs/UnrealBuildTool/Configuration/TestTargetRules.cs
@@ -1,7 +1,9 @@
 // Copyright Epic Games, Inc. All Rights Reserved.
 
 using EpicGames.Core;
+using System.Collections.Generic;
 using System.IO;
+using UnrealBuildBase;
 
 namespace UnrealBuildTool
 {
@@ -34,6 +36,8 @@ namespace UnrealBuildTool
 			TargetSourceFile = File = FileReference.Combine(UnrealBuildTool.EngineSourceDirectory, "UnrealBuildTool", "Configuration", "TestTargetRules.cs");
 			LaunchModuleName = TestedTarget.LaunchModuleName + "Tests";
 
+			ManifestFileNames = TestedTarget.ManifestFileNames;
+
 			WindowsPlatform = TestedTarget.WindowsPlatform;
 
 			Type = TargetType.Program;
@@ -41,7 +45,7 @@ namespace UnrealBuildTool
 
 			bBuildInSolutionByDefault = false;
 
-			bDeployAfterCompile = false;
+			bDeployAfterCompile = Target.Platform != UnrealTargetPlatform.Android;
 			bIsBuildingConsoleApplication = true;
 
 			// Disabling default true flags that aren't necessary for tests
@@ -62,6 +66,9 @@ namespace UnrealBuildTool
 			bCompileAgainstApplicationCore = false;
 			bUseLoggingInShipping = true;
 
+			// Allow exception handling
+			bForceEnableExceptions = true;
+
 			bool bDebugOrDevelopment = Target.Configuration == UnrealTargetConfiguration.Debug || Target.Configuration == UnrealTargetConfiguration.Development;
 			bBuildWithEditorOnlyData = Target.Platform.IsInGroup(UnrealPlatformGroup.Desktop) && bDebugOrDevelopment;
 
@@ -74,7 +81,9 @@ namespace UnrealBuildTool
 				bDebugBuildsActuallyUseDebugCRT = true;
 			}
 
+			GlobalDefinitions.Add("UE_LOW_LEVEL_TESTS=1");
 			GlobalDefinitions.Add("STATS=0");
+			GlobalDefinitions.Add("TEST_FOR_VALID_FILE_SYSTEM_MEMORY=0");
 
 			// Platform specific setup
 			if (Target.Platform == UnrealTargetPlatform.Android)
diff --git a/Engine/Source/Runtime/Cbor/Tests/CborTests.cpp b/Engine/Source/Runtime/Cbor/Tests/CborTests.cpp
--- a/Engine/Source/Runtime/Cbor/Tests/CborTests.cpp
+++ b/Engine/Source/Runtime/Cbor/Tests/CborTests.cpp
@@ -10,7 +10,7 @@
 #include "CoreGlobals.h"
 #include "TestHarness.h"
 
-TEST_CASE("Core::FCbor::Smoke Test", "[Core][Smoke]")
+TEST_CASE("CBor::FCbor::Smoke Test", "[CBor][Smoke]")
 {
 	// Run the test twices, once with little-endian encoding, once with big-endian encoding.
 	auto RunWithEndiannessFn = [](ECborEndianness Endianness)
@@ -322,7 +322,7 @@ TEST_CASE("Core::FCbor::Smoke Test", "[Core][Smoke]")
 	}
 }
 
-TEST_CASE("Core::FCbor::ByteSwap Performance Test", "[Core][Perf]")
+TEST_CASE("CBor::FCbor::ByteSwap::Performance", "[CBor][Perf]")
 {
 	//  NOTE: The most expensive value to write in CBOR is the double as it needs to swap 8 bytes all the time. Integers can be encoded in 1, 2, 4 or 8 bytes depending on their value, but not double.
 	//        String are not considered because they are written in UTF8. The test ensure this is not significantly longer to write/read while swapping the bytes.
diff --git a/Engine/Source/Runtime/Core/Private/Mac/MacPlatformProcess.cpp b/Engine/Source/Runtime/Core/Private/Mac/MacPlatformProcess.cpp
--- a/Engine/Source/Runtime/Core/Private/Mac/MacPlatformProcess.cpp
+++ b/Engine/Source/Runtime/Core/Private/Mac/MacPlatformProcess.cpp
@@ -1041,9 +1041,16 @@ const TCHAR* FMacPlatformProcess::BaseDir()
 					BasePath = [BasePath stringByDeletingLastPathComponent];
 				}
 			}
-		
+
 			FCString::Strcpy(Result, MAC_MAX_PATH, *FString(BasePath));
 			FCString::Strcat(Result, TEXT("/"));
+
+			FString CollapseResult(Result);
+#ifdef UE_RELATIVE_BASE_DIR
+			CollapseResult /= UE_RELATIVE_BASE_DIR;
+#endif
+			FPaths::CollapseRelativeDirectories(CollapseResult);
+			FCString::Strcpy(Result, *CollapseResult);		
 		}
 	}
 	return Result;
diff --git a/Engine/Source/Runtime/Core/Tests/Algorithm/UniqueTest.cpp b/Engine/Source/Runtime/Core/Tests/Algorithm/UniqueTest.cpp
new file mode 100644
--- /dev/null
+++ b/Engine/Source/Runtime/Core/Tests/Algorithm/UniqueTest.cpp
@@ -0,0 +1,49 @@
+// Copyright Epic Games, Inc. All Rights Reserved.
+
+#include "Containers/UnrealString.h"
+#include "Algo/Unique.h"
+#include "Containers/Array.h"
+#include "TestHarness.h"
+
+TEST_CASE("Core::Algorithm::Unique::Smoke Test", "[Core][Algorithm][Smoke]")
+{
+	using namespace Algo;
+
+	const TArray<int32> ExpectedArray{1, 2, 3};
+
+	{
+		TArray<int32> Array;
+		int32 RemoveFrom = Unique(Array);
+		TestEqual(TEXT("`Unique` must handle an empty container"), RemoveFrom, 0);
+	}
+	{
+		TArray<int32> Array{ 1, 2, 3 };
+		Array.SetNum(Unique(Array));
+		TestEqual(TEXT("Uniqued container with no duplicates must remain unchanged"), Array, ExpectedArray);
+	}
+	{
+		TArray<int32> Array{ 1, 1, 2, 2, 2, 3, 3, 3, 3};
+		Array.SetNum(Unique(Array));
+		TestEqual(TEXT("`Unique` with multiple duplicates must return correct result"), Array, ExpectedArray);
+	}
+	{
+		TArray<int32> Array{ 1, 1, 2, 3, 3, 3 };
+		Array.SetNum(Unique(Array));
+		TestEqual(TEXT("`Unique` with duplicates and unique items must return correct result"), Array, ExpectedArray);
+	}
+	{
+		FString Str = TEXT("aa");
+		Str = Str.Mid(0, Unique(Str));
+		TestEqual(TEXT("`Unique` on `FString` as an example of arbitrary random-access container must compile and return correct result"), Str, FString(TEXT("a")));
+	}
+	{
+		int32 Array[] = {1};
+		int32 NewSize = (int32)Unique(Array);
+		TestEqual(TEXT("`Unique` must support C arrays"), NewSize, 1);
+	}
+	{
+		TArray<int32> Array = { 1, 1 };
+		int32 NewSize = Unique(MakeArrayView(Array.GetData() + 1, 1));
+		TestEqual(TEXT("`Unique` must support ranges"), NewSize, 1);
+	}
+}
diff --git a/Engine/Source/Runtime/Core/Tests/Async/AsyncTest.cpp b/Engine/Source/Runtime/Core/Tests/Async/AsyncTest.cpp
new file mode 100644
--- /dev/null
+++ b/Engine/Source/Runtime/Core/Tests/Async/AsyncTest.cpp
@@ -0,0 +1,109 @@
+// Copyright Epic Games, Inc. All Rights Reserved.
+
+#include "CoreTypes.h"
+#include "Templates/Function.h"
+#include "Async/Async.h"
+#include "TestHarness.h"
+
+#include "TestCommon/CoreUtilities.h"
+
+/** Helper methods used in the test cases. */
+namespace AsyncTestUtils
+{
+	TFunction<int()> Task = [] {
+		return 123;
+	};
+
+	bool bHasVoidTaskFinished = false;
+
+	TFunction<void()> VoidTask = [] {
+		bHasVoidTaskFinished = true;
+	};
+}
+
+
+/** Test that task graph tasks return correctly. */
+TEST_CASE("Core::Async::TaskGraph::Task Graph", "[Core][Async][Smoke]")
+{
+	InitTaskGraphAndDependencies();
+
+	auto Future = Async(EAsyncExecution::TaskGraph, AsyncTestUtils::Task);
+	int Result = Future.Get();
+
+	TestEqual(TEXT("Task graph task must return expected value"), Result, 123);
+
+	CleanupTaskGraphAndDependencies();
+}
+
+
+/** Test that threaded tasks return correctly. */
+TEST_CASE("Core::Async::Thread::Thread", "[Core][Async][Smoke]")
+{
+	InitTaskGraphAndDependencies();
+
+	auto Future = Async(EAsyncExecution::Thread, AsyncTestUtils::Task);
+	int Result = Future.Get();
+
+	REQUIRE(Result == 123);
+
+	CleanupTaskGraphAndDependencies();
+}
+
+
+/** Test that threaded pool tasks return correctly. */
+TEST_CASE("Core::Async::ThreadedPool::Threaded Pool", "[Core][Async][Smoke]")
+{
+	InitTaskGraphAndDependencies();
+
+	auto Future = Async(EAsyncExecution::ThreadPool, AsyncTestUtils::Task);
+	int Result = Future.Get();
+
+	TestEqual(TEXT("Thread pool task must return expected value"), Result, 123);
+
+	CleanupTaskGraphAndDependencies();
+}
+
+
+/** Test that void tasks run without errors or warnings. */
+TEST_CASE("Core::Async::Task::Void Task", "[Core][Async][Smoke]")
+{
+	InitTaskGraphAndDependencies();
+
+	// Reset test variable before running
+	AsyncTestUtils::bHasVoidTaskFinished = false;
+	auto Future = Async(EAsyncExecution::TaskGraph, AsyncTestUtils::VoidTask);
+	Future.Get();
+
+	// Check that the variable state was updated by task
+	TestTrue(TEXT("Void tasks should run"), AsyncTestUtils::bHasVoidTaskFinished);
+
+	CleanupTaskGraphAndDependencies();
+}
+
+
+/** Test that asynchronous tasks have their completion callback called. */
+TEST_CASE("Core::Async::CompletionCallback::Completion Callback", "[Core][Async][Smoke]")
+{
+	InitTaskGraphAndDependencies();
+
+	bool Completed = false;
+	FEvent* CompletedEvent = FPlatformProcess::GetSynchEventFromPool(true);
+
+	auto Future = Async(EAsyncExecution::TaskGraph, AsyncTestUtils::Task, [&] {
+		Completed = true;
+		CompletedEvent->Trigger();
+	});
+
+	int Result = Future.Get();
+
+	// We need an additional synchronization point here since the future get above will return after
+	// the task is done but before the completion callback has returned!
+
+	bool CompletedEventTriggered = CompletedEvent->Wait(FTimespan(0 /* hours */, 0 /* minutes */, 5 /* seconds */));
+	FPlatformProcess::ReturnSynchEventToPool(CompletedEvent);
+
+	TestEqual(TEXT("Async Result"), Result, 123);
+	TestTrue(TEXT("Completion callback to be called"), CompletedEventTriggered && Completed);
+
+	CleanupTaskGraphAndDependencies();
+}
diff --git a/Engine/Source/Runtime/Core/Tests/Compression/CompressedBufferTest.cpp b/Engine/Source/Runtime/Core/Tests/Compression/CompressedBufferTest.cpp
--- a/Engine/Source/Runtime/Core/Tests/Compression/CompressedBufferTest.cpp
+++ b/Engine/Source/Runtime/Core/Tests/Compression/CompressedBufferTest.cpp
@@ -11,8 +11,6 @@
 
 TEST_CASE("Core::Compression::CompressedBuffer::Compress", "[Core][Compression][Smoke]")
 {
-	InitCommandLine();
-
 	const uint8 ZeroBuffer[1024]{};
 	const FIoHash ZeroBufferHash = FIoHash::HashBuffer(MakeMemoryView(ZeroBuffer));
 
@@ -78,14 +76,10 @@ TEST_CASE("Core::Compression::CompressedBuffer::Compress", "[Core][Compression][
 		CHECK(CompressionLevel == ECompressedBufferCompressionLevel::VeryFast);
 		CHECK(FMath::IsPowerOfTwo(BlockSize));
 	}
-
-	CleanupCommandLine();
 }
 
 TEST_CASE("Core::Compression::CompressedBuffer::Decompress", "[Core][Compression][Smoke]")
 {
-	InitCommandLine();
-
 	const auto GenerateData = [](int32 N) -> TArray<uint64>
 	{
 		TArray<uint64> Data;
@@ -247,6 +241,4 @@ TEST_CASE("Core::Compression::CompressedBuffer::Decompress", "[Core][Compression
 			CHECK(Uncompressed.IsNull());
 		}
 	}
-
-	CleanupCommandLine();
 }
