Implement the RunGame command for SteamDeck

Still have to get tail'ing of the log working
As well as Exiting the process done

diff --git a/Engine/Source/Programs/AutomationTool/SteamDeck/SteamDeckSupport.cs b/Engine/Source/Programs/AutomationTool/SteamDeck/SteamDeckSupport.cs
--- a/Engine/Source/Programs/AutomationTool/SteamDeck/SteamDeckSupport.cs
+++ b/Engine/Source/Programs/AutomationTool/SteamDeck/SteamDeckSupport.cs
@@ -10,10 +10,11 @@ using UnrealBuildBase;
 
 public static class SteamDeckSupport
 {
-	public static string RSyncPath = Path.Combine(Unreal.RootDirectory.FullName, "Engine\\Extras\\ThirdPartyNotUE\\cwrsync\\bin\\rsync.exe");
-	public static string SSHPath = Path.Combine(Unreal.RootDirectory.FullName, "Engine\\Extras\\ThirdPartyNotUE\\cwrsync\\bin\\ssh.exe");
+	static string RSyncPath = Path.Combine(Unreal.RootDirectory.FullName, "Engine\\Extras\\ThirdPartyNotUE\\cwrsync\\bin\\rsync.exe");
+	static string SSHPath = Path.Combine(Unreal.RootDirectory.FullName, "Engine\\Extras\\ThirdPartyNotUE\\cwrsync\\bin\\ssh.exe");
 	static string DevKitRSAPath = Path.Combine(Environment.GetFolderPath(Environment.SpecialFolder.LocalApplicationData), @"steamos-devkit\steamos-devkit\devkit_rsa");
 	static string KnownHostsPath = Path.Combine(Environment.GetFolderPath(Environment.SpecialFolder.UserProfile), @".ssh\known_hosts");
+	static string SteamDeckScripts = Path.Combine(Unreal.RootDirectory.FullName, "Engine\\Build\\SteamDeck");
 	
 	#region Devices
 
@@ -112,7 +113,7 @@ public static class SteamDeckSupport
 
 	#region Deploying
 
-	public static string GetRegisterGameScript(ProjectParams Params, UnrealTargetPlatform RuntimePlatform, string GameFolderPath, string GameRunArgs, string MsvsmonVersion)
+	private static string GetRegisterGameScript(ProjectParams Params, UnrealTargetPlatform RuntimePlatform, string GameFolderPath, string GameRunArgs, string MsvsmonVersion)
 	{
 		FileReference ExePath = Params.GetProjectExeForPlatform(RuntimePlatform);
 		string RelGameExePath = ExePath.MakeRelativeTo(DirectoryReference.Combine(ExePath.Directory, "../../..")).Replace('\\', '/');
@@ -146,6 +147,14 @@ public static class SteamDeckSupport
 		return $"#!/bin/bash\nrm ~/devkit-game/{GameId}-settings.json\npython3 ~/devkit-utils/steam-client-create-shortcut --parms '{{{JoinedParams}}}'";
 	}
 
+	private static string GetRunGameScript(ProjectParams Params, UnrealTargetPlatform RuntimePlatform)
+	{
+		string GameId = $"{Params.ShortProjectName}_{RuntimePlatform}";
+
+		// NOTE this path is setup based on our steam-client-run-game script, if this changes other places need to change
+		return $"#!/bin/bash\npython3 ~/epic-scripts/steam-client-run-game --parms '{{\"gameid\":\"{GameId}\"}}'";
+	}
+
 	// Rsync is a Cygwin utility, so convert windows paths to a path that Cygwin Rsync will understand
 	// This will not work with UNC paths
 	public static string FixupHostPath(string HostPath)
@@ -170,7 +179,6 @@ public static class SteamDeckSupport
 		return HostPath;
 	}
 
-
 	private static IProcessResult Rsync(string SourceFolder, string DestFolder, string UserName, string IpAddr, string[] ExcludeList)
 	{
 		// make a standard set of options to pass to the rsync for auth's -e option
@@ -192,9 +200,9 @@ public static class SteamDeckSupport
 
 		string[] RsyncOpts =
 		{
-			// ?
+			// archive, verbose, human readable
 			$"-avh",
-			// ?
+			// ensure certain permissions on the files copied
 			$"--chmod=Du=rwx,Dgo=rx,Fu=rwx,Fog=rx",
 			// standard authorization options
 			AuthOptions,
@@ -216,6 +224,12 @@ public static class SteamDeckSupport
 		return CommandUtils.Run(RSyncPath, string.Join(" ", RsyncOpts), "");
 	}
 
+	private static IProcessResult SSHCommand(string Command, string UserName, string IpAddr)
+	{
+		string SSHArgs = $"-i {DevKitRSAPath} {UserName}@{IpAddr}";
+		return CommandUtils.Run(SSHPath, $"{SSHArgs} \"{Command}\"", "");
+	}
+
 	/* Deploying to a steam deck currently does 2 things
 	 *
 	 * 1) Generates a script CreateShortcutHelper.sh that will register the game on the SteamDeck once uploaded
@@ -231,8 +245,6 @@ public static class SteamDeckSupport
 
 		string GameFolderPath = $"[REDACTED_USER_HOME]/devkit-game/{Params.ShortProjectName}_{RuntimePlatform}";
 		string GameRunArgs = $"{SC.ProjectArgForCommandLines} {Params.StageCommandline} {Params.RunCommandline}".Replace("\"", "\\\"");
-		// string DevkitUtilPath = Microsoft.Win32.Registry.GetValue(@"HKEY_LOCAL_MACHINE\SOFTWARE\Microsoft\Windows\CurrentVersion\Uninstall\Steam App 943760", "InstallLocation", null) as string;
-
 
 		if (!File.Exists(DevKitRSAPath))
 		{
@@ -247,24 +259,44 @@ public static class SteamDeckSupport
 			if (DeployMsvsmon(IpAddr, UserName, out MsvsmonVersion) == false)
 			{
 				return;
-			}	
+			}
 		}
 
-		// write a file into the staged directory to be copied with the rest
-		string ScriptFileName = "CreateShortcutHelper.sh";
-		string ScriptFile = Path.Combine(SC.StageDirectory.FullName, ScriptFileName);
-		File.WriteAllText(ScriptFile, GetRegisterGameScript(Params, RuntimePlatform, GameFolderPath, GameRunArgs, MsvsmonVersion));
+		// write a create shortcut and run game script into the stage dir so it will be copied with out project. Note this is required due to the json
+		// strings requiring a single quote, which our run command replaces with a double quote due to a dotnet limitation.
+		string CreateShortcutScriptFileName = "CreateShortcutHelper.sh";
+		string CreateShortcutScriptFile = Path.Combine(SC.StageDirectory.FullName, CreateShortcutScriptFileName);
+		File.WriteAllText(CreateShortcutScriptFile, GetRegisterGameScript(Params, RuntimePlatform, GameFolderPath, GameRunArgs, MsvsmonVersion));
+
+		string RunGameScriptFileName = "RunGameHelper.sh";
+		string RunGameScriptFile = Path.Combine(SC.StageDirectory.FullName, RunGameScriptFileName);
+		File.WriteAllText(RunGameScriptFile, GetRunGameScript(Params, RuntimePlatform));
 
 		// copy the staged directory, skipping over huge debug files that we don't need on the target
 		IProcessResult Result = Rsync(SC.StageDirectory.FullName, GameFolderPath, UserName, IpAddr, new string[] { /*"*.debug", */ "*.pdb", "Saved/" });
 
+		if (Result.ExitCode > 0)
+		{
+			CommandUtils.LogWarning($"Failed to rsync the {SC.StageDirectory.FullName} to {GameFolderPath}. Check connection on ip Check connection on ip {UserName}@{IpAddr}");
+			return;
+		}
+
+		// copy any scripts we have SteamDeckScripts to our custom epic script location
+		string DevkitUtilsPath = $"[REDACTED_USER_HOME]/epic-scripts/";
+		Result = Rsync(SteamDeckScripts, DevkitUtilsPath, UserName, IpAddr, new string[] {});
+
+		if (Result.ExitCode > 0)
+		{
+			CommandUtils.LogWarning($"Failed to rsync the {SteamDeckScripts} to {DevkitUtilsPath}. Check connection on ip Check connection on ip {UserName}@{IpAddr}");
+			return;
+		}
+
 		// Run the script to register the game with the Deck
-		string SSHArgs = $"-i {DevKitRSAPath} {UserName}@{IpAddr}";
-		Result = CommandUtils.Run(SSHPath, $"{SSHArgs} \"chmod +x {GameFolderPath}/{ScriptFileName} && {GameFolderPath}/{ScriptFileName}\"", "");
+		Result = SSHCommand($"chmod +x {GameFolderPath}/{CreateShortcutScriptFileName} && {GameFolderPath}/{CreateShortcutScriptFileName}", UserName, IpAddr);
 
 		if (Result.ExitCode > 0)
 		{
-			CommandUtils.LogWarning($"Failed to run the {ScriptFileName}.sh script. Check connection on ip Check connection on ip {UserName}@{IpAddr}");
+			CommandUtils.LogWarning($"Failed to run the {CreateShortcutScriptFileName}.sh script. Check connection on ip Check connection on ip {UserName}@{IpAddr}");
 			return;
 		}
 	}
@@ -302,12 +334,18 @@ public static class SteamDeckSupport
 
 	public static IProcessResult RunClient(UnrealTargetPlatform RuntimePlatform, CommandUtils.ERunOptions ClientRunFlags, string ClientApp, string ClientCmdLine, ProjectParams Params)
 	{
-		// TODO figure out how to get the steam app num id. Then can run like this:
-		// steam steam://rungameid/<GameNumId>
-		// TODO would be great if we could tail the log while running. Figure out how to cancel/exit app
+		string IpAddr, UserName;
+		if (GetDeviceInfo(RuntimePlatform, Params, out IpAddr, out UserName) == false)
+		{
+			return null;
+		}
+
+		string GameFolderPath = $"[REDACTED_USER_HOME]/devkit-game/{Params.ShortProjectName}_{RuntimePlatform}";
+		string RunGameScriptFileName = "RunGameHelper.sh";
 
-		return null;
+		// TODO would be great if we could tail the log while running. Figure out how to cancel/exit app
+		return SSHCommand($"chmod +x {GameFolderPath}/{RunGameScriptFileName} && {GameFolderPath}/{RunGameScriptFileName}", UserName, IpAddr);
 	}
 
 	#endregion
-}
\ No newline at end of file
+}
