Bring UE5/Main files in line with RES
#rnx

diff --git a/Engine/Source/Programs/AutomationTool/BuildGraph/Tasks/ModifyJsonValue.cs b/Engine/Source/Programs/AutomationTool/BuildGraph/Tasks/ModifyJsonValue.cs
--- a/Engine/Source/Programs/AutomationTool/BuildGraph/Tasks/ModifyJsonValue.cs
+++ b/Engine/Source/Programs/AutomationTool/BuildGraph/Tasks/ModifyJsonValue.cs
@@ -2,7 +2,6 @@
 
 using AutomationTool;
 using EpicGames.BuildGraph;
-using System;
 using System.Collections.Generic;
 using System.Linq;
 using System.Xml;
