Gauntlet - Rename BuildType to BuildTarget in metadata

#rnx

diff --git a/Engine/Source/Programs/AutomationTool/Gauntlet/Unreal/Base/Gauntlet.UnrealTestNode.cs b/Engine/Source/Programs/AutomationTool/Gauntlet/Unreal/Base/Gauntlet.UnrealTestNode.cs
--- a/Engine/Source/Programs/AutomationTool/Gauntlet/Unreal/Base/Gauntlet.UnrealTestNode.cs
+++ b/Engine/Source/Programs/AutomationTool/Gauntlet/Unreal/Base/Gauntlet.UnrealTestNode.cs
@@ -1367,7 +1367,7 @@ namespace Gauntlet
 			var AllRoleTypes = Roles.Select(R =>  R.Type);
 			var AllRoleContexts = Roles.Select(R => Context.GetRoleContext(R.Type));
 			Report.SetMetadata("Platform", string.Join("+", AllRoleContexts.Select(R => R.Platform.ToString()).Distinct().OrderBy(P => P)));
-			Report.SetMetadata("BuildType", string.Join("+", AllRoleTypes.Select(R => R.ToString()).Distinct().OrderBy(B => B)));
+			Report.SetMetadata("BuildTarget", string.Join("+", AllRoleTypes.Select(R => R.ToString()).Distinct().OrderBy(B => B)));
 			Report.SetMetadata("Configuration", string.Join("+", AllRoleContexts.Select(R => R.Configuration.ToString()).Distinct().OrderBy(C => C)));
 			Report.SetMetadata("Project", Context.BuildInfo.ProjectName);
 			// Additional metadata passed through the command line arguments
