Add MergeDirectory_NoExceptions which does not delete the contents of the targer directory.

diff --git a/Engine/Source/Programs/AutomationTool/AutomationUtils/CommandUtils.cs b/Engine/Source/Programs/AutomationTool/AutomationUtils/CommandUtils.cs
--- a/Engine/Source/Programs/AutomationTool/AutomationUtils/CommandUtils.cs
+++ b/Engine/Source/Programs/AutomationTool/AutomationUtils/CommandUtils.cs
@@ -1463,6 +1463,30 @@ namespace AutomationTool
 			return CopyDirectory_NoExceptions(Source, Dest, CopyDirectoryOptions.Default | (bQuiet ? CopyDirectoryOptions.SuppressLogging : CopyDirectoryOptions.None));
 		}
 
+		/// <summary>
+		/// Copies a directory and all of it's contents recursively. Merges with the destination directory. Does not throw exceptions.
+		/// </summary>
+		/// <param name="Source"></param>
+		/// <param name="Dest"></param>
+		/// <param name="bQuiet">When true, logging is suppressed.</param>
+		/// <returns>True if the operation was successful, false otherwise.</returns>
+		public static bool MergeDirectory_NoExceptions(DirectoryReference Source, DirectoryReference Dest, bool bQuiet = false)
+		{
+			return CopyDirectory_NoExceptions(Source.FullName, Dest.FullName, CopyDirectoryOptions.KeepExistingDirectories | (bQuiet ? CopyDirectoryOptions.SuppressLogging : CopyDirectoryOptions.None));
+		}
+
+		/// <summary>
+		/// Copies a directory and all of it's contents recursively. Merges with the destination directory. Does not throw exceptions.
+		/// </summary>
+		/// <param name="Source"></param>
+		/// <param name="Dest"></param>
+		/// <param name="bQuiet">When true, logging is suppressed.</param>
+		/// <returns>True if the operation was successful, false otherwise.</returns>
+		public static bool MergeDirectory_NoExceptions(string Source, string Dest, bool bQuiet = false)
+		{
+			return CopyDirectory_NoExceptions(Source, Dest, CopyDirectoryOptions.KeepExistingDirectories | (bQuiet ? CopyDirectoryOptions.SuppressLogging : CopyDirectoryOptions.None));
+		}
+
 		/// <summary>
 		/// Copies a directory and all of it's contents recursively. Does not throw exceptions.
 		/// </summary>
