GoogleDriveHelper: correct rethrow

#trivial
#rnx

diff --git a/Engine/Source/Programs/AutomationTool/AutomationUtils/GoogleDriveHelper.cs b/Engine/Source/Programs/AutomationTool/AutomationUtils/GoogleDriveHelper.cs
--- a/Engine/Source/Programs/AutomationTool/AutomationUtils/GoogleDriveHelper.cs
+++ b/Engine/Source/Programs/AutomationTool/AutomationUtils/GoogleDriveHelper.cs
@@ -371,7 +371,7 @@ namespace DriveHelper
 			{
 				if (!ex.Message.Contains("File not found:"))
 				{
-					throw ex;
+					throw;
 				}
 			}
 			return ReturnFile;
