Gauntlet: Remove call to Thread.Abort(), unsafe & unavailable in later .net versions

#rnx

diff --git a/Engine/Source/Programs/AutomationTool/Gauntlet/Framework/Gauntlet.TestExecutor.cs b/Engine/Source/Programs/AutomationTool/Gauntlet/Framework/Gauntlet.TestExecutor.cs
--- a/Engine/Source/Programs/AutomationTool/Gauntlet/Framework/Gauntlet.TestExecutor.cs
+++ b/Engine/Source/Programs/AutomationTool/Gauntlet/Framework/Gauntlet.TestExecutor.cs
@@ -390,14 +390,15 @@ namespace Gauntlet
 						}
 					}
 
-					// tick anything running, this will also check IsCancelled and stop them
-					// forcibly kill anything waiting
 					if (StartingTestThreads.Count > 0)
 					{
 						foreach (Thread T in StartingTestThreads)
 						{
-							Log.Info("Aborting startup thread");
-							T.Abort();
+							// SYSLIB0006: 'Thread.Abort()' is obsolete: 'Thread.Abord it not supported and throws PlatformNotSupportedException'
+							//Log.Info("Aborting startup thread");
+							//T.Abort();
+
+							Log.Info($"Thread is still running: {T.Name} {T.ManagedThreadId}");
 						}
 						Thread.Sleep(1000);
 					}
