Launch UnrealBuildTool via dotnet unrealbuildtool.dll, not unrealbuildtool.exe

#rnx

diff --git a/Engine/Source/Programs/AutomationTool/AutomationUtils/CommandEnvironment.cs b/Engine/Source/Programs/AutomationTool/AutomationUtils/CommandEnvironment.cs
--- a/Engine/Source/Programs/AutomationTool/AutomationUtils/CommandEnvironment.cs
+++ b/Engine/Source/Programs/AutomationTool/AutomationUtils/CommandEnvironment.cs
@@ -224,16 +224,7 @@ namespace AutomationTool
 			Log.TraceVerbose("CompilationEvironment.HasCapabilityToCompile={0}", HasCapabilityToCompile);
 			Log.TraceVerbose("CompilationEvironment.FrameworkMsbuildExe={0}", FrameworkMsbuildPath);
 
-			try
-			{
-				DotnetMsbuildPath = HostPlatform.Current.GetDotnetMsbuildExe();
-			}
-			catch (Exception Ex)
-			{
-				Log.WriteLine(LogEventType.Warning, Ex.Message);
-				Log.WriteLine(LogEventType.Warning, "Assuming no compilation capability for NET Core projects.");
-				FrameworkMsbuildPath = "";
-			}
+			DotnetMsbuildPath = Unreal.DotnetPath.FullName;
 			Log.TraceVerbose("CompilationEvironment.DotnetMsbuildExe={0}", DotnetMsbuildPath);
 		}
 
diff --git a/Engine/Source/Programs/AutomationTool/AutomationUtils/HostPlatform.cs b/Engine/Source/Programs/AutomationTool/AutomationUtils/HostPlatform.cs
--- a/Engine/Source/Programs/AutomationTool/AutomationUtils/HostPlatform.cs
+++ b/Engine/Source/Programs/AutomationTool/AutomationUtils/HostPlatform.cs
@@ -40,18 +40,6 @@ namespace AutomationTool
 		/// <returns></returns>
 		abstract public string GetFrameworkMsbuildExe();
 
-		/// <summary>
-		/// Gets the path to dotnet
-		/// </summary>
-		/// <returns></returns>
-		abstract public FileReference GetDotnetExe();
-
-		/// <summary>
-		/// Gets the build executable filename for NET Core projects. Typically, the path to the bundled dotnet executable.
-		/// </summary>
-		/// <returns></returns>
-		abstract public string GetDotnetMsbuildExe();
-
 		/// <summary>
 		/// Folder under UE/ to the platform's binaries.
 		/// </summary>
diff --git a/Engine/Source/Programs/AutomationTool/AutomationUtils/LinuxHostPlatform.cs b/Engine/Source/Programs/AutomationTool/AutomationUtils/LinuxHostPlatform.cs
--- a/Engine/Source/Programs/AutomationTool/AutomationUtils/LinuxHostPlatform.cs
+++ b/Engine/Source/Programs/AutomationTool/AutomationUtils/LinuxHostPlatform.cs
@@ -39,16 +39,6 @@ namespace AutomationTool
 			return CachedFrameworkMsbuildTool;
 		}
 
-		public override FileReference GetDotnetExe()
-		{
-			return FileReference.Combine(Unreal.EngineDirectory, @"Binaries\ThirdParty\DotNet\Linux\dotnet");
-		}
-
-		public override string GetDotnetMsbuildExe()
-		{
-			return "../../ThirdParty/DotNet/Linux/dotnet";
-		}
-
 		public override string RelativeBinariesFolder
 		{
 			get { return @"Engine/Binaries/Linux/"; }
@@ -129,7 +119,6 @@ namespace AutomationTool
 					AppName = AppName.Replace(".exe", "");
 				}
 				// some of our C# applications are converted to dotnet core, do not run those via mono
-				// they are instead assumed to produce a host executable that can just be run
 				else if (AppName.Contains("UnrealBuildTool") || AppName.Contains("AutomationTool"))
 				{
 					Options &= ~CommandUtils.ERunOptions.AppMustExist;
diff --git a/Engine/Source/Programs/AutomationTool/AutomationUtils/MacHostPlatform.cs b/Engine/Source/Programs/AutomationTool/AutomationUtils/MacHostPlatform.cs
--- a/Engine/Source/Programs/AutomationTool/AutomationUtils/MacHostPlatform.cs
+++ b/Engine/Source/Programs/AutomationTool/AutomationUtils/MacHostPlatform.cs
@@ -38,16 +38,6 @@ namespace AutomationTool
 			return CachedFrameworkMsbuildExe;
 		}
 
-		public override FileReference GetDotnetExe()
-		{
-			return FileReference.Combine(Unreal.EngineDirectory, @"Binaries\ThirdParty\DotNet\Mac\dotnet");
-		}
-
-		public override string GetDotnetMsbuildExe()
-		{
-			return "../../ThirdParty/DotNet/Mac/dotnet";
-		}
-
 		public override string RelativeBinariesFolder
 		{
 			get { return @"Engine/Binaries/Mac/"; }
@@ -147,9 +137,6 @@ namespace AutomationTool
 					}
 				}
 				// some of our C# applications are converted to dotnet core, do not run those via mono
-				// they are instead assumed to produce a host executable that can just be run
-
-				// this needs fixing: these applications should be launched through the bundled dotnet
 
 				else if (AppName.Contains("UnrealBuildTool") || AppName.Contains("AutomationTool"))
 				{
diff --git a/Engine/Source/Programs/AutomationTool/AutomationUtils/UBTUtils.cs b/Engine/Source/Programs/AutomationTool/AutomationUtils/UBTUtils.cs
--- a/Engine/Source/Programs/AutomationTool/AutomationUtils/UBTUtils.cs
+++ b/Engine/Source/Programs/AutomationTool/AutomationUtils/UBTUtils.cs
@@ -19,6 +19,7 @@ namespace AutomationTool
 		/// <param name="Env">Environment to use.</param>
 		/// <param name="CommandLine">Commandline to pass on to UBT.</param>
 		/// <param name="LogName">Optional logfile name.</param>
+		[Obsolete("Deprecated in UE5.1; function signature has changed")]
         public static void RunUBT(CommandEnvironment Env, string UBTExecutable, string CommandLine)
 		{
 			if (!FileExists(UBTExecutable))
@@ -52,6 +53,47 @@ namespace AutomationTool
 			}
 		}
 
+		/// <summary>
+		/// Runs UBT with the specified commandline. Automatically creates a logfile. When 
+		/// no LogName is specified, the executable name is used as logfile base name.
+		/// </summary>
+		/// <param name="Env">Environment to use.</param>
+		/// <param name="CommandLine">Commandline to pass on to UBT.</param>
+		/// <param name="LogName">Optional logfile name.</param>
+        public static void RunUBT(CommandEnvironment Env, FileReference UnrealBuildToolDll, string CommandLine)
+		{
+			if (!FileReference.Exists(UnrealBuildToolDll))
+			{
+				throw new AutomationException($"Unable to find UnrealBuildTool.dll: {UnrealBuildToolDll}");
+			}
+
+			string BaseLogName = String.Format("UBT-{0}", String.Join("-", SharedUtils.ParseCommandLine(CommandLine).Where(x => !x.Contains('/') && !x.Contains('\\') && !x.StartsWith("-"))));
+			string LogName;
+			for(int Attempt = 1;;Attempt++)
+			{
+				LogName = String.Format("{0}.txt", (Attempt == 1)? BaseLogName : String.Format("{0}_{1}", BaseLogName, Attempt));
+
+				FileReference LogLocation = FileReference.Combine(new DirectoryReference(Env.LogFolder), LogName);
+				if(!FileReference.Exists(LogLocation))
+				{
+					CommandLine += String.Format(" -log=\"{0}\"", LogLocation);
+					break;
+				}
+
+				if (Attempt >= 50)
+				{
+					throw new AutomationException("Unable to find name for UBT log file after {0} attempts", Attempt);
+				}
+			}
+
+			CommandLine = $"\"{UnrealBuildToolDll}\" {CommandLine}";
+			IProcessResult Result = Run(Unreal.DotnetPath.FullName, CommandLine, Options: ERunOptions.AllowSpew | ERunOptions.NoStdOutCapture);
+			if(Result.ExitCode != 0)
+			{
+				throw new AutomationException((ExitCode)Result.ExitCode, "UnrealBuildTool failed. See log for more details. ({0})", CommandUtils.CombinePaths(Env.FinalLogFolder, LogName)) { OutputFormat = AutomationExceptionOutputFormat.Minimal };
+			}
+		}
+
 		/// <summary>
 		/// Builds a UBT Commandline.
 		/// </summary>
@@ -85,11 +127,28 @@ namespace AutomationTool
 		/// <param name="Config">Configuration to build.</param>
 		/// <param name="AdditionalArgs">Additional arguments to pass on to UBT.</param>
 		/// <param name="LogName">Optional logifle name.</param>
+		[Obsolete("Deprecated in UE5.1; Use the alternate version below with the UnrealBuildTool.dll, not executable")]
 		public static void RunUBT(CommandEnvironment Env, string UBTExecutable, FileReference Project, string Target, UnrealTargetPlatform Platform, UnrealTargetConfiguration Config, string AdditionalArgs = "")
 		{
 			RunUBT(Env, UBTExecutable, UBTCommandline(Project, Target, Platform, Config, AdditionalArgs));
 		}
 
+		/// <summary>
+		/// Builds a target using UBT.  Automatically creates a logfile. When 
+		/// no LogName is specified, the executable name is used as logfile base name.
+		/// </summary>
+		/// <param name="Env">BuildEnvironment to use.</param>
+		/// <param name="Project">Unreal project to build (optional)</param>
+		/// <param name="Target">Target to build.</param>
+		/// <param name="Platform">Platform to build for.</param>
+		/// <param name="Config">Configuration to build.</param>
+		/// <param name="AdditionalArgs">Additional arguments to pass on to UBT.</param>
+		/// <param name="LogName">Optional logifle name.</param>
+		public static void RunUBT(CommandEnvironment Env, FileReference UnrealBuildToolDll, FileReference Project, string Target, UnrealTargetPlatform Platform, UnrealTargetConfiguration Config, string AdditionalArgs = "")
+		{
+			RunUBT(Env, UnrealBuildToolDll, UBTCommandline(Project, Target, Platform, Config, AdditionalArgs));
+		}
+
 	}
 
 }
\ No newline at end of file
diff --git a/Engine/Source/Programs/AutomationTool/AutomationUtils/UnrealBuild.cs b/Engine/Source/Programs/AutomationTool/AutomationUtils/UnrealBuild.cs
--- a/Engine/Source/Programs/AutomationTool/AutomationUtils/UnrealBuild.cs
+++ b/Engine/Source/Programs/AutomationTool/AutomationUtils/UnrealBuild.cs
@@ -82,10 +82,15 @@ namespace AutomationTool
 		}
 		
 		private void PrepareUBT()
-		{			
-			if (CommandUtils.FileExists(UBTExecutable) == false)
+		{
+			if (!FileReference.Exists(UnrealBuildToolDll))
+			{
+				throw new UnrealBuildException($"UnrealBuildTool.dll does not exist at {UnrealBuildToolDll}");
+			}
+
+			if (!FileReference.Exists(Unreal.DotnetPath))
 			{
-				throw new UnrealBuildException("UBT does not exist in {0}.", UBTExecutable);
+				throw new UnrealBuildException($"dotnet executable does not exist at {Unreal.DotnetPath}");
 			}
 		}
 
@@ -120,7 +125,7 @@ namespace AutomationTool
 
 			string UHTArg = this.AlwaysBuildUHT ? "" : "-nobuilduht";
 
-			CommandUtils.RunUBT(CommandUtils.CmdEnv, UBTExecutable: UBTExecutable, Project: UprojectPath, Target: TargetName, Platform: Platform, Config: Config, AdditionalArgs: String.Format("-Manifest={0} {1} -NoHotReload -xgeexport {2}", CommandUtils.MakePathSafeToUseWithCommandLine(ManifestFile.FullName), UHTArg, AddArgs));
+			CommandUtils.RunUBT(CommandUtils.CmdEnv, UnrealBuildToolDll: UnrealBuildToolDll, Project: UprojectPath, Target: TargetName, Platform: Platform, Config: Config, AdditionalArgs: String.Format("-Manifest={0} {1} -NoHotReload -xgeexport {2}", CommandUtils.MakePathSafeToUseWithCommandLine(ManifestFile.FullName), UHTArg, AddArgs));
 
 			XGEItem Result = new XGEItem();
 			Result.Platform = Platform;
@@ -163,7 +168,7 @@ namespace AutomationTool
 			// run the deployment steps, if necessary
 			foreach(string DeployTargetFile in Item.Manifest.DeployTargetFiles)
 			{
-				CommandUtils.RunUBT(CommandUtils.CmdEnv, UBTExecutable, String.Format("-Mode=Deploy -Receipt=\"{0}\"", DeployTargetFile));
+				CommandUtils.RunUBT(CommandUtils.CmdEnv, UnrealBuildToolDll, String.Format("-Mode=Deploy -Receipt=\"{0}\"", DeployTargetFile));
 			}
 
 			foreach (string ManifestItem in Item.Manifest.BuildProducts)
@@ -196,7 +201,7 @@ namespace AutomationTool
 				Scope.Span.SetTag("target", TargetName);
 				Scope.Span.SetTag("platform", Platform.ToString());
 				Scope.Span.SetTag("config", Config.ToString());
-				CommandUtils.RunUBT(CommandUtils.CmdEnv, UBTExecutable: UBTExecutable, Project: UprojectPath, Target: TargetName, Platform: Platform, Config: Config, AdditionalArgs: "-Clean -NoHotReload" + AddArgs);
+				CommandUtils.RunUBT(CommandUtils.CmdEnv, UnrealBuildToolDll: UnrealBuildToolDll, Project: UprojectPath, Target: TargetName, Platform: Platform, Config: Config, AdditionalArgs: "-Clean -NoHotReload" + AddArgs);
 			}
         }
 
@@ -223,7 +228,7 @@ namespace AutomationTool
 			FileReference ManifestFile = GetManifestFile(UprojectPath);
 			CommandUtils.DeleteFile(ManifestFile);
 
-			CommandUtils.RunUBT(CommandUtils.CmdEnv, UBTExecutable: UBTExecutable, Project: UprojectPath, Target: TargetName, Platform: TargetPlatform, Config: Config, AdditionalArgs: String.Format("{0} -Manifest={1} -NoHotReload", AddArgs, CommandUtils.MakePathSafeToUseWithCommandLine(ManifestFile.FullName)));
+			CommandUtils.RunUBT(CommandUtils.CmdEnv, UnrealBuildToolDll: UnrealBuildToolDll, Project: UprojectPath, Target: TargetName, Platform: TargetPlatform, Config: Config, AdditionalArgs: String.Format("{0} -Manifest={1} -NoHotReload", AddArgs, CommandUtils.MakePathSafeToUseWithCommandLine(ManifestFile.FullName)));
 
 			BuildManifest Manifest = AddBuildProductsFromManifest(ManifestFile);
 			CommandUtils.DeleteFile(ManifestFile);
@@ -1111,6 +1116,8 @@ namespace AutomationTool
 			string Args = null;
 			if (XGETool != null) 
 			{
+				//Debug.Assert(OperatingSystem.IsWindows());
+
 				Args = "\"" + TaskFilePath + "\" /Rebuild /NoLogo /ShowAgent /ShowTime";
 				if (ParseParam("StopOnErrors"))
 				{
@@ -1225,7 +1232,7 @@ namespace AutomationTool
 		/// </summary>
 		public void AddUBTFilesToBuildProducts()
 		{
-			string UBTLocation = Path.GetDirectoryName(GetUBTExecutable());
+			string UBTLocation = UnrealBuildToolDll.Directory.FullName;
 			// copy all the files from the UBT output directory
 			string[] UBTFiles = CommandUtils.FindFiles_NoExceptions("*.*", true, UBTLocation);
 
@@ -1265,11 +1272,13 @@ namespace AutomationTool
 			}
 		}
 
+		[Obsolete("Deprecated in UE5.1; use UnrealBuildToolDll")]
 		public static string GetUBTExecutable()
 		{
 			return CommandUtils.CombinePaths(CommandUtils.CmdEnv.LocalRoot, @"Engine/Binaries/DotNET/UnrealBuildTool/UnrealBuildTool" + (RuntimePlatform.IsWindows ? ".exe" : ""));
 		}
 
+		[Obsolete("Deprecated in UE5.1; use UnrealBuildToolDll")]
 		public string UBTExecutable
 		{
 			get
@@ -1278,6 +1287,9 @@ namespace AutomationTool
 			}
 		}
 
+		public static FileReference UnrealBuildToolDll => 
+			FileReference.FromString(CommandUtils.CombinePaths(CommandUtils.CmdEnv.LocalRoot, "Engine/Binaries/DotNET/UnrealBuildTool/UnrealBuildTool.dll"));
+
 		// List of everything we built so far
 		public readonly HashSet<string> BuildProductFiles = new HashSet<string>(StringComparer.InvariantCultureIgnoreCase);
 	}
diff --git a/Engine/Source/Programs/AutomationTool/AutomationUtils/WindowsHostPlatform.cs b/Engine/Source/Programs/AutomationTool/AutomationUtils/WindowsHostPlatform.cs
--- a/Engine/Source/Programs/AutomationTool/AutomationUtils/WindowsHostPlatform.cs
+++ b/Engine/Source/Programs/AutomationTool/AutomationUtils/WindowsHostPlatform.cs
@@ -9,26 +9,18 @@ using System.Diagnostics;
 using UnrealBuildTool;
 using EpicGames.Core;
 using UnrealBuildBase;
+using System.Runtime.Versioning;
 
 namespace AutomationTool
 {
 	class WindowsHostPlatform : HostPlatform
 	{
+		//[SupportedOSPlatform("windows")]
 		public override string GetFrameworkMsbuildExe()
 		{
 			return WindowsExports.GetMSBuildToolPath();
 		}
 
-		public override FileReference GetDotnetExe()
-		{
-			return FileReference.Combine(Unreal.EngineDirectory, @"Binaries\ThirdParty\DotNet\Windows\dotnet.exe");
-		}
-
-		public override string GetDotnetMsbuildExe()
-		{
-			return @"..\..\ThirdParty\DotNet\Windows\dotnet.exe";
-		}
-
 		public override string RelativeBinariesFolder
 		{
 			get { return @"Engine/Binaries/Win64/"; }
diff --git a/Engine/Source/Programs/AutomationTool/BuildGraph/Scripts/InstalledBuild.cs b/Engine/Source/Programs/AutomationTool/BuildGraph/Scripts/InstalledBuild.cs
--- a/Engine/Source/Programs/AutomationTool/BuildGraph/Scripts/InstalledBuild.cs
+++ b/Engine/Source/Programs/AutomationTool/BuildGraph/Scripts/InstalledBuild.cs
@@ -300,13 +300,9 @@ namespace AutomationTool
 
 		static async Task<BgFileSet> CompileRulesAssemblies(BgContext State)
 		{
-			FileReference UbtPath = FileReference.Combine(RootDir, "Engine/Binaries/DotNET/UnrealBuildTool/UnrealBuildTool");
-			if (HostPlatform.Current.HostEditorPlatform == UnrealTargetPlatform.Win64)
-			{
-				UbtPath = UbtPath.ChangeExtension(".exe");
-			}
+			FileReference UnrealBuildToolDll = FileReference.Combine(RootDir, "Engine/Binaries/DotNET/UnrealBuildTool/UnrealBuildTool.dll");
 
-			await SpawnAsync(UbtPath.FullName, Arguments: "-Mode=QueryTargets");
+			await SpawnAsync(Unreal.DotnetPath.FullName, Arguments: $"\"{UnrealBuildToolDll}\" -Mode=QueryTargets");
 
 			return Workspace.Filter("Engine/Intermediate/Build/BuildRules/...");
 		}
diff --git a/Engine/Source/Programs/AutomationTool/BuildGraph/Tasks/DotNetTask.cs b/Engine/Source/Programs/AutomationTool/BuildGraph/Tasks/DotNetTask.cs
--- a/Engine/Source/Programs/AutomationTool/BuildGraph/Tasks/DotNetTask.cs
+++ b/Engine/Source/Programs/AutomationTool/BuildGraph/Tasks/DotNetTask.cs
@@ -10,6 +10,7 @@ using System.IO;
 using System.Text;
 using System.Threading.Tasks;
 using System.Xml;
+using UnrealBuildBase;
 
 namespace AutomationTool.Tasks
 {
@@ -83,7 +84,7 @@ namespace AutomationTool.Tasks
 		/// <param name="TagNameToFileSet">Mapping from tag names to the set of files they include</param>
 		public override async Task ExecuteAsync(JobContext Job, HashSet<FileReference> BuildProducts, Dictionary<string, HashSet<FileReference>> TagNameToFileSet)
 		{
-			FileReference DotNetFile = Parameters.DotNetPath == null ? HostPlatform.Current.GetDotnetExe() : Parameters.DotNetPath;
+			FileReference DotNetFile = Parameters.DotNetPath == null ? Unreal.DotnetPath : Parameters.DotNetPath;
 			if(!FileReference.Exists(DotNetFile))
 			{
 				throw new AutomationException("DotNet is missing from {0}", DotNetFile);
diff --git a/Engine/Source/Programs/AutomationTool/CompileScriptModules.cs b/Engine/Source/Programs/AutomationTool/CompileScriptModules.cs
--- a/Engine/Source/Programs/AutomationTool/CompileScriptModules.cs
+++ b/Engine/Source/Programs/AutomationTool/CompileScriptModules.cs
@@ -588,7 +588,8 @@ namespace AutomationToolDriver
 		{
 			Dictionary<FileReference, (CsProjBuildRecord BuildRecord, FileReference _)> BuildRecords = new Dictionary<FileReference, (CsProjBuildRecord, FileReference)>();
 			
-			Dictionary<string, Project> Projects = new Dictionary<string, Project>();
+			var Projects = new Dictionary<string, Project>();
+			var	SkippedProjects = new HashSet<string>();
 			
 			// Microsoft.Build.Evaluation.Project provides access to information stored in the .csproj xml that is 
 			// not available when using Microsoft.Build.Execution.ProjectInstance (used later in this function and
@@ -600,7 +601,7 @@ namespace AutomationToolDriver
 				void LoadProjectAndReferences(string ProjectPath, string ReferencedBy)
 				{
 					ProjectPath = Path.GetFullPath(ProjectPath);
-					if (!Projects.ContainsKey(ProjectPath))
+					if (!Projects.ContainsKey(ProjectPath) && !SkippedProjects.Contains(ProjectPath))
 					{
 						Project Project;
 
@@ -631,6 +632,20 @@ namespace AutomationToolDriver
 							throw IPFEx;
 						}
 
+/*
+						if (!OperatingSystem.IsWindows())
+						{
+							// check the TargetFramework of the project: we can't build Windows-only projects on 
+							// non-Windows platforms.
+							if (Project.GetProperty("TargetFramework").EvaluatedValue.Contains("windows"))
+							{
+								SkippedProjects.Add(ProjectPath);
+								Log.TraceInformation($"Skipping windows-only project {ProjectPath}");
+								return;
+							}
+						}
+*/
+						
 						Projects.Add(ProjectPath, Project);
 						ReferencedBy = String.IsNullOrEmpty(ReferencedBy) ? ProjectPath : $"{ProjectPath}{Environment.NewLine}{ReferencedBy}";
 						foreach (string ReferencedProject in Project.GetItems("ProjectReference").
@@ -788,7 +803,10 @@ namespace AutomationToolDriver
 			// Potential optimization: Contructing the ProjectGraph here gives the full graph of dependencies - which is nice,
 			// but not strictly necessary, and slower than doing it some other way.
 			ProjectGraph InputProjectGraph;
-			InputProjectGraph = new ProjectGraph(FoundAutomationProjects.Select(P => P.FullName), GlobalProperties);
+			InputProjectGraph = new ProjectGraph(FoundAutomationProjects
+				// Build the graph without anything that can't be built on this platform
+				.Where(x => !SkippedProjects.Contains(x.FullName))
+				.Select(P => P.FullName), GlobalProperties);
 
 			// A ProjectGraph that will represent the set of projects that we actually want to build
 			ProjectGraph BuildProjectGraph = null;
diff --git a/Engine/Source/Programs/AutomationTool/CookedEditor/CookedEditor.Automation.cs b/Engine/Source/Programs/AutomationTool/CookedEditor/CookedEditor.Automation.cs
--- a/Engine/Source/Programs/AutomationTool/CookedEditor/CookedEditor.Automation.cs
+++ b/Engine/Source/Programs/AutomationTool/CookedEditor/CookedEditor.Automation.cs
@@ -574,7 +574,7 @@ public class MakeCookedEditor : BuildCommand
 		if (Context.IniPlatformName == "Linux")
 		{
 			// linux needs dotnet runtime
-			SC.StageFiles(StagedFileType.NonUFS, DirectoryReference.Combine(Context.EngineDirectory, "Binaries", "ThirdParty", "DotNet", Context.IniPlatformName), StageFilesSearch.AllDirectories);
+			SC.StageFiles(StagedFileType.NonUFS, Unreal.DotnetDirectory, StageFilesSearch.AllDirectories);
 		}
 
 		// not sure if we need this or not now
diff --git a/Engine/Source/Programs/AutomationTool/IOS/IOSPlatform.Automation.cs b/Engine/Source/Programs/AutomationTool/IOS/IOSPlatform.Automation.cs
--- a/Engine/Source/Programs/AutomationTool/IOS/IOSPlatform.Automation.cs
+++ b/Engine/Source/Programs/AutomationTool/IOS/IOSPlatform.Automation.cs
@@ -1062,10 +1062,8 @@ public class IOSPlatform : Platform
 			// If engine is installed then UBT doesn't need to be built
 			if (Unreal.IsEngineInstalled())
 			{
-				// Get the path to UBT
-				string InstalledUBT = UnrealBuild.GetUBTExecutable();
 				Arguments = "-XcodeProjectFiles " + Arguments;
-				RunUBT(CmdEnv, InstalledUBT, Arguments);
+				RunUBT(CmdEnv, UnrealBuild.UnrealBuildToolDll, Arguments);
 			}
 			else
 			{
@@ -1467,8 +1465,7 @@ public class IOSPlatform : Platform
 		if (bXCArchive && !RuntimePlatform.IsWindows)
 		{
 			// Always put the archive in the current user's Library/Developer/Xcode/Archives path if not on the build machine
-			WindowsIdentity id = WindowsIdentity.GetCurrent();
-			string ArchivePath = "[REDACTED_USER_HOME] + id.Name + "/Library/Developer/Xcode/Archives";
+			string ArchivePath = "[REDACTED_USER_HOME] + Environment.UserName + "/Library/Developer/Xcode/Archives";
 			if (IsBuildMachine)
 			{
 				ArchivePath = Params.ArchiveDirectoryParam;
diff --git a/Engine/Source/Programs/AutomationTool/Scripts/BuildPluginCommand.Automation.cs b/Engine/Source/Programs/AutomationTool/Scripts/BuildPluginCommand.Automation.cs
--- a/Engine/Source/Programs/AutomationTool/Scripts/BuildPluginCommand.Automation.cs
+++ b/Engine/Source/Programs/AutomationTool/Scripts/BuildPluginCommand.Automation.cs
@@ -24,14 +24,14 @@ public sealed class BuildPlugin : BuildCommand
 {
 	const string AndroidArchitectures = "armv7+arm64";
 	const string HoloLensArchitecture = "arm64+x64";
-	string UBTExeRelativePath = @"Engine/Binaries/DotNET/UnrealBuildTool/UnrealBuildTool" + (RuntimePlatform.IsWindows ? ".exe" : "");
-	string UBTExe;
+	string UnrealBuildToolDllRelativePath = @"Engine/Binaries/DotNET/UnrealBuildTool/UnrealBuildTool.dll";
+	FileReference UnrealBuildToolDll;
 
 	public override void ExecuteBuild()
 	{
 		// See if an engine dir was specified, fall back on default if not.
 		DirectoryReference EngineDirParam = ParseOptionalDirectoryReferenceParam("EngineDir");
-		UBTExe = EngineDirParam == null ? UnrealBuild.GetUBTExecutable() : CommandUtils.CombinePaths(EngineDirParam.ToString(), UBTExeRelativePath);
+		UnrealBuildToolDll = EngineDirParam == null ? UnrealBuild.UnrealBuildToolDll : FileReference.FromString(CommandUtils.CombinePaths(EngineDirParam.ToString(), UnrealBuildToolDllRelativePath));
 
 		// Get the plugin filename
 		string PluginParam = ParseParamValue("Plugin");
@@ -114,7 +114,7 @@ public sealed class BuildPlugin : BuildCommand
 		// Compile the plugin for all the target platforms
 		List<UnrealTargetPlatform> HostPlatforms = GetHostPlatforms(this);
 		List<UnrealTargetPlatform> TargetPlatforms = GetTargetPlatforms(this, BuildHostPlatform.Current.Platform);
-		FileReference[] BuildProducts = CompilePlugin(UBTExe, HostProjectFile, HostProjectPluginFile, Plugin, HostPlatforms, TargetPlatforms, AdditionalArgs.ToString());
+		FileReference[] BuildProducts = CompilePlugin(UnrealBuildToolDll, HostProjectFile, HostProjectPluginFile, Plugin, HostPlatforms, TargetPlatforms, AdditionalArgs.ToString());
 
 		// Package up the final plugin data
 		PackagePlugin(HostProjectPluginFile, BuildProducts, PackageDir, ParseParam("unversioned"), TargetPlatforms);
@@ -143,6 +143,7 @@ public sealed class BuildPlugin : BuildCommand
 		return FileReference.Combine(HostProjectPluginDir, PluginFile.GetFileName());
 	}
 
+	[Obsolete("Deprecated in UE5.1; function signature changed")]
 	public static FileReference[] CompilePlugin(string UBTExe, FileReference HostProjectFile, FileReference HostProjectPluginFile, PluginDescriptor Plugin, List<UnrealTargetPlatform> HostPlatforms, List<UnrealTargetPlatform> TargetPlatforms, string AdditionalArgs = "")
 	{
 		List<FileReference> ManifestFileNames = new List<FileReference>();
@@ -184,6 +185,48 @@ public sealed class BuildPlugin : BuildCommand
 		return BuildProducts.ToArray();
 	}
 
+	public static FileReference[] CompilePlugin(FileReference UnrealBuildToolDll, FileReference HostProjectFile, FileReference HostProjectPluginFile, PluginDescriptor Plugin, List<UnrealTargetPlatform> HostPlatforms, List<UnrealTargetPlatform> TargetPlatforms, string AdditionalArgs = "")
+	{
+		List<FileReference> ManifestFileNames = new List<FileReference>();
+
+		// Build the host platforms
+		if(HostPlatforms.Count > 0)
+		{
+			CommandUtils.LogInformation("Building plugin for host platforms: {0}", String.Join(", ", HostPlatforms));
+			foreach (UnrealTargetPlatform HostPlatform in HostPlatforms)
+			{
+				if (Plugin.SupportedPrograms != null && Plugin.SupportedPrograms.Contains("UnrealHeaderTool"))
+				{
+					CompilePluginWithUBT(UnrealBuildToolDll, HostProjectFile, HostProjectPluginFile, Plugin, "UnrealHeaderTool", TargetType.Program, HostPlatform, UnrealTargetConfiguration.Development, ManifestFileNames, String.Format("{0} -plugin={1}", AdditionalArgs, CommandUtils.MakePathSafeToUseWithCommandLine(HostProjectPluginFile.FullName)));
+				}
+				CompilePluginWithUBT(UnrealBuildToolDll, HostProjectFile, HostProjectPluginFile, Plugin, "UnrealEditor", TargetType.Editor, HostPlatform, UnrealTargetConfiguration.Development, ManifestFileNames, AdditionalArgs);
+			}
+		}
+
+		// Add the supported game targets
+		if(TargetPlatforms.Count > 0)
+		{
+			List<UnrealTargetPlatform> SupportedTargetPlatforms = TargetPlatforms.FindAll(Plugin.SupportsTargetPlatform);
+			CommandUtils.LogInformation("Building plugin for target platforms: {0}", String.Join(", ", SupportedTargetPlatforms));
+			foreach (UnrealTargetPlatform TargetPlatform in SupportedTargetPlatforms)
+			{
+				string AdditionalTargetArgs = AdditionalArgs;
+				CompilePluginWithUBT(UnrealBuildToolDll, HostProjectFile, HostProjectPluginFile, Plugin, "UnrealGame", TargetType.Game, TargetPlatform, UnrealTargetConfiguration.Development, ManifestFileNames, AdditionalTargetArgs);
+				CompilePluginWithUBT(UnrealBuildToolDll, HostProjectFile, HostProjectPluginFile, Plugin, "UnrealGame", TargetType.Game, TargetPlatform, UnrealTargetConfiguration.Shipping, ManifestFileNames, AdditionalTargetArgs);
+			}
+		}
+
+		// Package the plugin to the output folder
+		HashSet<FileReference> BuildProducts = new HashSet<FileReference>();
+		foreach(FileReference ManifestFileName in ManifestFileNames)
+		{
+			BuildManifest Manifest = CommandUtils.ReadManifest(ManifestFileName);
+			BuildProducts.UnionWith(Manifest.BuildProducts.Select(x => new FileReference(x)));
+		}
+		return BuildProducts.ToArray();
+	}
+
+	[Obsolete("Deprecated in UE5.1; function signature has changed")]
 	static void CompilePluginWithUBT(string UBTExe, FileReference HostProjectFile, FileReference HostProjectPluginFile, PluginDescriptor Plugin, string TargetName, TargetType TargetType, UnrealTargetPlatform Platform, UnrealTargetConfiguration Configuration, List<FileReference> ManifestFileNames, string InAdditionalArgs)
 	{
 		// Find a list of modules that need to be built for this plugin
@@ -244,6 +287,66 @@ public sealed class BuildPlugin : BuildCommand
 		}
 	}
 
+	static void CompilePluginWithUBT(FileReference UnrealBuildToolDll, FileReference HostProjectFile, FileReference HostProjectPluginFile, PluginDescriptor Plugin, string TargetName, TargetType TargetType, UnrealTargetPlatform Platform, UnrealTargetConfiguration Configuration, List<FileReference> ManifestFileNames, string InAdditionalArgs)
+	{
+		// Find a list of modules that need to be built for this plugin
+		bool bCompilePlatform = false;
+		if (Plugin.Modules != null)
+		{
+			bool bBuildDeveloperTools = (TargetType == TargetType.Editor || TargetType == TargetType.Program || (Configuration != UnrealTargetConfiguration.Test && Configuration != UnrealTargetConfiguration.Shipping));
+			bool bBuildRequiresCookedData = (TargetType != TargetType.Editor && TargetType != TargetType.Program);
+
+			foreach (ModuleDescriptor Module in Plugin.Modules)
+			{
+				if (Module.IsCompiledInConfiguration(Platform, Configuration, TargetName, TargetType, bBuildDeveloperTools, bBuildRequiresCookedData))
+				{
+					bCompilePlatform = true;
+				}
+			}
+		}
+
+		// Add these modules to the build agenda
+		if(bCompilePlatform)
+		{
+			if (Platform == UnrealTargetPlatform.HoloLens)
+			{
+				// Make sure to save the manifests for each architecture with unique names so they don't get overwritten.
+				// This fixes packaging issues when building from binary engine releases, where the build produces a manifest for the plugin for ARM64, which
+				// then gets overwritten by the manifest for x64. Then during packaging, the plugin is referencing a manifest for the wrong architecture.
+				foreach (string Arch in HoloLensArchitecture.Split('+'))
+				{
+					FileReference ManifestFileName = FileReference.Combine(HostProjectFile.Directory, "Saved", String.Format("Manifest-{0}-{1}-{2}-{3}.xml", TargetName, Platform, Configuration, Arch));
+					ManifestFileNames.Add(ManifestFileName);
+					string Arguments = String.Format("-plugin={0} -iwyu -noubtmakefiles -manifest={1} -nohotreload", CommandUtils.MakePathSafeToUseWithCommandLine(HostProjectPluginFile.FullName), CommandUtils.MakePathSafeToUseWithCommandLine(ManifestFileName.FullName));
+					Arguments += String.Format(" -Architecture={0}", Arch);
+					if (!String.IsNullOrEmpty(InAdditionalArgs))
+					{
+						Arguments += InAdditionalArgs;
+					}
+					CommandUtils.RunUBT(CmdEnv, UnrealBuildToolDll, HostProjectFile, TargetName, Platform, Configuration, Arguments);
+				}
+			}
+			else
+			{
+				FileReference ManifestFileName = FileReference.Combine(HostProjectFile.Directory, "Saved", String.Format("Manifest-{0}-{1}-{2}.xml", TargetName, Platform, Configuration));
+				ManifestFileNames.Add(ManifestFileName);
+				
+				string Arguments = String.Format("-plugin={0} -iwyu -noubtmakefiles -manifest={1} -nohotreload", CommandUtils.MakePathSafeToUseWithCommandLine(HostProjectPluginFile.FullName), CommandUtils.MakePathSafeToUseWithCommandLine(ManifestFileName.FullName));
+				if (Platform == UnrealTargetPlatform.Android)
+				{
+					Arguments += String.Format(" -architectures={0}", AndroidArchitectures);
+				}
+
+				if (!String.IsNullOrEmpty(InAdditionalArgs))
+				{
+					Arguments += InAdditionalArgs;
+				}
+
+				CommandUtils.RunUBT(CmdEnv, UnrealBuildToolDll, HostProjectFile, TargetName, Platform, Configuration, Arguments);
+			}
+		}
+	}
+
 	public static void PackagePlugin(FileReference SourcePluginFile, IEnumerable<FileReference> BuildProducts, DirectoryReference TargetDir, bool bUnversioned, IEnumerable<UnrealTargetPlatform> TargetPlatforms)
 	{
 		DirectoryReference SourcePluginDir = SourcePluginFile.Directory;
diff --git a/Engine/Source/Programs/AutomationTool/Scripts/BuildTarget.Automation.cs b/Engine/Source/Programs/AutomationTool/Scripts/BuildTarget.Automation.cs
--- a/Engine/Source/Programs/AutomationTool/Scripts/BuildTarget.Automation.cs
+++ b/Engine/Source/Programs/AutomationTool/Scripts/BuildTarget.Automation.cs
@@ -193,8 +193,6 @@ namespace AutomationTool
 			}
 			else
 			{
-				// Get the path to UBT
-				FileReference InstalledUBT = FileReference.Combine(Unreal.EngineDirectory, "Binaries", "DotNET", "UnrealBuildTool", "UnrealBuildTool.exe");
 				UnrealTargetPlatform PlatformToBuild = PlatformList.First();
 				UnrealTargetConfiguration ConfigToBuild = ConfigurationList.First();
 				string TargetToBuild = BuildTargets.First().TargetName;
@@ -206,9 +204,9 @@ namespace AutomationTool
 
 					if (Clean)
 					{
-						CommandUtils.RunUBT(CommandUtils.CmdEnv, InstalledUBT.FullName, CommandLine + " -clean");
+						CommandUtils.RunUBT(CommandUtils.CmdEnv, Unreal.UnrealBuildToolDllPath, CommandLine + " -clean");
 					}
-					CommandUtils.RunUBT(CommandUtils.CmdEnv, InstalledUBT.FullName, CommandLine);
+					CommandUtils.RunUBT(CommandUtils.CmdEnv, Unreal.UnrealBuildToolDllPath, CommandLine);
 				}
 				else
 				{ 
diff --git a/Engine/Source/Programs/AutomationTool/Scripts/ListThirdPartySoftware.Automation.cs b/Engine/Source/Programs/AutomationTool/Scripts/ListThirdPartySoftware.Automation.cs
--- a/Engine/Source/Programs/AutomationTool/Scripts/ListThirdPartySoftware.Automation.cs
+++ b/Engine/Source/Programs/AutomationTool/Scripts/ListThirdPartySoftware.Automation.cs
@@ -33,7 +33,7 @@ class ListThirdPartySoftware : BuildCommand
 
 			IProcessResult Result;
 
-			Result = Run(UnrealBuild.GetUBTExecutable(), String.Format("{0} {1} -Mode=JsonExport -OutputFile=\"{2}\"", Target.Replace('|', ' '),  ProjectPath, OutputFile.FullName), Options: ERunOptions.Default);
+			Result = Run(Unreal.DotnetPath.FullName, $"\"{UnrealBuild.UnrealBuildToolDll}\" {Target.Replace('|', ' ')} {ProjectPath} -Mode=JsonExport -OutputFile=\"{OutputFile.FullName}\"", Options: ERunOptions.Default);
 
 			if (Result.ExitCode != 0)
 			{
diff --git a/Engine/Source/Programs/Shared/EpicGames.Build/Unreal.cs b/Engine/Source/Programs/Shared/EpicGames.Build/Unreal.cs
--- a/Engine/Source/Programs/Shared/EpicGames.Build/Unreal.cs
+++ b/Engine/Source/Programs/Shared/EpicGames.Build/Unreal.cs
@@ -6,6 +6,7 @@ using System.Collections.Generic;
 using System.IO;
 using System.Linq;
 using System.Reflection;
+using System.Runtime.InteropServices;
 using System.Text;
 
 namespace UnrealBuildBase
@@ -63,44 +64,55 @@ namespace UnrealBuildBase
 			return FoundRootDirectory;
 		}
 
-		private static FileReference FindUnrealBuildTool()
+		private static FileReference FindUnrealBuildToolDll()
 		{
-			// todo: use UnrealBuildTool.dll (same on all platforms). Will require changes wherever UnrealBuildTool is invoked.
-			string UBTName = "UnrealBuildTool" + RuntimePlatform.ExeExtension;
+			// UnrealBuildTool.dll is assumed to be located under {RootDirectory}/Engine/Binaries/DotNET/UnrealBuildTool/
+			FileReference UnrealBuildToolDllPath = FileReference.Combine(EngineDirectory, "Binaries", "DotNET", "UnrealBuildTool", "UnrealBuildTool.dll");
 
-			// the UnrealBuildTool executable is assumed to be located under {RootDirectory}/Engine/Binaries/DotNET/UnrealBuildTool/
-			FileReference UnrealBuildToolPath = FileReference.Combine(EngineDirectory, "Binaries", "DotNET", "UnrealBuildTool", UBTName);
+			UnrealBuildToolDllPath = FileReference.FindCorrectCase(UnrealBuildToolDllPath);
 
-			UnrealBuildToolPath = FileReference.FindCorrectCase(UnrealBuildToolPath);
-
-			if (!FileReference.Exists(UnrealBuildToolPath))
-			{
-				// if there's no exe found, use the dll as a fallback
-				UnrealBuildToolPath = FileReference.FindCorrectCase(UnrealBuildToolPath.ChangeExtension(".dll"));
-			}
-
-			if (!FileReference.Exists(UnrealBuildToolPath))
+			if (!FileReference.Exists(UnrealBuildToolDllPath))
 			{
-				throw new Exception($"Unable to find UnrealBuildTool in the expected location at {UnrealBuildToolPath.FullName}");
+				throw new Exception($"Unable to find UnrealBuildTool.dll in the expected location at {UnrealBuildToolDllPath.FullName}");
 			}
 
-			return UnrealBuildToolPath;
+			return UnrealBuildToolDllPath;
 		}
 
-		static private DirectoryReference FindDotnetDirectory()
+		//static private string DotnetVersionDirectory = "6.0.200";
+			
+		static private string FindRelativeDotnetDirectory()
 		{
 			string HostDotNetDirectoryName;
 			switch (RuntimePlatform.Current)
 			{
 				case RuntimePlatform.Type.Windows: HostDotNetDirectoryName = "Windows"; break;
-				case RuntimePlatform.Type.Mac:     HostDotNetDirectoryName = "Mac"; break;
+				case RuntimePlatform.Type.Mac: HostDotNetDirectoryName = "Mac"; break;
+					/*
+				{
+					HostDotNetDirectoryName = "mac-x64";
+					if (RuntimeInformation.ProcessArchitecture == Architecture.Arm64)
+					{
+						HostDotNetDirectoryName = "mac-arm64";
+					}
+					break;
+				}
+					*/
 				case RuntimePlatform.Type.Linux:   HostDotNetDirectoryName = "Linux"; break;
 				default: throw new Exception("Unknown host platform");
 			}
 
-			return DirectoryReference.Combine(EngineDirectory, "Binaries", "ThirdParty", "DotNet", HostDotNetDirectoryName);
+			return Path.Combine("Binaries", "ThirdParty", "DotNet", /*DotnetVersionDirectory,*/ HostDotNetDirectoryName);
 		}
 
+		/// <summary>
+		/// Relative path to the dotnet executable from EngineDir
+		/// </summary>
+		/// <returns></returns>
+		static public string RelativeDotnetDirectory = FindRelativeDotnetDirectory();
+
+		static private DirectoryReference FindDotnetDirectory() => DirectoryReference.Combine(EngineDirectory, RelativeDotnetDirectory);
+
 		/// <summary>
 		/// The full name of the root UE directory
 		/// </summary>
@@ -114,12 +126,13 @@ namespace UnrealBuildBase
 		/// <summary>
 		/// The path to UBT
 		/// </summary>
-		public static readonly FileReference UnrealBuildToolPath = FindUnrealBuildTool();
+		[Obsolete("Deprecated in UE5.1; to launch UnrealBuildTool, use this dll as the first argument with DonetPath")]
+		public static readonly FileReference UnrealBuildToolPath = FindUnrealBuildToolDll().ChangeExtension(RuntimePlatform.ExeExtension);
 
 		/// <summary>
 		/// The path to UBT
 		/// </summary>
-		public static readonly FileReference UnrealBuildToolDllPath = FindUnrealBuildTool().ChangeExtension(".dll");
+		public static readonly FileReference UnrealBuildToolDllPath = FindUnrealBuildToolDll();
 
 		/// <summary>
 		/// The directory containing the bundled .NET installation
diff --git a/Engine/Source/Programs/UnrealBuildTool/Platform/Windows/PVSToolChain.cs b/Engine/Source/Programs/UnrealBuildTool/Platform/Windows/PVSToolChain.cs
--- a/Engine/Source/Programs/UnrealBuildTool/Platform/Windows/PVSToolChain.cs
+++ b/Engine/Source/Programs/UnrealBuildTool/Platform/Windows/PVSToolChain.cs
@@ -137,7 +137,7 @@ namespace UnrealBuildTool
 					XmlSerializer Serializer = new XmlSerializer(typeof(PVSApplicationSettings));
 					using (FileStream Stream = new FileStream(SettingsPath.FullName, FileMode.Open, FileAccess.Read, FileShare.Read))
 					{
-						return (PVSApplicationSettings)Serializer.Deserialize(Stream);
+						return (PVSApplicationSettings?)Serializer.Deserialize(Stream);
 					}
 				}
 				catch (Exception Ex)
@@ -635,8 +635,8 @@ namespace UnrealBuildTool
 			Action AnalyzeAction = Makefile.CreateAction(ActionType.Compile);
 			AnalyzeAction.ActionType = ActionType.PostBuildStep;
 			AnalyzeAction.CommandDescription = "Process PVS-Studio Results";
-			AnalyzeAction.CommandPath = UnrealBuildTool.GetUBTPath();
-			AnalyzeAction.CommandArguments = String.Format("-Mode=PVSGather -Input=\"{0}\" -Output=\"{1}\" ", InputFileListItem.Location, OutputFile);
+			AnalyzeAction.CommandPath = Unreal.DotnetPath;
+			AnalyzeAction.CommandArguments = $"\"{Unreal.UnrealBuildToolDllPath}\" -Mode=PVSGather -Input=\"{InputFileListItem.Location}\" -Output=\"{OutputFile}\" ";
 			AnalyzeAction.WorkingDirectory = UnrealBuildTool.EngineSourceDirectory;
 			AnalyzeAction.PrerequisiteItems.Add(InputFileListItem);
 			AnalyzeAction.PrerequisiteItems.AddRange(Makefile.OutputItems);
diff --git a/Engine/Source/Programs/UnrealBuildTool/ProjectFiles/CodeLite/CodeLiteProject.cs b/Engine/Source/Programs/UnrealBuildTool/ProjectFiles/CodeLite/CodeLiteProject.cs
--- a/Engine/Source/Programs/UnrealBuildTool/ProjectFiles/CodeLite/CodeLiteProject.cs
+++ b/Engine/Source/Programs/UnrealBuildTool/ProjectFiles/CodeLite/CodeLiteProject.cs
@@ -193,8 +193,8 @@ namespace UnrealBuildTool
 							{
 								//
 								// Look the the following folder
-								XAttribute attribute = element.Attribute("Name");
-								if (attribute.Value == FolderName)
+								XAttribute? attribute = element.Attribute("Name");
+								if (attribute?.Value == FolderName)
 								{
 									// Ok, we found the folder as subfolder, let's use it.
 									root = element;
@@ -397,7 +397,7 @@ namespace UnrealBuildTool
 					// Add the working directory for the custom build commands.
 					//
 					XElement CustomBuildWorkingDirectory = new XElement("WorkingDirectory");
-					XText CustuomBuildWorkingDirectory = new XText(UnrealBuildTool.GetUBTPath().Directory.FullName);
+					XText CustuomBuildWorkingDirectory = new XText(Unreal.UnrealBuildToolDllPath.Directory.FullName);
 					CustomBuildWorkingDirectory.Add(CustuomBuildWorkingDirectory);
 					CodeLiteConfigurationCustomBuild.Add(CustomBuildWorkingDirectory);
 
@@ -426,7 +426,7 @@ namespace UnrealBuildTool
 					}
 					else
 					{
-						BuildTarget = UnrealBuildTool.GetUBTPath().GetFileName() + " " + BuildTarget;
+						BuildTarget = $"{Unreal.DotnetPath} \"{Unreal.UnrealBuildToolDllPath}\" {BuildTarget}";
 					}
 
 					if (GameProjectFile.Length > 0)
diff --git a/Engine/Source/Programs/UnrealBuildTool/ProjectFiles/KDevelop/KDevelopGenerator.cs b/Engine/Source/Programs/UnrealBuildTool/ProjectFiles/KDevelop/KDevelopGenerator.cs
--- a/Engine/Source/Programs/UnrealBuildTool/ProjectFiles/KDevelop/KDevelopGenerator.cs
+++ b/Engine/Source/Programs/UnrealBuildTool/ProjectFiles/KDevelop/KDevelopGenerator.cs
@@ -70,15 +70,13 @@ namespace UnrealBuildTool
 		private void WriteCommandSubSection(ref StringBuilder FileContent, string TargetName, string ConfName, int BuildConfigIndex, int Type)
 		{
 			string ToolType = "";
-			string Executable = "";
+			string Executable = "bash";
 			string ProjectCmdArg = "";
-			string BuildCommand = "";
+			string BuildCommand = "Engine/Build/BatchFiles/Linux/Build.sh";
 
 			if (TargetName == GameProjectName)
 			{
 				ProjectCmdArg = " -project=\"" + OnlyGameProject!.FullName + "\"";
-				Executable = "Engine/Build/BatchFiles/Linux/RunMono.sh";
-				BuildCommand = "Engine/Binaries/DotNET/UnrealBuildTool.exe";
 
 				if (Type == 1)
 				{
@@ -89,8 +87,6 @@ namespace UnrealBuildTool
 			else if (TargetName == (GameProjectName + "Editor"))
 			{
 				ProjectCmdArg = " -editorrecompile -project=\"" + OnlyGameProject!.FullName + "\"";
-				Executable = "Engine/Build/BatchFiles/Linux/RunMono.sh";
-				BuildCommand = "Engine/Binaries/DotNET/UnrealBuildTool.exe";
 
 				if (Type == 1)
 				{
@@ -100,9 +96,6 @@ namespace UnrealBuildTool
 			}
 			else
 			{
-				Executable = "bash";
-				BuildCommand = "Engine/Build/BatchFiles/Linux/Build.sh";
-
 				if (Type == 1)
 				{
 					// Override BuildCommand and ProjectCmdArg
@@ -259,7 +252,7 @@ namespace UnrealBuildTool
 
 			int IncludeIndex = 1;
 			// Iterate through all the include paths that
-			// UnrealBuildTool.exe generates
+			// UnrealBuildTool generates
 
 			foreach (ProjectFile CurProject in GeneratedProjectFiles)
 			{
diff --git a/Engine/Source/Programs/UnrealBuildTool/ProjectFiles/Make/MakefileGenerator.cs b/Engine/Source/Programs/UnrealBuildTool/ProjectFiles/Make/MakefileGenerator.cs
--- a/Engine/Source/Programs/UnrealBuildTool/ProjectFiles/Make/MakefileGenerator.cs
+++ b/Engine/Source/Programs/UnrealBuildTool/ProjectFiles/Make/MakefileGenerator.cs
@@ -68,7 +68,7 @@ namespace UnrealBuildTool
 			{
 				GameProjectFile = OnlyGameProject!.FullName;
 				MakeGameProjectFile = "GAMEPROJECTFILE =" + GameProjectFile + "\n";
-				ProjectBuildCommand = "PROJECTBUILD = bash \"$(UNREALROOTPATH)/Engine/Build/BatchFiles/Linux/RunMono.sh\" \"$(UNREALROOTPATH)/Engine/Binaries/DotNET/UnrealBuildTool.exe\"\n";
+				ProjectBuildCommand = $"PROJECTBUILD = bash \"$(UNREALROOTPATH)/Engine/{Unreal.RelativeDotnetDirectory}/dotnet\" \"$(UNREALROOTPATH)/Engine/Binaries/DotNET/UnrealBuildTool.dll\"\n";
 			}
 
 			BuildCommand = "BUILD = bash \"$(UNREALROOTPATH)/Engine/Build/BatchFiles/Linux/Build.sh\"\n";
diff --git a/Engine/Source/Programs/UnrealBuildTool/ProjectFiles/QMake/QMakefileGenerator.cs b/Engine/Source/Programs/UnrealBuildTool/ProjectFiles/QMake/QMakefileGenerator.cs
--- a/Engine/Source/Programs/UnrealBuildTool/ProjectFiles/QMake/QMakefileGenerator.cs
+++ b/Engine/Source/Programs/UnrealBuildTool/ProjectFiles/QMake/QMakefileGenerator.cs
@@ -141,7 +141,7 @@ namespace UnrealBuildTool
 			SystemIncludeDirectories = SystemIncludeDirectories.Distinct().ToList();
 
 			// Iterate through all the defines for the projects that are generated by 
-			// UnrealBuildTool.exe
+			// UnrealBuildTool
 			// !RAKE: move to seperate function
 			QMakeDefinesPriFileContent.Append("DEFINES += \\\n");
 			foreach (ProjectFile CurProject in GeneratedProjectFiles)
@@ -187,7 +187,7 @@ namespace UnrealBuildTool
 			}
 
 			// Iterate through all the include paths that
-			// UnrealBuildTool.exe generates
+			// UnrealBuildTool generates
 			// !RAKE: Move to seperate function
 			QMakeIncludesPriFileContent.Append("INCLUDEPATH += \\\n");
 			foreach (string CurPath in IncludeDirectories)
@@ -210,7 +210,7 @@ namespace UnrealBuildTool
 				GameProjectPath = OnlyGameProject!.Directory.FullName;
 				GameProjectFile = OnlyGameProject.FullName;
 				QMakeGameProjectFile = "gameProjectFile=" + GameProjectFile + "\n";
-				BuildCommand = "build=bash $$unrealRootPath/Engine/Build/BatchFiles/Linux/RunMono.sh $$unrealRootPath/Engine/Binaries/DotNET/UnrealBuildTool.exe\n\n";
+				BuildCommand = $"build=bash $$unrealRootPath/Engine/{Unreal.RelativeDotnetDirectory}/dotnet $$unrealRootPath/Engine/Binaries/DotNET/UnrealBuildTool.exe\n\n";
 			}
 			else
 			{
diff --git a/Engine/Source/Programs/UnrealBuildTool/System/DynamicCompilation.cs b/Engine/Source/Programs/UnrealBuildTool/System/DynamicCompilation.cs
--- a/Engine/Source/Programs/UnrealBuildTool/System/DynamicCompilation.cs
+++ b/Engine/Source/Programs/UnrealBuildTool/System/DynamicCompilation.cs
@@ -55,10 +55,10 @@ namespace UnrealBuildTool
 			// recently than the dynamically-compiled assembly, then we'll always recompile it.  This is
 			// because Unreal Build Tool's code may have changed in such a way that invalidate these
 			// previously-compiled assembly files.
-			FileItem ExecutableItem = FileItem.GetItemByFileReference(UnrealBuildTool.GetUBTPath());
-			if (ExecutableItem.LastWriteTimeUtc > OutputAssemblyInfo.LastWriteTimeUtc)
+			FileItem UnrealBuildToolDllItem = FileItem.GetItemByFileReference(Unreal.UnrealBuildToolDllPath);
+			if (UnrealBuildToolDllItem.LastWriteTimeUtc > OutputAssemblyInfo.LastWriteTimeUtc)
 			{
-				Log.TraceLog("Compiling {0}: {1} is newer", OutputAssemblyPath, ExecutableItem.Name);
+				Log.TraceLog("Compiling {0}: {1} is newer", OutputAssemblyPath, UnrealBuildToolDllItem.Name);
 				return true;
 			}
 
diff --git a/Engine/Source/Programs/UnrealBuildTool/UnrealBuildTool.cs b/Engine/Source/Programs/UnrealBuildTool/UnrealBuildTool.cs
--- a/Engine/Source/Programs/UnrealBuildTool/UnrealBuildTool.cs
+++ b/Engine/Source/Programs/UnrealBuildTool/UnrealBuildTool.cs
@@ -223,6 +223,7 @@ namespace UnrealBuildTool
 		/// Gets the absolute path to the UBT assembly.
 		/// </summary>
 		/// <returns>A string containing the path to the UBT assembly.</returns>
+		[Obsolete("Deprecated in UE5.1 - use UnrealBuildTool.DotnetPath Unreal.UnrealBuildToolDllPath")]
 		static public FileReference GetUBTPath()
 		{
 			return Unreal.UnrealBuildToolPath;
@@ -533,7 +534,7 @@ namespace UnrealBuildTool
 				{
 					using (GlobalTracer.Instance.BuildSpan("XmlConfig.ReadConfigFiles()").StartActive())
 					{
-						string XmlConfigMutexName = SingleInstanceMutex.GetUniqueMutexForPath("UnrealBuildTool_Mutex_XmlConfig", Assembly.GetExecutingAssembly().CodeBase!);
+						string XmlConfigMutexName = SingleInstanceMutex.GetUniqueMutexForPath("UnrealBuildTool_Mutex_XmlConfig", Assembly.GetExecutingAssembly().Location);
 						using(SingleInstanceMutex XmlConfigMutex = new SingleInstanceMutex(XmlConfigMutexName, true))
 						{
 							FileReference? XmlConfigCache = Arguments.GetFileReferenceOrDefault("-XmlConfigCache=", null);
@@ -558,7 +559,7 @@ namespace UnrealBuildTool
 				{
 					using (GlobalTracer.Instance.BuildSpan("SingleInstanceMutex.Acquire()").StartActive())
 					{
-						string MutexName = SingleInstanceMutex.GetUniqueMutexForPath("UnrealBuildTool_Mutex", Assembly.GetExecutingAssembly().CodeBase!);
+						string MutexName = SingleInstanceMutex.GetUniqueMutexForPath("UnrealBuildTool_Mutex", Assembly.GetExecutingAssembly().Location);
 						Mutex = new SingleInstanceMutex(MutexName, Options.bWaitMutex);
 					}
 				}
