UAT: Ctrl-c exit fixes.
In ProcessManager.KillAll() make a local copy of global ActiveProcesses and to avoid calling any external code while iterating the internal global list.
In ProcessResult.StopProcess() make sure to set the ExitCode member from the killed process, else the information is lost forever since the member Proc has already been nulled out.
In CommandUtils.Run() check Result.ProcessObject for null for the case when WaitForExit() returns after StopProcess() has been called from another thread (ctrl-c handler).

#rnx

diff --git a/Engine/Source/Programs/AutomationTool/AutomationUtils/ProcessUtils.cs b/Engine/Source/Programs/AutomationTool/AutomationUtils/ProcessUtils.cs
--- a/Engine/Source/Programs/AutomationTool/AutomationUtils/ProcessUtils.cs
+++ b/Engine/Source/Programs/AutomationTool/AutomationUtils/ProcessUtils.cs
@@ -109,22 +109,22 @@ namespace AutomationTool
 		/// </summary>
 		public static void KillAll()
 		{
-			List<IProcess> ProcessesToKill = new List<IProcess>();
+			List<IProcess> ProcessesToKill = null;
 			lock (SyncObject)
 			{
-				foreach (var ProcResult in ActiveProcesses)
-				{
-					if (!ProcResult.HasExited)
-					{
-						ProcessesToKill.Add(ProcResult);
-					}
-				}
+				ProcessesToKill = new List<IProcess>(ActiveProcesses);
 				ActiveProcesses.Clear();
 			}
-			// Remove processes that can't be killed
+
+			// Remove processes that have exited or can't be killed
 			for (int ProcessIndex = ProcessesToKill.Count - 1; ProcessIndex >= 0; --ProcessIndex )
 			{
-				var ProcessName = ProcessesToKill[ProcessIndex].GetProcessName();
+				IProcess Process =  ProcessesToKill[ProcessIndex];
+				var ProcessName = Process.GetProcessName();
+				if (Process.HasExited)
+				{
+					ProcessesToKill.RemoveAt(ProcessIndex);
+				}
 				if (!String.IsNullOrEmpty(ProcessName) && !CanBeKilled(ProcessName))
 				{
 					CommandUtils.LogLog("Ignoring process \"{0}\" because it can't be killed.", ProcessName);
@@ -665,6 +665,7 @@ namespace AutomationTool
 					}
 					else
 					{
+						ExitCode = ProcToKill.ExitCode;
 						CommandUtils.LogLog("Process {0} successfully exited.", ProcToKillName);
 						OnProcessExited();
 					}
@@ -900,7 +901,11 @@ namespace AutomationTool
 			{
 				var BuildDuration = (DateTime.UtcNow - StartTime).TotalMilliseconds;
 				//AddRunTime(App, (int)(BuildDuration));
-				Result.ExitCode = Result.ProcessObject.ExitCode;
+				Process Proc = Result.ProcessObject;
+				if (Proc != null)
+				{
+					Result.ExitCode = Proc.ExitCode;
+				}
 				if (!Options.HasFlag(ERunOptions.NoLoggingOfRunCommand) || Options.HasFlag(ERunOptions.LoggingOfRunDuration))
 				{
 					LogWithVerbosity(SpewVerbosity, "Took {0}s to run {1}, ExitCode={2}", BuildDuration / 1000, Path.GetFileName(App), Result.ExitCode);
