Using System.Text.Json for Device reservation instead of fastJSON and making the deserialization not case sensitive

#rnx

diff --git a/Engine/Source/Programs/AutomationTool/AutomationUtils/DeviceReservation.cs b/Engine/Source/Programs/AutomationTool/AutomationUtils/DeviceReservation.cs
--- a/Engine/Source/Programs/AutomationTool/AutomationUtils/DeviceReservation.cs
+++ b/Engine/Source/Programs/AutomationTool/AutomationUtils/DeviceReservation.cs
@@ -4,11 +4,11 @@ using System.Collections.Generic;
 using System.IO;
 using System.Linq;
 using System.Threading;
-using System.Xml.Serialization;
 using System.Net;
-using UnrealBuildTool;
 using System.Text;
 using System.Text.RegularExpressions;
+using System.Text.Json;
+using System.Text.Json.Serialization;
 
 #pragma warning disable SYSLIB0014
 
@@ -173,7 +173,8 @@ namespace AutomationTool.DeviceReservation
 
 		public static T InvokeAPI<T>(Uri UriToRequest, string Method, object ObjectToSerialize)
 		{
-			fastJSON.JSON.Instance.RegisterCustomType(typeof(TimeSpan), (o) => o.ToString(), (s) => TimeSpan.Parse(s));
+			var SerializeOptions = new JsonSerializerOptions() { PropertyNameCaseInsensitive = true };
+			SerializeOptions.Converters.Add(new TimeSpanJsonConverter());
 
 			var Request = (HttpWebRequest)WebRequest.Create(UriToRequest);
 			Request.UseDefaultCredentials = true;
@@ -185,7 +186,7 @@ namespace AutomationTool.DeviceReservation
 
 				using (var RequestStream = Request.GetRequestStream())
 				{
-					var JsonString = fastJSON.JSON.Instance.ToJSON(ObjectToSerialize, new fastJSON.JSONParameters() { UseExtensions = false });
+					var JsonString = JsonSerializer.Serialize(ObjectToSerialize, SerializeOptions);
 					var Writer = new StreamWriter(RequestStream);
 					Writer.Write(JsonString);
 					Writer.Flush();
@@ -199,7 +200,7 @@ namespace AutomationTool.DeviceReservation
 				MemoryStream MemoryStream = new MemoryStream();
 				ResponseStream.CopyTo(MemoryStream);
 				string JsonString = Encoding.UTF8.GetString(MemoryStream.ToArray());
-				return fastJSON.JSON.Instance.ToObject<T>(JsonString);
+				return JsonSerializer.Deserialize<T>(JsonString, SerializeOptions );
 			}
 		}
 
@@ -215,7 +216,7 @@ namespace AutomationTool.DeviceReservation
 
 				using (var RequestStream = Request.GetRequestStream())
 				{
-					var JsonString = fastJSON.JSON.Instance.ToJSON(ObjectToSerialize, new fastJSON.JSONParameters() { UseExtensions = false });
+					var JsonString = JsonSerializer.Serialize(ObjectToSerialize, new JsonSerializerOptions());
 					var Writer = new StreamWriter(RequestStream);
 					Writer.Write(JsonString);
 					Writer.Flush();
@@ -246,13 +247,13 @@ namespace AutomationTool.DeviceReservation
 
 		private sealed class CreateReservationData
 		{
-			public string[] DeviceTypes;
-			public string Hostname;
-			public TimeSpan Duration;
-			public string ReservationDetails;
-			public string PoolId;
-			public string JobId;
-			public string StepId;
+			public string[] DeviceTypes { get; set; }
+			public string Hostname { get; set; }
+			public TimeSpan Duration { get; set; }
+			public string ReservationDetails { get; set; }
+			public string PoolId { get; set; }
+			public string JobId { get; set; }
+			public string StepId { get; set; }
 		}
 
 		public static Reservation Create(Uri BaseUri, string[] DeviceTypes, TimeSpan Duration, int RetryMax = 5, string PoolID = "")
@@ -418,4 +419,23 @@ namespace AutomationTool.DeviceReservation
 			return Utils.InvokeAPI<Device>(BaseUri.AppendPath("api/v1/devices/" + DeviceName), "GET", null);
 		}
 	}
+
+	public class TimeSpanJsonConverter : JsonConverter<TimeSpan>
+	{
+		public override bool CanConvert(Type ObjectType)
+		{
+			return ObjectType == typeof(TimeSpan);
+		}
+
+		public override TimeSpan Read(ref Utf8JsonReader Reader, Type TypeToConvert, JsonSerializerOptions Options)
+		{
+			return TimeSpan.Parse(Reader.GetString());
+		}
+
+		public override void Write(Utf8JsonWriter Writer, TimeSpan Value, JsonSerializerOptions Options)
+		{
+			var StrValue = Value.ToString();
+			Writer.WriteStringValue(Options.PropertyNamingPolicy?.ConvertName(StrValue) ?? StrValue);
+		}
+	}
 }
diff --git a/Engine/Source/Programs/AutomationTool/Gauntlet/Framework/Devices/Gauntlet.DevicePool.cs b/Engine/Source/Programs/AutomationTool/Gauntlet/Framework/Devices/Gauntlet.DevicePool.cs
--- a/Engine/Source/Programs/AutomationTool/Gauntlet/Framework/Devices/Gauntlet.DevicePool.cs
+++ b/Engine/Source/Programs/AutomationTool/Gauntlet/Framework/Devices/Gauntlet.DevicePool.cs
@@ -6,10 +6,10 @@ using System.IO;
 using AutomationTool;
 using AutomationTool.DeviceReservation;
 using UnrealBuildTool;
-using System.Threading;
 using System.Text.RegularExpressions;
 using System.Linq;
-using Newtonsoft.Json;
+using System.Text.Json;
+using System.Text.Json.Serialization;
 
 namespace Gauntlet
 {
@@ -30,27 +30,27 @@ namespace Gauntlet
 	/// </summary>
 	public class DeviceDefinition
 	{
-		public string Name;
+		public string Name { get; set; }
 
-		public string Address;
+		public string Address { get; set; }
 
-		public string DeviceData;
+		public string DeviceData { get; set; }
 
 		// legacy - remove!
 		[JsonConverter(typeof(UnrealTargetPlatformConvertor))]
-		public UnrealTargetPlatform Type;
+		public UnrealTargetPlatform Type { get; set; }
 
 		[JsonConverter(typeof(UnrealTargetPlatformConvertor))]
-		public UnrealTargetPlatform? Platform;
+		public UnrealTargetPlatform Platform { get; set; }
+
+		public EPerfSpec PerfSpec { get; set; }
+
+		public string Model { get; set; } = string.Empty;
 
-		public EPerfSpec PerfSpec;
+		public string Available { get; set; }
 
-		public string Model = string.Empty;
+		public bool RemoveOnShutdown { get; set; }
 
-		public string Available;
-		
-		public bool RemoveOnShutdown;
-		
 		public override string ToString()
 		{
 			return string.Format("{0} @ {1}. Platform={2} Model={3}", Name, Address, Platform, string.IsNullOrEmpty(Model) ? "Unspecified" : Model);
@@ -589,7 +589,7 @@ namespace Gauntlet
 					{
 						DeviceDefinition Def = new DeviceDefinition();
 						Def.Name = string.Format("Virtual{0}{1}", DevicePlatform.ToString(), i);
-						Def.Platform = DevicePlatform;
+						Def.Platform = DevicePlatform??BuildHostPlatform.Current.Platform;
 						UnprovisionedDevices.Add(Def);
 					}
 				}
@@ -687,12 +687,14 @@ namespace Gauntlet
 					Def.DeviceData = Device.DeviceData;
 					Def.Model = string.Empty;
 
-					if (!String.IsNullOrEmpty(Device.PerfSpec) && !Enum.TryParse<EPerfSpec>(Device.PerfSpec, true, out Def.PerfSpec))
+					EPerfSpec Out = EPerfSpec.Unspecified;
+					if (!String.IsNullOrEmpty(Device.PerfSpec) && !Enum.TryParse<EPerfSpec>(Device.PerfSpec, true, out Out))
 					{
 						throw new AutomationException("Unable to convert perfspec '{0}' into an EPerfSpec", Device.PerfSpec);
 					}
+					Def.PerfSpec = Out;
 
-					ITargetDevice TargetDevice = CreateAndRegisterDeviceFromDefinition(Def);
+				    ITargetDevice TargetDevice = CreateAndRegisterDeviceFromDefinition(Def);
 
 					// If a device from service can't be added, fail reservation and cleanup devices
 					// @todo: device problem reporting, requesting additional devices
@@ -755,7 +757,10 @@ namespace Gauntlet
 				if (PossibleFileName && File.Exists(InputReference))
 				{
 					Gauntlet.Log.Info("Adding devices from {0}", InputReference);
-					List<DeviceDefinition> DeviceDefinitions = JsonConvert.DeserializeObject<List<DeviceDefinition>>(File.ReadAllText(InputReference));
+					List<DeviceDefinition> DeviceDefinitions = JsonSerializer.Deserialize<List < DeviceDefinition >>(
+						File.ReadAllText(InputReference),
+						new JsonSerializerOptions { PropertyNameCaseInsensitive = true }
+					);
 
 					foreach (DeviceDefinition Def in DeviceDefinitions)
 					{
@@ -939,7 +944,7 @@ namespace Gauntlet
 
 			try
 			{
-				bool IsDesktop = Def.Platform != null && UnrealBuildTool.Utils.GetPlatformsInClass(UnrealPlatformClass.Desktop).Contains(Def.Platform.Value);
+				bool IsDesktop = Def.Platform != null && UnrealBuildTool.Utils.GetPlatformsInClass(UnrealPlatformClass.Desktop).Contains(Def.Platform);
 
 				string ClientTempDir = GetCleanCachePath(Def);
 
diff --git a/Engine/Source/Programs/AutomationTool/Gauntlet/Unreal/Utils/Gauntlet.UnrealHelpers.cs b/Engine/Source/Programs/AutomationTool/Gauntlet/Unreal/Utils/Gauntlet.UnrealHelpers.cs
--- a/Engine/Source/Programs/AutomationTool/Gauntlet/Unreal/Utils/Gauntlet.UnrealHelpers.cs
+++ b/Engine/Source/Programs/AutomationTool/Gauntlet/Unreal/Utils/Gauntlet.UnrealHelpers.cs
@@ -3,14 +3,14 @@
 using System;
 using System.Collections.Generic;
 using System.IO;
-using System.Threading;
 using AutomationTool;
 using UnrealBuildTool;
 using System.Linq;
 using System.Net.NetworkInformation;
 using System.Text.RegularExpressions;
 using System.Net;
-using Newtonsoft.Json;
+using System.Text.Json;
+using System.Text.Json.Serialization;
 
 namespace Gauntlet
 {
@@ -451,29 +451,33 @@ namespace Gauntlet
 	/// <summary>
 	///  Converts between json and UnrealTargetPlatform
 	/// </summary>
-	public class UnrealTargetPlatformConvertor : JsonConverter
+	public class UnrealTargetPlatformConvertor : JsonConverter<UnrealTargetPlatform>
 	{
 		public override bool CanConvert(Type ObjectType)
 		{
-			return ObjectType == typeof(string) || ObjectType == typeof(UnrealTargetPlatform);
+			return ObjectType == typeof(UnrealTargetPlatform);
 		}
 
-		public override object ReadJson(JsonReader Reader, Type ObjectType, object ExistingValue, JsonSerializer Serializer)
+		public override UnrealTargetPlatform Read(ref Utf8JsonReader Reader, Type TypeToConvert, JsonSerializerOptions Options)
 		{
-			UnrealTargetPlatform Platform;
-			if (!UnrealTargetPlatform.TryParse((string)Reader.Value, out Platform))
+			if(Reader.TokenType == JsonTokenType.Null)
 			{
-				return null;
+				return BuildHostPlatform.Current.Platform;
 			}
-			return Platform;
+
+			UnrealTargetPlatform StructValue;
+			if (UnrealTargetPlatform.TryParse(Reader.GetString(), out StructValue))
+			{
+				return StructValue;
+			}
+			throw new JsonException();
 		}
 
-		public override void WriteJson(JsonWriter Writer, object Value, JsonSerializer Serializer)
+		public override void Write(Utf8JsonWriter Writer, UnrealTargetPlatform StructValue, JsonSerializerOptions Options)
 		{
-			UnrealTargetPlatform? Platform = (UnrealTargetPlatform)Value;
-			Writer.WriteValue(Platform);
+			var Value = StructValue.ToString();
+			Writer.WriteStringValue(Options.PropertyNamingPolicy?.ConvertName(Value) ?? Value);
 		}
-
 	}
 
 
