Clarify error message when scanning temp storage for cleanup

#rnx

diff --git a/Engine/Source/Programs/AutomationTool/BuildGraph/TempStorage.cs b/Engine/Source/Programs/AutomationTool/BuildGraph/TempStorage.cs
--- a/Engine/Source/Programs/AutomationTool/BuildGraph/TempStorage.cs
+++ b/Engine/Source/Programs/AutomationTool/BuildGraph/TempStorage.cs
@@ -1442,7 +1442,7 @@ namespace AutomationTool
 						}
 						catch (Exception Ex)
 						{
-							Log.TraceError("Exception while trying to delete {0}: {1}", BuildDirectory, Ex);
+							Log.TraceError("Exception while trying to scan files under {0}: {1}", BuildDirectory, Ex);
 						}
 					}
 				}
