Re-enable Audio Tests now that we have a Null Audio device working for SDL

diff --git a/Engine/Source/Programs/AutomationTool/Gauntlet/Unreal/Base/Gauntlet.UnrealTestConfiguration.cs b/Engine/Source/Programs/AutomationTool/Gauntlet/Unreal/Base/Gauntlet.UnrealTestConfiguration.cs
--- a/Engine/Source/Programs/AutomationTool/Gauntlet/Unreal/Base/Gauntlet.UnrealTestConfiguration.cs
+++ b/Engine/Source/Programs/AutomationTool/Gauntlet/Unreal/Base/Gauntlet.UnrealTestConfiguration.cs
@@ -895,7 +895,7 @@ namespace Gauntlet
 				// Disable cef as it seems to want to talk to an X11 server so unlikely its even working
 				if (AppConfig.Platform == UnrealTargetPlatform.Linux)
 				{
-					AppConfig.CommandLine += " -renderoffscreen -nosound -nocef";
+					AppConfig.CommandLine += " -renderoffscreen";
 				}
 			}
 
