Android editor integration improvements.

UAT RunClient: Continuously redirect output to uat stdout while client is running.
UAT RunClient: Add a process wrapper to kill client process when uat exits and to write the adb device log file to the host pc directories.
UAT RunClient: Push the ClientCmdLine args as a UECommandLine.txt file to the device to the correct location (different for adb deploy and android file server) to override the stage/apk command line.
UAT MakeApk: Add a UECommandLine.txt cache file to the intermediate directory to invalidate the apk when the stage/apk command line changes (e.g. when switching between CBTB or COTF or cooking with/without Zen).
LaunchAndroid.cpp: Add debug output for which override command line file that is being used.
LaunchAndroid.cpp: Implement missing -waitforattach/-waitfordebugger feature.
DataDrivenPlatformInfo.ini: set bUATClosesAfterLaunch=false to stop the editor from passing -nokill to uat to let it stop all processes when the client exits..

#rnx

diff --git a/Engine/Config/Android/DataDrivenPlatformInfo.ini b/Engine/Config/Android/DataDrivenPlatformInfo.ini
--- a/Engine/Config/Android/DataDrivenPlatformInfo.ini
+++ b/Engine/Config/Android/DataDrivenPlatformInfo.ini
@@ -17,7 +17,7 @@ AutoSDKPath=Android
 TutorialPath=SharingAndReleasing/Mobile/Android
 bIsEnabled=true
 bUsesHostCompiler=false
-bUATClosesAfterLaunch=true
+bUATClosesAfterLaunch=false
 PlatformGroupName=Mobile
 
 [PreviewPlatform AndroidES31]
@@ -81,4 +81,4 @@ bNeedsOfflineCompiler=true
 bSupportsIndexBufferUAVs=true
 bSupportsManualVertexFetch=false
 bRequiresDisableForwardLocalLights=true
-bWaterUsesSimpleForwardShading=true
\ No newline at end of file
+bWaterUsesSimpleForwardShading=true
diff --git a/Engine/Source/Programs/AutomationTool/Android/AndroidPlatform.Automation.cs b/Engine/Source/Programs/AutomationTool/Android/AndroidPlatform.Automation.cs
--- a/Engine/Source/Programs/AutomationTool/Android/AndroidPlatform.Automation.cs
+++ b/Engine/Source/Programs/AutomationTool/Android/AndroidPlatform.Automation.cs
@@ -28,6 +28,131 @@ public class AndroidPlatform : Platform
     private const string TargetAndroidLocation = "obb/";
 	private const string TargetAndroidTemp = "/data/local/tmp/";
 
+	public class AdbCreatedProcess : AutomationTool.IProcessResult
+	{
+		private readonly object StopSyncObject = new object();
+		IProcessResult AdbLogProcess;
+		string LogPath;
+		string PackageName;
+		string DeviceName;
+		bool bStopped = false;
+
+		public AdbCreatedProcess(
+				IProcessResult InAdbLogProcess,
+				string InLogPath,
+				string InPackageName,
+				string InDeviceName)
+		{
+			AdbLogProcess = InAdbLogProcess;
+			LogPath = InLogPath;
+			PackageName = InPackageName;
+			DeviceName = InDeviceName;
+			ProcessManager.AddProcess(this);
+		}
+
+		~AdbCreatedProcess()
+		{
+			ProcessManager.RemoveProcess(this);
+		}
+
+		public void StopProcess(bool KillDescendants = true)
+		{
+			lock (StopSyncObject)
+			{
+				if (!bStopped)
+				{
+					AndroidPlatform.RunAdbCommand(DeviceName, "shell am force-stop " + PackageName);
+					if (!AdbLogProcess.HasExited)
+					{
+						AdbLogProcess.StopProcess(KillDescendants);
+					}
+					DumpDeviceOutputToLogFiles();
+					bStopped = true;
+				}
+			}
+		}
+
+		public bool HasExited
+		{
+			get
+			{
+				if (!bStopped && (AdbLogProcess.HasExited || !IsPackageRunningOnDevice()))
+				{
+					StopProcess();
+				}
+				return bStopped;
+			}
+		}
+
+		public string GetProcessName()
+		{
+			return String.Format("{0}@{1}", PackageName, DeviceName);
+		}
+
+		public void OnProcessExited()
+		{
+		}
+
+		public void DisposeProcess()
+		{
+			AdbLogProcess.DisposeProcess();
+		}
+
+		public void StdOut(object sender, DataReceivedEventArgs e)
+		{
+		}
+
+		public void StdErr(object sender, DataReceivedEventArgs e)
+		{
+		}
+
+		public int ExitCode
+		{
+			get { return AdbLogProcess.ExitCode; }
+			set { AdbLogProcess.ExitCode = value; }
+		}
+
+		public string Output
+		{
+			get { return AdbLogProcess.Output; }
+		}
+
+		public Process ProcessObject
+		{
+			get { return AdbLogProcess.ProcessObject; }
+		}
+
+		public void WaitForExit()
+		{
+			while (!AdbLogProcess.HasExited && IsPackageRunningOnDevice())
+			{
+				Thread.Sleep(100);
+			}
+			StopProcess();
+		}
+
+		private bool IsPackageRunningOnDevice()
+		{
+			ERunOptions Options = ERunOptions.Default | ERunOptions.SpewIsVerbose | ERunOptions.NoLoggingOfRunCommand;
+			IProcessResult Result = AndroidPlatform.RunAdbCommand(DeviceName, "shell ps", null, Options);
+			string ProcessList = Result.Output;
+			bool bIsProcessRunning = ProcessList.Contains(PackageName);
+			return bIsProcessRunning;
+		}
+
+		private void DumpDeviceOutputToLogFiles()
+		{
+			string SanitizedDeviceName = DeviceName.Replace(":", "_");
+			string LogFilename = Path.Combine(LogPath, "devicelog" + SanitizedDeviceName + ".log");
+			string ServerLogFilename = Path.Combine(CmdEnv.LogFolder, "devicelog" + SanitizedDeviceName + ".log");
+			ERunOptions Options = ERunOptions.Default & ~ERunOptions.AllowSpew;
+			IProcessResult LogFileProcess = RunAdbCommand(DeviceName, "logcat -d", null, Options);
+			string AllOutput = LogFileProcess.Output;
+			File.WriteAllText(LogFilename, AllOutput);
+			File.WriteAllText(ServerLogFilename, AllOutput);
+		}
+	}
+
 	public AndroidPlatform()
 		: base(UnrealTargetPlatform.Android)
 	{
@@ -547,12 +672,14 @@ public class AndroidPlatform : Platform
 		return PluginExtras;
 	}
 
-	private bool UsingAndroidFileServer(DeploymentContext SC, out bool bEnablePlugin, out string AFSToken, out bool bIsShipping, out bool bIncludeInShipping, out bool bAllowExternalStartInShipping)
+	private bool UsingAndroidFileServer(ProjectParams Params, DeploymentContext SC, out bool bEnablePlugin, out string AFSToken, out bool bIsShipping, out bool bIncludeInShipping, out bool bAllowExternalStartInShipping)
 	{
-		UnrealTargetConfiguration TargetConfiguration = SC.StageTargetConfigurations[0];
+		FileReference RawProjectPath = SC != null ? SC.RawProjectPath : Params.RawProjectPath;
+		UnrealTargetPlatform TargetPlatform = SC != null ? SC.StageTargetPlatform.PlatformType : Params.ClientTargetPlatforms[0].Type;
+		UnrealTargetConfiguration TargetConfiguration = SC != null ? SC.StageTargetConfigurations[0] : Params.ClientConfigsToBuild[0];
 		bIsShipping = TargetConfiguration == UnrealTargetConfiguration.Shipping;
 
-		ConfigHierarchy Ini = ConfigCache.ReadHierarchy(ConfigHierarchyType.Engine, DirectoryReference.FromFile(SC.RawProjectPath), SC.StageTargetPlatform.PlatformType);
+		ConfigHierarchy Ini = ConfigCache.ReadHierarchy(ConfigHierarchyType.Engine, DirectoryReference.FromFile(RawProjectPath), TargetPlatform);
 		if (!Ini.GetBool("/Script/AndroidFileServerEditor.AndroidFileServerRuntimeSettings", "bEnablePlugin", out bEnablePlugin))
 		{
 			bEnablePlugin = true;
@@ -789,7 +916,7 @@ public class AndroidPlatform : Platform
 		bool bIsShipping;
 		bool bAFSIncludeInShipping;
 		bool bAFSAllowExternalStartInShipping;
-		UsingAndroidFileServer(SC, out bAFSEnablePlugin, out AFSToken, out bIsShipping, out bAFSIncludeInShipping, out bAFSAllowExternalStartInShipping);
+		UsingAndroidFileServer(Params, SC, out bAFSEnablePlugin, out AFSToken, out bIsShipping, out bAFSIncludeInShipping, out bAFSAllowExternalStartInShipping);
 
 		if (bAFSEnablePlugin && !bPackageDataInsideApk)
 		{
@@ -1571,7 +1698,7 @@ public class AndroidPlatform : Platform
 		bool bIsShipping;
 		bool bAFSIncludeInShipping;
 		bool bAFSAllowExternalStartInShipping;
-		UsingAndroidFileServer(SC, out bAFSEnablePlugin, out AFSToken, out bIsShipping, out bAFSIncludeInShipping, out bAFSAllowExternalStartInShipping);
+		UsingAndroidFileServer(Params, SC, out bAFSEnablePlugin, out AFSToken, out bIsShipping, out bAFSIncludeInShipping, out bAFSAllowExternalStartInShipping);
 		bool bUseAFS = bAFSEnablePlugin && !bPackageDataInsideApk;
 
 		List<string> AddedObbFiles = new List<string>();
@@ -1802,7 +1929,7 @@ public class AndroidPlatform : Platform
 		}
 	}
 
-	private static string GetAdbCommandLine(ProjectParams Params, string SerialNumber, string Args)
+	private static string GetAdbCommandLine(string SerialNumber, string Args)
 	{
 	    if (string.IsNullOrEmpty(SerialNumber) == false)
 		{
@@ -1836,21 +1963,26 @@ public class AndroidPlatform : Platform
 	}
 
 	public static IProcessResult RunAdbCommand(ProjectParams Params, string SerialNumber, string Args, string Input = null, ERunOptions Options = ERunOptions.Default, bool bShouldLogCommand = false)
+	{
+		return RunAdbCommand(SerialNumber, Args, Input, Options, bShouldLogCommand);
+	}
+
+	private static IProcessResult RunAdbCommand(string SerialNumber, string Args, string Input = null, ERunOptions Options = ERunOptions.Default, bool bShouldLogCommand = false)
 	{
 		string AdbCommand = Environment.ExpandEnvironmentVariables("%ANDROID_HOME%/platform-tools/adb" + (RuntimePlatform.IsWindows ? ".exe" : ""));
 		if (Options.HasFlag(ERunOptions.AllowSpew) || Options.HasFlag(ERunOptions.SpewIsVerbose))
 		{
 			LastSpewFilename = "";
-			return Run(AdbCommand, GetAdbCommandLine(Params, SerialNumber, Args), Input, Options, SpewFilterCallback: new ProcessResult.SpewFilterCallbackType(ADBSpewFilter));
+			return Run(AdbCommand, GetAdbCommandLine(SerialNumber, Args), Input, Options, SpewFilterCallback: new ProcessResult.SpewFilterCallbackType(ADBSpewFilter));
 		}
-		return Run(AdbCommand, GetAdbCommandLine(Params, SerialNumber, Args), Input, Options);
+		return Run(AdbCommand, GetAdbCommandLine(SerialNumber, Args), Input, Options);
 	}
 
 	private string RunAndLogAdbCommand(ProjectParams Params, string SerialNumber, string Args, out int SuccessCode)
 	{
 		string AdbCommand = Environment.ExpandEnvironmentVariables("%ANDROID_HOME%/platform-tools/adb" + (RuntimePlatform.IsWindows ? ".exe" : ""));
 		LastSpewFilename = "";
-		return RunAndLog(CmdEnv, AdbCommand, GetAdbCommandLine(Params, SerialNumber, Args), out SuccessCode, SpewFilterCallback: new ProcessResult.SpewFilterCallbackType(ADBSpewFilter));
+		return RunAndLog(CmdEnv, AdbCommand, GetAdbCommandLine(SerialNumber, Args), out SuccessCode, SpewFilterCallback: new ProcessResult.SpewFilterCallbackType(ADBSpewFilter));
 	}
 
 	public override void GetConnectedDevices(ProjectParams Params, out List<string> Devices)
@@ -1999,7 +2131,7 @@ public class AndroidPlatform : Platform
 		bool bIsShipping;
 		bool bAFSIncludeInShipping;
 		bool bAFSAllowExternalStartInShipping;
-		if (UsingAndroidFileServer(SC, out bAFSEnablePlugin, out AFSToken, out bIsShipping, out bAFSIncludeInShipping, out bAFSAllowExternalStartInShipping))
+		if (UsingAndroidFileServer(Params, SC, out bAFSEnablePlugin, out AFSToken, out bIsShipping, out bAFSIncludeInShipping, out bAFSAllowExternalStartInShipping))
 		{
 			return RetrieveDeployedManifestsAFS(Params, SC, DeviceName, out UFSManifests, out NonUFSManifests, AFSToken);
 		}
@@ -3157,7 +3289,7 @@ public class AndroidPlatform : Platform
 		bool bAFSAllowExternalStartInShipping;
 
 		// Pick the proper deploy method
-		if (UsingAndroidFileServer(SC, out bAFSEnablePlugin, out AFSToken, out bIsShipping, out bAFSIncludeInShipping, out bAFSAllowExternalStartInShipping))
+		if (UsingAndroidFileServer(Params, SC, out bAFSEnablePlugin, out AFSToken, out bIsShipping, out bAFSIncludeInShipping, out bAFSAllowExternalStartInShipping))
 		{
 			DeployAndroidFileServer(Params, SC, AFSToken);
 		}
@@ -3531,19 +3663,65 @@ public class AndroidPlatform : Platform
 		return DeviceArch;
 	}
 
+	private bool DeployClientCmdLineAFS(ProjectParams Params, string DeviceName, string PackageName, string AFSToken, string ClientCmdLine)
+	{
+		AndroidFileClient client = new AndroidFileClient(DeviceName);
+		if (!client.OpenConnection())
+		{
+			Log.TraceInformation("DeployClientCmdLine: Trying to start file server {0}", PackageName);
+			if (!client.StartServer(PackageName, AFSToken))
+			{
+				Log.TraceInformation("DeployClientCmdLine: Failed to start server {0}, ignoring client command line", PackageName);
+				return false;
+			}
+		}
+
+		// verify connection to the correct server
+		string DevicePackageName = client.Query("^packagename");
+		if (DevicePackageName != PackageName)
+		{
+			Log.TraceInformation("DeployClientCmdLine: Connected to wrong server {0}, trying again", DevicePackageName);
+			client.TerminateServer();
+
+			Log.TraceInformation("DeployClientCmdLine: Trying to start file server {0}", PackageName);
+			if (!client.StartServer(PackageName, AFSToken))
+			{
+				Log.TraceInformation("DeployClientCmdLine: Failed to start server {0}, ignoring client command line", PackageName);
+				return false;
+			}
+		}
+
+		Log.TraceInformation("Writing ClientCmdLine to remote ^commandfile: {0}", ClientCmdLine);
+		client.FileWriteString(ClientCmdLine, "^commandfile");
+		client.TerminateServer();
+		client.CloseConnection();
+		return true;
+	}
+
+	private bool DeployClientCmdLineADB(ProjectParams Params, string DeviceName, string PackageName, string ClientCmdLine)
+	{
+		string DeviceStorageQueryCommand = GetStorageQueryCommand();
+		IProcessResult StorageResult = RunAdbCommand(Params, DeviceName, DeviceStorageQueryCommand, null, ERunOptions.AppMustExist);
+		string StorageLocation = StorageResult.Output.Trim();
+		string RemoteDir = StorageLocation + "/UnrealGame/" + Params.ShortProjectName;
+		string ClientCmdLineTmpFile = Path.GetTempFileName();
+		string ClientCmdLineRemoteFile = RemoteDir + "/UECommandLine.txt";
+		File.WriteAllText(ClientCmdLineTmpFile, ClientCmdLine);
+		Log.TraceInformation("Pushing ClientCmdLine to remote file {0}: {1}", ClientCmdLineRemoteFile, ClientCmdLine);
+		RunAdbCommand(DeviceName, String.Format("push {0} {1}", ClientCmdLineTmpFile, ClientCmdLineRemoteFile));
+		File.Delete(ClientCmdLineTmpFile);
+		return true;
+	}
+
 	public override IProcessResult RunClient(ERunOptions ClientRunFlags, string ClientApp, string ClientCmdLine, ProjectParams Params)
 	{
 		IProcessResult Result = null;
-		//make a copy of the device names, we'll be working through them
-		List<string> DeviceNames = new List<string>();
-		//same with the package names
-		List<string> PackageNames = new List<string>();
+
+		string LogPath = Path.Combine(Params.BaseStageDirectory, "Android\\logs");
+		Directory.CreateDirectory(LogPath);
 
 		foreach (string DeviceName in Params.DeviceNames)
 		{
-			//save the device name
-			DeviceNames.Add(DeviceName);
-
 			//get the package name and save that
 			string DeviceArchitecture = GetBestDeviceArchitecture(Params, DeviceName);
 
@@ -3560,7 +3738,6 @@ public class AndroidPlatform : Platform
 			{
 				throw new AutomationException(ExitCode.Error_AppNotFound, "Failed to find application " + ApkName);
 			}
-			
 
 			// run aapt to get the name of the intent
 			string PackageName = GetPackageInfo(ApkName, false);
@@ -3569,87 +3746,42 @@ public class AndroidPlatform : Platform
 				throw new AutomationException(ExitCode.Error_FailureGettingPackageInfo, "Failed to get package name from " + ClientApp);
 			}
 
-			PackageNames.Add(PackageName);
-
-			// Message back to the Unreal Editor to correctly set the app id for each device
-			Console.WriteLine("Running Package@Device:{0}@{1}", PackageName, DeviceName);
-
-			// clear the log for the device
-			RunAdbCommand(Params, DeviceName, "logcat -c");
-
-			// start the app on device!
-			string CommandLine = "shell am start -n " + PackageName + "/" + GetLaunchableActivityName();
-			RunAdbCommand(Params, DeviceName, CommandLine, null, ClientRunFlags);
-
-			// save the output to the staging directory
-			string LogPath = Path.Combine(Params.BaseStageDirectory, "Android\\logs");
-			Directory.CreateDirectory(LogPath);
-		}
-
-		//now check if each device still has the game running, and time out if it's taking too long
-		DateTime StartTime = DateTime.Now;
-		int TimeOutSeconds = Params.RunTimeoutSeconds;
-
-		// wait before getting the process list with "adb shell ps"
-		// on some devices the list is not yet ready
-		Thread.Sleep(2000);
-
-		while (DeviceNames.Count > 0)
-		{
-			for(int DeviceIndex = 0; DeviceIndex < DeviceNames.Count; DeviceIndex++)
+			// push ClientCmdLine args as a file to the device to override the stage/apk command line
 			{
-				string DeviceName = DeviceNames[DeviceIndex];
-				
-				//replace the port name in the case of deploy while adb is using wifi
-				string SanitizedDeviceName = DeviceName.Replace(":", "_");
-
-				bool FinishedRunning = false;
-				IProcessResult ProcessesResult = RunAdbCommand(Params, DeviceName, "shell ps", null, ERunOptions.SpewIsVerbose);
-
-				string RunningProcessList = ProcessesResult.Output;
-				if (!RunningProcessList.Contains(PackageNames[DeviceIndex]))
+				bool bAFSEnablePlugin;
+				string AFSToken;
+				bool bIsShipping;
+				bool bAFSIncludeInShipping;
+				bool bAFSAllowExternalStartInShipping;
+				if (UsingAndroidFileServer(Params, null, out bAFSEnablePlugin, out AFSToken, out bIsShipping, out bAFSIncludeInShipping, out bAFSAllowExternalStartInShipping))
 				{
-					FinishedRunning = true;
+					DeployClientCmdLineAFS(Params, DeviceName, PackageName, AFSToken, ClientCmdLine);
 				}
-
-				Thread.Sleep(1000);
-
-				if(!FinishedRunning)
+				else
 				{
-					TimeSpan DeltaRunTime = DateTime.Now - StartTime;
-					if ((DeltaRunTime.TotalSeconds > TimeOutSeconds) && (TimeOutSeconds != 0))
-					{
-						LogInformation("Device: " + DeviceName + " timed out while waiting for run to finish");
-						FinishedRunning = true;
-					}
+					DeployClientCmdLineADB(Params, DeviceName, PackageName, ClientCmdLine);
 				}
+			}
 
-				//log the results, then clear out the device from our list
-				if(FinishedRunning)
-				{
-					// this is just to get the ue log to go to the output
-					RunAdbCommand(Params, DeviceName, "logcat -d -s UE debug Debug DEBUG");
-
-					// get the log we actually want to save
-					IProcessResult LogFileProcess = RunAdbCommand(Params, DeviceName, "logcat -d", null, ERunOptions.AppMustExist);
+			// Message back to the Unreal Editor to correctly set the app id for each device
+			CommandUtils.LogInformation("Running Package@Device:{0}@{1}", PackageName, DeviceName);
 
-					string LogPath = Path.Combine(Params.BaseStageDirectory, "Android\\logs");
-					string LogFilename = Path.Combine(LogPath, "devicelog" + SanitizedDeviceName + ".log");
-					string ServerLogFilename = Path.Combine(CmdEnv.LogFolder, "devicelog" + SanitizedDeviceName + ".log");
+			// clear the log for the device
+			RunAdbCommand(DeviceName, "logcat -c");
 
-					File.WriteAllText(LogFilename, LogFileProcess.Output);
-					File.WriteAllText(ServerLogFilename, LogFileProcess.Output);
+			// start the app on device!
+			string CommandLine = "shell am start -n " + PackageName + "/" + GetLaunchableActivityName();
+			RunAdbCommand(DeviceName, CommandLine);
 
-					if (Result == null)
-					{
-						Result = LogFileProcess;
-					}
-					DeviceNames.RemoveAt(DeviceIndex);
-					PackageNames.RemoveAt(DeviceIndex);
+			// wait before getting the process list with "adb shell ps" from AdbCreatedProcess
+			// on some devices the list is not yet ready
+			Thread.Sleep(2000);
 
-					--DeviceIndex;
-				}
-			}
+			// Start logging process and return immediately.
+			// Stdout from the title is continuosly emitted to stdout in UAT.
+			// When process is done the AdbCreatedProcess wrapper will save the output to the log directories
+			Result = RunAdbCommand(DeviceName, "logcat -s UE debug Debug DEBUG", null, ClientRunFlags | ERunOptions.NoWaitForExit);
+			Result = new AdbCreatedProcess(Result, LogPath, PackageName, DeviceName);
 		}
 
 		return Result;
diff --git a/Engine/Source/Programs/UnrealBuildTool/Platform/Android/UEDeployAndroid.cs b/Engine/Source/Programs/UnrealBuildTool/Platform/Android/UEDeployAndroid.cs
--- a/Engine/Source/Programs/UnrealBuildTool/Platform/Android/UEDeployAndroid.cs
+++ b/Engine/Source/Programs/UnrealBuildTool/Platform/Android/UEDeployAndroid.cs
@@ -4007,6 +4007,20 @@ namespace UnrealBuildTool
 
 			}
 
+			string CommandLineSourceFileName = Path.Combine(Path.GetDirectoryName(ObbFileLocation)!, Path.GetFileNameWithoutExtension(ObbFileLocation), "UECommandLine.txt");
+			string CommandLineCacheFileName = Path.Combine(IntermediateAndroidPath, "UECommandLine.txt");
+			if (bBuildSettingsMatch)
+			{
+				bool bCommandLineMatch =
+					(!File.Exists(CommandLineSourceFileName) && !File.Exists(CommandLineCacheFileName)) ||
+					BinaryFileEquals(CommandLineSourceFileName, CommandLineCacheFileName);
+				if (!bCommandLineMatch)
+				{
+					bBuildSettingsMatch = false;
+					Log.TraceInformation("Previous .apk file(s) were made with different stage/apk command line, forcing repackage.");
+				}
+			}
+
 			// Initialize UPL contexts for each architecture enabled
 			UPL.Init(NDKArches, bForDistribution, EngineDirectory, IntermediateAndroidPath, ProjectDirectory, Configuration.ToString(), bSkipGradleBuild, bPerArchBuildDir:true, ArchRemapping:ArchRemapping);
 
@@ -4039,6 +4053,14 @@ namespace UnrealBuildTool
 
 			// at this point, we can write out the cached build settings to compare for a next build
 			File.WriteAllText(BuildSettingsCacheFile, CurrentBuildSettings);
+			if (File.Exists(CommandLineSourceFileName))
+			{
+				CopyIfDifferent(CommandLineSourceFileName, CommandLineCacheFileName, true, true);
+			}
+			else
+			{
+				SafeDeleteFile(CommandLineCacheFileName);
+			}
 
 			// make up a dictionary of strings to replace in xml files (strings.xml)
 			Dictionary<string, string> Replacements = new Dictionary<string, string>();
@@ -4081,7 +4103,6 @@ namespace UnrealBuildTool
 				}
 
 				// See if we need to stage a UECommandLine.txt file in assets
-				string CommandLineSourceFileName = Path.Combine(Path.GetDirectoryName(ObbFileLocation)!, Path.GetFileNameWithoutExtension(ObbFileLocation), "UECommandLine.txt");
 				string CommandLineDestFileName = Path.Combine(UnrealBuildPath, "assets", "UECommandLine.txt");
 				if (File.Exists(CommandLineSourceFileName))
 				{
diff --git a/Engine/Source/Runtime/Launch/Private/Android/LaunchAndroid.cpp b/Engine/Source/Runtime/Launch/Private/Android/LaunchAndroid.cpp
--- a/Engine/Source/Runtime/Launch/Private/Android/LaunchAndroid.cpp
+++ b/Engine/Source/Runtime/Launch/Private/Android/LaunchAndroid.cpp
@@ -315,6 +315,8 @@ static void InitCommandLine()
 
 	if(CommandLineFile)
 	{
+		FPlatformMisc::LowLevelOutputDebugStringf(TEXT("Using override commandline file: %s"), *CommandLineFilePath);
+
 		char CommandLine[CMD_LINE_MAX];
 		fgets(CommandLine, UE_ARRAY_COUNT(CommandLine) - 1, CommandLineFile);
 
@@ -450,6 +452,16 @@ int32 AndroidMain(struct android_app* state)
 	InitCommandLine();
 	FPlatformMisc::LowLevelOutputDebugStringf(TEXT("Final commandline: %s\n"), FCommandLine::Get());
 
+#if !(UE_BUILD_SHIPPING)
+	// If "-waitforattach" or "-WaitForDebugger" was specified, halt startup and wait for a debugger to attach before continuing
+	if (FParse::Param(FCommandLine::Get(), TEXT("waitforattach")) || FParse::Param(FCommandLine::Get(), TEXT("WaitForDebugger")))
+	{
+		FPlatformMisc::LowLevelOutputDebugString(TEXT("Waiting for debugger to attach...\n"));
+		while (!FPlatformMisc::IsDebuggerPresent());
+		UE_DEBUG_BREAK();
+	}
+#endif
+
 	EventHandlerEvent = FPlatformProcess::GetSynchEventFromPool(false);
 	FPlatformMisc::LowLevelOutputDebugString(TEXT("Created sync event\n"));
 	FAppEventManager::GetInstance()->SetEventHandlerEvent(EventHandlerEvent);
