UBT: recognize shared libraries with version number suffixes as allowed duplicate build products

Shared libraries on Linux commonly use a scheme of version number suffixes
and/or symlinks to help ensure that client applications can link against the
appropriate version of a library (e.g. "libfoo.so.1.70.0"). This change updates
IsDuplicateBuildProduct() with an additional case to cover these libraries so
that they are not flagged as an error by the build system if multiple steps
produce them.

diff --git a/Engine/Source/Programs/AutomationTool/BuildGraph/TempStorage.cs b/Engine/Source/Programs/AutomationTool/BuildGraph/TempStorage.cs
--- a/Engine/Source/Programs/AutomationTool/BuildGraph/TempStorage.cs
+++ b/Engine/Source/Programs/AutomationTool/BuildGraph/TempStorage.cs
@@ -1243,6 +1243,11 @@ namespace AutomationTool
 			{
 				return true;
 			}
+			if (FileName.StartsWith("lib", StringComparison.OrdinalIgnoreCase) && FileName.Contains(".so.", StringComparison.OrdinalIgnoreCase))
+			{
+				// e.g. a Unix shared library with a version number suffix.
+				return true;
+			}
 			if (FileName.Equals("plugInfo.json", StringComparison.OrdinalIgnoreCase))
 			{
 				return true;
