UAT: Add BuildCookRun support for running both a dedicated server and cotf server.
Merge RunClientWithServer() and RunStandAloneClient() functions into one RunClient() function to resolve numerous inconsistencies between them and between unattended/non-unattended mode.
Always use SC.StageTargetPlatform.RunClient() to launch clients so that the logic works on all platforms.
When either the client, dedicated server or cook server exits, then all other processes exits too.
Hitting Ctrl-c in UAT kills all processes.

ShooterGame testing on Win64, consoles and android:
BuildCookRun -project=ShooterGame -platform=Win64 -config=development -skipbuild -cook -stage -pak -deploy -run
BuildCookRun -project=ShooterGame -platform=Win64 -config=development -skipbuild -skipcook -skipstage -skipdeploy -run
BuildCookRun -project=ShooterGame -platform=Win64 -config=development -skipbuild -skipcook -skipstage -skipdeploy -run -unattended
BuildCookRun -project=ShooterGame -platform=Win64 -config=development -skipbuild -cookonthefly -stage -deploy -run
BuildCookRun -project=ShooterGame -platform=Win64 -config=development -skipbuild -cookonthefly -stage -deploy -run -unattended
BuildCookRun -project=ShooterGame -platform=Win64 -config=development -skipbuild -cookonthefly -zenstore -stage -deploy -run
BuildCookRun -project=ShooterGame -platform=Win64 -config=development -skipbuild -cookonthefly -zenstore -stage -deploy -run -unattended

ShooterGame Win64 Desktop specific testing with lan server and multiple clients:
BuildCookRun -project=ShooterGame -platform=Win64 -config=development -skipbuild -cookonthefly -zenstore -server -stage -run -map=Highrise -numclients=3
BuildCookRun -project=ShooterGame -platform=Win64 -config=development -skipbuild -cookonthefly -zenstore -server -stage -run -map=Highrise -numclients=3 -unattended

EngineTest testing on Win64 and console:
BuildCookRun -project=EngineTest -platform=Win64 -config=development -skipbuild -cook -stage -pak -deploy -run -runautomationtest=Project.System
BuildCookRun -project=EngineTest -platform=Win64 -config=development -skipbuild -cookonthefly -stage -deploy -run -runautomationtest=Project.System
BuildCookRun -project=EngineTest -platform=Win64 -config=development -skipbuild -cookonthefly -zenstore -stage -deploy -run -runautomationtest=Project.System

#rnx

diff --git a/Engine/Source/Programs/AutomationTool/AutomationUtils/ProjectParams.cs b/Engine/Source/Programs/AutomationTool/AutomationUtils/ProjectParams.cs
--- a/Engine/Source/Programs/AutomationTool/AutomationUtils/ProjectParams.cs
+++ b/Engine/Source/Programs/AutomationTool/AutomationUtils/ProjectParams.cs
@@ -859,15 +859,13 @@ namespace AutomationTool
 			this.Deploy = GetParamValueIfNotSpecified(Command, Deploy, this.Deploy, "deploy");
 			this.DeployFolder = ParseParamValueIfNotSpecified(Command, DeployFolder, "deploy", null);
 
-			// if the user specified -deploy but no folder, set the default
-			if (this.Deploy && string.IsNullOrEmpty(this.DeployFolder))
+			// always set the default deploy folder, so that it is available in -skipdeploy scenarios too
+			if (string.IsNullOrEmpty(this.DeployFolder))
 			{
 				this.DeployFolder = UnrealBuildTool.DeployExports.GetDefaultDeployFolder(this.ShortProjectName);
 			}
-			else if (string.IsNullOrEmpty(this.DeployFolder) == false)
+			else
 			{
-				// if the user specified a folder set deploy to true.
-				//@todo - remove 'deploy' var and check deployfolder != null?
 				this.Deploy = true;
 			}
 
@@ -2914,11 +2912,6 @@ namespace AutomationTool
                 throw new AutomationException("DedicatedServer cannot be used with RunAutomationTests");
             }
 
-			if ((CookOnTheFly || FileServer) && DedicatedServer)
-			{
-				throw new AutomationException("Don't use either -cookonthefly or -fileserver with -server.");
-			}
-
 			if (NoClient && !DedicatedServer && !CookOnTheFly)
 			{
 				throw new AutomationException("-noclient can only be used with -server or -cookonthefly.");
diff --git a/Engine/Source/Programs/AutomationTool/Scripts/CookCommand.Automation.cs b/Engine/Source/Programs/AutomationTool/Scripts/CookCommand.Automation.cs
--- a/Engine/Source/Programs/AutomationTool/Scripts/CookCommand.Automation.cs
+++ b/Engine/Source/Programs/AutomationTool/Scripts/CookCommand.Automation.cs
@@ -255,43 +255,29 @@ namespace AutomationScripts
 				throw new AutomationException("Missing " + UEEditorExe + " executable. Needs to be built first.");
 			}
 
-			if (Params.CookOnTheFly && !Params.SkipServer)
+			if (Params.CookOnTheFly)
 			{
 				if (Params.HasDLCName)
 				{
 					throw new AutomationException("Cook on the fly doesn't support cooking dlc");
 				}
-				if (Params.ClientTargetPlatforms.Count > 0)
-				{
-					var LogFolderOutsideOfSandbox = GetLogFolderOutsideOfSandbox();
-					if (!Unreal.IsEngineInstalled())
-					{
-						// In the installed runs, this is the same folder as CmdEnv.LogFolder so delete only in not-installed
-						DeleteDirectory(LogFolderOutsideOfSandbox);
-						CreateDirectory(LogFolderOutsideOfSandbox);
-					}
-
-					string COTFCommandLine = GetGenericCookCommandletParams(Params);
-					if (Params.ZenStore)
-					{
-						COTFCommandLine += " -messaging";
-					}
 
-					var ServerLogFile = CombinePaths(LogFolderOutsideOfSandbox, "Server.log");
-					Platform ClientPlatformInst = Params.ClientTargetPlatformInstances[0];
-					string TargetCook = ClientPlatformInst.GetCookPlatform(false, false); // cook on he fly doesn't support server cook platform... 
-					ServerProcess = RunCookOnTheFlyServer(Params.RawProjectPath, Params.NoClient ? "" : ServerLogFile, TargetCook, COTFCommandLine);
-
-					if (ServerProcess != null)
-					{
-						LogInformation("Waiting a few seconds for the server to start...");
-						Thread.Sleep(5000);
-					}
+				var LogFolderOutsideOfSandbox = GetLogFolderOutsideOfSandbox();
+				if (!Unreal.IsEngineInstalled())
+				{
+					// In the installed runs, this is the same folder as CmdEnv.LogFolder so delete only in not-installed
+					DeleteDirectory(LogFolderOutsideOfSandbox);
+					CreateDirectory(LogFolderOutsideOfSandbox);
 				}
-				else
+
+				string COTFCommandLine = GetGenericCookCommandletParams(Params);
+				if (Params.ZenStore)
 				{
-					throw new AutomationException("Failed to run, client target platform not specified");
+					COTFCommandLine += " -messaging";
 				}
+
+				var CookServerLogFile = CombinePaths(LogFolderOutsideOfSandbox, "CookServer.log");
+				CookServerProcess = RunCookOnTheFlyServer(Params.RawProjectPath, Params.NoClient ? "" : CookServerLogFile, COTFCommandLine);
 			}
 			else
 			{
diff --git a/Engine/Source/Programs/AutomationTool/Scripts/RunProjectCommand.Automation.cs b/Engine/Source/Programs/AutomationTool/Scripts/RunProjectCommand.Automation.cs
--- a/Engine/Source/Programs/AutomationTool/Scripts/RunProjectCommand.Automation.cs
+++ b/Engine/Source/Programs/AutomationTool/Scripts/RunProjectCommand.Automation.cs
@@ -3,6 +3,7 @@ using System;
 using System.Collections.Generic;
 using System.Diagnostics;
 using System.IO;
+using System.Text;
 using System.Text.RegularExpressions;
 using System.Threading;
 using System.Reflection;
@@ -36,9 +37,14 @@ namespace AutomationScripts
 		private static Thread ClientLogReaderThread = null;
 
 		/// <summary>
-		/// Process for the server, can be set by the cook command when a cook on the fly server is used
+		/// Process for the cook server, can be set by the cook command when a cook on the fly server is used
 		/// </summary>
-		public static IProcessResult ServerProcess;
+		public static IProcessResult CookServerProcess;
+
+		/// <summary>
+		/// Process for the dedicated server
+		/// </summary>
+		private static IProcessResult DedicatedServerProcess;
 
 		// debug commands for the engine to crash
 		public static string[] CrashCommands =
@@ -84,6 +90,57 @@ namespace AutomationScripts
 			}
 		}
 
+		private static bool WaitForProcessReady(IProcessResult Process, string Name, string LogFile, string[] ReadyTexts = null)
+		{
+			while (!FileExists(LogFile) && !Process.HasExited)
+			{
+				LogInformation("Waiting for {0} logging process to start...", Name);
+				Thread.Sleep(2000);
+			}
+
+			if (!FileExists(LogFile))
+			{
+				throw new AutomationException("{0} exited without creating a log file.", Name);
+			}
+
+			LogInformation("Logging started for {0} at: {1}", Name, LogFile);
+			Thread.Sleep(1000);
+
+			if (ReadyTexts == null)
+			{
+				LogInformation("{0} is ready!", Name);
+				return true;
+			}
+
+			using (FileStream ProcessLog = File.Open(LogFile, FileMode.Open, FileAccess.Read, FileShare.ReadWrite))
+			{
+				StreamReader LogReader = new StreamReader(ProcessLog);
+
+				// Read until the process has exited or text has been found
+				while (!Process.HasExited)
+				{
+					while (!LogReader.EndOfStream)
+					{
+						string Output = LogReader.ReadToEnd();
+						if (!String.IsNullOrEmpty(Output))
+						{
+							foreach (string ReadyText in ReadyTexts)
+							{
+								if (Output.Contains(ReadyText))
+								{
+									LogInformation("{1} is ready! \"{0}\" was found in log.", ReadyText, Name);
+									return true;
+								}
+							}
+						}
+					}
+					LogInformation("Waiting for {0} to get ready...", Name);
+					Thread.Sleep(2000);
+				}
+			}
+			throw new AutomationException("{0} exited before we asked it to.", Name);
+		}
+
 		public static void Run(ProjectParams Params)
 		{
 			Params.ValidateAndLog();
@@ -95,18 +152,19 @@ namespace AutomationScripts
 			LogInformation("********** RUN COMMAND STARTED **********");
 
 			var LogFolderOutsideOfSandbox = GetLogFolderOutsideOfSandbox();
-			if (!Unreal.IsEngineInstalled() && ServerProcess == null)
+			if (!Unreal.IsEngineInstalled() && CookServerProcess == null)
 			{
 				// In the installed runs, this is the same folder as CmdEnv.LogFolder so delete only in not-installed
 				DeleteDirectory(LogFolderOutsideOfSandbox);
 				CreateDirectory(LogFolderOutsideOfSandbox);
 			}
-			var ServerLogFile = CombinePaths(LogFolderOutsideOfSandbox, "Server.log");
+			var CookServerLogFile = CombinePaths(LogFolderOutsideOfSandbox, "CookServer.log");
+			var DedicatedServerLogFile = CombinePaths(LogFolderOutsideOfSandbox, "DedicatedServer.log");
 			var ClientLogFile = CombinePaths(LogFolderOutsideOfSandbox, Params.EditorTest ? "Editor.log" : "Client.log");
 
 			try
 			{
-				RunInternal(Params, ServerLogFile, ClientLogFile);
+				RunInternal(Params, CookServerLogFile, DedicatedServerLogFile, ClientLogFile);
 			}
 			catch
 			{
@@ -114,83 +172,99 @@ namespace AutomationScripts
 			}
 			finally
 			{
+				if (!GlobalCommandLine.NoKill)
+				{
+					if (CookServerProcess != null && !CookServerProcess.HasExited)
+					{
+						LogInformation("Stopping cook server...");
+						CookServerProcess.StopProcess();
+					}
+					if (DedicatedServerProcess != null && !DedicatedServerProcess.HasExited)
+					{
+						LogInformation("Stopping dedicated server...");
+						DedicatedServerProcess.StopProcess();
+					}
+				}
 				CopyLogsBackToLogFolder();
 			}
 
 			LogInformation("********** RUN COMMAND COMPLETED **********");
 		}
 
-		private static void RunInternal(ProjectParams Params, string ServerLogFile, string ClientLogFile)
+		private static void RunInternal(
+				ProjectParams Params,
+				string CookServerLogFile,
+				string DedicatedServerLogFile,
+				string ClientLogFile)
 		{
-			// Start the UnrealTrace
 			StartUnrealTrace();
 
-			// Setup server process if required.
+			if (CookServerProcess != null)
+			{
+				WaitForProcessReady(CookServerProcess, "Cook server", CookServerLogFile,
+						new string[] {"Unreal Network File Server is ready", "COTF server is ready"});
+			}
+
 			if (Params.DedicatedServer && !Params.SkipServer)
 			{
-				if (Params.ServerTargetPlatforms.Count > 0)
+				// With dedicated server, the client connects to local host to load a map,
+				// unless client parameters are already specified
+				if (String.IsNullOrEmpty(Params.ClientCommandline))
 				{
-					TargetPlatformDescriptor ServerPlatformDesc = Params.ServerTargetPlatforms[0];
-					ServerProcess = RunDedicatedServer(Params, ServerLogFile, Params.RunCommandline);
-
-					// With dedicated server, the client connects to local host to load a map, unless client parameters are already specified
-					if (String.IsNullOrEmpty(Params.ClientCommandline))
+					if (!String.IsNullOrEmpty(Params.ServerDeviceAddress))
 					{
-						if (!String.IsNullOrEmpty(Params.ServerDeviceAddress))
-						{
-							Params.ClientCommandline = Params.ServerDeviceAddress;
-						}
-						else
-						{
-							Params.ClientCommandline = "127.0.0.1";
-						}
+						Params.ClientCommandline = Params.ServerDeviceAddress;
+					}
+					else
+					{
+						Params.ClientCommandline = "127.0.0.1";
 					}
 				}
-				else
-				{
-					throw new AutomationException("Failed to run, server target platform not specified");
-				}
-			}
 
-			if (ServerProcess != null)
-			{
-				LogInformation("Waiting a few seconds for the server to start...");
-				Thread.Sleep(5000);
+				DedicatedServerProcess = RunDedicatedServer(Params, DedicatedServerLogFile, Params.RunCommandline);
+				WaitForProcessReady(DedicatedServerProcess, "Dedicated server", DedicatedServerLogFile,
+						new string[] {"Game Engine Initialized"});
 			}
 
 			if (!Params.NoClient)
 			{
-				LogInformation("Starting Client....");
-
 				var SC = CreateDeploymentContext(Params, false);
-
 				ERunOptions ClientRunFlags;
 				string ClientApp;
 				string ClientCmdLine;
 				SetupClientParams(SC, Params, ClientLogFile, out ClientRunFlags, out ClientApp, out ClientCmdLine);
-
-				// Run the client.
-				if (ServerProcess != null)
-				{
-					RunClientWithServer(SC, ServerLogFile, ServerProcess, ClientApp, ClientCmdLine, ClientRunFlags, ClientLogFile, Params);
-				}
-				else
-				{
-					RunStandaloneClient(SC, ClientLogFile, ClientRunFlags, ClientApp, ClientCmdLine, Params);
-				}
+				RunClient(SC, ClientLogFile, ClientRunFlags, ClientApp, ClientCmdLine, Params);
 			}
 		}
 
-		private static void RunStandaloneClient(List<DeploymentContext> DeployContextList, string ClientLogFile, ERunOptions ClientRunFlags, string ClientApp, string ClientCmdLine, ProjectParams Params)
+		private static void RunClient(
+				List<DeploymentContext> DeployContextList,
+				string ClientLogFile,
+				ERunOptions ClientRunFlags,
+				string ClientApp,
+				string ClientCmdLine,
+				ProjectParams Params)
 		{
+			var ExtraClients = new List<IProcessResult>();
+			int NumClients = Params.NumClients;
+			int RunTimeoutSeconds = Params.RunTimeoutSeconds;
+			IProcessResult ClientProcess = null;
+			var SC = DeployContextList[0];
+			bool bTestExitTextFound = false;
+			ERunOptions ExtraClientRunFlags = ClientRunFlags | ERunOptions.NoStdOutRedirect;
+
+			DateTime ClientStartTime = DateTime.UtcNow;
+
 			if (Params.Unattended)
 			{
 				string LookFor = "Bringing up level for play took";
 				bool bCommandlet = false;
+				bool bAutomation = false;
 
 				if (Params.RunAutomationTest != "" || Params.RunAutomationTests)
 				{
 					LookFor = "Automation Test Queue Empty";
+					bAutomation = true;
 				}
 				else if (Params.EditorTest)
 				{
@@ -202,391 +276,236 @@ namespace AutomationScripts
 					LookFor = "Game engine shut down";
 					bCommandlet = true;
 				}
-
+				else if (Params.DedicatedServer)
 				{
+					LookFor = "Welcomed by server";
+				}
+				ClientCmdLine += "-testexit=\"" + LookFor + "\"";
+
+				string AllClientOutput = "";
+				int AllClientOutputLength = 0;
+				int LastAutoFailIndex = -1;
 
-					string AllClientOutput = "";
-					int LastAutoFailIndex = -1;
-					IProcessResult ClientProcess = null;
-					FileStream ClientProcessLog = null;
-					StreamReader ClientLogReader = null;
-					LogInformation("Starting Client for unattended test....");
-					ClientProcess = Run(ClientApp, ClientCmdLine + " -testexit=\"" + LookFor + "\"", null, ClientRunFlags | ERunOptions.NoWaitForExit);
-					while (!FileExists(ClientLogFile) && !ClientProcess.HasExited)
+				LogInformation("Starting Client for unattended test....");
+				ClientProcess = SC.StageTargetPlatform.RunClient(ClientRunFlags, ClientApp, ClientCmdLine, Params);
+
+				if (DedicatedServerProcess != null)
+				{
+					if (NumClients > 1 && NumClients < 9)
 					{
-						LogInformation("Waiting for client logging process to start...{0}", ClientLogFile);
-						Thread.Sleep(2000);
+						for (int i = 1; i < NumClients; i++)
+						{
+							LogInformation("Starting Extra Client {0} for unattended test....", i);
+							ExtraClients.Add(SC.StageTargetPlatform.RunClient(ExtraClientRunFlags, ClientApp, ClientCmdLine, Params));
+						}
 					}
-					if (FileExists(ClientLogFile))
+				}
+
+				bool bKeepReading = ClientProcess != null;
+				while (bKeepReading)
+				{
+					Thread.Sleep(100);
+
+					if (RunTimeoutSeconds > 0)
 					{
-						Thread.Sleep(2000);
-						LogInformation("Client logging process started...{0}", ClientLogFile);
-						ClientProcessLog = File.Open(ClientLogFile, FileMode.Open, FileAccess.Read, FileShare.ReadWrite);
-						ClientLogReader = new StreamReader(ClientProcessLog);
+						if ((DateTime.UtcNow - ClientStartTime).TotalSeconds > RunTimeoutSeconds)
+						{
+							LogInformation("The run timed out after {0} seconds. Stopping client...", RunTimeoutSeconds);
+							ClientProcess.StopProcess();
+						}
 					}
-					if (ClientLogReader == null)
+					if (ClientProcess.HasExited)
 					{
-						throw new AutomationException("Client exited without creating a log file.");
+						LogInformation("Client exited, waiting for stdout...");
+						ClientProcess.WaitForExit();
+						LogInformation("Client exited, logging done!");
+						bKeepReading = false;
 					}
-					bool bKeepReading = true;
-					bool WelcomedCorrectly = false;
-					bool bClientExited = false;
-					DateTime ExitTime = DateTime.UtcNow;
 
-					while (bKeepReading)
+					AllClientOutput = ClientProcess.Output;
+					if (AllClientOutput.Length > AllClientOutputLength)
 					{
-						if (!bClientExited && ClientProcess.HasExited)
-						{
-							ExitTime = DateTime.UtcNow;
-							bClientExited = true;
-						}
-						string ClientOutput = ClientLogReader.ReadToEnd();
-						if (!String.IsNullOrEmpty(ClientOutput))
-						{
-							if (bClientExited)
-							{
-								ExitTime = DateTime.UtcNow; // as long as it is spewing, we reset the timer
-							}
-							AllClientOutput += ClientOutput;
-							Console.Write(ClientOutput);
+						int StartIndex = AllClientOutputLength;
+						AllClientOutputLength = AllClientOutput.Length;
 
-							if (AllClientOutput.LastIndexOf(LookFor) > AllClientOutput.IndexOf(LookFor))
+						// look for the test exit phrase, but ignore any output of the actual commandline string
+						// which can look like either -testexit=Bringing, -testexit="Bringing, -testexit=\"Bringing
+						if (AllClientOutput.IndexOf(LookFor, StartIndex) > 0 &&
+							AllClientOutput.IndexOf("-testexit=", Math.Max(0, StartIndex - 12)) < 0)
+						{
+							if (DedicatedServerProcess != null)
 							{
-								WelcomedCorrectly = true;
-								LogInformation("Test complete...");
-								bKeepReading = false;
+								LogInformation("Welcomed by server or client loaded");
 							}
-							else if (Params.RunAutomationTests)
+							LogInformation("Test complete");
+							LogInformation("**** UNATTENDED TEST COMPLETE: {0:0.00} seconds ****",
+								(DateTime.UtcNow - ClientStartTime).TotalMilliseconds / 1000);
+							bTestExitTextFound = true;
+							bKeepReading = false;
+						}
+						if (bAutomation)
+						{
+							int FailIndex = AllClientOutput.LastIndexOf("Automation Test Failed");
+							int ParenIndex = AllClientOutput.LastIndexOf(")");
+							if (FailIndex >= 0 && ParenIndex > FailIndex && FailIndex > LastAutoFailIndex)
 							{
-								int FailIndex = AllClientOutput.LastIndexOf("Automation Test Failed");
-								int ParenIndex = AllClientOutput.LastIndexOf(")");
-								if (FailIndex >= 0 && ParenIndex > FailIndex && FailIndex > LastAutoFailIndex)
+								string Tail = AllClientOutput.Substring(FailIndex);
+								int CloseParenIndex = Tail.IndexOf(")");
+								int OpenParenIndex = Tail.IndexOf("(");
+								string Test = "";
+								if (OpenParenIndex >= 0 && CloseParenIndex > OpenParenIndex)
 								{
-									string Tail = AllClientOutput.Substring(FailIndex);
-									int CloseParenIndex = Tail.IndexOf(")");
-									int OpenParenIndex = Tail.IndexOf("(");
-									string Test = "";
-									if (OpenParenIndex >= 0 && CloseParenIndex > OpenParenIndex)
-									{
-										Test = Tail.Substring(OpenParenIndex + 1, CloseParenIndex - OpenParenIndex - 1);
-										LogError("Automated test failed ({0}).", Test);
-										LastAutoFailIndex = FailIndex;
-									}
+									Test = Tail.Substring(OpenParenIndex + 1, CloseParenIndex - OpenParenIndex - 1);
+									LogError("Automated test failed ({0}).", Test);
+									LastAutoFailIndex = FailIndex;
 								}
 							}
-							// Detect commandlet failure
-							else if (bCommandlet)
-							{
-								const string ResultLog = "Commandlet->Main return this error code: ";
+						}
+						// Detect commandlet failure
+						else if (bCommandlet)
+						{
+							const string ResultLog = "Commandlet->Main return this error code: ";
 
-								int ResultStart = AllClientOutput.LastIndexOf(ResultLog);
-								int ResultValIdx = ResultStart + ResultLog.Length;
+							int ResultStart = AllClientOutput.LastIndexOf(ResultLog);
+							int ResultValIdx = ResultStart + ResultLog.Length;
 
-								if (ResultStart >= 0 && ResultValIdx < AllClientOutput.Length &&
+							if (ResultStart >= 0 && ResultValIdx < AllClientOutput.Length &&
 									AllClientOutput.Substring(ResultValIdx, 1) == "1")
-								{
-									// Parse the full commandlet warning/error summary
-									string FullSummary = "";
-									int SummaryStart = AllClientOutput.LastIndexOf("Warning/Error Summary");
-
-									if (SummaryStart >= 0 && SummaryStart < ResultStart)
-									{
-										FullSummary = AllClientOutput.Substring(SummaryStart, ResultStart - SummaryStart);
-									}
+							{
+								// Parse the full commandlet warning/error summary
+								string FullSummary = "";
+								int SummaryStart = AllClientOutput.LastIndexOf("Warning/Error Summary");
 
+								if (SummaryStart >= 0 && SummaryStart < ResultStart)
+								{
+									FullSummary = AllClientOutput.Substring(SummaryStart, ResultStart - SummaryStart);
+								}
 
-									if (FullSummary.Length > 0)
-									{
-										LogError("Commandlet failed, summary:" + Environment.NewLine +
-																						FullSummary);
-									}
-									else
-									{
-										LogError("Commandlet failed.");
-									}
+								if (FullSummary.Length > 0)
+								{
+									LogError("Commandlet failed, summary:" + Environment.NewLine + FullSummary);
+								}
+								else
+								{
+									LogError("Commandlet failed.");
 								}
 							}
 						}
-						else if (bClientExited && (DateTime.UtcNow - ExitTime).TotalSeconds > 30)
-						{
-							LogInformation("Client exited and has been quiet for 30 seconds...exiting");
-							bKeepReading = false;
-						}
 					}
-					if (ClientProcess != null && !ClientProcess.HasExited)
+					if (DedicatedServerProcess != null && DedicatedServerProcess.HasExited)
 					{
-						LogInformation("Client is supposed to exit, lets wait a while for it to exit naturally...");
-						for (int i = 0; i < 120 && !ClientProcess.HasExited; i++)
-						{
-							Thread.Sleep(1000);
-						}
+						LogInformation("Dedicated server exited, stopping client...");
+						ClientProcess.StopProcess();
 					}
-					if (ClientProcess != null && !ClientProcess.HasExited)
+					else if (CookServerProcess != null && CookServerProcess.HasExited)
 					{
-						LogInformation("Stopping client...");
+						LogInformation("Cook server exited, stopping client...");
 						ClientProcess.StopProcess();
-						Thread.Sleep(10000);
 					}
-					while (!ClientLogReader.EndOfStream)
+				}
+
+				if (ClientProcess != null && !ClientProcess.HasExited)
+				{
+					LogInformation("Client is supposed to exit, lets wait 20 seconds for it to exit naturally...");
+					for (int i = 0; i < 20 && !ClientProcess.HasExited; i++)
 					{
-						string ClientOutput = ClientLogReader.ReadToEnd();
-						if (!String.IsNullOrEmpty(ClientOutput))
-						{
-							Console.Write(ClientOutput);
-						}
+						Thread.Sleep(1000);
 					}
-
-					if (!WelcomedCorrectly)
+					if (!ClientProcess.HasExited)
 					{
-						throw new AutomationException("Client exited before we asked it to.");
+						LogInformation("Stopping client...");
+						ClientProcess.StopProcess();
 					}
 				}
 			}
 			else
 			{
-				var SC = DeployContextList[0];
-				IProcessResult ClientProcess = SC.StageTargetPlatform.RunClient(ClientRunFlags, ClientApp, ClientCmdLine, Params);
-				if (ClientProcess != null)
-				{
-					// If the client runs without StdOut redirect we're going to read the log output directly from log file on
-					// a separate thread.
-					if ((ClientRunFlags & ERunOptions.NoStdOutRedirect) == ERunOptions.NoStdOutRedirect)
-					{
-						ClientLogReaderThread = new System.Threading.Thread(ClientLogReaderProc);
-						ClientLogReaderThread.Start(new object[] { ClientLogFile, ClientProcess });
-					}
-
-					do
-					{
-						Thread.Sleep(100);
-					}
-					while (ClientProcess.HasExited == false);
-
-					SC.StageTargetPlatform.PostRunClient(ClientProcess, Params);
+				LogInformation("Starting Client....");
+				ClientProcess = SC.StageTargetPlatform.RunClient(ClientRunFlags, ClientApp, ClientCmdLine, Params);
 
-					// any non-zero exit code should propagate an exception. The Virtual function above may have
-					// already thrown a more specific exception or given a more specific ErrorCode, but this catches the rest.
-					if (ClientProcess.ExitCode != 0)
+				if (DedicatedServerProcess != null)
+				{
+					if (NumClients > 1 && NumClients < 9)
 					{
-						throw new AutomationException("Client exited with error code: " + ClientProcess.ExitCode);
+						for (int i = 1; i < NumClients; i++)
+						{
+							LogInformation("Starting Extra Client {0}....", i);
+							ExtraClients.Add(SC.StageTargetPlatform.RunClient(ExtraClientRunFlags, ClientApp, ClientCmdLine, Params));
+						}
 					}
 				}
-			}
-		}
-
-		private static void RunClientWithServer(List<DeploymentContext> DeployContextList, string ServerLogFile, IProcessResult ServerProcess, string ClientApp, string ClientCmdLine, ERunOptions ClientRunFlags, string ClientLogFile, ProjectParams Params)
-		{
-			IProcessResult ClientProcess = null;
-			var OtherClients = new List<IProcessResult>();
-
-			bool WelcomedCorrectly = false;
-			int NumClients = Params.NumClients;
-			string AllClientOutput = "";
-			int LastAutoFailIndex = -1;
 
-			if (Params.Unattended)
-			{
-				string LookFor = "Bringing up level for play took";
-				if (Params.DedicatedServer)
-				{
-					LookFor = "Welcomed by server";
-				}
-				else if (Params.RunAutomationTest != "" || Params.RunAutomationTests)
-				{
-					LookFor = "Automation Test Queue Empty";
-				}
+				if (ClientProcess != null)
 				{
-					while (!FileExists(ServerLogFile) && !ServerProcess.HasExited)
+					// If the client runs with LogWindow (without StdOut redirect),
+					// then fetch output from log file on a separate thread.
+					if (SC.StageTargetPlatform.UseAbsLog)
 					{
-						LogInformation("Waiting for logging process to start...");
-						Thread.Sleep(2000);
+						if ((ClientRunFlags & ERunOptions.NoStdOutRedirect) == ERunOptions.NoStdOutRedirect)
+						{
+							ClientLogReaderThread = new System.Threading.Thread(ClientLogReaderProc);
+							ClientLogReaderThread.Start(new object[] { ClientLogFile, ClientProcess });
+						}
 					}
-					Thread.Sleep(1000);
 
-					string AllServerOutput = "";
-					using (FileStream ProcessLog = File.Open(ServerLogFile, FileMode.Open, FileAccess.Read, FileShare.ReadWrite))
+					do
 					{
-						StreamReader LogReader = new StreamReader(ProcessLog);
-						bool bKeepReading = true;
-
-						FileStream ClientProcessLog = null;
-						StreamReader ClientLogReader = null;
-
-						// Read until the process has exited.
-						while (!ServerProcess.HasExited && bKeepReading)
+						Thread.Sleep(100);
+						if (RunTimeoutSeconds > 0)
 						{
-							while (!LogReader.EndOfStream && bKeepReading && ClientProcess == null)
-							{
-								string Output = LogReader.ReadToEnd();
-								if (!String.IsNullOrEmpty(Output))
-								{
-									AllServerOutput += Output;
-									if (ClientProcess == null &&
-											(AllServerOutput.Contains("Game Engine Initialized") || AllServerOutput.Contains("Unreal Network File Server is ready")))
-									{
-										LogInformation("Starting Client for unattended test....");
-										ClientProcess = Run(ClientApp, ClientCmdLine + " -testexit=\"" + LookFor + "\"", null, ClientRunFlags | ERunOptions.NoWaitForExit);
-										//@todo no testing is done on these
-										if (NumClients > 1 && NumClients < 9)
-										{
-											for (int i = 1; i < NumClients; i++)
-											{
-												LogInformation("Starting Extra Client....");
-												OtherClients.Add(Run(ClientApp, ClientCmdLine, null, ClientRunFlags | ERunOptions.NoWaitForExit));
-											}
-										}
-										while (!FileExists(ClientLogFile) && !ClientProcess.HasExited)
-										{
-											LogInformation("Waiting for client logging process to start...{0}", ClientLogFile);
-											Thread.Sleep(2000);
-										}
-										if (!ClientProcess.HasExited)
-										{
-											Thread.Sleep(2000);
-											LogInformation("Client logging process started...{0}", ClientLogFile);
-											ClientProcessLog = File.Open(ClientLogFile, FileMode.Open, FileAccess.Read, FileShare.ReadWrite);
-											ClientLogReader = new StreamReader(ClientProcessLog);
-										}
-									}
-									else if (ClientProcess == null && !ServerProcess.HasExited)
-									{
-										LogInformation("Waiting for server to start....");
-										Thread.Sleep(2000);
-									}
-									if (ClientProcess != null && ClientProcess.HasExited)
-									{
-										ServerProcess.StopProcess();
-										throw new AutomationException("Client exited before we asked it to.");
-									}
-								}
-							}
-							if (ClientLogReader != null)
+							if ((DateTime.UtcNow - ClientStartTime).TotalSeconds > RunTimeoutSeconds)
 							{
-								if (ClientProcess.HasExited)
-								{
-									ServerProcess.StopProcess();
-									throw new AutomationException("Client exited or closed the log before we asked it to.");
-								}
-								while (!ClientProcess.HasExited && !ServerProcess.HasExited && bKeepReading)
-								{
-									while (!ClientLogReader.EndOfStream && bKeepReading && !ServerProcess.HasExited && !ClientProcess.HasExited)
-									{
-										string ClientOutput = ClientLogReader.ReadToEnd();
-										if (!String.IsNullOrEmpty(ClientOutput))
-										{
-											AllClientOutput += ClientOutput;
-											Console.Write(ClientOutput);
-
-											if (AllClientOutput.LastIndexOf(LookFor) > AllClientOutput.IndexOf(LookFor))
-											{
-												if (Params.FakeClient)
-												{
-													LogInformation("Welcomed by server or client loaded, lets wait ten minutes...");
-													Thread.Sleep(60000 * 10);
-												}
-												else
-												{
-													LogInformation("Welcomed by server or client loaded, lets wait 30 seconds...");
-													Thread.Sleep(30000);
-												}
-												WelcomedCorrectly = true;
-												bKeepReading = false;
-											}
-											else if (Params.RunAutomationTests)
-											{
-												int FailIndex = AllClientOutput.LastIndexOf("Automation Test Failed");
-												int ParenIndex = AllClientOutput.LastIndexOf(")");
-												if (FailIndex >= 0 && ParenIndex > FailIndex && FailIndex > LastAutoFailIndex)
-												{
-													string Tail = AllClientOutput.Substring(FailIndex);
-													int CloseParenIndex = Tail.IndexOf(")");
-													int OpenParenIndex = Tail.IndexOf("(");
-													string Test = "";
-													if (OpenParenIndex >= 0 && CloseParenIndex > OpenParenIndex)
-													{
-														Test = Tail.Substring(OpenParenIndex + 1, CloseParenIndex - OpenParenIndex - 1);
-														LogError("Automated test failed ({0}).", Test);
-														LastAutoFailIndex = FailIndex;
-													}
-												}
-											}
-										}
-									}
-								}
+								LogInformation("The run timed out after {0} seconds. Stopping client...", RunTimeoutSeconds);
+								ClientProcess.StopProcess();
 							}
 						}
+						if (DedicatedServerProcess != null && DedicatedServerProcess.HasExited)
+						{
+							LogInformation("Dedicated server exited, stopping client...");
+							ClientProcess.StopProcess();
+						}
+						else if (CookServerProcess != null && CookServerProcess.HasExited)
+						{
+							LogInformation("Cook server exited, stopping client...");
+							ClientProcess.StopProcess();
+						}
 					}
+					while (ClientProcess.HasExited == false);
 				}
 			}
-			else
+
+			if (ExtraClients.Count > 0)
 			{
-				LogFileReaderProcess(ServerLogFile, ServerProcess, (string Output) =>
+				LogInformation("Client exited, stopping extra clients...");
+				foreach (var OtherClient in ExtraClients)
 				{
-					bool bKeepReading = true;
-					if (ClientProcess == null && !String.IsNullOrEmpty(Output))
+					if (OtherClient != null && !OtherClient.HasExited)
 					{
-						AllClientOutput += Output;
-						if (AllClientOutput.Contains("Game Engine Initialized") || AllClientOutput.Contains("Unreal Network File Server is ready") || AllClientOutput.Contains("COTF server is ready"))
-						{
-							LogInformation("Starting Client....");
-							var SC = DeployContextList[0];
-							ClientProcess = SC.StageTargetPlatform.RunClient(ClientRunFlags | ERunOptions.NoWaitForExit, ClientApp, ClientCmdLine, Params);
-						//						ClientProcess = Run(ClientApp, ClientCmdLine, null, ClientRunFlags | ERunOptions.NoWaitForExit);
-						if (NumClients > 1 && NumClients < 9)
-							{
-								for (int i = 1; i < NumClients; i++)
-								{
-									LogInformation("Starting Extra Client....");
-									IProcessResult NewClient = SC.StageTargetPlatform.RunClient(ClientRunFlags | ERunOptions.NoWaitForExit, ClientApp, ClientCmdLine, Params);
-									OtherClients.Add(NewClient);
-								}
-							}
-						}
-					}
-					else if (ClientProcess == null && !ServerProcess.HasExited)
-					{
-						LogInformation("Waiting for server to start....");
-						Thread.Sleep(2000);
-					}
-
-					if (String.IsNullOrEmpty(Output) == false)
-					{
-						Console.Write(Output);
+						OtherClient.StopProcess();
 					}
+				}
+			}
 
-					if (ClientProcess != null && ClientProcess.HasExited)
-					{
-
-						LogInformation("Client exited, stopping server....");
-						if (!GlobalCommandLine.NoKill)
-						{
-							ServerProcess.StopProcess();
-						}
-						bKeepReading = false;
-					}
+			SC.StageTargetPlatform.PostRunClient(ClientProcess, Params);
 
-					return bKeepReading; // Keep reading
-			});
-			}
-			LogInformation("Server exited....");
-			if (ClientProcess != null && !ClientProcess.HasExited)
-			{
-				ClientProcess.StopProcess();
-			}
-			foreach (var OtherClient in OtherClients)
+			if (Params.Unattended)
 			{
-				if (OtherClient != null && !OtherClient.HasExited)
+				// In unattended/-testexit mode we only throw if testexit text was not found
+				if (!bTestExitTextFound)
 				{
-					OtherClient.StopProcess();
+					throw new AutomationException("Client exited before we asked it to.");
 				}
 			}
-			if (Params.Unattended)
+			else
 			{
-				if (!WelcomedCorrectly)
+				// Any non-zero exit code should propagate an exception. The PostRunClient function above may have
+				// already thrown a more specific exception or given a more specific ErrorCode, but this catches the rest.
+				if (ClientProcess != null && ClientProcess.ExitCode != 0)
 				{
-					throw new AutomationException("Server or client exited before we asked it to.");
+					throw new AutomationException("Client exited with error code: " + ClientProcess.ExitCode);
 				}
 			}
+			LogInformation("Client exited with error code: " + ClientProcess.ExitCode);
 		}
 
 		private static void SetupClientParams(List<DeploymentContext> DeployContextList, ProjectParams Params, string ClientLogFile, out ERunOptions ClientRunFlags, out string ClientApp, out string ClientCmdLine)
@@ -596,21 +515,22 @@ namespace AutomationScripts
 				throw new AutomationException("No ClientTargetPlatform set for SetupClientParams.");
 			}
 
-			//		var DeployContextList = CreateDeploymentContext(Params, false);
-
 			if (DeployContextList.Count == 0)
 			{
 				throw new AutomationException("No DeployContextList for SetupClientParams.");
 			}
 
+			// set default output variables
+			ClientRunFlags = ERunOptions.Default | ERunOptions.NoWaitForExit;
+			ClientApp = "";
+			ClientCmdLine = "";
+
+			var TargetPlatform = Params.ClientTargetPlatforms[0];
 			var SC = DeployContextList[0];
 
 			// Get client app name and command line.
-			ClientRunFlags = ERunOptions.AllowSpew | ERunOptions.AppMustExist;
-			ClientApp = "";
-			ClientCmdLine = "";
 			string TempCmdLine = SC.ProjectArgForCommandLines + " ";
-			var PlatformName = Params.ClientTargetPlatforms[0].ToString();
+			var PlatformName = TargetPlatform.ToString();
 			if (Params.Cook || Params.CookOnTheFly)
 			{
 				List<FileReference> Exes = SC.StageTargetPlatform.GetExecutableNames(SC);
@@ -627,18 +547,16 @@ namespace AutomationScripts
 
 				if (Params.CookOnTheFly || Params.FileServer)
 				{
-					string FileHostCommandline = GetFileHostCommandline(Params, SC);
-					TempCmdLine += FileHostCommandline + " ";
+					String FileHostCommandline = GetFileHostCommandline(Params, SC);
+					if (!string.IsNullOrEmpty(FileHostCommandline))
+					{
+						TempCmdLine += FileHostCommandline + " ";
+					}
 
 					if (Params.CookOnTheFlyStreaming)
 					{
 						TempCmdLine += "-streaming ";
 					}
-					else if (!Params.ZenStore && SC.StageTargetPlatform.PlatformType != UnrealTargetPlatform.IOS)
-					{
-						// per josh, allowcaching is deprecated/doesn't make sense for iOS.
-						TempCmdLine += "-allowcaching ";
-					}
 				}
 				else if (Params.UsePak(SC.StageTargetPlatform))
 				{
@@ -677,21 +595,12 @@ namespace AutomationScripts
 					TempCmdLine += "-ddc=" + Params.DDCGraph + " ";
 				}
 			}
-			if (Params.LogWindow)
-			{
-				// Without NoStdOutRedirect '-log' doesn't log anything to the window
-				ClientRunFlags |= ERunOptions.NoStdOutRedirect;
-				TempCmdLine += "-log ";
-			}
-			else
-			{
-				TempCmdLine += "-stdout ";
-			}
+
 			if (Params.Unattended)
 			{
 				TempCmdLine += "-unattended ";
 			}
-			if (IsBuildMachine || Params.Unattended)
+			if (IsBuildMachine)
 			{
 				TempCmdLine += "-buildmachine ";
 			}
@@ -707,24 +616,37 @@ namespace AutomationScripts
 			}
 			else if (Params.RunAutomationTest != "")
 			{
-				TempCmdLine += "-execcmds=\"automation list;runtests " + Params.RunAutomationTest + "\" ";
+				TempCmdLine += "-execcmds=\"automation runtests " + Params.RunAutomationTest + ";quit\" ";
 			}
 			else if (Params.RunAutomationTests)
 			{
-				TempCmdLine += "-execcmds=\"automation list;runall\" ";
+				TempCmdLine += "-execcmds=\"automation runall;quit;\" ";
 			}
 			if (SC.StageTargetPlatform.UseAbsLog)
 			{
 				TempCmdLine += "-abslog=" + CommandUtils.MakePathSafeToUseWithCommandLine(ClientLogFile) + " ";
 			}
-			if (SC.StageTargetPlatform.PlatformType == UnrealTargetPlatform.Win64)
+			if (SC.StageTargetPlatform.PlatformType == BuildHostPlatform.Current.Platform)
 			{
-				TempCmdLine += "-Messaging -Windowed ";
+				if (Params.LogWindow && !Params.Unattended)
+				{
+					// Without NoStdOutRedirect '-log' doesn't log anything to the window
+					ClientRunFlags |= ERunOptions.NoStdOutRedirect;
+					TempCmdLine += "-log ";
+				}
+				else
+				{
+					// unattended run logic depends on parsing stdout
+					TempCmdLine += "-stdout -AllowStdOutLogVerbosity ";
+				}
 			}
-			else
+			if (SC.StageTargetPlatform.PlatformType == UnrealTargetPlatform.Win64)
 			{
-				TempCmdLine += "-Messaging ";
+				TempCmdLine += "-Windowed ";
 			}
+
+			TempCmdLine += "-Messaging ";
+
 			if (Params.NullRHI && SC.StageTargetPlatform.PlatformType != UnrealTargetPlatform.Mac) // all macs have GPUs, and currently the mac dies with nullrhi
 			{
 				TempCmdLine += "-nullrhi ";
@@ -758,7 +680,7 @@ namespace AutomationScripts
 				}
 				ClientCmdLine += " ";
 				ClientCmdLine += "-Exe=\"" + ClientApp + "\" ";
-				ClientCmdLine += "-Targetplatform=" + Params.ClientTargetPlatforms[0].ToString() + " ";
+				ClientCmdLine += "-Targetplatform=" + PlatformName + " ";
 				ClientCmdLine += "-Params=\"" + TempCmdLine + "\"";
 				ClientApp = CombinePaths(CmdEnv.LocalRoot, "Engine/Binaries/Win64/UnrealFrontend.exe");
 
@@ -909,7 +831,7 @@ namespace AutomationScripts
 				FileHostParams += " ";
 			}
 
-			return FileHostParams;
+			return FileHostParams.Trim();
 		}
 
 		private static void ClientLogReaderProc(object ArgsContainer)
@@ -947,12 +869,13 @@ namespace AutomationScripts
 			var SC = DeployContextList[0];
 
 			var ServerApp = CombinePaths(CmdEnv.LocalRoot, "Engine/Binaries/Win64/UnrealEditor.exe");
-			if (ServerParams.Cook)
+			bool bCooked = ServerParams.Cook || ServerParams.CookOnTheFly;
+			if (bCooked)
 			{
 				List<FileReference> Exes = SC.StageTargetPlatform.GetExecutableNames(SC);
 				ServerApp = Exes[0].FullName;
 			}
-			var Args = ServerParams.Cook ? "" : (SC.ProjectArgForCommandLines + " ");
+			var Args = bCooked ? "" : (SC.ProjectArgForCommandLines + " ");
 			Console.WriteLine(Params.ServerDeviceAddress);
 			TargetPlatformDescriptor ServerPlatformDesc = ServerParams.ServerTargetPlatforms[0];
 			if (ServerParams.Cook && ServerPlatformDesc.Type == UnrealTargetPlatform.Linux && !String.IsNullOrEmpty(ServerParams.ServerDeviceAddress))
@@ -979,7 +902,7 @@ namespace AutomationScripts
 				{
 					Map += "?fake";
 				}
-				Args += String.Format("{0} -server -abslog={1}  -log -Messaging", Map, CommandUtils.MakePathSafeToUseWithCommandLine(ServerLogFile));
+				Args += String.Format("{0} -server -abslog={1} -log -Messaging", Map, CommandUtils.MakePathSafeToUseWithCommandLine(ServerLogFile));
 				if (Params.Unattended)
 				{
 					Args += " -unattended";
@@ -988,6 +911,12 @@ namespace AutomationScripts
 				if (Params.ServerCommandline.Length > 0)
 				{
 					Args += " " + Params.ServerCommandline;
+				} 
+
+				String FileHostCommandline = GetFileHostCommandline(Params, SC);
+				if (!string.IsNullOrEmpty(FileHostCommandline))
+				{
+					Args += " " + FileHostCommandline;
 				}
 			}
 
@@ -999,12 +928,12 @@ namespace AutomationScripts
 				}
 			}
 
-			if (Params.HasDDCGraph)
+			if (Params.HasDDCGraph && !bCooked)
 			{
 				Args += " -ddc=" + Params.DDCGraph;
 			}
 
-			if (IsBuildMachine || Params.Unattended)
+			if (IsBuildMachine)
 			{
 				Args += " -buildmachine";
 			}
@@ -1034,13 +963,14 @@ namespace AutomationScripts
 			}
 
 			PushDir(Path.GetDirectoryName(ServerApp));
-			var Result = Run(ServerApp, Args, null, ERunOptions.AllowSpew | ERunOptions.NoWaitForExit | ERunOptions.AppMustExist | ERunOptions.NoStdOutRedirect);
+			var Result = Run(ServerApp, Args, null, ERunOptions.Default | ERunOptions.NoWaitForExit | ERunOptions.NoStdOutRedirect);
 			PopDir();
 
+			LogInformation("Running DedicatedServer@Process:{0}@{1}", ServerApp, Result.ProcessObject.Id);
 			return Result;
 		}
 
-		private static IProcessResult RunCookOnTheFlyServer(FileReference ProjectName, string ServerLogFile, string TargetPlatform, string AdditionalCommandLine)
+		private static IProcessResult RunCookOnTheFlyServer(FileReference ProjectName, string ServerLogFile, string AdditionalCommandLine)
 		{
 			var ServerApp = HostPlatform.Current.GetUnrealExePath("UnrealEditor.exe");
 			var Args = String.Format("{0} -run=cook -cookonthefly -unattended -CrashForUAT -log",
@@ -1049,23 +979,25 @@ namespace AutomationScripts
 			{
 				// Issue with dotnet not allowing any files with an exclusive advisory lock to be opened for read-only or copied
 				// [REDACTED_PROFILE_URL]/runtime/issues/34126
-				if (TargetPlatform == "Linux")
+				if (HostPlatform.Current.HostEditorPlatform == UnrealBuildTool.UnrealTargetPlatform.Linux)
 				{
 					Args += " -noexclusivelockonwrite";
 				}
 
 				Args += " -abslog=" + CommandUtils.MakePathSafeToUseWithCommandLine(ServerLogFile);
 			}
-			if (IsBuildMachine)
+
+			if (!String.IsNullOrEmpty(AdditionalCommandLine))
 			{
-				Args += " -buildmachine";
+				Args += " " + AdditionalCommandLine;
 			}
-			Args += " " + AdditionalCommandLine;
 
 			// Run the server (Without NoStdOutRedirect -log doesn't log anything to the window)
 			PushDir(Path.GetDirectoryName(ServerApp));
-			var Result = Run(ServerApp, Args, null, ERunOptions.AllowSpew | ERunOptions.NoWaitForExit | ERunOptions.AppMustExist | ERunOptions.NoStdOutRedirect);
+			var Result = Run(ServerApp, Args, null, ERunOptions.Default | ERunOptions.NoWaitForExit | ERunOptions.NoStdOutRedirect);
 			PopDir();
+
+			LogInformation("Running CookServer@Process:{0}@{1}", ServerApp, Result.ProcessObject.Id);
 			return Result;
 		}
 
