BuildGraph: Fix compilation of C# projects with the bundled NET runtime.

diff --git a/Engine/Source/Programs/AutomationTool/BuildGraph/Tasks/CsCompileTask.cs b/Engine/Source/Programs/AutomationTool/BuildGraph/Tasks/CsCompileTask.cs
--- a/Engine/Source/Programs/AutomationTool/BuildGraph/Tasks/CsCompileTask.cs
+++ b/Engine/Source/Programs/AutomationTool/BuildGraph/Tasks/CsCompileTask.cs
@@ -149,6 +149,7 @@ namespace AutomationTool.Tasks
 			if(!Parameters.EnumerateOnly)
 			{
 				List<string> Arguments = new List<string>();
+				Arguments.Add("build");
 				foreach(KeyValuePair<string, string> PropertyPair in Properties)
 				{
 					Arguments.Add(String.Format("/property:{0}={1}", CommandUtils.MakePathSafeToUseWithCommandLine(PropertyPair.Key), CommandUtils.MakePathSafeToUseWithCommandLine(PropertyPair.Value)));
@@ -161,16 +162,12 @@ namespace AutomationTool.Tasks
 				{
 					Arguments.Add(String.Format("/target:{0}", CommandUtils.MakePathSafeToUseWithCommandLine(Parameters.Target)));
 				}
-				if(!CommandUtils.CmdEnv.FrameworkMsbuildPath.Equals("xbuild"))
-				{
-					// not supported by xbuild
-					Arguments.Add("/restore");
-				}
+				Arguments.Add("/restore");
 				Arguments.Add("/verbosity:minimal");
 				Arguments.Add("/nologo");
 				foreach(FileReference ProjectFile in ProjectFiles)
 				{
-					CommandUtils.MsBuild(CommandUtils.CmdEnv, ProjectFile.FullName, String.Join(" ", Arguments), null);
+					CommandUtils.RunAndLog(CommandUtils.CmdEnv, CommandUtils.CmdEnv.DotnetMsbuildPath, String.Join(" ", Arguments), null);
 				}
 			}
 
