CsCompile: Include project in command line

#rnx

diff --git a/Engine/Source/Programs/AutomationTool/BuildGraph/Tasks/CsCompileTask.cs b/Engine/Source/Programs/AutomationTool/BuildGraph/Tasks/CsCompileTask.cs
--- a/Engine/Source/Programs/AutomationTool/BuildGraph/Tasks/CsCompileTask.cs
+++ b/Engine/Source/Programs/AutomationTool/BuildGraph/Tasks/CsCompileTask.cs
@@ -165,9 +165,17 @@ namespace AutomationTool.Tasks
 				Arguments.Add("/restore");
 				Arguments.Add("/verbosity:minimal");
 				Arguments.Add("/nologo");
+
+				string JoinedArguments = String.Join(" ", Arguments);
+
 				foreach(FileReference ProjectFile in ProjectFiles)
 				{
-					CommandUtils.RunAndLog(CommandUtils.CmdEnv, CommandUtils.CmdEnv.DotnetMsbuildPath, String.Join(" ", Arguments), null);
+					if (!FileReference.Exists(ProjectFile))
+					{
+						throw new AutomationException("Project {0} does not exist!", ProjectFile);
+					}
+
+					CommandUtils.RunAndLog(CommandUtils.CmdEnv, CommandUtils.CmdEnv.DotnetMsbuildPath, $"{CommandUtils.MakePathSafeToUseWithCommandLine(ProjectFile.FullName)} {JoinedArguments}");
 				}
 			}
 
