CsCompile: Fix argument order (and use msbuild rather than build, consistent with elsewhere)

#rnx

diff --git a/Engine/Source/Programs/AutomationTool/BuildGraph/Tasks/CsCompileTask.cs b/Engine/Source/Programs/AutomationTool/BuildGraph/Tasks/CsCompileTask.cs
--- a/Engine/Source/Programs/AutomationTool/BuildGraph/Tasks/CsCompileTask.cs
+++ b/Engine/Source/Programs/AutomationTool/BuildGraph/Tasks/CsCompileTask.cs
@@ -149,7 +149,6 @@ namespace AutomationTool.Tasks
 			if(!Parameters.EnumerateOnly)
 			{
 				List<string> Arguments = new List<string>();
-				Arguments.Add("build");
 				foreach(KeyValuePair<string, string> PropertyPair in Properties)
 				{
 					Arguments.Add(String.Format("/property:{0}={1}", CommandUtils.MakePathSafeToUseWithCommandLine(PropertyPair.Key), CommandUtils.MakePathSafeToUseWithCommandLine(PropertyPair.Value)));
@@ -175,7 +174,7 @@ namespace AutomationTool.Tasks
 						throw new AutomationException("Project {0} does not exist!", ProjectFile);
 					}
 
-					CommandUtils.RunAndLog(CommandUtils.CmdEnv, CommandUtils.CmdEnv.DotnetMsbuildPath, $"{CommandUtils.MakePathSafeToUseWithCommandLine(ProjectFile.FullName)} {JoinedArguments}");
+					CommandUtils.RunAndLog(CommandUtils.CmdEnv, CommandUtils.CmdEnv.DotnetMsbuildPath, $"msbuild {CommandUtils.MakePathSafeToUseWithCommandLine(ProjectFile.FullName)} {JoinedArguments}");
 				}
 			}
 
