.net6.0 is not supported in VS2019. Suppress warnings about that when using that IDE (with this change, when building UnrealBuildTool & AutomationTool)

#rnx

diff --git a/Engine/Source/Programs/AutomationTool/Android/Android.Automation.csproj b/Engine/Source/Programs/AutomationTool/Android/Android.Automation.csproj
--- a/Engine/Source/Programs/AutomationTool/Android/Android.Automation.csproj
+++ b/Engine/Source/Programs/AutomationTool/Android/Android.Automation.csproj
@@ -1,4 +1,6 @@
 <Project Sdk="Microsoft.NET.Sdk">
+  <Import Project="..\..\Shared\UnrealEngine.csproj.props" />
+  
   <PropertyGroup>
     <TargetFramework>net6.0</TargetFramework>
     <Configuration Condition=" '$(Configuration)' == '' ">Development</Configuration>
diff --git a/Engine/Source/Programs/AutomationTool/AutomationTool.csproj b/Engine/Source/Programs/AutomationTool/AutomationTool.csproj
--- a/Engine/Source/Programs/AutomationTool/AutomationTool.csproj
+++ b/Engine/Source/Programs/AutomationTool/AutomationTool.csproj
@@ -1,4 +1,6 @@
 <Project Sdk="Microsoft.NET.Sdk">
+  <Import Project="..\Shared\UnrealEngine.csproj.props" />
+  
   <PropertyGroup>
     <IsWindows Condition="'$([System.Runtime.InteropServices.RuntimeInformation]::IsOSPlatform($([System.Runtime.InteropServices.OSPlatform]::Windows)))' == 'true'">true</IsWindows> 
     <IsOSX Condition="'$([System.Runtime.InteropServices.RuntimeInformation]::IsOSPlatform($([System.Runtime.InteropServices.OSPlatform]::OSX)))' == 'true'">true</IsOSX> 
@@ -8,7 +10,7 @@
   <Choose>
     <When Condition="'$(IsWindows)'=='true'">
       <PropertyGroup>
-        <TargetFramework>net6.0-windows</TargetFramework>
+        <TargetFramework>net6.0-windows10.0.19041</TargetFramework>
         <UseWindowsForms>true</UseWindowsForms>
       </PropertyGroup>
     </When>
@@ -71,6 +73,8 @@
     </When>
   </Choose>
 
+  <!-- This set of dependencies should be kept as small as possible
+       This list must match the list of locations searched for changed files in BuildUAT.bat and BuildUAT.sh -->
   <ItemGroup>
     <Compile Include="..\Shared\MetaData.cs">
       <Link>Properties\MetaData.cs</Link>
@@ -106,6 +110,8 @@
     <PackageReference Include="Microsoft.Build.Locator" Version="1.4.1" PrivateAssets="all" />
   </ItemGroup>
   
+  <!-- This set of dependencies should be kept as small as possible
+       This list must match the list of locations searched for changed files in BuildUAT.bat and BuildUAT.sh -->
   <ItemGroup>
     <ProjectReference Include="..\Shared\EpicGames.Core\EpicGames.Core.csproj" PrivateAssets="all" />
     <ProjectReference Include="..\Shared\EpicGames.Build\EpicGames.Build.csproj" PrivateAssets="all" />
diff --git a/Engine/Source/Programs/AutomationTool/AutomationUtils/AutomationUtils.Automation.csproj b/Engine/Source/Programs/AutomationTool/AutomationUtils/AutomationUtils.Automation.csproj
--- a/Engine/Source/Programs/AutomationTool/AutomationUtils/AutomationUtils.Automation.csproj
+++ b/Engine/Source/Programs/AutomationTool/AutomationUtils/AutomationUtils.Automation.csproj
@@ -1,4 +1,6 @@
 <Project Sdk="Microsoft.NET.Sdk">
+  <Import Project="..\..\Shared\UnrealEngine.csproj.props" />
+  
   <PropertyGroup>
     <TargetFramework>net6.0</TargetFramework>
     <Configuration Condition=" '$(Configuration)' == '' ">Development</Configuration>
diff --git a/Engine/Source/Programs/AutomationTool/BuildGraph/BuildGraph.Automation.csproj b/Engine/Source/Programs/AutomationTool/BuildGraph/BuildGraph.Automation.csproj
--- a/Engine/Source/Programs/AutomationTool/BuildGraph/BuildGraph.Automation.csproj
+++ b/Engine/Source/Programs/AutomationTool/BuildGraph/BuildGraph.Automation.csproj
@@ -1,4 +1,6 @@
 <Project Sdk="Microsoft.NET.Sdk">
+  <Import Project="..\..\Shared\UnrealEngine.csproj.props" />
+  
   <PropertyGroup>
     <TargetFramework>net6.0</TargetFramework>
     <Configuration Condition=" '$(Configuration)' == '' ">Development</Configuration>
diff --git a/Engine/Source/Programs/AutomationTool/CookedEditor/CookedEditor.Automation.csproj b/Engine/Source/Programs/AutomationTool/CookedEditor/CookedEditor.Automation.csproj
--- a/Engine/Source/Programs/AutomationTool/CookedEditor/CookedEditor.Automation.csproj
+++ b/Engine/Source/Programs/AutomationTool/CookedEditor/CookedEditor.Automation.csproj
@@ -1,4 +1,6 @@
 ﻿<Project Sdk="Microsoft.NET.Sdk">
+  <Import Project="..\..\Shared\UnrealEngine.csproj.props" />
+  
   <PropertyGroup>
     <TargetFramework>net6.0</TargetFramework>
     <Configuration Condition=" '$(Configuration)' == '' ">Development</Configuration>
diff --git a/Engine/Source/Programs/AutomationTool/CrowdinLocalization/CrowdinLocalization.Automation.csproj b/Engine/Source/Programs/AutomationTool/CrowdinLocalization/CrowdinLocalization.Automation.csproj
--- a/Engine/Source/Programs/AutomationTool/CrowdinLocalization/CrowdinLocalization.Automation.csproj
+++ b/Engine/Source/Programs/AutomationTool/CrowdinLocalization/CrowdinLocalization.Automation.csproj
@@ -1,4 +1,6 @@
 ﻿<Project Sdk="Microsoft.NET.Sdk">
+  <Import Project="..\..\Shared\UnrealEngine.csproj.props" />
+  
   <PropertyGroup>
     <TargetFramework>net6.0</TargetFramework>
     <Configuration Condition=" '$(Configuration)' == '' ">Development</Configuration>
@@ -19,7 +21,7 @@
   <PropertyGroup Condition=" '$(Configuration)|$(Platform)' == 'Development|AnyCPU' ">
     <DefineConstants>$(DefineConstants);TRACE</DefineConstants>
     <Optimize>true</Optimize>
-    <NoWarn>1701;1702;CS1998;CS0649</NoWarn>
+    <NoWarn>$(NoWarn);1701;1702;CS1998;CS0649</NoWarn>
   </PropertyGroup>
   <PropertyGroup Condition=" '$(Configuration)|$(Platform)' == 'Release|AnyCPU' ">
     <DefineConstants>$(DefineConstants);TRACE</DefineConstants>
diff --git a/Engine/Source/Programs/AutomationTool/Gauntlet/Gauntlet.Automation.csproj b/Engine/Source/Programs/AutomationTool/Gauntlet/Gauntlet.Automation.csproj
--- a/Engine/Source/Programs/AutomationTool/Gauntlet/Gauntlet.Automation.csproj
+++ b/Engine/Source/Programs/AutomationTool/Gauntlet/Gauntlet.Automation.csproj
@@ -1,4 +1,6 @@
 ﻿<Project Sdk="Microsoft.NET.Sdk">
+  <Import Project="..\..\Shared\UnrealEngine.csproj.props" />
+  
   <PropertyGroup>
     <TargetFramework>net6.0</TargetFramework>
     <Configuration Condition=" '$(Configuration)' == '' ">Development</Configuration>
diff --git a/Engine/Source/Programs/AutomationTool/HoloLens/HoloLens.Automation.csproj b/Engine/Source/Programs/AutomationTool/HoloLens/HoloLens.Automation.csproj
--- a/Engine/Source/Programs/AutomationTool/HoloLens/HoloLens.Automation.csproj
+++ b/Engine/Source/Programs/AutomationTool/HoloLens/HoloLens.Automation.csproj
@@ -1,4 +1,6 @@
 ﻿<Project Sdk="Microsoft.NET.Sdk">
+  <Import Project="..\..\Shared\UnrealEngine.csproj.props" />
+  
   <PropertyGroup>
     <TargetFramework>net6.0-windows10.0.19041.0</TargetFramework>
     <Configuration Condition=" '$(Configuration)' == '' ">Development</Configuration>
diff --git a/Engine/Source/Programs/AutomationTool/IOS/IOS.Automation.csproj b/Engine/Source/Programs/AutomationTool/IOS/IOS.Automation.csproj
--- a/Engine/Source/Programs/AutomationTool/IOS/IOS.Automation.csproj
+++ b/Engine/Source/Programs/AutomationTool/IOS/IOS.Automation.csproj
@@ -1,4 +1,6 @@
 <Project Sdk="Microsoft.NET.Sdk">
+  <Import Project="..\..\Shared\UnrealEngine.csproj.props" />
+  
   <PropertyGroup>
     <TargetFramework>net6.0</TargetFramework>
     <Configuration Condition=" '$(Configuration)' == '' ">Development</Configuration>
diff --git a/Engine/Source/Programs/AutomationTool/Linux/Linux.Automation.csproj b/Engine/Source/Programs/AutomationTool/Linux/Linux.Automation.csproj
--- a/Engine/Source/Programs/AutomationTool/Linux/Linux.Automation.csproj
+++ b/Engine/Source/Programs/AutomationTool/Linux/Linux.Automation.csproj
@@ -1,4 +1,6 @@
 <Project Sdk="Microsoft.NET.Sdk">
+  <Import Project="..\..\Shared\UnrealEngine.csproj.props" />
+  
   <PropertyGroup>
     <TargetFramework>net6.0</TargetFramework>
     <Configuration Condition=" '$(Configuration)' == '' ">Development</Configuration>
diff --git a/Engine/Source/Programs/AutomationTool/Localization/Localization.Automation.csproj b/Engine/Source/Programs/AutomationTool/Localization/Localization.Automation.csproj
--- a/Engine/Source/Programs/AutomationTool/Localization/Localization.Automation.csproj
+++ b/Engine/Source/Programs/AutomationTool/Localization/Localization.Automation.csproj
@@ -1,4 +1,6 @@
 <Project Sdk="Microsoft.NET.Sdk">
+  <Import Project="..\..\Shared\UnrealEngine.csproj.props" />
+  
   <PropertyGroup>
     <TargetFramework>net6.0</TargetFramework>
     <Configuration Condition=" '$(Configuration)' == '' ">Development</Configuration>
diff --git a/Engine/Source/Programs/AutomationTool/LowLevelTests/LowLevelTests.Automation.csproj b/Engine/Source/Programs/AutomationTool/LowLevelTests/LowLevelTests.Automation.csproj
--- a/Engine/Source/Programs/AutomationTool/LowLevelTests/LowLevelTests.Automation.csproj
+++ b/Engine/Source/Programs/AutomationTool/LowLevelTests/LowLevelTests.Automation.csproj
@@ -1,4 +1,6 @@
 <Project Sdk="Microsoft.NET.Sdk">
+  <Import Project="..\..\Shared\UnrealEngine.csproj.props" />
+  
   <PropertyGroup>
     <TargetFramework>net6.0</TargetFramework>
 	<Configuration Condition=" '$(Configuration)' == '' ">Development</Configuration>
diff --git a/Engine/Source/Programs/AutomationTool/Mac/Mac.Automation.csproj b/Engine/Source/Programs/AutomationTool/Mac/Mac.Automation.csproj
--- a/Engine/Source/Programs/AutomationTool/Mac/Mac.Automation.csproj
+++ b/Engine/Source/Programs/AutomationTool/Mac/Mac.Automation.csproj
@@ -1,4 +1,6 @@
 <Project Sdk="Microsoft.NET.Sdk">
+  <Import Project="..\..\Shared\UnrealEngine.csproj.props" />
+  
   <PropertyGroup>
     <TargetFramework>net6.0</TargetFramework>
     <Configuration Condition=" '$(Configuration)' == '' ">Development</Configuration>
diff --git a/Engine/Source/Programs/AutomationTool/OneSkyLocalization/OneSkyLocalization.Automation.csproj b/Engine/Source/Programs/AutomationTool/OneSkyLocalization/OneSkyLocalization.Automation.csproj
--- a/Engine/Source/Programs/AutomationTool/OneSkyLocalization/OneSkyLocalization.Automation.csproj
+++ b/Engine/Source/Programs/AutomationTool/OneSkyLocalization/OneSkyLocalization.Automation.csproj
@@ -1,4 +1,6 @@
 <Project Sdk="Microsoft.NET.Sdk">
+  <Import Project="..\..\Shared\UnrealEngine.csproj.props" />
+  
   <PropertyGroup>
     <TargetFramework>net6.0</TargetFramework>
     <Configuration Condition=" '$(Configuration)' == '' ">Development</Configuration>
diff --git a/Engine/Source/Programs/AutomationTool/Scripts/AutomationScripts.Automation.csproj b/Engine/Source/Programs/AutomationTool/Scripts/AutomationScripts.Automation.csproj
--- a/Engine/Source/Programs/AutomationTool/Scripts/AutomationScripts.Automation.csproj
+++ b/Engine/Source/Programs/AutomationTool/Scripts/AutomationScripts.Automation.csproj
@@ -1,4 +1,6 @@
 <Project Sdk="Microsoft.NET.Sdk">
+  <Import Project="..\..\Shared\UnrealEngine.csproj.props" />
+  
   <PropertyGroup>
     <TargetFramework>net6.0</TargetFramework>
     <Configuration Condition=" '$(Configuration)' == '' ">Development</Configuration>
diff --git a/Engine/Source/Programs/AutomationTool/SteamDeck/SteamDeck.Automation.csproj b/Engine/Source/Programs/AutomationTool/SteamDeck/SteamDeck.Automation.csproj
--- a/Engine/Source/Programs/AutomationTool/SteamDeck/SteamDeck.Automation.csproj
+++ b/Engine/Source/Programs/AutomationTool/SteamDeck/SteamDeck.Automation.csproj
@@ -1,4 +1,6 @@
 <Project Sdk="Microsoft.NET.Sdk">
+  <Import Project="..\..\Shared\UnrealEngine.csproj.props" />
+  
   <PropertyGroup>
     <TargetFramework>net6.0</TargetFramework>
     <Configuration Condition=" '$(Configuration)' == '' ">Development</Configuration>
diff --git a/Engine/Source/Programs/AutomationTool/TVOS/TVOS.Automation.csproj b/Engine/Source/Programs/AutomationTool/TVOS/TVOS.Automation.csproj
--- a/Engine/Source/Programs/AutomationTool/TVOS/TVOS.Automation.csproj
+++ b/Engine/Source/Programs/AutomationTool/TVOS/TVOS.Automation.csproj
@@ -1,4 +1,6 @@
 <Project Sdk="Microsoft.NET.Sdk">
+  <Import Project="..\..\Shared\UnrealEngine.csproj.props" />
+  
   <PropertyGroup>
     <TargetFramework>net6.0</TargetFramework>
     <Configuration Condition=" '$(Configuration)' == '' ">Development</Configuration>
diff --git a/Engine/Source/Programs/AutomationTool/Turnkey/Turnkey.Automation.csproj b/Engine/Source/Programs/AutomationTool/Turnkey/Turnkey.Automation.csproj
--- a/Engine/Source/Programs/AutomationTool/Turnkey/Turnkey.Automation.csproj
+++ b/Engine/Source/Programs/AutomationTool/Turnkey/Turnkey.Automation.csproj
@@ -1,4 +1,6 @@
 ﻿<Project Sdk="Microsoft.NET.Sdk">
+  <Import Project="..\..\Shared\UnrealEngine.csproj.props" />
+  
   <PropertyGroup>
     <TargetFramework>net6.0</TargetFramework>
     <Configuration Condition=" '$(Configuration)' == '' ">Development</Configuration>
diff --git a/Engine/Source/Programs/AutomationTool/Win/Win.Automation.csproj b/Engine/Source/Programs/AutomationTool/Win/Win.Automation.csproj
--- a/Engine/Source/Programs/AutomationTool/Win/Win.Automation.csproj
+++ b/Engine/Source/Programs/AutomationTool/Win/Win.Automation.csproj
@@ -1,4 +1,6 @@
 <Project Sdk="Microsoft.NET.Sdk">
+  <Import Project="..\..\Shared\UnrealEngine.csproj.props" />
+  
   <PropertyGroup>
     <TargetFramework>net6.0</TargetFramework>
     <Configuration Condition=" '$(Configuration)' == '' ">Development</Configuration>
diff --git a/Engine/Source/Programs/AutomationTool/XLocLocalization/XLocLocalization.Automation.csproj b/Engine/Source/Programs/AutomationTool/XLocLocalization/XLocLocalization.Automation.csproj
--- a/Engine/Source/Programs/AutomationTool/XLocLocalization/XLocLocalization.Automation.csproj
+++ b/Engine/Source/Programs/AutomationTool/XLocLocalization/XLocLocalization.Automation.csproj
@@ -1,4 +1,6 @@
 <Project Sdk="Microsoft.NET.Sdk">
+  <Import Project="..\..\Shared\UnrealEngine.csproj.props" />
+  
   <PropertyGroup>
     <TargetFramework>net6.0</TargetFramework>
     <Configuration Condition=" '$(Configuration)' == '' ">Development</Configuration>
diff --git a/Engine/Source/Programs/IncludeTool/IncludeTool/IncludeTool.csproj b/Engine/Source/Programs/IncludeTool/IncludeTool/IncludeTool.csproj
--- a/Engine/Source/Programs/IncludeTool/IncludeTool/IncludeTool.csproj
+++ b/Engine/Source/Programs/IncludeTool/IncludeTool/IncludeTool.csproj
@@ -1,4 +1,6 @@
 <Project Sdk="Microsoft.NET.Sdk">
+  <Import Project="..\..\Shared\UnrealEngine.csproj.props" />
+
   <PropertyGroup>
     <TargetFramework>net6.0</TargetFramework>
     <Configuration Condition=" '$(Configuration)' == '' ">Development</Configuration>
diff --git a/Engine/Source/Programs/Shared/EpicGames.UHT/EpicGames.UHT.csproj b/Engine/Source/Programs/Shared/EpicGames.UHT/EpicGames.UHT.csproj
--- a/Engine/Source/Programs/Shared/EpicGames.UHT/EpicGames.UHT.csproj
+++ b/Engine/Source/Programs/Shared/EpicGames.UHT/EpicGames.UHT.csproj
@@ -1,4 +1,6 @@
 ﻿<Project Sdk="Microsoft.NET.Sdk">
+  <Import Project="..\UnrealEngine.csproj.props" />
+  
   <PropertyGroup>
     <TargetFramework>netcoreapp3.1</TargetFramework>
     <Configuration Condition=" '$(Configuration)' == '' ">Development</Configuration>
diff --git a/Engine/Source/Programs/UnrealBuildTool/ProjectFiles/ProjectFileGenerator.cs b/Engine/Source/Programs/UnrealBuildTool/ProjectFiles/ProjectFileGenerator.cs
--- a/Engine/Source/Programs/UnrealBuildTool/ProjectFiles/ProjectFileGenerator.cs
+++ b/Engine/Source/Programs/UnrealBuildTool/ProjectFiles/ProjectFileGenerator.cs
@@ -616,8 +616,13 @@ namespace UnrealBuildTool
 				using FileStream Stream = FileReference.Open(ModuleFilesProjectLocation, FileMode.Create, FileAccess.Write, FileShare.Read);
 				using StreamWriter Writer = new StreamWriter(Stream, Encoding.UTF8);
 
+				Writer.WriteLine("<!-- This file was generated by UnrealBuildTool.ProjectFileGenerator.CreateRulesAssemblyProject() -->");
 				Writer.WriteLine("<Project Sdk=\"Microsoft.NET.Sdk\">");
 
+				Writer.WriteLine("  <Import Project=\"" + 
+					FileReference.Combine(Unreal.EngineDirectory, "Source/Programs/Shared/UnrealEngine.csproj.props").MakeRelativeTo(ProjectFilesDirectory) + 
+					"\" />");
+
 				Writer.WriteLine("  <PropertyGroup>");
 				Writer.WriteLine("    <TargetFramework>net6.0</TargetFramework>");
 				Writer.WriteLine("    <AppendTargetFrameworkToOutputPath>false</AppendTargetFrameworkToOutputPath>"); // Shorten intermediate filepath slightly
diff --git a/Engine/Source/Programs/UnrealBuildTool/UnrealBuildTool.csproj b/Engine/Source/Programs/UnrealBuildTool/UnrealBuildTool.csproj
--- a/Engine/Source/Programs/UnrealBuildTool/UnrealBuildTool.csproj
+++ b/Engine/Source/Programs/UnrealBuildTool/UnrealBuildTool.csproj
@@ -1,4 +1,6 @@
 <Project Sdk="Microsoft.NET.Sdk">
+  <Import Project="..\Shared\UnrealEngine.csproj.props" />
+  
   <PropertyGroup>
     <TargetFramework>net6.0</TargetFramework>
     <Configuration Condition=" '$(Configuration)' == '' ">Development</Configuration>
