Changed Gauntlet so editor artifacts are archived if there's a crash on a build machine.

diff --git a/Engine/Source/Programs/AutomationTool/Gauntlet/Unreal/Base/Gauntlet.UnrealSession.cs b/Engine/Source/Programs/AutomationTool/Gauntlet/Unreal/Base/Gauntlet.UnrealSession.cs
--- a/Engine/Source/Programs/AutomationTool/Gauntlet/Unreal/Base/Gauntlet.UnrealSession.cs
+++ b/Engine/Source/Programs/AutomationTool/Gauntlet/Unreal/Base/Gauntlet.UnrealSession.cs
@@ -894,10 +894,15 @@ namespace Gauntlet
 
 			Directory.CreateDirectory(InDestArtifactPath);
 
-			// Don't archive editor data, there can be a *lot* of stuff in that saved folder!
-			bool IsEditor = InRunningRole.Role.RoleType.UsesEditor();
-
 			bool IsDevBuild = InContext.TestParams.ParseParam("dev");
+			bool IsEditorBuild = InRunningRole.Role.RoleType.UsesEditor();
+
+			// Unless there was a crash on a builder don't archive editor data (there can be a *lot* of stuff in their).
+			bool SkipArchivingAssets = IsDevBuild ||
+										(IsEditorBuild && 
+											(CommandUtils.IsBuildMachine == false || InRunningRole.AppInstance.ExitCode == 0)
+										);
+
 
 			string DestSavedDir = InDestArtifactPath;
 			string SourceSavedDir = "";
@@ -989,8 +994,19 @@ namespace Gauntlet
 			});
 
 			// don't archive data in dev mode, because peoples saved data could be huuuuuuuge!
-			if (IsEditor == false)
+			if (SkipArchivingAssets)
 			{
+				if (IsEditorBuild)
+				{
+					Log.Info("Skipping archival of assets for editor {0}", RoleName);
+				}
+				else if (IsDevBuild)
+				{
+					Log.Info("Skipping archival of assets for dev build");
+				}
+			}
+			else
+			{ 
 				LogLevel OldLevel = Log.Level;
 				Log.Level = LogLevel.Normal;
 
@@ -1006,17 +1022,6 @@ namespace Gauntlet
 
 				Log.Level = OldLevel;
 			}
-			else
-			{
-				if (IsEditor)
-				{
-					Log.Info("Skipping archival of assets for editor {0}", RoleName);
-				}
-				else if (IsDevBuild)
-				{
-					Log.Info("Skipping archival of assets for dev build");
-				}
-			}
 
 			foreach (EIntendedBaseCopyDirectory ArtifactDir in InRunningRole.Role.AdditionalArtifactDirectories)
 			{
