BuildGraph: RunUAT via dotnet & AutomationTool.dll

#rnx

diff --git a/Engine/Source/Programs/AutomationTool/AutomationUtils/CommandEnvironment.cs b/Engine/Source/Programs/AutomationTool/AutomationUtils/CommandEnvironment.cs
--- a/Engine/Source/Programs/AutomationTool/AutomationUtils/CommandEnvironment.cs
+++ b/Engine/Source/Programs/AutomationTool/AutomationUtils/CommandEnvironment.cs
@@ -58,7 +58,7 @@ namespace AutomationTool
 		public string RobocopyExe { get; protected set; }
 		public string MountExe { get; protected set; }
 		public string CmdExe { get; protected set; }
-		public string UATExe { get; protected set; }
+		public string AutomationToolDll { get; protected set; }
 		public string TimestampAsString { get; protected set; }
 		public bool HasCapabilityToCompile { get; protected set; }
 		public string FrameworkMsbuildPath { get; protected set; }
@@ -71,26 +71,19 @@ namespace AutomationTool
 		/// </summary>
 		internal CommandEnvironment()
 		{
-			// Get the path to the UAT executable
-			// the entry assembly is the .dll but it is easier to use apphost so change extension to the executable
-			UATExe = Path.ChangeExtension(Assembly.GetEntryAssembly().GetOriginalLocation(), RuntimePlatform.IsWindows ? "exe" : null);
+			// Get the path to AutomationTool.dll
+			AutomationToolDll = Assembly.GetEntryAssembly().GetOriginalLocation();
 
-			if (!CommandUtils.FileExists(UATExe))
+			if (!CommandUtils.FileExists(AutomationToolDll))
 			{
-				// use the dll as a fallback if the exe doesn't exist
-				UATExe = Assembly.GetEntryAssembly().GetOriginalLocation();
-			}
-
-			if (!CommandUtils.FileExists(UATExe))
-			{
-				throw new AutomationException("Could not find AutomationTool.exe. Reflection indicated it was here: {0}", UATExe);
+				throw new AutomationException("Could not find AutomationTool.dll. Reflection indicated it was here: {0}", AutomationToolDll);
 			}
 
 			// Find the root directory (containing the Engine folder)
 			LocalRoot = CommandUtils.GetEnvVar(EnvVarNames.LocalRoot);
 			if(String.IsNullOrEmpty(LocalRoot))
 			{
-				LocalRoot = CommandUtils.ConvertSeparators(PathSeparator.Slash, Path.GetFullPath(Path.Combine(Path.GetDirectoryName(UATExe), "..", "..", "..", "..")));
+				LocalRoot = CommandUtils.ConvertSeparators(PathSeparator.Slash, Path.GetFullPath(Path.Combine(Path.GetDirectoryName(AutomationToolDll), "..", "..", "..", "..")));
 				CommandUtils.ConditionallySetEnvVar(EnvVarNames.LocalRoot, LocalRoot);
 			}
 
@@ -199,7 +192,7 @@ namespace AutomationTool
 			Log.TraceVerbose("DotnetMsbuildExe={0}", DotnetMsbuildPath);
 			Log.TraceVerbose("RobocopyExe={0}", RobocopyExe);
 			Log.TraceVerbose("TimestampAsString={0}", TimestampAsString);
-			Log.TraceVerbose("UATExe={0}", UATExe);			
+			Log.TraceVerbose("AutomationToolDll={0}", AutomationToolDll);			
 		}
 
 		/// <summary>
diff --git a/Engine/Source/Programs/AutomationTool/AutomationUtils/P4Environment.cs b/Engine/Source/Programs/AutomationTool/AutomationUtils/P4Environment.cs
--- a/Engine/Source/Programs/AutomationTool/AutomationUtils/P4Environment.cs
+++ b/Engine/Source/Programs/AutomationTool/AutomationUtils/P4Environment.cs
@@ -153,7 +153,7 @@ namespace AutomationTool
 				if(String.IsNullOrEmpty(Client))
 				{
 					string HostName = System.Net.Dns.GetHostName();
-					ThisClient = DetectClient(DefaultConnection, User, HostName, CmdEnv.UATExe);
+					ThisClient = DetectClient(DefaultConnection, User, HostName, CmdEnv.AutomationToolDll);
 					Log.TraceInformation("Using user {0} clientspec {1} {2}", User, ThisClient.Name, ThisClient.RootPath);
 					Client = ThisClient.Name;
 					CommandUtils.SetEnvVar(EnvVarNames.Client, Client);
@@ -328,7 +328,7 @@ namespace AutomationTool
 		/// <param name="HostName">Host</param>
 		/// <param name="UATLocation">Path to UAT exe, this will be checked agains the root path.</param>
 		/// <returns>Client to use.</returns>
-		private static P4ClientInfo DetectClient(P4Connection Connection, string UserName, string HostName, string UATLocation)
+		private static P4ClientInfo DetectClient(P4Connection Connection, string UserName, string HostName, string AutomationToolDll)
 		{
 			CommandUtils.LogVerbose("uebp_CLIENT not set, detecting current client...");
 
@@ -341,7 +341,7 @@ namespace AutomationTool
 				{
 					string ClientName = Result.Output.Substring(KeyName.Length).Trim();
 					P4ClientInfo ClientInfo = Connection.GetClientInfo(ClientName, true);
-					if (ClientInfo != null && Connection.IsValidClientForFile(ClientInfo, UATLocation))
+					if (ClientInfo != null && Connection.IsValidClientForFile(ClientInfo, AutomationToolDll))
 					{
 						return ClientInfo;
 					}
@@ -350,7 +350,7 @@ namespace AutomationTool
 
 			// Otherwise search for all clients that match
 			List<P4ClientInfo> MatchingClients = new List<P4ClientInfo>();
-			P4ClientInfo[] P4Clients = Connection.GetClientsForUser(UserName, UATLocation);
+			P4ClientInfo[] P4Clients = Connection.GetClientsForUser(UserName, AutomationToolDll);
 			foreach (P4ClientInfo Client in P4Clients)
 			{
 				if (!String.IsNullOrEmpty(Client.Host) && String.Compare(Client.Host, HostName, true) != 0)
diff --git a/Engine/Source/Programs/AutomationTool/AutomationUtils/ProcessUtils.cs b/Engine/Source/Programs/AutomationTool/AutomationUtils/ProcessUtils.cs
--- a/Engine/Source/Programs/AutomationTool/AutomationUtils/ProcessUtils.cs
+++ b/Engine/Source/Programs/AutomationTool/AutomationUtils/ProcessUtils.cs
@@ -1076,7 +1076,7 @@ namespace AutomationTool
 				EnvironmentVars.Add(AutomationTool.EnvVarNames.LocalRoot, ""); // if we don't clear this out, it will think it is a build machine; it will rederive everything
 			}
 
-			IProcessResult Result = Run(CmdEnv.UATExe, CommandLine, null, ERunOptions.Default, EnvironmentVars, Identifier: Identifier);
+			IProcessResult Result = Run(Unreal.DotnetPath.FullName, $"\"{Env.AutomationToolDll}\" {CommandLine}", null, ERunOptions.Default, EnvironmentVars, Identifier: Identifier);
 			if (Result.ExitCode != 0)
 			{
 				throw new CommandFailedException(String.Format("Recursive UAT command failed (exit code {0})", Result.ExitCode)){ OutputFormat = AutomationExceptionOutputFormat.Silent };
