Temporary fix to running UAT from visual studio where the system version of dotnet is used.

#rnx

diff --git a/Engine/Source/Programs/AutomationTool/AutomationTool.csproj b/Engine/Source/Programs/AutomationTool/AutomationTool.csproj
--- a/Engine/Source/Programs/AutomationTool/AutomationTool.csproj
+++ b/Engine/Source/Programs/AutomationTool/AutomationTool.csproj
@@ -108,6 +108,9 @@
   <ItemGroup>
     <PackageReference Include="Microsoft.Build" Version="17.1.0" ExcludeAssets="runtime" PrivateAssets="all" />
     <PackageReference Include="Microsoft.Build.Locator" Version="1.4.1" PrivateAssets="all" />
+    <!-- Added as a temporary workaround to running from visual studio where the system version of dotnet is run. 
+         Adding a loader context resolves many of the issues, but it still fails when MSBuild tries to write the json -->
+    <PackageReference Include="Newtonsoft.Json" Version="13.0.1" />
   </ItemGroup>
   
   <!-- This set of dependencies should be kept as small as possible
diff --git a/Engine/Source/Programs/UnrealBuildTool/UnrealBuildTool.csproj b/Engine/Source/Programs/UnrealBuildTool/UnrealBuildTool.csproj
--- a/Engine/Source/Programs/UnrealBuildTool/UnrealBuildTool.csproj
+++ b/Engine/Source/Programs/UnrealBuildTool/UnrealBuildTool.csproj
@@ -113,5 +113,9 @@
 
     <!-- Required for Hololens -->
     <PackageReference Include="System.Reflection.MetadataLoadContext" Version="4.7.2" /> 
+    
+    <!-- Added as a temporary workaround to running from visual studio where the system version of dotnet is run. 
+         Adding a loader context resolves many of the issues, but it still fails when MSBuild tries to write the json -->
+    <PackageReference Include="Newtonsoft.Json" Version="13.0.1" />
   </ItemGroup>
 </Project>
