BuildCommonTools: Don't try to build a non-existent library

#trival
#rnx

diff --git a/Engine/Source/Programs/AutomationTool/Scripts/BuildCommonTools.Automation.cs b/Engine/Source/Programs/AutomationTool/Scripts/BuildCommonTools.Automation.cs
--- a/Engine/Source/Programs/AutomationTool/Scripts/BuildCommonTools.Automation.cs
+++ b/Engine/Source/Programs/AutomationTool/Scripts/BuildCommonTools.Automation.cs
@@ -1,4 +1,4 @@
-﻿// Copyright Epic Games, Inc. All Rights Reserved.
+// Copyright Epic Games, Inc. All Rights Reserved.
 
 using System;
 using System.IO;
@@ -85,7 +85,6 @@ public class BuildCommonTools : BuildCommand
 		Agenda.SwarmAgentProject = @"Engine\Source\Programs\UnrealSwarm\SwarmAgent.sln";
 		Agenda.SwarmCoordinatorProject = @"Engine\Source\Programs\UnrealSwarm\SwarmCoordinator.sln";
 		Agenda.DotNetProjects.Add(@"Engine/Source/Editor/SwarmInterface/DotNET/SwarmInterface.csproj");
-		Agenda.DotNetProjects.Add(@"Engine/Source/Programs/DotNETCommon/DotNETUtilities/DotNETUtilities.csproj");
 		Agenda.DotNetProjects.Add(@"Engine/Source/Programs/UnrealControls/UnrealControls.csproj");
 
 		// Windows binaries
