Fix case of DotNET path for archiving UnrealAndroidFileTool

#android

diff --git a/Engine/Source/Programs/AutomationTool/Android/AndroidPlatform.Automation.cs b/Engine/Source/Programs/AutomationTool/Android/AndroidPlatform.Automation.cs
--- a/Engine/Source/Programs/AutomationTool/Android/AndroidPlatform.Automation.cs
+++ b/Engine/Source/Programs/AutomationTool/Android/AndroidPlatform.Automation.cs
@@ -1880,7 +1880,7 @@ public class AndroidPlatform : Platform
 				}
 				if (bUseAFS && (bHaveAPK || bHaveUniversal))
 				{
-					SC.ArchiveFiles(Path.Combine(SC.EngineRoot.FullName, "Binaries", "DotNet", "Android", "UnrealAndroidFileTool"), GetAFSExecutable(Target));
+					SC.ArchiveFiles(Path.Combine(SC.EngineRoot.FullName, "Binaries", "DotNET", "Android", "UnrealAndroidFileTool"), GetAFSExecutable(Target));
 				}
 				//SC.ArchiveFiles(Path.GetDirectoryName(NoOBBBatchName), Path.GetFileName(NoOBBBatchName));
 			}
