Add iTunes download to Turnkey on Windows

diff --git a/Engine/Extras/iTunes/Install_iTunes.bat b/Engine/Extras/iTunes/Install_iTunes.bat
new file mode 100755
--- /dev/null
+++ b/Engine/Extras/iTunes/Install_iTunes.bat
@@ -0,0 +1,23 @@
+@echo off
+echo.
+if NOT "%1"=="" goto install
+echo This batch file is intended to be called by Turnkey and should not be run independently.
+goto end
+
+:install
+echo This batch file will try to unninstall any installed iTunes and install the version retrieved by Turnkey. 
+echo.
+
+echo.
+echo Trying to uninstall current iTunes
+echo.
+pause
+
+wmic product where name="iTunes" call uninstall/nointeractive
+
+echo.
+echo Trying to install downloaded iTunes
+echo.
+pause
+
+%1
\ No newline at end of file
diff --git a/Engine/Source/Programs/AutomationTool/IOS/IOSPlatform.Automation.cs b/Engine/Source/Programs/AutomationTool/IOS/IOSPlatform.Automation.cs
--- a/Engine/Source/Programs/AutomationTool/IOS/IOSPlatform.Automation.cs
+++ b/Engine/Source/Programs/AutomationTool/IOS/IOSPlatform.Automation.cs
@@ -168,24 +168,43 @@ public class IOSPlatform : Platform
 
 	public override bool GetSDKInstallCommand(out string Command, out string Params, ref bool bRequiresPrivilegeElevation, ref bool bCreateWindow, ITurnkeyContext TurnkeyContext)
 	{
-		TurnkeyContext.Log("Moving your original Xcode application from /Applications to the Trash, and unzipping the new version into /Applications");
-
-		// put current Xcode in the trash, and unzip a new one. Xcode in the dock will have to be fixed up tho!
-		Command = "osascript";
-		Params =
-			" -e \"try\"" +
-			" -e   \"tell application \\\"Finder\\\" to delete POSIX file \\\"/Applications/Xcode.app\\\"\"" +
-			" -e \"end try\"" +
-			" -e \"do shell script \\\"cd /Applications; xip --expand $(CopyOutputPath);\\\"\"" +
-			" -e \"try\"" +
-			" -e   \"do shell script \\\"xcode-select -s /Applications/Xcode.app; xcode-select --install; xcodebuild -license accept; xcodebuild -runFirstLaunch\\\" with administrator privileges\"" +
-			" -e \"end try\"";
+		Command = "";
+		Params = "";
 
+		if (UnrealBuildTool.BuildHostPlatform.Current.Platform == UnrealTargetPlatform.Mac)
+		{
+			TurnkeyContext.Log("Moving your original Xcode application from /Applications to the Trash, and unzipping the new version into /Applications");
+
+			// put current Xcode in the trash, and unzip a new one. Xcode in the dock will have to be fixed up tho!
+			Command = "osascript";
+			Params =
+				" -e \"try\"" +
+				" -e   \"tell application \\\"Finder\\\" to delete POSIX file \\\"/Applications/Xcode.app\\\"\"" +
+				" -e \"end try\"" +
+				" -e \"do shell script \\\"cd /Applications; xip --expand $(CopyOutputPath);\\\"\"" +
+				" -e \"try\"" +
+				" -e   \"do shell script \\\"xcode-select -s /Applications/Xcode.app; xcode-select --install; xcodebuild -license accept; xcodebuild -runFirstLaunch\\\" with administrator privileges\"" +
+				" -e \"end try\"";
+		}
+		else if (UnrealBuildTool.BuildHostPlatform.Current.Platform == UnrealTargetPlatform.Win64)
+		{
+
+			TurnkeyContext.Log("Uninstalling old iTunes and preparing the new one to be installed.");
+
+			Command = "$(EngineDir)/Extras/iTunes/Install_iTunes.bat";
+			Params = "$(CopyOutputPath)";
+		}
 		return true;
 	}
 
 	public override bool GetDeviceUpdateSoftwareCommand(out string Command, out string Params, ref bool bRequiresPrivilegeElevation, ref bool bCreateWindow, ITurnkeyContext TurnkeyContext, DeviceInfo Device)
 	{
+		if (UnrealBuildTool.BuildHostPlatform.Current.Platform != UnrealTargetPlatform.Mac)
+		{
+			Command = Params = null;
+			return true;
+		}
+
 		TurnkeyContext.Log("Installing an offline downloaded .ipsw onto your device using the Apple Configurator application.");
 
 		// cfgtool needs ECID, not UDID, so find it
@@ -215,7 +234,10 @@ public class IOSPlatform : Platform
 		{
 			if (ExitCode == 0)
 			{
-				TurnkeyContext.PauseForUser("If you had Xcode in your Dock, you will need to remove it and add the new one (even though it was in the same location). macOS follows the move to the Trash for the Dock icon");
+				if (UnrealBuildTool.BuildHostPlatform.Current.Platform == UnrealTargetPlatform.Mac)
+				{
+					TurnkeyContext.PauseForUser("If you had Xcode in your Dock, you will need to remove it and add the new one (even though it was in the same location). macOS follows the move to the Trash for the Dock icon");
+				}
 			}
 		}
 
diff --git a/Engine/Source/Programs/UnrealBuildTool/Platform/Mac/ApplePlatformSDK.cs b/Engine/Source/Programs/UnrealBuildTool/Platform/Mac/ApplePlatformSDK.cs
--- a/Engine/Source/Programs/UnrealBuildTool/Platform/Mac/ApplePlatformSDK.cs
+++ b/Engine/Source/Programs/UnrealBuildTool/Platform/Mac/ApplePlatformSDK.cs
@@ -14,14 +14,14 @@ namespace UnrealBuildTool
 	{
 		public override string GetMainVersion()
 		{
-			return "11.5";
+			return "13.3";
 		}
 
 		protected override void GetValidVersionRange(out string MinVersion, out string MaxVersion)
 		{
 			if (RuntimePlatform.IsMac)
 			{
-				MinVersion = "11.0.0";
+				MinVersion = "12.4.0";
 				MaxVersion = "13.9.9";
 			}
 			else
