UAT: Fix non-zero UAT exit code when closing game launched on Android.
Restore the behavior as it was before cl 19621549.

#rnx

diff --git a/Engine/Source/Programs/AutomationTool/Android/AndroidPlatform.Automation.cs b/Engine/Source/Programs/AutomationTool/Android/AndroidPlatform.Automation.cs
--- a/Engine/Source/Programs/AutomationTool/Android/AndroidPlatform.Automation.cs
+++ b/Engine/Source/Programs/AutomationTool/Android/AndroidPlatform.Automation.cs
@@ -35,6 +35,7 @@ public class AndroidPlatform : Platform
 		string LogPath;
 		string PackageName;
 		string DeviceName;
+		int LogFileProcessExitCode = 0;
 		bool bStopped = false;
 
 		public AdbCreatedProcess(
@@ -108,8 +109,8 @@ public class AndroidPlatform : Platform
 
 		public int ExitCode
 		{
-			get { return AdbLogProcess.ExitCode; }
-			set { AdbLogProcess.ExitCode = value; }
+			get { return LogFileProcessExitCode; }
+			set { LogFileProcessExitCode = value; }
 		}
 
 		public string Output
@@ -150,6 +151,8 @@ public class AndroidPlatform : Platform
 			string AllOutput = LogFileProcess.Output;
 			File.WriteAllText(LogFilename, AllOutput);
 			File.WriteAllText(ServerLogFilename, AllOutput);
+
+			ExitCode = LogFileProcess.ExitCode;
 		}
 	}
 
