Clean up leftover OBB files which may override new data deployment for Android quick launch

#android

diff --git a/Engine/Source/Programs/AutomationTool/Android/AndroidPlatform.Automation.cs b/Engine/Source/Programs/AutomationTool/Android/AndroidPlatform.Automation.cs
--- a/Engine/Source/Programs/AutomationTool/Android/AndroidPlatform.Automation.cs
+++ b/Engine/Source/Programs/AutomationTool/Android/AndroidPlatform.Automation.cs
@@ -2255,7 +2255,8 @@ public class AndroidPlatform : Platform
 
 			// Setup the OBB name and add the storage path (queried from the device) to it
 			string QueryStorageResult = adb.Shell(DeviceName, "echo $EXTERNAL_STORAGE");
-			string StorageLocation = QueryStorageResult.Trim() + "/Android";  // "mnt/sdcard"
+			string ExternalStorage = QueryStorageResult.Trim();		// "mnt/sdcard"
+			string StorageLocation = ExternalStorage + "/Android";	// "mnt/sdcard/Android"
 			string DeviceObbName = StorageLocation + "/" + GetDeviceObbName(ApkName, SC);
 			string DevicePatchName = StorageLocation + "/" + GetDevicePatchName(ApkName, SC);
 			string DeviceOverflow1Name = StorageLocation + "/" + GetDeviceOverflowName(ApkName, SC, 1);
@@ -2272,6 +2273,15 @@ public class AndroidPlatform : Platform
 				adb.Shell(DeviceName, "setprop security.perf_harden 0");
 			}
 
+			// remove any main or patch OBB file which would override deployed data (note: not the same as later delete or deploy)
+			{
+				string DeviceOldObbName = ExternalStorage + "/" + GetDeviceObbName(ApkName, SC);
+				string DeviceOldPatchName = ExternalStorage + "/" + GetDevicePatchName(ApkName, SC);
+
+				adb.Shell(DeviceName, "rm " + DeviceOldObbName);
+				adb.Shell(DeviceName, "rm " + DeviceOldPatchName);
+			}
+
 			// close connection to since uninstall/reinstall will reset server
 			client.CloseConnection();
 
@@ -2830,6 +2840,15 @@ public class AndroidPlatform : Platform
 				RunAdbCommand(Params, DeviceName, "shell setprop security.perf_harden 0");
 			}
 
+			// remove any main or patch OBB file which would override deployed data (note: not the same as later delete or deploy)
+			{
+				string DeviceOldObbName = StorageLocation + "/Android/" + GetDeviceObbName(ApkName, SC);
+				string DeviceOldPatchName = StorageLocation + "/Android/" + GetDevicePatchName(ApkName, SC);
+
+				RunAdbCommand(Params, DeviceName, "shell rm " + DeviceOldObbName);
+				RunAdbCommand(Params, DeviceName, "shell rm " + DeviceOldPatchName);
+			}
+
             // determine if APK out of date
             string APKLastUpdateTime = new FileInfo(ApkName).LastWriteTime.ToString();
             bool bNeedAPKInstall = true;
