Add commandline boilerplate to support optional cooking through UAT

diff --git a/Engine/Source/Programs/AutomationTool/AutomationUtils/ProjectParams.cs b/Engine/Source/Programs/AutomationTool/AutomationUtils/ProjectParams.cs
--- a/Engine/Source/Programs/AutomationTool/AutomationUtils/ProjectParams.cs
+++ b/Engine/Source/Programs/AutomationTool/AutomationUtils/ProjectParams.cs
@@ -331,6 +331,7 @@ namespace AutomationTool
 			this.CookOnTheFly = InParams.CookOnTheFly;
             this.CookOnTheFlyStreaming = InParams.CookOnTheFlyStreaming;
             this.UnversionedCookedContent = InParams.UnversionedCookedContent;
+			this.OptionalContent = InParams.OptionalContent;
 			this.SkipCookingEditorContent = InParams.SkipCookingEditorContent;
             this.NumCookersToSpawn = InParams.NumCookersToSpawn;
 			this.FileServer = InParams.FileServer;
@@ -470,6 +471,7 @@ namespace AutomationTool
             bool? CookOnTheFly = null,
             bool? CookOnTheFlyStreaming = null,
             bool? UnversionedCookedContent = null,
+			bool? OptionalContent = null,
 			bool? EncryptIniFiles = null,
             bool? EncryptPakIndex = null,
 			bool? EncryptEverything = null,
@@ -739,6 +741,7 @@ namespace AutomationTool
             }
             this.CookOnTheFlyStreaming = GetParamValueIfNotSpecified(Command, CookOnTheFlyStreaming, this.CookOnTheFlyStreaming, "cookontheflystreaming");
             this.UnversionedCookedContent = GetOptionalParamValueIfNotSpecified(Command, UnversionedCookedContent, this.UnversionedCookedContent, "UnversionedCookedContent", "VersionCookedContent");
+			this.OptionalContent = GetOptionalParamValueIfNotSpecified(Command, OptionalContent, this.OptionalContent, "editoroptional", "noeditoroptional");
 			this.SkipCookingEditorContent = GetParamValueIfNotSpecified(Command, SkipCookingEditorContent, this.SkipCookingEditorContent, "SkipCookingEditorContent");
             if (NumCookersToSpawn.HasValue)
             {
@@ -1790,6 +1793,10 @@ namespace AutomationTool
         /// </summary>
         public bool UnversionedCookedContent = true;
 
+		/// <summary>
+		/// Cook: Cook with optional data enabled
+		/// </summary>
+		public bool OptionalContent = false;
 
 		/// <summary>
 		/// Cook: Uses the iterative cooking, command line: -iterativecooking or -iterate
@@ -3002,6 +3009,7 @@ namespace AutomationTool
 				CommandUtils.LogLog("CookOnTheFly={0}", CookOnTheFly);
 				CommandUtils.LogLog("CookOnTheFlyStreaming={0}", CookOnTheFlyStreaming);
 				CommandUtils.LogLog("UnversionedCookedContent={0}", UnversionedCookedContent);
+				CommandUtils.LogLog("OptionalContent={0}", OptionalContent);
 				CommandUtils.LogLog("SkipCookingEditorContent={0}", SkipCookingEditorContent);
                 CommandUtils.LogLog("NumCookersToSpawn={0}", NumCookersToSpawn);
                 CommandUtils.LogLog("GeneratePatch={0}", GeneratePatch);
diff --git a/Engine/Source/Programs/AutomationTool/Scripts/CookCommand.Automation.cs b/Engine/Source/Programs/AutomationTool/Scripts/CookCommand.Automation.cs
--- a/Engine/Source/Programs/AutomationTool/Scripts/CookCommand.Automation.cs
+++ b/Engine/Source/Programs/AutomationTool/Scripts/CookCommand.Automation.cs
@@ -44,6 +44,10 @@ namespace AutomationScripts
 			{
 				CommandletParams += " -unversioned";
 			}
+			if (Params.OptionalContent)
+			{
+				CommandletParams += " -editoroptional";
+			}
 			if (Params.FastCook)
 			{
 				CommandletParams += " -FastCook";
