CookedEditor staging: Fix warnings about missing generated map files - do not add the uncooked version of map files for generated files that do not exist uncooked.
#rnx

diff --git a/Engine/Source/Programs/AutomationTool/CookedEditor/CookedEditor.Automation.cs b/Engine/Source/Programs/AutomationTool/CookedEditor/CookedEditor.Automation.cs
--- a/Engine/Source/Programs/AutomationTool/CookedEditor/CookedEditor.Automation.cs
+++ b/Engine/Source/Programs/AutomationTool/CookedEditor/CookedEditor.Automation.cs
@@ -168,6 +168,7 @@ public class ModifyStageContext
 		}
 
 		// remove already-cooked assets to be replaced with 
+		List<StagedFileReference> UncookedFilesThatDoNotExist = new List<StagedFileReference>();
 		string[] CookedExtensions = { ".uasset", ".umap", ".ubulk", ".uexp" };
 		foreach (var UncookedFile in StagedUncookFiles)
 		{
@@ -180,6 +181,19 @@ public class ModifyStageContext
 				SC.FilesToStage.UFSFiles.Remove(PathWithExtension);
 				StagedUFSFiles.Remove(PathWithExtension);
 			}
+
+			// Some uncooked packages are generated at cook time and do not exist in the uncooked depot.
+			// We removed the cooked version of these files from staging above, but we should not add an entry for their
+			// non-existent uncooked file. Add them to a list for removal after the loop.
+			FileReference FullPathToUncooked = UncookedFile.Value;
+			if (!FullPathToUncooked.ToFileInfo().Exists)
+			{
+				UncookedFilesThatDoNotExist.Add(UncookedFile.Key);
+			}
+		}
+		foreach (StagedFileReference UncookedFile in UncookedFilesThatDoNotExist)
+		{
+			StagedUncookFiles.Remove(UncookedFile);
 		}
 
 		// stage the filtered UFSFiles
