CsCompile: Don't use /restore when compiling with xbuild

#rnx

diff --git a/Engine/Source/Programs/AutomationTool/BuildGraph/Tasks/CsCompileTask.cs b/Engine/Source/Programs/AutomationTool/BuildGraph/Tasks/CsCompileTask.cs
--- a/Engine/Source/Programs/AutomationTool/BuildGraph/Tasks/CsCompileTask.cs
+++ b/Engine/Source/Programs/AutomationTool/BuildGraph/Tasks/CsCompileTask.cs
@@ -167,7 +167,11 @@ namespace AutomationTool.Tasks
 				{
 					Arguments.Add(String.Format("/target:{0}", CommandUtils.MakePathSafeToUseWithCommandLine(Parameters.Target)));
 				}
-				Arguments.Add("/restore");
+				if(!CommandUtils.CmdEnv.FrameworkMsbuildPath.Equals("xbuild"))
+				{
+					// not supported by xbuild
+					Arguments.Add("/restore");
+				}
 				Arguments.Add("/verbosity:minimal");
 				Arguments.Add("/nologo");
 
