Removed temporary code that was causing UAT to always rebuild plugins.

#rnx

diff --git a/Engine/Source/Programs/AutomationTool/Program.cs b/Engine/Source/Programs/AutomationTool/Program.cs
--- a/Engine/Source/Programs/AutomationTool/Program.cs
+++ b/Engine/Source/Programs/AutomationTool/Program.cs
@@ -419,7 +419,7 @@ namespace AutomationToolDriver
 				: AutomationToolCommandLine.CommandsToExecute;
 			bool bBuildSuccess;
 			HashSet<FileReference> ScriptModuleAssemblyPaths = CompileScriptModule.InitializeScriptModules(
-					Rules.RulesFileType.AutomationModule, ScriptsForProject, AdditionalScriptDirs, bForceCompile | true, bNoCompile, bUseBuildRecords, 
+					Rules.RulesFileType.AutomationModule, ScriptsForProject, AdditionalScriptDirs, bForceCompile, bNoCompile, bUseBuildRecords, 
 					out bBuildSuccess, (int Count) =>
                     {
 						Log.TraceInformation($"Building {Count} projects (see Log 'Engine/Programs/AutomationTool/Saved/Logs/Log.txt' for more details)");
