Gauntlet - Support different build target configuration for editor with RunUnreal command.

#rnx

diff --git a/Engine/Source/Programs/AutomationTool/Gauntlet/Unreal/BuildSource/Gauntlet.UnrealBuildSource.cs b/Engine/Source/Programs/AutomationTool/Gauntlet/Unreal/BuildSource/Gauntlet.UnrealBuildSource.cs
--- a/Engine/Source/Programs/AutomationTool/Gauntlet/Unreal/BuildSource/Gauntlet.UnrealBuildSource.cs
+++ b/Engine/Source/Programs/AutomationTool/Gauntlet/Unreal/BuildSource/Gauntlet.UnrealBuildSource.cs
@@ -274,16 +274,16 @@ namespace Gauntlet
 		/// Adds an Editor build to our list of available builds if one exists
 		/// </summary>
 		/// <param name="InUnrealPath"></param>
-		virtual protected IBuild CreateEditorBuild(DirectoryReference InUnrealPath)
+		virtual protected IBuild CreateEditorBuild(DirectoryReference InUnrealPath, UnrealTargetConfiguration InConfiguration = UnrealTargetConfiguration.Development)
 		{
 			if (InUnrealPath != null)
 			{
 				// check for the editor
-				string EditorExe = Path.Combine(InUnrealPath.FullName, GetRelativeExecutablePath(UnrealTargetRole.Editor, BuildHostPlatform.Current.Platform, UnrealTargetConfiguration.Development));
+				string EditorExe = Path.Combine(InUnrealPath.FullName, GetRelativeExecutablePath(UnrealTargetRole.Editor, BuildHostPlatform.Current.Platform, InConfiguration));
 
 				if (Utils.SystemHelpers.ApplicationExists(EditorExe))
 				{
-					EditorBuild NewBuild = new EditorBuild(EditorExe, UnrealTargetConfiguration.Development);
+					EditorBuild NewBuild = new EditorBuild(EditorExe, InConfiguration);
 
 					return NewBuild;
 				}
@@ -316,7 +316,7 @@ namespace Gauntlet
 		/// for the provided platform
 		/// </summary>
 		/// <param name="InPlatform"></param>
-		virtual protected void DiscoverBuilds(UnrealTargetPlatform InPlatform)
+		virtual protected void DiscoverBuilds(UnrealTargetPlatform InPlatform, UnrealTargetConfiguration InConfiguration = UnrealTargetConfiguration.Development)
 		{
 			if (!HaveDiscoveredBuilds(InPlatform))
 			{
@@ -326,7 +326,7 @@ namespace Gauntlet
 				// Add an editor build if this is our current platform.
 				if (InPlatform == BuildHostPlatform.Current.Platform)
 				{
-					IBuild EditorBuild = CreateEditorBuild(UnrealPath);
+					IBuild EditorBuild = CreateEditorBuild(UnrealPath, InConfiguration);
 
 					if (EditorBuild == null)
 					{
@@ -393,7 +393,7 @@ namespace Gauntlet
 
 			if (!HaveDiscoveredBuilds(InPlatform.Value))
 			{
-				DiscoverBuilds(InPlatform.Value);
+				DiscoverBuilds(InPlatform.Value, InConfiguration);
 			}
 
 			IEnumerable<IBuild> PlatformBuilds = DiscoveredBuilds[InPlatform.Value];
@@ -639,7 +639,14 @@ namespace Gauntlet
 
 				if (BuildName == "Editor")
 				{
-					EditorExe = ProjectUtils.GetProjectTarget(ProjectPath, UnrealBuildTool.TargetType.Editor, TargetPlatform, TargetConfiguration);
+					try
+					{
+						EditorExe = ProjectUtils.GetProjectTarget(ProjectPath, UnrealBuildTool.TargetType.Editor, TargetPlatform, TargetConfiguration);
+					}
+					catch
+					{
+						throw new AutomationException("No suitable editor build for {0} configuration", TargetConfiguration);
+					}
 				}
 
 				if (EditorExe != null)
diff --git a/Engine/Source/Programs/AutomationTool/Gauntlet/Unreal/RunUnreal.cs b/Engine/Source/Programs/AutomationTool/Gauntlet/Unreal/RunUnreal.cs
--- a/Engine/Source/Programs/AutomationTool/Gauntlet/Unreal/RunUnreal.cs
+++ b/Engine/Source/Programs/AutomationTool/Gauntlet/Unreal/RunUnreal.cs
@@ -254,6 +254,14 @@ namespace Gauntlet
 							}
 						}
 					}
+					else if (Type.IsEditor())
+					{
+						Role.Configuration = RequestedConfiguration;
+						if (Role.Configuration > UnrealTargetConfiguration.Development)
+						{
+							Role.Configuration = UnrealTargetConfiguration.Development;
+						}
+					}
 
 					Gauntlet.Log.Verbose("Mapped Role {0} to RoleContext {1}", Type, Role);
 
