Allow programs to be built using the BuildTarget command.

diff --git a/Engine/Source/Programs/AutomationTool/Scripts/BuildTarget.Automation.cs b/Engine/Source/Programs/AutomationTool/Scripts/BuildTarget.Automation.cs
--- a/Engine/Source/Programs/AutomationTool/Scripts/BuildTarget.Automation.cs
+++ b/Engine/Source/Programs/AutomationTool/Scripts/BuildTarget.Automation.cs
@@ -279,6 +279,12 @@ namespace AutomationTool
 				}
 			}		
 
+			// If no target is found, try to build the provided target name. This enables programs to build (Ex. UnrealInsights).
+			if(ProjectTarget == null)
+			{
+				ProjectTarget = new SimpleTargetInfo(InTargetName, TargetType.Program);
+			}
+
 			if (ProjectTarget == null)
 			{
 				throw new AutomationException("{0} is not a valid target in {1}", InTargetName, InProjectFile);
