UAT: Restrict the COTF -filehostip list to ip v4 host addresses and only include 127.0.0.1 on host platforms or platforms that support port redirection.
Makes connecting to cotf faster on some platforms.
Removes warnings like: LogSockets: Warning: Could not serialize fe80::edc3:3284:5b1b:1f32%7, got error code SE_NO_ERROR [0]
Makes the string shorter and more readable.

Test:
RunUAT.bat BuildCookRun -project=Samples\Games\ShooterGame\ShooterGame.uproject -platform=XYZ -config=development -build -cookonthefly -stage -run

FileHostIp example on PC & android:
Before: -filehostip=127.0.0.1+fe80::edc3:3284:5b1b:1f32%7+10.24.101.44+::1+127.0.0.1+fe80::58d4:b8b3:ff7a:4974%9+172.18.240.1
After : -filehostip=127.0.0.1+10.24.101.44+172.18.240.1

FileHostIp example on console:
Before: -filehostip=127.0.0.1+fe80::edc3:3284:5b1b:1f32%7+10.24.101.44+::1+127.0.0.1+fe80::58d4:b8b3:ff7a:4974%9+172.18.240.1
After : -filehostip=10.24.101.44+172.18.240.1

#rnx

diff --git a/Engine/Source/Programs/AutomationTool/Android/AndroidPlatform.Automation.cs b/Engine/Source/Programs/AutomationTool/Android/AndroidPlatform.Automation.cs
--- a/Engine/Source/Programs/AutomationTool/Android/AndroidPlatform.Automation.cs
+++ b/Engine/Source/Programs/AutomationTool/Android/AndroidPlatform.Automation.cs
@@ -3737,6 +3737,11 @@ public class AndroidPlatform : Platform
 		return true;
 	}
 
+	public override void ModifyFileHostAddresses(List<string> HostAddresses)
+	{
+		HostAddresses.Insert(0, "127.0.0.1");
+	}
+
 	public override IProcessResult RunClient(ERunOptions ClientRunFlags, string ClientApp, string ClientCmdLine, ProjectParams Params)
 	{
 		IProcessResult Result = null;
diff --git a/Engine/Source/Programs/AutomationTool/AutomationUtils/Platform.cs b/Engine/Source/Programs/AutomationTool/AutomationUtils/Platform.cs
--- a/Engine/Source/Programs/AutomationTool/AutomationUtils/Platform.cs
+++ b/Engine/Source/Programs/AutomationTool/AutomationUtils/Platform.cs
@@ -619,6 +619,13 @@ namespace AutomationTool
 			return "";
 		}
 
+		/// <summary>
+		/// Modify or override the list of file host addresses for this platform.
+		/// </summary>
+		public virtual void ModifyFileHostAddresses(List<string> HostAddresses)
+		{
+		}
+
 		/// <summary>
 		/// True if this platform can write to the abslog path that's on the host desktop.
 		/// </summary>
diff --git a/Engine/Source/Programs/AutomationTool/Scripts/CopyBuildToStagingDirectory.Automation.cs b/Engine/Source/Programs/AutomationTool/Scripts/CopyBuildToStagingDirectory.Automation.cs
--- a/Engine/Source/Programs/AutomationTool/Scripts/CopyBuildToStagingDirectory.Automation.cs
+++ b/Engine/Source/Programs/AutomationTool/Scripts/CopyBuildToStagingDirectory.Automation.cs
@@ -4089,7 +4089,7 @@ namespace AutomationScripts
 				CommandLineBuilder.AppendFormat("{0} ", Params.RunCommandline.Trim(new char[] { '\"' }));
 			}
 
-			String FileHostCommandline = GetFileHostCommandline(Params, SC).Trim();
+			String FileHostCommandline = GetFileHostCommandline(Params, SC);
 			if (!string.IsNullOrEmpty(FileHostCommandline))
 			{
 				CommandLineBuilder.AppendFormat("{0} ", FileHostCommandline);
diff --git a/Engine/Source/Programs/AutomationTool/Scripts/RunProjectCommand.Automation.cs b/Engine/Source/Programs/AutomationTool/Scripts/RunProjectCommand.Automation.cs
--- a/Engine/Source/Programs/AutomationTool/Scripts/RunProjectCommand.Automation.cs
+++ b/Engine/Source/Programs/AutomationTool/Scripts/RunProjectCommand.Automation.cs
@@ -693,145 +693,106 @@ namespace AutomationScripts
 			}
 		}
 
-		private static string GetFileHostCommandline(ProjectParams Params, DeploymentContext SC)
+		private static List<string> GetFileHostAddresses(DeploymentContext SC)
 		{
-			string FileHostParams = "";
-			if (Params.CookOnTheFly || Params.FileServer)
+			List<string> HostAddresses = new List<string>();
+
+			// Add localhost first for host platforms and skip it completely for other platforms.
+			// Any Platform can implement ModifyFileHostAddresses to tweak this default behavior.
+			string LocalHost = "127.0.0.1";
+			if (UnrealBuildTool.BuildHostPlatform.Current.Platform == SC.StageTargetPlatform.PlatformType)
+			{
+				HostAddresses.Add(LocalHost);
+			}
+
+			bool bIsMac = UnrealBuildTool.BuildHostPlatform.Current.Platform == UnrealTargetPlatform.Mac;
+			NetworkInterface[] Interfaces = NetworkInterface.GetAllNetworkInterfaces();
+			foreach (NetworkInterface Adapter in Interfaces)
 			{
-				if (Params.ZenStore)
+				if (bIsMac)
 				{
-					if (Params.CookOnTheFly)
+					if (Adapter.NetworkInterfaceType == NetworkInterfaceType.Loopback)
 					{
-						FileHostParams += "-cookonthefly ";
+						continue;
 					}
-					FileHostParams += "-zenstoreproject=" + ProjectUtils.GetProjectPathId(SC.RawProjectPath) + " ";
-					FileHostParams += "-zenstorehost=";
 				}
 				else
 				{
-					FileHostParams += "-filehostip=";
+					if (Adapter.OperationalStatus != OperationalStatus.Up)
+					{
+						continue;
+					}
 				}
 
-				// add localhost first for platforms using redirection
-				const string LocalHost = "127.0.0.1";
-				if (!IsNullOrEmpty(Params.Port))
+				IPInterfaceProperties IP = Adapter.GetIPProperties();
+				foreach (UnicastIPAddressInformation UnicastAddress in IP.UnicastAddresses)
 				{
-					bool FirstParam = true;
-					foreach (var Port in Params.Port)
+					if (!InternalUtils.IsDnsEligible(UnicastAddress))
 					{
-						if (!FirstParam)
-						{
-							FileHostParams += "+";
-						}
-						FirstParam = false;
-						string[] PortProtocol = Port.Split(new char[] { ':' });
-						if (PortProtocol.Length > 1)
-						{
-							FileHostParams += String.Format("{0}://{1}:{2}", PortProtocol[0], LocalHost, PortProtocol[1]);
-						}
-						else
-						{
-							FileHostParams += LocalHost;
-							FileHostParams += ":";
-							FileHostParams += Params.Port;
-						}
+						continue;
+					}
 
+					if (UnicastAddress.Address.AddressFamily != System.Net.Sockets.AddressFamily.InterNetwork)
+					{
+						continue;
 					}
-				}
-				else
-				{
-					// use default port
-					FileHostParams += LocalHost;
-				}
 
-				if (UnrealBuildTool.BuildHostPlatform.Current.Platform == UnrealTargetPlatform.Mac)
-				{
-					NetworkInterface[] Interfaces = NetworkInterface.GetAllNetworkInterfaces();
-					foreach (NetworkInterface adapter in Interfaces)
+					string HostAddress = UnicastAddress.Address.ToString();
+					if (HostAddress == LocalHost)
 					{
-						if (adapter.NetworkInterfaceType != NetworkInterfaceType.Loopback)
-						{
-							IPInterfaceProperties IP = adapter.GetIPProperties();
-							for (int Index = 0; Index < IP.UnicastAddresses.Count; ++Index)
-							{
-								if (InternalUtils.IsDnsEligible(IP.UnicastAddresses[Index]) && IP.UnicastAddresses[Index].Address.AddressFamily == System.Net.Sockets.AddressFamily.InterNetwork)
-								{
-									if (!IsNullOrEmpty(Params.Port))
-									{
-										foreach (var Port in Params.Port)
-										{
-											FileHostParams += "+";
-											string[] PortProtocol = Port.Split(new char[] { ':' });
-											if (PortProtocol.Length > 1)
-											{
-												FileHostParams += String.Format("{0}://{1}:{2}", PortProtocol[0], IP.UnicastAddresses[Index].Address.ToString(), PortProtocol[1]);
-											}
-											else
-											{
-												FileHostParams += IP.UnicastAddresses[Index].Address.ToString();
-												FileHostParams += ":";
-												FileHostParams += Params.Port;
-											}
-
-										}
-									}
-									else
-									{
-										FileHostParams += "+";
-										// use default port
-										FileHostParams += IP.UnicastAddresses[Index].Address.ToString();
-									}
-								}
-							}
-						}
+						continue;
 					}
+					HostAddresses.Add(HostAddress);
 				}
-				else
+			}
+			return HostAddresses.ToList();
+		}
+
+		private static string GetFileHostCommandline(ProjectParams Params, DeploymentContext SC)
+		{
+			string FileHostParams = "";
+			if (!Params.CookOnTheFly && !Params.FileServer)
+			{
+				return FileHostParams;
+			}
+
+			List<string> HostAddresses = GetFileHostAddresses(SC);
+			SC.StageTargetPlatform.ModifyFileHostAddresses(HostAddresses);
+
+			if (!IsNullOrEmpty(Params.Port))
+			{
+				foreach (var Port in Params.Port)
 				{
-					NetworkInterface[] Interfaces = NetworkInterface.GetAllNetworkInterfaces();
-					foreach (NetworkInterface adapter in Interfaces)
+					string[] PortProtocol = Port.Split(new char[] { ':' });
+					for (int I = 0; I < HostAddresses.Count; I++)
 					{
-						if (adapter.OperationalStatus == OperationalStatus.Up)
+						if (PortProtocol.Length > 1)
 						{
-							IPInterfaceProperties IP = adapter.GetIPProperties();
-							for (int Index = 0; Index < IP.UnicastAddresses.Count; ++Index)
-							{
-								if (InternalUtils.IsDnsEligible(IP.UnicastAddresses[Index]))
-								{
-									if (!IsNullOrEmpty(Params.Port))
-									{
-										foreach (var Port in Params.Port)
-										{
-											FileHostParams += "+";
-											string[] PortProtocol = Port.Split(new char[] { ':' });
-											if (PortProtocol.Length > 1)
-											{
-												FileHostParams += String.Format("{0}://{1}:{2}", PortProtocol[0], IP.UnicastAddresses[Index].Address.ToString(), PortProtocol[1]);
-											}
-											else
-											{
-												FileHostParams += IP.UnicastAddresses[Index].Address.ToString();
-												FileHostParams += ":";
-												FileHostParams += Params.Port;
-											}
-										}
-									}
-									else
-									{
-										FileHostParams += "+";
-										// use default port
-										FileHostParams += IP.UnicastAddresses[Index].Address.ToString();
-									}
-								}
-							}
+							HostAddresses[I] = String.Format("{0}://{1}:{2}", PortProtocol[0], HostAddresses[I], PortProtocol[1]);
+						}
+						else
+						{
+							HostAddresses[I] = String.Format("{0}:{1}", HostAddresses[I], Port);
 						}
 					}
 				}
-
-				FileHostParams += " ";
 			}
 
-			return FileHostParams.Trim();
+			if (Params.ZenStore)
+			{
+				if (Params.CookOnTheFly)
+				{
+					FileHostParams += "-cookonthefly ";
+				}
+				FileHostParams += "-zenstoreproject=" + ProjectUtils.GetProjectPathId(SC.RawProjectPath) + " ";
+				FileHostParams += "-zenstorehost=";
+			}
+			else
+			{
+				FileHostParams += "-filehostip=";
+			}
+			FileHostParams += String.Join("+", HostAddresses);
+			return FileHostParams;
 		}
 
 		private static void ClientLogReaderProc(object ArgsContainer)
