Couple of temp storage clean up changes to let multiple clean up jobs run less painfully.
* Flag folders as delete in progress to allow other clean up jobs to skip a folder and go to the next one.
* Delete from newest to oldest in an attempt to provide a buffer before multiple jobs will attempt to delete the same folder.

#rnx

diff --git a/Engine/Source/Programs/AutomationTool/BuildGraph/TempStorage.cs b/Engine/Source/Programs/AutomationTool/BuildGraph/TempStorage.cs
--- a/Engine/Source/Programs/AutomationTool/BuildGraph/TempStorage.cs
+++ b/Engine/Source/Programs/AutomationTool/BuildGraph/TempStorage.cs
@@ -1459,11 +1459,27 @@ namespace AutomationTool
 			CommandUtils.LogInformation("Found {0} builds; {1} to delete.", NumBuilds, BuildsToDelete.Count);
 
 			// Loop through them all, checking for files older than the delete time
-			for (int Idx = 0; Idx < BuildsToDelete.Count; Idx++)
+			int Idx = BuildsToDelete.Count;
+			while (Idx-- > 0)
 			{
 				try
 				{
-					CommandUtils.LogInformation("[{0}/{1}] Deleting {2}...", Idx + 1, BuildsToDelete.Count, BuildsToDelete[Idx].FullName);
+					// Done if something already cleaned up this folder.
+					if (!BuildsToDelete[Idx].Exists)
+					{
+						continue;
+					}
+
+					// Check if there is a marker file, if so skip this folder unless it has been twenty minutes.
+					FileInfo DeleteInProgressFile = BuildsToDelete[Idx].GetFiles("DeleteInProgress.tmp").FirstOrDefault();
+					if (DeleteInProgressFile != null && DeleteInProgressFile.LastWriteTimeUtc < (DateTimeOffset.UtcNow - TimeSpan.FromMinutes(20)))
+					{
+						CommandUtils.LogInformation("[{0}/{1}] {2} flagged as delete in progress, skipping...", BuildsToDelete.Count - Idx, BuildsToDelete.Count, BuildsToDelete[Idx].FullName);
+						continue;
+					}
+
+					File.WriteAllBytes(Path.Combine(BuildsToDelete[Idx].FullName, "DeleteInProgress.tmp"), Array.Empty<byte>());
+					CommandUtils.LogInformation("[{0}/{1}] Deleting {2}...", BuildsToDelete.Count - Idx, BuildsToDelete.Count, BuildsToDelete[Idx].FullName);
 					BuildsToDelete[Idx].Delete(true);
 				}
 				catch (Exception Ex)
