Expose -editoroptional option to cooked editor makerelease

diff --git a/Engine/Source/Programs/AutomationTool/CookedEditor/CookedEditor.Automation.cs b/Engine/Source/Programs/AutomationTool/CookedEditor/CookedEditor.Automation.cs
--- a/Engine/Source/Programs/AutomationTool/CookedEditor/CookedEditor.Automation.cs
+++ b/Engine/Source/Programs/AutomationTool/CookedEditor/CookedEditor.Automation.cs
@@ -1061,6 +1061,7 @@ public class MakeCookedEditor : BuildCommand
 			, SkipBuildClient: !Options.Contains("build", StringComparer.InvariantCultureIgnoreCase)
 			, SkipCook: !Options.Contains("cook", StringComparer.InvariantCultureIgnoreCase)
 			, SkipStage: !Options.Contains("stage", StringComparer.InvariantCultureIgnoreCase)
+			, OptionalContent: Options.Contains("editoroptional", StringComparer.InvariantCultureIgnoreCase)
 		);
 
 		// if the MainParams override the ReleaseVersion path, use it directly
