UAT: Fixed issue with not finding packagestore.manifest when staging plugins from project folders not directly under the workspace root.

diff --git a/Engine/Source/Programs/AutomationTool/Scripts/CopyBuildToStagingDirectory.Automation.cs b/Engine/Source/Programs/AutomationTool/Scripts/CopyBuildToStagingDirectory.Automation.cs
--- a/Engine/Source/Programs/AutomationTool/Scripts/CopyBuildToStagingDirectory.Automation.cs
+++ b/Engine/Source/Programs/AutomationTool/Scripts/CopyBuildToStagingDirectory.Automation.cs
@@ -658,7 +658,23 @@ namespace AutomationScripts
 			{
 				// Making a plugin
 				DirectoryReference DLCRoot = Params.DLCFile.Directory;
-				string RelativeDLCRootPath = (Params.DLCOverrideCookedSubDir == null) ? DLCRoot.MakeRelativeTo(SC.LocalRoot) : Params.DLCOverrideCookedSubDir;
+				string DLCCookedSubDir;
+				if (Params.DLCOverrideCookedSubDir != null)
+				{
+					DLCCookedSubDir = Params.DLCOverrideCookedSubDir;
+				}
+				else if (DLCRoot.IsUnderDirectory(SC.EngineRoot))
+				{
+					DLCCookedSubDir = Path.Combine("Engine", DLCRoot.MakeRelativeTo(SC.EngineRoot));
+				}
+				else if (DLCRoot.IsUnderDirectory(SC.ProjectRoot))
+				{
+					DLCCookedSubDir = Path.Combine(SC.ShortProjectName, DLCRoot.MakeRelativeTo(SC.ProjectRoot));
+				}
+				else
+				{
+					DLCCookedSubDir = DLCRoot.MakeRelativeTo(SC.LocalRoot);
+				}
 
 				// Put all of the cooked dir into the staged dir
 				if (String.IsNullOrEmpty(Params.CookOutputDir))
@@ -676,7 +692,7 @@ namespace AutomationScripts
 
 				DirectoryReference PlatformEngineDir = DirectoryReference.Combine(SC.PlatformCookDir, "Engine");
 				DirectoryReference ProjectMetadataDir = DirectoryReference.Combine(SC.PlatformCookDir, SC.ShortProjectName, "Metadata");
-				SC.MetadataDir = DirectoryReference.Combine(SC.PlatformCookDir, RelativeDLCRootPath, "Metadata");
+				SC.MetadataDir = DirectoryReference.Combine(SC.PlatformCookDir, DLCCookedSubDir, "Metadata");
 
 				// The .uplugin file is staged differently for different DLC
 				// The .uplugin file doesn't actually exist for mobile DLC
