- FIxed a Turnkey bug where flags weren't surrounded in quotes, so the editor wasn't parsting all of the flags properly

diff --git a/Engine/Source/Programs/AutomationTool/Turnkey/Commands/VerifySdk.cs b/Engine/Source/Programs/AutomationTool/Turnkey/Commands/VerifySdk.cs
--- a/Engine/Source/Programs/AutomationTool/Turnkey/Commands/VerifySdk.cs
+++ b/Engine/Source/Programs/AutomationTool/Turnkey/Commands/VerifySdk.cs
@@ -103,7 +103,7 @@ namespace Turnkey.Commands
 						PlatformProperties.Add(SDKInfo.ToString("Allowed", ""));
 					}
 
-					PlatformProperties.Add($"Flags={ReportedState}");
+					PlatformProperties.Add($"Flags=\"{ReportedState}\"");
 
 					if (TurnkeyContext.ErrorMessages.Count() > 0)
 					{
