UAT: Add additional startup and BuildCookRun logging and timings

#rnx

diff --git a/Engine/Source/Programs/AutomationTool/AutomationUtils/Automation.cs b/Engine/Source/Programs/AutomationTool/AutomationUtils/Automation.cs
--- a/Engine/Source/Programs/AutomationTool/AutomationUtils/Automation.cs
+++ b/Engine/Source/Programs/AutomationTool/AutomationUtils/Automation.cs
@@ -237,6 +237,7 @@ namespace AutomationTool
 		/// <param name="Commands"></param>
 		public static async Task<ExitCode> ExecuteAsync(List<CommandInfo> CommandsToExecute, Dictionary<string, Type> Commands)
 		{
+			CommandUtils.LogInformation("Executing commands...");
 			for (int CommandIndex = 0; CommandIndex < CommandsToExecute.Count; ++CommandIndex)
 			{
 				var CommandInfo = CommandsToExecute[CommandIndex];
diff --git a/Engine/Source/Programs/AutomationTool/Program.cs b/Engine/Source/Programs/AutomationTool/Program.cs
--- a/Engine/Source/Programs/AutomationTool/Program.cs
+++ b/Engine/Source/Programs/AutomationTool/Program.cs
@@ -315,6 +315,7 @@ namespace AutomationToolDriver
 		{
 			// Initialize the log system, buffering the output until we can create the log file
 			Log.AddTraceListener(StartupListener);
+			Log.TraceInformation("Starting AutomationTool...");
 
 			// Populate AutomationToolCommandLine and CommandsToExecute
 			try
@@ -409,6 +410,8 @@ namespace AutomationToolDriver
 
 		static async Task<ExitCode> MainProc()
 		{
+			Log.TraceInformation("Initializing script modules...");
+			var StartTime = DateTime.UtcNow;
 			string ScriptsForProject = (string)AutomationToolCommandLine.GetValueUnchecked("-ScriptsForProject");
 			List<string> AdditionalScriptDirs = (List<string>) AutomationToolCommandLine.GetValueUnchecked("-ScriptDir");
 			bool bForceCompile = AutomationToolCommandLine.IsSetGlobal("-Compile");
@@ -449,6 +452,7 @@ namespace AutomationToolDriver
 
 			Type AutomationTools_Automation = AutomationUtilsAssembly.GetType("AutomationTool.Automation");
 			MethodInfo Automation_Process = AutomationTools_Automation.GetMethod("ProcessAsync");
+			Log.TraceInformation("Total script module initialization time: {0:0.00} s.", (DateTime.UtcNow - StartTime).TotalMilliseconds / 1000);
 			return await (Task<ExitCode>) Automation_Process.Invoke(null,
 				new object[] {AutomationToolCommandLine, StartupListener, ScriptModuleAssemblyPaths});
 		}
diff --git a/Engine/Source/Programs/AutomationTool/Scripts/ArchiveCommand.Automation.cs b/Engine/Source/Programs/AutomationTool/Scripts/ArchiveCommand.Automation.cs
--- a/Engine/Source/Programs/AutomationTool/Scripts/ArchiveCommand.Automation.cs
+++ b/Engine/Source/Programs/AutomationTool/Scripts/ArchiveCommand.Automation.cs
@@ -54,6 +54,7 @@ namespace AutomationScripts
 			}
 
 			LogInformation("********** ARCHIVE COMMAND STARTED **********");
+			var StartTime = DateTime.UtcNow;
 
 			LogInformation("Archiving to {0}", Params.ArchiveDirectoryParam);
 
@@ -79,7 +80,8 @@ namespace AutomationScripts
 					SC.StageTargetPlatform.ProcessArchivedProject(Params, SC);
 				}
 			}
+			LogInformation("Archive command time: {0:0.00} s", (DateTime.UtcNow - StartTime).TotalMilliseconds / 1000);
 			LogInformation("********** ARCHIVE COMMAND COMPLETED **********");
 		}
 	}
-}
\ No newline at end of file
+}
diff --git a/Engine/Source/Programs/AutomationTool/Scripts/BuildCookRun.Automation.cs b/Engine/Source/Programs/AutomationTool/Scripts/BuildCookRun.Automation.cs
--- a/Engine/Source/Programs/AutomationTool/Scripts/BuildCookRun.Automation.cs
+++ b/Engine/Source/Programs/AutomationTool/Scripts/BuildCookRun.Automation.cs
@@ -24,6 +24,8 @@ public class BuildCookRun : BuildCommand
 {
 	public override void ExecuteBuild()
 	{
+		var StartTime = DateTime.UtcNow;
+
 		// these need to be done first
 		var bForeign = ParseParam("foreign");
 		var bForeignCode = ParseParam("foreigncode");
@@ -38,6 +40,8 @@ public class BuildCookRun : BuildCommand
 		var Params = SetupParams();
 
 		DoBuildCookRun(Params);
+
+		LogInformation("BuildCookRun time: {0:0.00} s", (DateTime.UtcNow - StartTime).TotalMilliseconds / 1000);
 	}
 
 	protected ProjectParams SetupParams()
diff --git a/Engine/Source/Programs/AutomationTool/Scripts/BuildProjectCommand.Automation.cs b/Engine/Source/Programs/AutomationTool/Scripts/BuildProjectCommand.Automation.cs
--- a/Engine/Source/Programs/AutomationTool/Scripts/BuildProjectCommand.Automation.cs
+++ b/Engine/Source/Programs/AutomationTool/Scripts/BuildProjectCommand.Automation.cs
@@ -70,6 +70,7 @@ namespace AutomationScripts
 			}
 
 			LogInformation("********** BUILD COMMAND STARTED **********");
+			var StartTime = DateTime.UtcNow;
 
 			var UnrealBuild = new UnrealBuild(Command);
 			var Agenda = new UnrealBuild.BuildAgenda();
@@ -234,6 +235,7 @@ namespace AutomationScripts
 				UnrealBuild.AddBuildProductsToChangelist(WorkingCL, UnrealBuild.BuildProductFiles);
 			}
 
+			LogInformation("Build command time: {0:0.00} s", (DateTime.UtcNow - StartTime).TotalMilliseconds / 1000);
 			LogInformation("********** BUILD COMMAND COMPLETED **********");
 		}
 	}
diff --git a/Engine/Source/Programs/AutomationTool/Scripts/CookCommand.Automation.cs b/Engine/Source/Programs/AutomationTool/Scripts/CookCommand.Automation.cs
--- a/Engine/Source/Programs/AutomationTool/Scripts/CookCommand.Automation.cs
+++ b/Engine/Source/Programs/AutomationTool/Scripts/CookCommand.Automation.cs
@@ -252,6 +252,7 @@ namespace AutomationScripts
 			Params.ValidateAndLog();
 
 			LogInformation("********** COOK COMMAND STARTED **********");
+			var StartTime = DateTime.UtcNow;
 
 			string UEEditorExe = HostPlatform.Current.GetUnrealExePath(Params.UnrealExe);
 			if (!FileExists(UEEditorExe))
@@ -327,6 +328,7 @@ namespace AutomationScripts
 			}
 
 
+			LogInformation("Cook command time: {0:0.00} s", (DateTime.UtcNow - StartTime).TotalMilliseconds / 1000);
 			LogInformation("********** COOK COMMAND COMPLETED **********");
 		}
 
diff --git a/Engine/Source/Programs/AutomationTool/Scripts/CopyBuildToStagingDirectory.Automation.cs b/Engine/Source/Programs/AutomationTool/Scripts/CopyBuildToStagingDirectory.Automation.cs
--- a/Engine/Source/Programs/AutomationTool/Scripts/CopyBuildToStagingDirectory.Automation.cs
+++ b/Engine/Source/Programs/AutomationTool/Scripts/CopyBuildToStagingDirectory.Automation.cs
@@ -4453,6 +4453,7 @@ namespace AutomationScripts
 				Params.ValidateAndLog();
 
 				LogInformation("********** STAGE COMMAND STARTED **********");
+				var StartTime = DateTime.UtcNow;
 
 				if (!Params.NoClient)
 				{
@@ -4526,6 +4527,7 @@ namespace AutomationScripts
 						ApplyStagingManifest(Params, SC);
 					}
 				}
+				LogInformation("Stage command time: {0:0.00} s", (DateTime.UtcNow - StartTime).TotalMilliseconds / 1000);
 				LogInformation("********** STAGE COMMAND COMPLETED **********");
 			}
 		}
diff --git a/Engine/Source/Programs/AutomationTool/Scripts/DeployCommand.Automation.cs b/Engine/Source/Programs/AutomationTool/Scripts/DeployCommand.Automation.cs
--- a/Engine/Source/Programs/AutomationTool/Scripts/DeployCommand.Automation.cs
+++ b/Engine/Source/Programs/AutomationTool/Scripts/DeployCommand.Automation.cs
@@ -20,6 +20,7 @@ namespace AutomationScripts
 			}
 
 			LogInformation("********** DEPLOY COMMAND STARTED **********");
+			var StartTime = DateTime.UtcNow;
 
 			if (!Params.NoClient)
 			{
@@ -68,6 +69,7 @@ namespace AutomationScripts
 				}
 			}
 
+			LogInformation("Deploy command time: {0:0.00} s", (DateTime.UtcNow - StartTime).TotalMilliseconds / 1000);
 			LogInformation("********** DEPLOY COMMAND COMPLETED **********");
 		}
 	}
diff --git a/Engine/Source/Programs/AutomationTool/Scripts/GetFileCommand.Automation.cs b/Engine/Source/Programs/AutomationTool/Scripts/GetFileCommand.Automation.cs
--- a/Engine/Source/Programs/AutomationTool/Scripts/GetFileCommand.Automation.cs
+++ b/Engine/Source/Programs/AutomationTool/Scripts/GetFileCommand.Automation.cs
@@ -22,6 +22,7 @@ namespace AutomationScripts
 			}
 
 			LogInformation("********** GETFILE COMMAND STARTED **********");
+			var StartTime = DateTime.UtcNow;
 
 			var FileName = Path.GetFileName(Params.GetFile);
 			var LocalFile = CombinePaths(CmdEnv.EngineSavedFolder, FileName);
@@ -34,6 +35,7 @@ namespace AutomationScripts
 
 			SC[0].StageTargetPlatform.GetTargetFile(Params.GetFile, LocalFile, Params);
 
+			LogInformation("GetFile command time: {0:0.00} s", (DateTime.UtcNow - StartTime).TotalMilliseconds / 1000);
 			LogInformation("********** GETFILE COMMAND COMPLETED **********");
 		}
 	}
diff --git a/Engine/Source/Programs/AutomationTool/Scripts/PackageCommand.Automation.cs b/Engine/Source/Programs/AutomationTool/Scripts/PackageCommand.Automation.cs
--- a/Engine/Source/Programs/AutomationTool/Scripts/PackageCommand.Automation.cs
+++ b/Engine/Source/Programs/AutomationTool/Scripts/PackageCommand.Automation.cs
@@ -26,9 +26,20 @@ namespace AutomationScripts
 					DeployContextList.AddRange(CreateDeploymentContext(Params, true, false));
 				}
 
-				if (DeployContextList.Count > 0)
+				bool bShouldPackage = false; 
+				foreach (var SC in DeployContextList)
+				{
+					if (Params.Package || (SC.StageTargetPlatform.RequiresPackageToDeploy && Params.Deploy))
+					{
+						bShouldPackage = true;
+						break;
+					}
+				}
+
+				if (bShouldPackage)
 				{
 					LogInformation("********** PACKAGE COMMAND STARTED **********");
+					var StartTime = DateTime.UtcNow;
 
 					foreach (var SC in DeployContextList)
 					{
@@ -38,6 +49,7 @@ namespace AutomationScripts
 						}
 					}
 
+					LogInformation("Package command time: {0:0.00} s", (DateTime.UtcNow - StartTime).TotalMilliseconds / 1000);
 					LogInformation("********** PACKAGE COMMAND COMPLETED **********");
 				}
 			}
diff --git a/Engine/Source/Programs/AutomationTool/Scripts/RunProjectCommand.Automation.cs b/Engine/Source/Programs/AutomationTool/Scripts/RunProjectCommand.Automation.cs
--- a/Engine/Source/Programs/AutomationTool/Scripts/RunProjectCommand.Automation.cs
+++ b/Engine/Source/Programs/AutomationTool/Scripts/RunProjectCommand.Automation.cs
@@ -150,6 +150,7 @@ namespace AutomationScripts
 			}
 
 			LogInformation("********** RUN COMMAND STARTED **********");
+			var StartTime = DateTime.UtcNow;
 
 			var LogFolderOutsideOfSandbox = GetLogFolderOutsideOfSandbox();
 			if (!Unreal.IsEngineInstalled() && CookServerProcess == null)
@@ -188,6 +189,7 @@ namespace AutomationScripts
 				CopyLogsBackToLogFolder();
 			}
 
+			LogInformation("Run command time: {0:0.00} s", (DateTime.UtcNow - StartTime).TotalMilliseconds / 1000);
 			LogInformation("********** RUN COMMAND COMPLETED **********");
 		}
 
diff --git a/Engine/Source/Programs/Shared/EpicGames.MsBuild/CsProjBuilder.cs b/Engine/Source/Programs/Shared/EpicGames.MsBuild/CsProjBuilder.cs
--- a/Engine/Source/Programs/Shared/EpicGames.MsBuild/CsProjBuilder.cs
+++ b/Engine/Source/Programs/Shared/EpicGames.MsBuild/CsProjBuilder.cs
@@ -461,6 +461,7 @@ namespace EpicGames.MsBuild
 
 		private static bool BuildProjects(ProjectGraph ProjectGraph, Dictionary<string, string> GlobalProperties)
 		{
+			var StartTime = DateTime.UtcNow;
 			var Logger = new MLogger();
 
 			string[] TargetsToBuild = { "Restore", "Build" };
@@ -502,7 +503,7 @@ namespace EpicGames.MsBuild
 					Result = false;
 				}
 			}
-			Log.TraceInformation(" build complete.");
+			Log.TraceInformation("Build projects time: {0:0.00} s", (DateTime.UtcNow - StartTime).TotalMilliseconds / 1000);
 
 			return Result;
 		}
