Low-Level test runner for windows-based platform extensions.

#rnx

diff --git a/Engine/Source/Programs/AutomationTool/LowLevelTests/RunLowLevelTests.cs b/Engine/Source/Programs/AutomationTool/LowLevelTests/RunLowLevelTests.cs
--- a/Engine/Source/Programs/AutomationTool/LowLevelTests/RunLowLevelTests.cs
+++ b/Engine/Source/Programs/AutomationTool/LowLevelTests/RunLowLevelTests.cs
@@ -471,7 +471,7 @@ namespace LowLevelTests
 		public LowLevelTestsBuild CreateBuild(UnrealTargetPlatform InPlatform, string InTestApp, string InBuildPath)
 		{
 			string DesktopExecutableRegEx;
-			if (InPlatform == UnrealTargetPlatform.Win64)
+			if (InPlatform.IsInGroup(UnrealPlatformGroup.Windows))
 			{
 				DesktopExecutableRegEx = @"\w+Tests.exe$";
 			}
@@ -505,7 +505,7 @@ namespace LowLevelTests
 			string ReportRelativePath = GetTargetReportPath(InPlatform, InTestApp, InBuildPath);
 			string ReportPath = Path.Combine(InBuildPath, ReportRelativePath);
 			string ExpectedLocalPath = Path.Combine(InTargetDirectory, ReportRelativePath);
-			File.Copy(ReportPath, InTargetDirectory);
+			File.Copy(ReportPath, ExpectedLocalPath);
 			return ExpectedLocalPath;
 		}
 	}
