- Added option for IoStore to write optional files (.o.ucas) to a different directory. This is used when making a a cooked editor as DLC against a client build (the client build has to make the optional editor data, but we do not want to ship the optional data with the client). The MakeCookedEditor script will then copy the files from that location into the editor's Stage directory (under Content/Paks)

diff --git a/Engine/Source/Developer/IoStoreUtilities/Private/IoStoreUtilities.cpp b/Engine/Source/Developer/IoStoreUtilities/Private/IoStoreUtilities.cpp
--- a/Engine/Source/Developer/IoStoreUtilities/Private/IoStoreUtilities.cpp
+++ b/Engine/Source/Developer/IoStoreUtilities/Private/IoStoreUtilities.cpp
@@ -286,6 +286,7 @@ struct FContainerSourceSpec
 {
 	FName Name;
 	FString OutputPath;
+	FString OptionalOutputPath;
 	TArray<FContainerSourceFile> SourceFiles;
 	FString PatchTargetFile;
 	TArray<FString> PatchSourceContainerFiles;
@@ -2280,7 +2281,15 @@ void InitializeContainerTargetsAndPackages(
 
 			if (bHasOptionalSegmentPackages)
 			{
-				ContainerTarget->OptionalSegmentOutputPath = ContainerTarget->OutputPath + FPackagePath::GetOptionalSegmentExtensionModifier();
+				if (ContainerSource.OptionalOutputPath.IsEmpty())
+				{
+					ContainerTarget->OptionalSegmentOutputPath = ContainerTarget->OutputPath + FPackagePath::GetOptionalSegmentExtensionModifier();
+				}
+				else
+				{
+					// if we have an optional output location, use that directory, combined with the name of the output path 
+					ContainerTarget->OptionalSegmentOutputPath = FPaths::Combine(ContainerSource.OptionalOutputPath, FPaths::GetCleanFilename(ContainerTarget->OutputPath) + FPackagePath::GetOptionalSegmentExtensionModifier());
+				}
 			}
 		}
 	}
@@ -5866,6 +5875,8 @@ bool ParseContainerGenerationArguments(FIoStoreArguments& Arguments, FIoStoreWri
 			}
 			ContainerSpec.OutputPath = FPaths::ChangeExtension(ContainerSpec.OutputPath, TEXT(""));
 
+			FParse::Value(*Command, TEXT("OptionalOutput="), ContainerSpec.OptionalOutputPath);
+
 			FString ContainerName;
 			if (FParse::Value(*Command, TEXT("ContainerName="), ContainerName))
 			{
diff --git a/Engine/Source/Programs/AutomationTool/AutomationUtils/DeploymentContext.cs b/Engine/Source/Programs/AutomationTool/AutomationUtils/DeploymentContext.cs
--- a/Engine/Source/Programs/AutomationTool/AutomationUtils/DeploymentContext.cs
+++ b/Engine/Source/Programs/AutomationTool/AutomationUtils/DeploymentContext.cs
@@ -196,6 +196,11 @@ public class DeploymentContext //: ProjectParams
 	/// </summary>
 	public DirectoryReference DebugStageDirectory;
 
+	/// <summary>
+	/// Directory to put all of the optional (ie editor only) files in: d:\stagedir\Windows
+	/// </summary>
+	public DirectoryReference OptionalFileStageDirectory = null;
+
 	/// <summary>
 	/// Directory name for staged projects
 	/// </summary>
@@ -381,6 +386,7 @@ public class DeploymentContext //: ProjectParams
 		FileReference RawProjectPathOrName,
 		DirectoryReference InLocalRoot,
 		DirectoryReference BaseStageDirectory,
+		DirectoryReference OptionalStageDirectory,
 		DirectoryReference BaseArchiveDirectory,
 		Platform InSourcePlatform,
         Platform InTargetPlatform,
@@ -448,6 +454,7 @@ public class DeploymentContext //: ProjectParams
 			StageDirectory = DirectoryReference.Combine(BaseStageDirectory, FinalCookPlatform);
 			DebugStageDirectory = InSeparateDebugStageDirectory? DirectoryReference.Combine(BaseStageDirectory, FinalCookPlatform + "Debug") : StageDirectory;
 		}
+		OptionalFileStageDirectory = OptionalStageDirectory;
 
 		if (BaseArchiveDirectory != null)
 		{
diff --git a/Engine/Source/Programs/AutomationTool/AutomationUtils/ProjectParams.cs b/Engine/Source/Programs/AutomationTool/AutomationUtils/ProjectParams.cs
--- a/Engine/Source/Programs/AutomationTool/AutomationUtils/ProjectParams.cs
+++ b/Engine/Source/Programs/AutomationTool/AutomationUtils/ProjectParams.cs
@@ -1481,7 +1481,10 @@ namespace AutomationTool
 
 		[Help("stagingdirectory=Path", "Directory to copy the builds to, i.e. -stagingdirectory=C:\\Stage")]
 		public string StageDirectoryParam;
-        
+
+		[Help("optionalstagingdirectory=Path", "Directory to copy the optional files to, i.e. -optionalstagingdirectory=C:\\StageOptional")]
+		public string OptionalStageDirectory;
+
 		[Help("unrealexe=ExecutableName", "Name of the Unreal Editor executable, i.e. -unrealexe=UnrealEditor.exe")]
 		public string UnrealExe;
 
diff --git a/Engine/Source/Programs/AutomationTool/CookedEditor/CookedEditor.Automation.cs b/Engine/Source/Programs/AutomationTool/CookedEditor/CookedEditor.Automation.cs
--- a/Engine/Source/Programs/AutomationTool/CookedEditor/CookedEditor.Automation.cs
+++ b/Engine/Source/Programs/AutomationTool/CookedEditor/CookedEditor.Automation.cs
@@ -375,6 +375,10 @@ public class MakeCookedEditor : BuildCommand
 	protected DirectoryReference CookedEditorStageDirectory = null;
 	protected DirectoryReference ReleaseStageDirectory = null;
 
+	// with -makerelease, this will have the location of optional editor only files, but a subclass can just set this if the optional files
+	// were made and saved off somewhere, it can point to this and the optional files will be automatically staged into Content/Paks
+	protected DirectoryReference ReleaseOptionalFileStageDirectory = null;
+
 	public override void ExecuteBuild()
 	{
 		LogInformation("************************* MakeCookedEditor");
@@ -494,6 +498,13 @@ public class MakeCookedEditor : BuildCommand
 				SC.StageFile(StagedFileType.UFS, EditorThumbnails, new StagedFileReference($"{Context.ProjectName}/CachedEditorThumbnails.bin"));
 			}
 		}
+
+		if (ReleaseOptionalFileStageDirectory != null)
+		{
+			// these files were already staged by the client/release build, so we stage them as NonUFS
+			// these could be just copied, but StageFiles handles copying easily
+			SC.StageFiles(StagedFileType.NonUFS, ReleaseOptionalFileStageDirectory, StageFilesSearch.AllDirectories, new StagedDirectoryReference($"{Context.ProjectName}/Content/Paks"));
+		}
 	}
 
 	static void ReadProjectsRecursively(FileReference File, Dictionary<string, string> InitialProperties, Dictionary<FileReference, CsProjectInfo> FileToProjectInfo)
@@ -650,7 +661,6 @@ public class MakeCookedEditor : BuildCommand
 
 	protected virtual void FinalizeRelease(ProjectParams ReleaseParams)
 	{
-
 	}
 
 	protected virtual void PreModifyDeploymentContext(ProjectParams Params, DeploymentContext SC)
@@ -1073,6 +1083,9 @@ public class MakeCookedEditor : BuildCommand
 		// copy off the staging dir
 		ReleaseParams.PreModifyDeploymentContextCallback = new Action<ProjectParams, DeploymentContext>((ProjectParams P, DeploymentContext SC) => { ReleaseStageDirectory = SC.StageDirectory; });
 
+		ReleaseOptionalFileStageDirectory = DirectoryReference.Combine(MainParams.RawProjectPath.Directory, "Saved", "CookedEditor", "OptionalData");
+		ReleaseParams.OptionalStageDirectory = ReleaseOptionalFileStageDirectory.FullName;
+
 		ModifyReleaseParams(ReleaseParams);
 
 		return ReleaseParams;
diff --git a/Engine/Source/Programs/AutomationTool/Scripts/CopyBuildToStagingDirectory.Automation.cs b/Engine/Source/Programs/AutomationTool/Scripts/CopyBuildToStagingDirectory.Automation.cs
--- a/Engine/Source/Programs/AutomationTool/Scripts/CopyBuildToStagingDirectory.Automation.cs
+++ b/Engine/Source/Programs/AutomationTool/Scripts/CopyBuildToStagingDirectory.Automation.cs
@@ -243,6 +243,7 @@ namespace AutomationScripts
 			Dictionary<string, string> UnrealPakResponseFile,
 			string ContainerName,
 			FileReference PakOutputLocation,
+			DirectoryReference OptionalOutputLocation,
 			bool bCompressed,
 			EncryptionAndSigning.CryptoSettings CryptoSettings,
 			string EncryptionKeyGuid,
@@ -252,6 +253,10 @@ namespace AutomationScripts
 		{
 			StringBuilder CmdLine = new StringBuilder();
 			CmdLine.AppendFormat("-Output={0}", MakePathSafeToUseWithCommandLine(Path.ChangeExtension(PakOutputLocation.FullName, ".utoc")));
+			if (OptionalOutputLocation != null)
+			{
+				CmdLine.AppendFormat(" -OptionalOutput={0}", MakePathSafeToUseWithCommandLine(OptionalOutputLocation.FullName));
+			}
 			CmdLine.AppendFormat(" -ContainerName={0}", ContainerName);
 			if (!bIsDLC && !String.IsNullOrEmpty(PatchSourceContentPath))
 			{
@@ -3001,6 +3006,7 @@ namespace AutomationScripts
 								IoStoreResponseFile,
 								OutputLocation.GetFileNameWithoutAnyExtensions(),
 								OutputLocation,
+								SC.OptionalFileStageDirectory,
 								bCompressContainers,
 								CryptoSettings,
 								PakParams.EncryptionKeyGuid,
@@ -4005,6 +4011,7 @@ namespace AutomationScripts
 					try
 					{
 						DeleteDirectory(SC.StageDirectory.FullName);
+						DeleteDirectory(SC.OptionalFileStageDirectory.FullName);
 					}
 					catch (Exception Ex)
 					{
@@ -4406,6 +4413,7 @@ namespace AutomationScripts
 				//@todo should pull StageExecutables from somewhere else if not cooked
 				var SC = new DeploymentContext(Params.RawProjectPath, Unreal.RootDirectory,
 					String.IsNullOrEmpty(StageDirectory) ? null : new DirectoryReference(StageDirectory),
+					String.IsNullOrEmpty(Params.OptionalStageDirectory) ? null : new DirectoryReference(Params.OptionalStageDirectory),
 					String.IsNullOrEmpty(ArchiveDirectory) ? null : new DirectoryReference(ArchiveDirectory),
 					Platform.Platforms[CookedDataPlatform],
 					Platform.Platforms[StagePlatform],
