- FIxed an exception when NOT not using oiptional stage dir

diff --git a/Engine/Source/Programs/AutomationTool/Scripts/CopyBuildToStagingDirectory.Automation.cs b/Engine/Source/Programs/AutomationTool/Scripts/CopyBuildToStagingDirectory.Automation.cs
--- a/Engine/Source/Programs/AutomationTool/Scripts/CopyBuildToStagingDirectory.Automation.cs
+++ b/Engine/Source/Programs/AutomationTool/Scripts/CopyBuildToStagingDirectory.Automation.cs
@@ -4011,13 +4011,25 @@ namespace AutomationScripts
 					try
 					{
 						DeleteDirectory(SC.StageDirectory.FullName);
-						DeleteDirectory(SC.OptionalFileStageDirectory.FullName);
 					}
 					catch (Exception Ex)
 					{
 						// Delete cooked data (if any) as it may be incomplete / corrupted.
 						throw new AutomationException(ExitCode.Error_FailedToDeleteStagingDirectory, Ex, "Stage Failed. Failed to delete staging directory " + SC.StageDirectory.FullName);
 					}
+					if (SC.OptionalFileStageDirectory != null && DirectoryReference.Exists(SC.OptionalFileStageDirectory))
+					{
+						LogInformation("Cleaning Optional Stage Directory: {0}", SC.OptionalFileStageDirectory.FullName);
+						try
+						{
+							DeleteDirectory(SC.OptionalFileStageDirectory.FullName);
+						}
+						catch (Exception Ex)
+						{
+							// Delete cooked data (if any) as it may be incomplete / corrupted.
+							throw new AutomationException(ExitCode.Error_FailedToDeleteStagingDirectory, Ex, "Stage Failed. Failed to delete staging directory " + SC.OptionalFileStageDirectory.FullName);
+						}
+					}
 				}
 			}
 			else
