Use platform-remapped staging path when checking for restricted folder names

diff --git a/Engine/Source/Programs/AutomationTool/AutomationUtils/Platform.cs b/Engine/Source/Programs/AutomationTool/AutomationUtils/Platform.cs
--- a/Engine/Source/Programs/AutomationTool/AutomationUtils/Platform.cs
+++ b/Engine/Source/Programs/AutomationTool/AutomationUtils/Platform.cs
@@ -645,13 +645,21 @@ namespace AutomationTool
 		}
 
 		/// <summary>
-		/// Remaps the given content directory to its final location
+		/// Remaps the given content file to its final location
 		/// </summary>
 		public virtual StagedFileReference Remap(StagedFileReference Dest)
 		{
 			return Dest;
 		}
 
+		/// <summary>
+		/// Remaps the given content directory to its final location
+		/// </summary>
+		public virtual StagedDirectoryReference Remap(StagedDirectoryReference Dest)
+		{
+			return Dest;
+		}
+
 		/// <summary>
 		/// Tri-state - The intent is to override command line parameters for pak if needed per platform.
 		/// </summary>
diff --git a/Engine/Source/Programs/AutomationTool/IOS/IOSPlatform.Automation.cs b/Engine/Source/Programs/AutomationTool/IOS/IOSPlatform.Automation.cs
--- a/Engine/Source/Programs/AutomationTool/IOS/IOSPlatform.Automation.cs
+++ b/Engine/Source/Programs/AutomationTool/IOS/IOSPlatform.Automation.cs
@@ -1941,6 +1941,11 @@ public class IOSPlatform : Platform
 	{
 		return new StagedFileReference("cookeddata/" + Dest.Name);
 	}
+	public override StagedDirectoryReference Remap(StagedDirectoryReference Dest)
+	{
+		return new StagedDirectoryReference("cookeddata/" + Dest.Name);
+	}
+
 	public override List<string> GetDebugFileExtensions()
 	{
 		return new List<string> { ".dsym", ".udebugsymbols" };
diff --git a/Engine/Source/Programs/AutomationTool/Scripts/CopyBuildToStagingDirectory.Automation.cs b/Engine/Source/Programs/AutomationTool/Scripts/CopyBuildToStagingDirectory.Automation.cs
--- a/Engine/Source/Programs/AutomationTool/Scripts/CopyBuildToStagingDirectory.Automation.cs
+++ b/Engine/Source/Programs/AutomationTool/Scripts/CopyBuildToStagingDirectory.Automation.cs
@@ -1371,7 +1371,7 @@ namespace AutomationScripts
 				RestrictedFiles.UnionWith(SC.FilesToStage.NonUFSDebugFiles.Keys.Where(x => x.ContainsName(RestrictedName)));
 				RestrictedFiles.UnionWith(SC.CrashReporterUFSFiles.Keys.Where(x => x.ContainsName(RestrictedName)));
 			}
-			RestrictedFiles.RemoveWhere(RestrictedFile => SC.DirectoriesAllowList.Any(TestDirectory => RestrictedFile.Directory.IsUnderDirectory(TestDirectory)));
+			RestrictedFiles.RemoveWhere(RestrictedFile => SC.DirectoriesAllowList.Any(TestDirectory => RestrictedFile.Directory.IsUnderDirectory(SC.StageTargetPlatform.Remap(TestDirectory))));
 
 			if (RestrictedFiles.Count > 0)
 			{
