- Added missing commandline option for the optional file staging directory (-optionalfilestagingdirectory=)

diff --git a/Engine/Source/Programs/AutomationTool/AutomationUtils/ProjectParams.cs b/Engine/Source/Programs/AutomationTool/AutomationUtils/ProjectParams.cs
--- a/Engine/Source/Programs/AutomationTool/AutomationUtils/ProjectParams.cs
+++ b/Engine/Source/Programs/AutomationTool/AutomationUtils/ProjectParams.cs
@@ -801,6 +801,7 @@ namespace AutomationTool
 				this.Stage = true;
 			}
 			this.StageDirectoryParam = ParseParamValueIfNotSpecified(Command, StageDirectoryParam, "stagingdirectory", String.Empty, true);
+			this.OptionalStageDirectory = ParseParamValueIfNotSpecified(Command, OptionalStageDirectory, "optionalfilestagingdirectory", String.Empty, true);
 			this.bCodeSign = GetOptionalParamValueIfNotSpecified(Command, CodeSign, CommandUtils.IsBuildMachine, "CodeSign", "NoCodeSign");
 			this.bTreatNonShippingBinariesAsDebugFiles = GetParamValueIfNotSpecified(Command, TreatNonShippingBinariesAsDebugFiles, false, "TreatNonShippingBinariesAsDebugFiles");
 			this.bUseExtraFlavor = GetParamValueIfNotSpecified(Command, UseExtraFlavor, false, "UseExtraFlavor");
@@ -1482,7 +1483,7 @@ namespace AutomationTool
 		[Help("stagingdirectory=Path", "Directory to copy the builds to, i.e. -stagingdirectory=C:\\Stage")]
 		public string StageDirectoryParam;
 
-		[Help("optionalstagingdirectory=Path", "Directory to copy the optional files to, i.e. -optionalstagingdirectory=C:\\StageOptional")]
+		[Help("optionalfilestagingdirectory=Path", "Directory to copy the optional files to, i.e. -optionalfilestagingdirectory=C:\\StageOptional")]
 		public string OptionalStageDirectory;
 
 		[Help("unrealexe=ExecutableName", "Name of the Unreal Editor executable, i.e. -unrealexe=UnrealEditor.exe")]
