HoloLensPlatform::GetPathToVCLibsPackages() needs to call MicrosoftPlatformSDKs::EnumerateSdkRootDirs() function and loop through the result trying them in order, rather than always looking for a hardcoded path.  Also adding VCLibs for arm64 necessary for HoloLens to CarefullyRedist.  For some reason, they are in Microsoft SDKs/Windows Kits instead of just in Windows Kits/, so I matched that directory structure in CarefullyRedist, then the code I'm adding the HoloLensPlatform.Automation.cs can use the same path to look for it if it's installed manually or from AutoSDKs.
testing: Quick Launch HoloLensQA project to HoloLens device still works!

diff --git a/Engine/Source/Programs/AutomationTool/HoloLens/HoloLensPlatform.Automation.cs b/Engine/Source/Programs/AutomationTool/HoloLens/HoloLensPlatform.Automation.cs
--- a/Engine/Source/Programs/AutomationTool/HoloLens/HoloLensPlatform.Automation.cs
+++ b/Engine/Source/Programs/AutomationTool/HoloLens/HoloLensPlatform.Automation.cs
@@ -1728,8 +1728,11 @@ namespace HoloLens.Automation
 			}
 		}
 
-        private string[] GetPathToVCLibsPackages(bool UseDebugCrt, WindowsCompiler Compiler)
+		private string[] GetPathToVCLibsPackages(bool UseDebugCrt, WindowsCompiler Compiler)
 		{
+			List<DirectoryReference> SdkRootDirs = new List<DirectoryReference>();
+			WindowsExports.EnumerateSdkRootDirs(SdkRootDirs);
+			
 			string VCVersionFragment;
             switch (Compiler)
 			{
@@ -1747,25 +1750,56 @@ namespace HoloLens.Automation
 
 			List<string> Runtimes = new List<string>();
 
-			foreach(var Arch in ActualArchitectures)
+			bool foundVcLibs = false;
+			foreach (DirectoryReference SdkRootDir in SdkRootDirs)
 			{
-				string ArchitectureFragment = WindowsExports.GetArchitectureSubpath(Arch);
-
-				string RuntimePath = Path.Combine(Environment.GetFolderPath(Environment.SpecialFolder.ProgramFilesX86),
-					"Microsoft SDKs",
-					"Windows Kits",
-					"10",
-					"ExtensionSDKs",
-					"Microsoft.VCLibs",
-					string.Format("{0}.0", VCVersionFragment),
-					"Appx",
-					UseDebugCrt ? "Debug" : "Retail",
-					ArchitectureFragment,
-					string.Format("Microsoft.VCLibs.{0}.{1}.00.appx", ArchitectureFragment, VCVersionFragment));
+				if (!foundVcLibs)
+				{
+					foreach (var Arch in ActualArchitectures)
+					{
+						string ArchitectureFragment = WindowsExports.GetArchitectureSubpath(Arch);
 
-				Runtimes.Add(RuntimePath);
+						// For whatever reason, VCLibs are not in Windows Kits/, but Microsoft SDKs/Windows Kits, so go up two directories and start from there.
+						if (DirectoryExists(SdkRootDir.ParentDirectory.ToString()) && DirectoryExists(SdkRootDir.ParentDirectory.ParentDirectory.ToString()))
+						{
+							DirectoryReference Win64SDKsRootDir = SdkRootDir.ParentDirectory.ParentDirectory;
+
+							string RuntimePath = Path.Combine(Win64SDKsRootDir.ToString(),
+								"Microsoft SDKs",
+								"Windows Kits",
+								"10",
+								"ExtensionSDKs",
+								"Microsoft.VCLibs",
+								string.Format("{0}.0", VCVersionFragment),
+								"Appx",
+								UseDebugCrt ? "Debug" : "Retail",
+								ArchitectureFragment,
+								string.Format("Microsoft.VCLibs.{0}.{1}.00.appx", ArchitectureFragment, VCVersionFragment));
+
+							if (FileExists(RuntimePath))
+							{
+								Runtimes.Add(RuntimePath);
+								LogLog("found vclibs, path: " + RuntimePath);
+								foundVcLibs = true;
+							}
+							else
+							{
+								LogLog("no vclib found, path: " + RuntimePath);
+							}
+						}
+						else
+						{
+							LogLog("Directory structure for VCLibs not present , SDK path: " + SdkRootDir.ToString());
+						}
+					}
+				}
+			}
 
+			if (!foundVcLibs)
+			{
+				LogError("no appropriate VCLibs found in any of the SDK locations");
 			}
+
 			return Runtimes.ToArray();
 		}
 
diff --git a/Engine/Source/Programs/UnrealBuildTool/Platform/Windows/MicrosoftPlatformSDK.cs b/Engine/Source/Programs/UnrealBuildTool/Platform/Windows/MicrosoftPlatformSDK.cs
--- a/Engine/Source/Programs/UnrealBuildTool/Platform/Windows/MicrosoftPlatformSDK.cs
+++ b/Engine/Source/Programs/UnrealBuildTool/Platform/Windows/MicrosoftPlatformSDK.cs
@@ -262,7 +262,7 @@ namespace UnrealBuildTool
 		/// </summary>
 		/// <param name="RootDirs">Receives all the Windows 10 sdk root directories</param>
 		[SupportedOSPlatform("windows")]
-		private static void EnumerateSdkRootDirs(List<DirectoryReference> RootDirs)
+		public static void EnumerateSdkRootDirs(List<DirectoryReference> RootDirs)
 		{
 			DirectoryReference? RootDir;
 			if (TryReadInstallDirRegistryKey32("Microsoft\\Windows Kits\\Installed Roots", "KitsRoot10", out RootDir))
diff --git a/Engine/Source/Programs/UnrealBuildTool/Platform/Windows/WindowsExports.cs b/Engine/Source/Programs/UnrealBuildTool/Platform/Windows/WindowsExports.cs
--- a/Engine/Source/Programs/UnrealBuildTool/Platform/Windows/WindowsExports.cs
+++ b/Engine/Source/Programs/UnrealBuildTool/Platform/Windows/WindowsExports.cs
@@ -95,5 +95,16 @@ namespace UnrealBuildTool
 
 			return WindowsSdkDirs;
 		}
+
+		/// <summary>
+		/// Enumerates all the Windows 10 SDK root directories
+		/// </summary>
+		/// <param name="RootDirs">Receives all the Windows 10 sdk root directories</param>
+		[SupportedOSPlatform("windows")]
+		public static void EnumerateSdkRootDirs(List<DirectoryReference> RootDirs)
+		{
+			MicrosoftPlatformSDK.EnumerateSdkRootDirs(RootDirs);
+		}
+
 	}
 }
