BuildGraph: Add more parameters to Docker-Build task

Supporting BuildKit and .dockerignore files will help improve performance without resorting to temporary staging of files.

diff --git a/Engine/Source/Programs/AutomationTool/BuildGraph/Tasks/DockerBuildTask.cs b/Engine/Source/Programs/AutomationTool/BuildGraph/Tasks/DockerBuildTask.cs
--- a/Engine/Source/Programs/AutomationTool/BuildGraph/Tasks/DockerBuildTask.cs
+++ b/Engine/Source/Programs/AutomationTool/BuildGraph/Tasks/DockerBuildTask.cs
@@ -36,6 +36,24 @@ namespace AutomationTool.Tasks
 		[TaskParameter(Optional = true)]
 		public string DockerFile;
 
+		/// <summary>
+		/// Path to a .dockerignore. Will be copied to basedir if specified.
+		/// </summary>
+		[TaskParameter(Optional = true)]
+		public string DockerIgnoreFile;
+
+		/// <summary>
+		/// Use BuildKit in Docker
+		/// </summary>
+		[TaskParameter(Optional = true)]
+		public bool UseBuildKit;
+
+		/// <summary>
+		/// Type of progress output (--progress)
+		/// </summary>
+		[TaskParameter(Optional = true)]
+		public string ProgressOutput;
+
 		/// <summary>
 		/// Tag for the image
 		/// </summary>
@@ -100,6 +118,7 @@ namespace AutomationTool.Tasks
 			{
 				DirectoryReference BaseDir = ResolveDirectory(Parameters.BaseDir);
 				List<FileReference> SourceFiles = ResolveFilespec(BaseDir, Parameters.Files, TagNameToFileSet).ToList();
+				bool isStagingEnabled = SourceFiles.Count > 0;
 
 				DirectoryReference StagingDir = DirectoryReference.Combine(Unreal.EngineDirectory, "Intermediate", "Docker");
 				FileUtils.ForceDeleteDirectoryContents(StagingDir);
@@ -107,6 +126,9 @@ namespace AutomationTool.Tasks
 				List<FileReference> TargetFiles = SourceFiles.ConvertAll(x => FileReference.Combine(StagingDir, x.MakeRelativeTo(BaseDir)));
 				CommandUtils.ThreadedCopyFiles(SourceFiles, BaseDir, StagingDir);
 
+				FileReference DockerIgnoreFileInBaseDir = FileReference.Combine(BaseDir, ".dockerignore");
+				FileReference.Delete(DockerIgnoreFileInBaseDir);
+
 				if (!String.IsNullOrEmpty(Parameters.OverlayDirs))
 				{
 					foreach (string OverlayDir in Parameters.OverlayDirs.Split(';'))
@@ -129,12 +151,28 @@ namespace AutomationTool.Tasks
 					}
 					Arguments.Append($" -f {DockerFile.MakeRelativeTo(BaseDir).QuoteArgument()}");
 				}
+				if (Parameters.DockerIgnoreFile != null)
+				{
+					FileReference DockerIgnoreFile = ResolveFile(Parameters.DockerIgnoreFile);
+					FileReference.Copy(DockerIgnoreFile, DockerIgnoreFileInBaseDir);
+				}
+				if (Parameters.ProgressOutput != null)
+				{
+					Arguments.Append($" --progress={Parameters.ProgressOutput}");
+				}
 				if (Parameters.Arguments != null)
 				{
 					Arguments.Append($" {Parameters.Arguments}");
 				}
 
-				await SpawnTaskBase.ExecuteAsync("docker", Arguments.ToString(), EnvVars: ParseEnvVars(Parameters.Environment, Parameters.EnvironmentFile), WorkingDir: StagingDir.FullName);
+				Dictionary<string, string> EnvVars = ParseEnvVars(Parameters.Environment, Parameters.EnvironmentFile);
+				if (Parameters.UseBuildKit)
+				{
+					EnvVars["DOCKER_BUILDKIT"] = "1";
+				}
+				
+				string WorkingDir = isStagingEnabled ? StagingDir.FullName : BaseDir.FullName;
+				await SpawnTaskBase.ExecuteAsync("docker", Arguments.ToString(), EnvVars: EnvVars, WorkingDir: WorkingDir);
 			}
 		}
 
