Make FindTagNamesFromList static since it doesn't use any member variables and allows it to be called from other classes.

#rnx

diff --git a/Engine/Source/Programs/AutomationTool/BuildGraph/Task.cs b/Engine/Source/Programs/AutomationTool/BuildGraph/Task.cs
--- a/Engine/Source/Programs/AutomationTool/BuildGraph/Task.cs
+++ b/Engine/Source/Programs/AutomationTool/BuildGraph/Task.cs
@@ -235,7 +235,7 @@ namespace AutomationTool
 		/// </summary>
 		/// <param name="TagList">List of tags separated by semicolons</param>
 		/// <returns>Tag names from this filespec</returns>
-		protected IEnumerable<string> FindTagNamesFromList(string TagList)
+		protected static IEnumerable<string> FindTagNamesFromList(string TagList)
 		{
 			if(!String.IsNullOrEmpty(TagList))
 			{
