BuildCMakeLib: Added support for targeting arm64 on Win64 with VS2019

This allows a library to use a Win64+arm64 build for HoloLens without providing a separate set of libs for HoloLens.

diff --git a/Engine/Source/Programs/AutomationTool/Scripts/BuildCMakeLib.Automation.cs b/Engine/Source/Programs/AutomationTool/Scripts/BuildCMakeLib.Automation.cs
--- a/Engine/Source/Programs/AutomationTool/Scripts/BuildCMakeLib.Automation.cs
+++ b/Engine/Source/Programs/AutomationTool/Scripts/BuildCMakeLib.Automation.cs
@@ -21,7 +21,7 @@ using UnrealBuildBase;
 [Help("TargetPlatform", "Specify the name of the target platform to build (eg. -TargetPlatform=IOS).")]
 [Help("TargetArchitecture", "Specify the name of the target architecture to build (eg. -TargetArchitecture=x86_64).")]
 [Help("TargetConfigs", "Specify a list of configurations to build, separated by '+' characters (eg. -TargetConfigs=release+debug). Default is release+debug.")]
-[Help("BinOutputPath", "Override the path to output binaries to. (eg. -BinOutputPath=bin). Default is empty.")]
+[Help("BinOutputPath", "Override the path to output binaries to. (eg. -BinOutputPath=bin). Default is Binaries.")]
 [Help("LibOutputPath", "Override the path to output libraries to. (eg. -LibOutputPath=lib). Default is empty.")]
 [Help("CMakeGenerator", "Specify the CMake generator to use.")]
 [Help("CMakeProjectIncludeFile", "Specify the name of the CMake project include file to use, first looks in current directory then looks in global directory.")]
@@ -434,7 +434,7 @@ public sealed class BuildCMakeLib : BuildCommand
 		TargetLib.Name = ParseParamValue("TargetLib", "");
 		TargetLib.Version = ParseParamValue("TargetLibVersion", "");
 		TargetLib.SourcePath = ParseParamValue("TargetLibSourcePath", "");
-		TargetLib.BinOutputPath = ParseParamValue("BinOutputPath", "");
+		TargetLib.BinOutputPath = ParseParamValue("BinOutputPath", "Binaries");
 		TargetLib.LibOutputPath = ParseParamValue("LibOutputPath", "");
 		TargetLib.CMakeProjectIncludeFile = ParseParamValue("CMakeProjectIncludeFile", "");
 		TargetLib.CMakeAdditionalArguments = ParseParamValue("CMakeAdditionalArguments", "");
@@ -728,7 +728,21 @@ class VS2019TargetPlatform_Win64 : BuildCMakeLib.VS2019TargetPlatform
 	public override string DebugDatabaseExtension => "pdb";
 	public override string DynamicLibraryExtension => "dll";
 	public override string StaticLibraryExtension => "lib";
+	public override string VariantDirectory => Architecture == "Win64" ? "" : Architecture.ToLower();
 	public override bool IsPlatformExtension => false;
+
+	private readonly string Architecture;
+
+	public VS2019TargetPlatform_Win64(string Architecture = "Win64")
+	{
+		this.Architecture = Architecture;
+	}
+
+	public override string GetCMakeSetupArguments(BuildCMakeLib.TargetLib TargetLib, string TargetConfiguration)
+	{
+		return base.GetCMakeSetupArguments(TargetLib, TargetConfiguration)
+			+ (Architecture == "Win64" ? "" : string.Format(" -A {0}", Architecture));
+	}
 }
 
 class NMakeTargetPlatform_Win64 : BuildCMakeLib.NMakeTargetPlatform
