Perform restricted staged folder name checking before doing platform filename remapping
Back out incorrect fix from CL 20056908

diff --git a/Engine/Source/Programs/AutomationTool/AutomationUtils/Platform.cs b/Engine/Source/Programs/AutomationTool/AutomationUtils/Platform.cs
--- a/Engine/Source/Programs/AutomationTool/AutomationUtils/Platform.cs
+++ b/Engine/Source/Programs/AutomationTool/AutomationUtils/Platform.cs
@@ -644,18 +644,10 @@ namespace AutomationTool
 			return (FileType == StagedFileType.UFS || FileType == StagedFileType.NonUFS);
 		}
 
-		/// <summary>
-		/// Remaps the given content file to its final location
-		/// </summary>
-		public virtual StagedFileReference Remap(StagedFileReference Dest)
-		{
-			return Dest;
-		}
-
 		/// <summary>
 		/// Remaps the given content directory to its final location
 		/// </summary>
-		public virtual StagedDirectoryReference Remap(StagedDirectoryReference Dest)
+		public virtual StagedFileReference Remap(StagedFileReference Dest)
 		{
 			return Dest;
 		}
diff --git a/Engine/Source/Programs/AutomationTool/IOS/IOSPlatform.Automation.cs b/Engine/Source/Programs/AutomationTool/IOS/IOSPlatform.Automation.cs
--- a/Engine/Source/Programs/AutomationTool/IOS/IOSPlatform.Automation.cs
+++ b/Engine/Source/Programs/AutomationTool/IOS/IOSPlatform.Automation.cs
@@ -1941,11 +1941,6 @@ public class IOSPlatform : Platform
 	{
 		return new StagedFileReference("cookeddata/" + Dest.Name);
 	}
-	public override StagedDirectoryReference Remap(StagedDirectoryReference Dest)
-	{
-		return new StagedDirectoryReference("cookeddata/" + Dest.Name);
-	}
-
 	public override List<string> GetDebugFileExtensions()
 	{
 		return new List<string> { ".dsym", ".udebugsymbols" };
diff --git a/Engine/Source/Programs/AutomationTool/Scripts/CopyBuildToStagingDirectory.Automation.cs b/Engine/Source/Programs/AutomationTool/Scripts/CopyBuildToStagingDirectory.Automation.cs
--- a/Engine/Source/Programs/AutomationTool/Scripts/CopyBuildToStagingDirectory.Automation.cs
+++ b/Engine/Source/Programs/AutomationTool/Scripts/CopyBuildToStagingDirectory.Automation.cs
@@ -1327,41 +1327,6 @@ namespace AutomationScripts
 				}
 			}
 
-			// Make all the filenames lowercase
-			if (SC.StageTargetPlatform.DeployLowerCaseFilenames(StagedFileType.NonUFS))
-			{
-				SC.FilesToStage.NonUFSFiles = SC.FilesToStage.NonUFSFiles.ToDictionary(x => x.Key.ToLowerInvariant(), x => x.Value);
-			}
-			// DebugNonUFS is weird, so ask per-file
-			SC.FilesToStage.NonUFSDebugFiles = SC.FilesToStage.NonUFSDebugFiles.ToDictionary(x => SC.StageTargetPlatform.DeployLowerCaseFile(x.Value, StagedFileType.DebugNonUFS) ? x.Key.ToLowerInvariant() : x.Key, x => x.Value);
-			if (!Params.UsePak(SC.StageTargetPlatform) && SC.StageTargetPlatform.DeployLowerCaseFilenames(StagedFileType.UFS))
-			{
-				SC.FilesToStage.UFSFiles = SC.FilesToStage.UFSFiles.ToDictionary(x => x.Key.ToLowerInvariant(), x => x.Value);
-			}
-
-			// Remap all the non-ufs files if not using a PAK file
-			// Shared NonUFS files are staged in their remapped location, and may be duplicated in the to-stage list.
-			if (SC.StageTargetPlatform.RemapFileType(StagedFileType.NonUFS))
-			{
-				SC.FilesToStage.NonUFSFiles = SC.FilesToStage.NonUFSFiles.ToDictionary(x => SC.StageTargetPlatform.Remap(x.Key), x => x.Value);
-			}
-				
-			if (SC.StageTargetPlatform.RemapFileType(StagedFileType.DebugNonUFS))
-			{
-				SC.FilesToStage.NonUFSDebugFiles = SC.FilesToStage.NonUFSDebugFiles.ToDictionary(x => SC.StageTargetPlatform.Remap(x.Key), x => x.Value);
-			}
-
-			if (SC.StageTargetPlatform.RemapFileType(StagedFileType.UFS) && !Params.UsePak(SC.StageTargetPlatform))
-			{
-				SC.FilesToStage.UFSFiles = SC.FilesToStage.UFSFiles.ToDictionary(x => SC.StageTargetPlatform.Remap(x.Key), x => x.Value);
-			}
-
-			// Merge all the NonUFS system files back into the NonUFS list. Deployment is currently only set up to read from that.
-			foreach (KeyValuePair<StagedFileReference, FileReference> Pair in SC.FilesToStage.NonUFSSystemFiles)
-			{
-				SC.FilesToStage.NonUFSFiles[Pair.Key] = Pair.Value;
-			}
-
 			// Make sure there are no restricted folders in the output
 			HashSet<StagedFileReference> RestrictedFiles = new HashSet<StagedFileReference>();
 			foreach (string RestrictedName in SC.RestrictedFolderNames)
@@ -1371,7 +1336,7 @@ namespace AutomationScripts
 				RestrictedFiles.UnionWith(SC.FilesToStage.NonUFSDebugFiles.Keys.Where(x => x.ContainsName(RestrictedName)));
 				RestrictedFiles.UnionWith(SC.CrashReporterUFSFiles.Keys.Where(x => x.ContainsName(RestrictedName)));
 			}
-			RestrictedFiles.RemoveWhere(RestrictedFile => SC.DirectoriesAllowList.Any(TestDirectory => RestrictedFile.Directory.IsUnderDirectory(SC.StageTargetPlatform.Remap(TestDirectory))));
+			RestrictedFiles.RemoveWhere(RestrictedFile => SC.DirectoriesAllowList.Any(TestDirectory => RestrictedFile.Directory.IsUnderDirectory(TestDirectory)));
 
 			if (RestrictedFiles.Count > 0)
 			{
@@ -1406,6 +1371,41 @@ namespace AutomationScripts
 				}
 				throw new AutomationException(Message.ToString());
 			}
+
+			// Make all the filenames lowercase
+			if (SC.StageTargetPlatform.DeployLowerCaseFilenames(StagedFileType.NonUFS))
+			{
+				SC.FilesToStage.NonUFSFiles = SC.FilesToStage.NonUFSFiles.ToDictionary(x => x.Key.ToLowerInvariant(), x => x.Value);
+			}
+			// DebugNonUFS is weird, so ask per-file
+			SC.FilesToStage.NonUFSDebugFiles = SC.FilesToStage.NonUFSDebugFiles.ToDictionary(x => SC.StageTargetPlatform.DeployLowerCaseFile(x.Value, StagedFileType.DebugNonUFS) ? x.Key.ToLowerInvariant() : x.Key, x => x.Value);
+			if (!Params.UsePak(SC.StageTargetPlatform) && SC.StageTargetPlatform.DeployLowerCaseFilenames(StagedFileType.UFS))
+			{
+				SC.FilesToStage.UFSFiles = SC.FilesToStage.UFSFiles.ToDictionary(x => x.Key.ToLowerInvariant(), x => x.Value);
+			}
+
+			// Remap all the non-ufs files if not using a PAK file
+			// Shared NonUFS files are staged in their remapped location, and may be duplicated in the to-stage list.
+			if (SC.StageTargetPlatform.RemapFileType(StagedFileType.NonUFS))
+			{
+				SC.FilesToStage.NonUFSFiles = SC.FilesToStage.NonUFSFiles.ToDictionary(x => SC.StageTargetPlatform.Remap(x.Key), x => x.Value);
+			}
+				
+			if (SC.StageTargetPlatform.RemapFileType(StagedFileType.DebugNonUFS))
+			{
+				SC.FilesToStage.NonUFSDebugFiles = SC.FilesToStage.NonUFSDebugFiles.ToDictionary(x => SC.StageTargetPlatform.Remap(x.Key), x => x.Value);
+			}
+
+			if (SC.StageTargetPlatform.RemapFileType(StagedFileType.UFS) && !Params.UsePak(SC.StageTargetPlatform))
+			{
+				SC.FilesToStage.UFSFiles = SC.FilesToStage.UFSFiles.ToDictionary(x => SC.StageTargetPlatform.Remap(x.Key), x => x.Value);
+			}
+
+			// Merge all the NonUFS system files back into the NonUFS list. Deployment is currently only set up to read from that.
+			foreach (KeyValuePair<StagedFileReference, FileReference> Pair in SC.FilesToStage.NonUFSSystemFiles)
+			{
+				SC.FilesToStage.NonUFSFiles[Pair.Key] = Pair.Value;
+			}
 		}
 
 		/// <summary>
