Fix overflow obb file installs for Android

#android

diff --git a/Engine/Source/Programs/AutomationTool/Android/AndroidPlatform.Automation.cs b/Engine/Source/Programs/AutomationTool/Android/AndroidPlatform.Automation.cs
--- a/Engine/Source/Programs/AutomationTool/Android/AndroidPlatform.Automation.cs
+++ b/Engine/Source/Programs/AutomationTool/Android/AndroidPlatform.Automation.cs
@@ -1434,7 +1434,7 @@ public class AndroidPlatform : Platform
 				OBBInstallCommand = bNoObbInstall ? AFSCommand + " deletefile '^mainobb'" : AFSCommand + " push " + Path.GetFileName(ObbName) + " '^mainobb'";
 				PatchInstallCommand = bNoPatchInstall ? AFSCommand + " deletefile '^patchobb'" : AFSCommand + " push " + Path.GetFileName(PatchName) + " '^patchobb'";
 				Overflow1InstallCommand = bNoOverflow1Install ? AFSCommand + " deletefile '^overflow1obb'" : AFSCommand + " push " + Path.GetFileName(Overflow1Name) + " '^overflow1obb'";
-				Overflow1InstallCommand = bNoOverflow2Install ? AFSCommand + " deletefile '^overflow2obb'" : AFSCommand + " push " + Path.GetFileName(Overflow2Name) + " '^overflow2obb'";
+				Overflow2InstallCommand = bNoOverflow2Install ? AFSCommand + " deletefile '^overflow2obb'" : AFSCommand + " push " + Path.GetFileName(Overflow2Name) + " '^overflow2obb'";
 			}
 
 			LogInformation("Writing shell script for install with {0}", bPackageDataInsideApk ? "data in APK" : "separate obb");
@@ -1519,7 +1519,7 @@ public class AndroidPlatform : Platform
 				OBBInstallCommand = bNoObbInstall ? AFSCommand + " deletefile \"^mainobb\"" : AFSCommand + " push " + Path.GetFileName(ObbName) + " \"^mainobb\"";
 				PatchInstallCommand = bNoPatchInstall ? AFSCommand + " deletefile \"^patchobb\"" : AFSCommand + " push " + Path.GetFileName(PatchName) + " \"^patchobb\"";
 				Overflow1InstallCommand = bNoOverflow1Install ? AFSCommand + " deletefile \"^overflow1obb\"" : AFSCommand + " push " + Path.GetFileName(Overflow1Name) + " \"^overflow1obb\"";
-				Overflow1InstallCommand = bNoOverflow2Install ? AFSCommand + " deletefile \"^overflow2obb\"" : AFSCommand + " push " + Path.GetFileName(Overflow2Name) + " \"^overflow2obb\"";
+				Overflow2InstallCommand = bNoOverflow2Install ? AFSCommand + " deletefile \"^overflow2obb\"" : AFSCommand + " push " + Path.GetFileName(Overflow2Name) + " \"^overflow2obb\"";
 			}
 
 			LogInformation("Writing bat for install with {0}", bPackageDataInsideApk ? "data in APK" : "separate OBB");
