- Updated CookedEditor to use the OPtionalFileStagingDirectory param to stage optional files with, if the member wasn't set by a subclass
- Updated some buildscripts to use -editoroptional data and -optionalfilestagingdirectory options

diff --git a/Engine/Source/Programs/AutomationTool/CookedEditor/CookedEditor.Automation.cs b/Engine/Source/Programs/AutomationTool/CookedEditor/CookedEditor.Automation.cs
--- a/Engine/Source/Programs/AutomationTool/CookedEditor/CookedEditor.Automation.cs
+++ b/Engine/Source/Programs/AutomationTool/CookedEditor/CookedEditor.Automation.cs
@@ -499,6 +499,13 @@ public class MakeCookedEditor : BuildCommand
 			}
 		}
 
+		// if a subclass or -makerelease didn't set ReleaseOptionalFileStageDirectory, then look in the Params for the commandline option
+		if (ReleaseOptionalFileStageDirectory == null && !string.IsNullOrEmpty(Params.OptionalStageDirectory) &&
+			File.Exists(Params.OptionalStageDirectory))
+		{
+			ReleaseOptionalFileStageDirectory = new DirectoryReference(Params.OptionalStageDirectory);
+		}
+
 		if (ReleaseOptionalFileStageDirectory != null)
 		{
 			// these files were already staged by the client/release build, so we stage them as NonUFS
