- Always specify -editoroptional when making CookedEditor, and to the client build when -makerelease is specified
- If doing editor-as-DLC, and you are cooking client manually, you will need to add -editoroptional to the client cooker commandline

diff --git a/Engine/Source/Programs/AutomationTool/CookedEditor/CookedEditor.Automation.cs b/Engine/Source/Programs/AutomationTool/CookedEditor/CookedEditor.Automation.cs
--- a/Engine/Source/Programs/AutomationTool/CookedEditor/CookedEditor.Automation.cs
+++ b/Engine/Source/Programs/AutomationTool/CookedEditor/CookedEditor.Automation.cs
@@ -978,6 +978,8 @@ public class MakeCookedEditor : BuildCommand
 			, BasedOnReleaseVersion: BasedOnReleaseVersion
 			, DedicatedServer: bIsCookedCooker
 			, NoClient: bIsCookedCooker
+			, OptionalContent: true
+
 		);
 
 		// cook the cooked editor targetplatorm as the "client"
@@ -1078,7 +1080,6 @@ public class MakeCookedEditor : BuildCommand
 			, SkipBuildClient: !Options.Contains("build", StringComparer.InvariantCultureIgnoreCase)
 			, SkipCook: !Options.Contains("cook", StringComparer.InvariantCultureIgnoreCase)
 			, SkipStage: !Options.Contains("stage", StringComparer.InvariantCultureIgnoreCase)
-			, OptionalContent: Options.Contains("editoroptional", StringComparer.InvariantCultureIgnoreCase)
 		);
 
 		// if the MainParams override the ReleaseVersion path, use it directly
@@ -1090,7 +1091,9 @@ public class MakeCookedEditor : BuildCommand
 		// copy off the staging dir
 		ReleaseParams.PreModifyDeploymentContextCallback = new Action<ProjectParams, DeploymentContext>((ProjectParams P, DeploymentContext SC) => { ReleaseStageDirectory = SC.StageDirectory; });
 
+		// cooked editor doesn't work without OptionalContent now, so always generated it, and save it somewhere that staging of the cookededitor will get
 		ReleaseOptionalFileStageDirectory = DirectoryReference.Combine(MainParams.RawProjectPath.Directory, "Saved", "CookedEditor", "OptionalData");
+		ReleaseParams.OptionalContent = true;
 		ReleaseParams.OptionalStageDirectory = ReleaseOptionalFileStageDirectory.FullName;
 
 		ModifyReleaseParams(ReleaseParams);
