- Added some more debugging for dealing with missing EditorOptional files going from Client -> Editor

diff --git a/Engine/Source/Developer/IoStoreUtilities/Private/IoStoreUtilities.cpp b/Engine/Source/Developer/IoStoreUtilities/Private/IoStoreUtilities.cpp
--- a/Engine/Source/Developer/IoStoreUtilities/Private/IoStoreUtilities.cpp
+++ b/Engine/Source/Developer/IoStoreUtilities/Private/IoStoreUtilities.cpp
@@ -2441,6 +2441,8 @@ void InitializeContainerTargetsAndPackages(
 					// if we have an optional output location, use that directory, combined with the name of the output path 
 					ContainerTarget->OptionalSegmentOutputPath = FPaths::Combine(ContainerSource.OptionalOutputPath, FPaths::GetCleanFilename(ContainerTarget->OutputPath) + FPackagePath::GetOptionalSegmentExtensionModifier());
 				}
+
+				UE_LOG(LogIoStore, Display, TEXT("Saving optional container to: '%s'", *ContainerTarget->OptionalSegmentOutputPath));
 			}
 		}
 	}
diff --git a/Engine/Source/Programs/AutomationTool/CookedEditor/CookedEditor.Automation.cs b/Engine/Source/Programs/AutomationTool/CookedEditor/CookedEditor.Automation.cs
--- a/Engine/Source/Programs/AutomationTool/CookedEditor/CookedEditor.Automation.cs
+++ b/Engine/Source/Programs/AutomationTool/CookedEditor/CookedEditor.Automation.cs
@@ -511,6 +511,12 @@ public class MakeCookedEditor : BuildCommand
 			// these files were already staged by the client/release build, so we stage them as NonUFS
 			// these could be just copied, but StageFiles handles copying easily
 			SC.StageFiles(StagedFileType.NonUFS, ReleaseOptionalFileStageDirectory, StageFilesSearch.AllDirectories, new StagedDirectoryReference($"{Context.ProjectName}/Content/Paks"));
+
+			Log.TraceInformation($"Staging optional files from {ReleaseOptionalFileStageDirectory.FullName}:");
+			foreach (var OptionalFile in DirectoryReference.EnumerateFiles(ReleaseOptionalFileStageDirectory, "*"))
+			{
+				Log.TraceInformation($"  '{OptionalFile.FullName}'");
+			}
 		}
 	}
 
