Improve error message when build can't be installed because of space on the local machine.

#rnx

diff --git a/Engine/Source/Programs/AutomationTool/Gauntlet/Unreal/Base/Gauntlet.UnrealSession.cs b/Engine/Source/Programs/AutomationTool/Gauntlet/Unreal/Base/Gauntlet.UnrealSession.cs
--- a/Engine/Source/Programs/AutomationTool/Gauntlet/Unreal/Base/Gauntlet.UnrealSession.cs
+++ b/Engine/Source/Programs/AutomationTool/Gauntlet/Unreal/Base/Gauntlet.UnrealSession.cs
@@ -744,7 +744,21 @@ namespace Gauntlet
 					catch (System.Exception Ex)
 					{
 						// Warn, ignore the device, and do not continue
-						Log.Info("Failed to install app onto device {0} for role {1}. {2}. Will retry with new device", Device, Role, Ex);
+						string ErrorMessage = string.Format("Failed to install app onto device {0} for role {1}. {2}. Will retry with new device", Device, Role, Ex);
+						if (ErrorMessage.Contains("not enough space"))
+						{
+							Log.Warning(ErrorMessage);
+							if(Device.Platform == BuildHostPlatform.Current.Platform)
+							{
+								// If on desktop platform, we are not retrying.
+								// It is unlikely that space is going to be made and InstallBuildParallel has marked the build path as problematic.
+								SessionRetries = 0;
+							}
+						}
+						else
+						{
+							Log.Info(ErrorMessage);
+						}
 						UnrealDeviceReservation.MarkProblemDevice(Device);
 						InstallSuccess = false;
 						IDeviceUsageReporter.RecordEnd(Device.Name, Device.Platform, IDeviceUsageReporter.EventType.Install, IDeviceUsageReporter.EventState.Failure);
