Gauntlet - Moving RHIValidation option to UE.AutomationTestConfig

#rnx

diff --git a/Engine/Source/Programs/AutomationTool/Gauntlet/Unreal/Automation/UE.Automation.cs b/Engine/Source/Programs/AutomationTool/Gauntlet/Unreal/Automation/UE.Automation.cs
--- a/Engine/Source/Programs/AutomationTool/Gauntlet/Unreal/Automation/UE.Automation.cs
+++ b/Engine/Source/Programs/AutomationTool/Gauntlet/Unreal/Automation/UE.Automation.cs
@@ -40,6 +40,12 @@ namespace UE
 			}
 		}
 
+		/// <summary>
+		/// Enable the RHI validation layer.
+		/// </summary>
+		[AutoParam]
+		public virtual bool RHIValidation { get; set; } = false;
+
 		/// <summary>
 		/// Run with Nvidia cards for raytracing support
 		/// </summary>
@@ -266,6 +272,11 @@ namespace UE
 				}
 			}
 
+			if (RHIValidation)
+			{
+				AppConfig.CommandLine += " -rhivalidation";
+			}
+
 			// Options specific to roles running under the editor
 			if (ConfigRole.RoleType.UsesEditor())
 			{
diff --git a/Engine/Source/Programs/AutomationTool/Gauntlet/Unreal/Engine/UnrealTests.cs b/Engine/Source/Programs/AutomationTool/Gauntlet/Unreal/Engine/UnrealTests.cs
--- a/Engine/Source/Programs/AutomationTool/Gauntlet/Unreal/Engine/UnrealTests.cs
+++ b/Engine/Source/Programs/AutomationTool/Gauntlet/Unreal/Engine/UnrealTests.cs
@@ -329,12 +329,6 @@ namespace Gauntlet
 			[AutoParam]
 			public virtual bool UseEditor { get; set; } = false;
 
-			/// <summary>
-			/// Enable the RHI validation layer.
-			/// </summary>
-			[AutoParam]
-			public virtual bool RHIValidation { get; set; } = false;
-
 			/// <summary>
 			/// Filter or groups of tests to apply (alias for RunTest) 
 			/// </summary>
@@ -355,11 +349,6 @@ namespace Gauntlet
 				}
 
 				base.ApplyToConfig(AppConfig, ConfigRole, OtherRoles);
-
-				if (RHIValidation)
-				{
-					AppConfig.CommandLine += " -rhivalidation";
-				}
 			}
 		}
 	}
