UAT: Fixed uptnl files not being excluded along with the other cooked files when filtering files for the cooked editor target

diff --git a/Engine/Source/Programs/AutomationTool/CookedEditor/CookedEditor.Automation.cs b/Engine/Source/Programs/AutomationTool/CookedEditor/CookedEditor.Automation.cs
--- a/Engine/Source/Programs/AutomationTool/CookedEditor/CookedEditor.Automation.cs
+++ b/Engine/Source/Programs/AutomationTool/CookedEditor/CookedEditor.Automation.cs
@@ -169,7 +169,7 @@ public class ModifyStageContext
 
 		// remove already-cooked assets to be replaced with 
 		List<StagedFileReference> UncookedFilesThatDoNotExist = new List<StagedFileReference>();
-		string[] CookedExtensions = { ".uasset", ".umap", ".ubulk", ".uexp" };
+		string[] CookedExtensions = { ".uasset", ".umap", ".ubulk", ".uexp", ".uptnl" };
 		foreach (var UncookedFile in StagedUncookFiles)
 		{
 			string PathWithNoExtension = Path.ChangeExtension(UncookedFile.Key.Name, null);
