Move the remaining log matchers into AutomationUtils.

diff --git a/Engine/Source/Programs/AutomationTool/AutomationUtils/AutomationUtils.Automation.csproj b/Engine/Source/Programs/AutomationTool/AutomationUtils/AutomationUtils.Automation.csproj
--- a/Engine/Source/Programs/AutomationTool/AutomationUtils/AutomationUtils.Automation.csproj
+++ b/Engine/Source/Programs/AutomationTool/AutomationUtils/AutomationUtils.Automation.csproj
@@ -1,11 +1,11 @@
-<Project Sdk="Microsoft.NET.Sdk">
+﻿<Project Sdk="Microsoft.NET.Sdk">
   <Import Project="..\..\Shared\UnrealEngine.csproj.props" />
   
   <PropertyGroup>
     <TargetFramework>net6.0</TargetFramework>
     <Configuration Condition=" '$(Configuration)' == '' ">Development</Configuration>
     <OutputType>Library</OutputType>
-    <RootNamespace>AutomationUtils.Automation</RootNamespace>
+    <RootNamespace>AutomationUtils</RootNamespace>
     <AssemblyName>AutomationUtils.Automation</AssemblyName>
     <GenerateAssemblyInfo>false</GenerateAssemblyInfo>
     <GenerateTargetFrameworkAttribute>false</GenerateTargetFrameworkAttribute>
diff --git a/Engine/Source/Programs/AutomationTool/AutomationUtils/Matchers/ContentEventMatcher.cs b/Engine/Source/Programs/AutomationTool/AutomationUtils/Matchers/ContentEventMatcher.cs
new file mode 100644
--- /dev/null
+++ b/Engine/Source/Programs/AutomationTool/AutomationUtils/Matchers/ContentEventMatcher.cs
@@ -0,0 +1,61 @@
+// Copyright Epic Games, Inc. All Rights Reserved.
+
+using System;
+using System.Text.RegularExpressions;
+using EpicGames.Core;
+using Microsoft.Extensions.Logging;
+
+#nullable enable
+
+namespace AutomationUtils.Matchers
+{
+	/// <summary>
+	/// Matches compile errors and annotates with the source file path and revision
+	/// </summary>
+	class ContentEventMatcher : ILogEventMatcher
+	{
+		const string Pattern = 
+			@"^\s*" + 
+			@"(?<channel>[a-zA-Z0-9]+):\s+" + 
+			@"(?<severity>Error:|Warning:)\s+" + 
+			@"(?:\[AssetLog\] )?" + 
+			@"(?<asset>(?:[a-zA-Z]:)?[^:]+(?:.uasset|.umap)):\s*" + 
+			@"(?<message>.*)";
+
+		/// <inheritdoc/>
+		public LogEventMatch? Match(ILogCursor input)
+		{
+			// Do the match in two phases so we can early out if the strings "error" or "warning" are not present. The patterns before these strings can
+			// produce many false positives, making them very slow to execute.
+			if (input.IsMatch("Error:|Warning:"))
+			{
+				Match? match;
+				if(input.TryMatch(Pattern, out match))
+				{
+					LogEventBuilder builder = new LogEventBuilder(input);
+
+					builder.Annotate(match.Groups["channel"], LogEventMarkup.Channel);
+					builder.Annotate(match.Groups["severity"], LogEventMarkup.Severity);
+					builder.AnnotateAsset(match.Groups["asset"]);
+					builder.Annotate(match.Groups["message"], LogEventMarkup.Message);
+
+					return builder.ToMatch(LogEventPriority.AboveNormal, GetLogLevelFromSeverity(match), KnownLogEvents.Engine_AssetLog);
+				}
+			}
+			return null;
+		}
+
+		static LogLevel GetLogLevelFromSeverity(Match match)
+		{
+			string severity = match.Groups["severity"].Value;
+			if (severity.Equals("Warning:", StringComparison.Ordinal))
+			{
+				return LogLevel.Warning;
+			}
+			else
+			{
+				return LogLevel.Error;
+			}
+		}
+	}
+}
diff --git a/Engine/Source/Programs/AutomationTool/AutomationUtils/Matchers/CrashEventMatcher.cs b/Engine/Source/Programs/AutomationTool/AutomationUtils/Matchers/CrashEventMatcher.cs
new file mode 100644
--- /dev/null
+++ b/Engine/Source/Programs/AutomationTool/AutomationUtils/Matchers/CrashEventMatcher.cs
@@ -0,0 +1,65 @@
+// Copyright Epic Games, Inc. All Rights Reserved.
+
+using System.Text.RegularExpressions;
+using EpicGames.Core;
+using Microsoft.Extensions.Logging;
+
+#nullable enable
+
+namespace AutomationUtils.Matchers
+{
+	/// <summary>
+	/// Matcher for engine crashes
+	/// </summary>
+	class CrashEventMatcher : ILogEventMatcher
+	{
+		public LogEventMatch? Match(ILogCursor cursor)
+		{
+			if (cursor.IsMatch("begin: stack for UAT"))
+			{
+				for (int maxOffset = 1; maxOffset < 100; maxOffset++)
+				{
+					if (cursor.IsMatch(maxOffset, "end: stack for UAT"))
+					{
+						LogEventBuilder builder = new LogEventBuilder(cursor, lineCount: maxOffset + 1);
+						return builder.ToMatch(LogEventPriority.BelowNormal, GetLogLevel(cursor), KnownLogEvents.Engine_Crash);
+					}
+				}
+			}
+			if (cursor.IsMatch("AutomationTool: Stack:"))
+			{
+				LogEventBuilder builder = new LogEventBuilder(cursor);
+				while (builder.Current.IsMatch(1, "AutomationTool: Stack:"))
+				{
+					builder.MoveNext();
+				}
+				return builder.ToMatch(LogEventPriority.Low, LogLevel.Error, KnownLogEvents.AutomationTool_Crash);
+			}
+
+			Match? match;
+			if (cursor.TryMatch(@"ExitCode=(3|139|255)(?!\d)", out match))
+			{
+				LogEventBuilder builder = new LogEventBuilder(cursor);
+				builder.Annotate("exitCode", match.Groups[1]);
+				return builder.ToMatch(LogEventPriority.Low, LogLevel.Error, KnownLogEvents.AutomationTool_CrashExitCode);
+			}
+			return null;
+		}
+
+		static LogLevel GetLogLevel(ILogCursor cursor)
+		{
+			if(cursor.IsMatch(0, "[Ee]rror:"))
+			{
+				return LogLevel.Error;
+			}
+			else if(cursor.IsMatch(0, "[Ww]arning:"))
+			{
+				return LogLevel.Warning;
+			}
+			else
+			{
+				return LogLevel.Information;
+			}
+		}
+	}
+}
diff --git a/Engine/Source/Programs/AutomationTool/AutomationUtils/Matchers/ExceptionEventMatcher.cs b/Engine/Source/Programs/AutomationTool/AutomationUtils/Matchers/ExceptionEventMatcher.cs
new file mode 100644
--- /dev/null
+++ b/Engine/Source/Programs/AutomationTool/AutomationUtils/Matchers/ExceptionEventMatcher.cs
@@ -0,0 +1,30 @@
+// Copyright Epic Games, Inc. All Rights Reserved.
+
+using EpicGames.Core;
+using Microsoft.Extensions.Logging;
+
+#nullable enable
+
+namespace AutomationUtils.Matchers
+{
+	/// <summary>
+	/// Matches a generic C# exception
+	/// </summary>
+	class ExceptionEventMatcher : ILogEventMatcher
+	{
+		/// <inheritdoc/>
+		public LogEventMatch? Match(ILogCursor cursor)
+		{
+			if (cursor.IsMatch(@"^\s*Unhandled Exception: "))
+			{
+				LogEventBuilder builder = new LogEventBuilder(cursor);
+				while(builder.Current.IsMatch(1, @"^\s*at "))
+				{
+					builder.MoveNext();
+				}
+				return builder.ToMatch(LogEventPriority.Low, LogLevel.Error, KnownLogEvents.Exception);
+			}
+			return null;
+		}
+	}
+}
diff --git a/Engine/Source/Programs/AutomationTool/AutomationUtils/Matchers/ExitCodeEventMatcher.cs b/Engine/Source/Programs/AutomationTool/AutomationUtils/Matchers/ExitCodeEventMatcher.cs
new file mode 100644
--- /dev/null
+++ b/Engine/Source/Programs/AutomationTool/AutomationUtils/Matchers/ExitCodeEventMatcher.cs
@@ -0,0 +1,52 @@
+// Copyright Epic Games, Inc. All Rights Reserved.
+
+using EpicGames.Core;
+using Microsoft.Extensions.Logging;
+
+#nullable enable
+
+namespace AutomationUtils.Matchers
+{
+	/// <summary>
+	/// Matcher for editor/UAT instances exiting with an error
+	/// </summary>
+	class ExitCodeEventMatcher : ILogEventMatcher
+	{
+		public LogEventMatch? Match(ILogCursor cursor)
+		{
+			int numLines = 0;
+			for (; ; )
+			{
+				if (cursor.IsMatch(numLines, "Editor terminated with exit code [1-9]"))
+				{
+					numLines++;
+				}
+				else if (cursor.IsMatch(numLines, "AutomationTool exiting with ExitCode=[1-9]"))
+				{
+					numLines++;
+				}
+				else if (cursor.IsMatch(numLines, "BUILD FAILED"))
+				{
+					numLines++;
+				}
+				else if (cursor.IsMatch(numLines, "(Error executing.+)(tool returned code)(.+)"))
+				{
+					numLines++;
+				}
+				else
+				{
+					break;
+				}
+			}
+
+			if (numLines > 0)
+			{
+				LogEventBuilder builder = new LogEventBuilder(cursor);
+				builder.MoveNext(numLines - 1);
+				return builder.ToMatch(LogEventPriority.Low, LogLevel.Error, KnownLogEvents.ExitCode);
+			}
+			return null;
+		}
+	}
+}
+
diff --git a/Engine/Source/Programs/AutomationTool/AutomationUtils/Matchers/GauntletEventMatcher.cs b/Engine/Source/Programs/AutomationTool/AutomationUtils/Matchers/GauntletEventMatcher.cs
new file mode 100644
--- /dev/null
+++ b/Engine/Source/Programs/AutomationTool/AutomationUtils/Matchers/GauntletEventMatcher.cs
@@ -0,0 +1,81 @@
+// Copyright Epic Games, Inc. All Rights Reserved.
+
+using System.Text.RegularExpressions;
+using EpicGames.Core;
+using Microsoft.Extensions.Logging;
+
+#nullable enable
+
+namespace AutomationUtils.Matchers
+{
+	/// <summary>
+	/// Matcher for Gauntlet unit tests
+	/// </summary>
+	class GauntletEventMatcher : ILogEventMatcher
+	{
+		public LogEventMatch? Match(ILogCursor cursor)
+		{
+			Match? match;
+			if(cursor.TryMatch(@"^(?<indent>\s*)Error: EngineTest.RunTests Group:(?<group>[^\s]+) \(", out match))
+			{
+				string indent = match.Groups["indent"].Value + " ";
+
+				LogEventBuilder builder = new LogEventBuilder(cursor);
+				builder.Annotate(match.Groups["group"]);
+
+				string group = match.Groups["group"].Value;
+
+				bool inErrorList = false;
+				//				List<LogEventBuilder> ChildEvents = new List<LogEventBuilder>();
+
+				//				int LineCount = 1;
+				EventId eventId = KnownLogEvents.Gauntlet;
+				while (builder.Next.IsMatch($"^(?:{indent}.*|\\s*)$"))
+				{
+					builder.MoveNext();
+					if (inErrorList)
+					{
+						Match? testNameMatch;
+						if (builder.Current.IsMatch(@"^\s*#{1,3} "))
+						{
+							inErrorList = false;
+						}
+						else if (builder.Current.TryMatch(@"^\s*#####\s+(?<friendly_name>.*):\s*(?<name>\S+)\s*", out testNameMatch))
+						{
+							builder.AddProperty("group", group);//.Annotate().AddProperty("group", Group);
+
+							builder.Annotate(testNameMatch.Groups["name"]);
+							builder.Annotate(testNameMatch.Groups["friendly_name"]);
+
+							eventId = KnownLogEvents.Gauntlet_UnitTest;//							ChildEvents.Add(ChildEventBuilder);
+						}
+					}
+					else
+					{
+						if (builder.Current.IsMatch(@"^\s*### The following tests failed:"))
+						{
+							inErrorList = true;
+						}
+					}
+				}
+
+				return builder.ToMatch(LogEventPriority.High, LogLevel.Error, eventId);
+			}
+
+			if (cursor.TryMatch(@"Error: Screenshot '(?<screenshot>[^']+)' test failed", out match))
+			{
+				LogEventBuilder builder = new LogEventBuilder(cursor);
+				builder.Annotate(match.Groups["screenshot"], LogEventMarkup.ScreenshotTest);//.MarkAsScreenshotTest();
+				return builder.ToMatch(LogEventPriority.High, LogLevel.Error, KnownLogEvents.Gauntlet_ScreenshotTest);
+			}
+
+			if (cursor.TryMatch("\\[ERROR\\] (.*)$", out _))
+			{
+				LogEventBuilder builder = new LogEventBuilder(cursor);
+				return builder.ToMatch(LogEventPriority.High, LogLevel.Error, KnownLogEvents.Generic);
+			}
+
+			return null;
+		}
+	}
+}
diff --git a/Engine/Source/Programs/AutomationTool/AutomationUtils/Matchers/GenericEventMatcher.cs b/Engine/Source/Programs/AutomationTool/AutomationUtils/Matchers/GenericEventMatcher.cs
new file mode 100644
--- /dev/null
+++ b/Engine/Source/Programs/AutomationTool/AutomationUtils/Matchers/GenericEventMatcher.cs
@@ -0,0 +1,58 @@
+// Copyright Epic Games, Inc. All Rights Reserved.
+
+using System;
+using System.Text.RegularExpressions;
+using EpicGames.Core;
+using Microsoft.Extensions.Logging;
+
+#nullable enable
+
+namespace AutomationUtils.Matchers
+{
+	/// <summary>
+	/// Low-priority matcher for generic error strings like "warning:" and "error:"
+	/// </summary>
+	class GenericEventMatcher : ILogEventMatcher
+	{
+		/// <inheritdoc/>
+		public LogEventMatch? Match(ILogCursor cursor)
+		{
+			Match? match;
+			if (cursor.TryMatch(@"^\s*(FATAL|fatal error):", out _))
+			{
+				return new LogEventBuilder(cursor).ToMatch(LogEventPriority.Low, LogLevel.Error, KnownLogEvents.Generic);
+			}
+			if (cursor.TryMatch(@"(?<!\w)(?i)(WARNING|ERROR) ?(\([^)]+\)|\[[^\]]+\])?: ", out match))
+			{
+				// Careful to match the first WARNING or ERROR in the line here.
+				LogLevel level = LogLevel.Error;
+				if(match.Groups[1].Value.Equals("WARNING", StringComparison.OrdinalIgnoreCase))
+				{
+					level = LogLevel.Warning;
+				}
+
+				LogEventBuilder builder = new LogEventBuilder(cursor);
+				while (builder.Current.IsMatch(1, String.Format(@"^({0} | *$)", ExtractIndent(cursor[0]!))))
+				{
+					builder.MoveNext();
+				}
+				return builder.ToMatch(LogEventPriority.Lowest, level, KnownLogEvents.Generic);
+			}
+			if (cursor.IsMatch(@"[Ee]rror [A-Z]\d+\s:"))
+			{
+				return new LogEventBuilder(cursor).ToMatch(LogEventPriority.Lowest, LogLevel.Error, KnownLogEvents.Generic);
+			}
+			return null;
+		}
+
+		static string ExtractIndent(string line)
+		{
+			int length = 0;
+			while (length < line.Length && line[length] == ' ')
+			{
+				length++;
+			}
+			return new string(' ', length);
+		}
+	}
+}
diff --git a/Engine/Source/Programs/AutomationTool/AutomationUtils/Matchers/GradleEventMatcher.cs b/Engine/Source/Programs/AutomationTool/AutomationUtils/Matchers/GradleEventMatcher.cs
new file mode 100644
--- /dev/null
+++ b/Engine/Source/Programs/AutomationTool/AutomationUtils/Matchers/GradleEventMatcher.cs
@@ -0,0 +1,50 @@
+// Copyright Epic Games, Inc. All Rights Reserved.
+
+using System.Text.RegularExpressions;
+using EpicGames.Core;
+using Microsoft.Extensions.Logging;
+
+#nullable enable
+
+namespace AutomationUtils.Matchers
+{
+	/// <summary>
+	/// Matcher for Gradle errors
+	/// </summary>
+	class GradleEventMatcher : ILogEventMatcher
+	{
+		public LogEventMatch? Match(ILogCursor cursor)
+		{
+			if (cursor.IsMatch(@"^\s*FAILURE:"))
+			{
+				int maxOffset = 1;
+				for (; ; )
+				{
+					int newMaxOffset = maxOffset;
+					while (cursor.IsMatch(newMaxOffset, @"^\s*$"))
+					{
+						newMaxOffset++;
+					}
+
+					Match? match;
+					if (!cursor.TryMatch(newMaxOffset, @"^(\s*)\*", out match))
+					{
+						break;
+					}
+					maxOffset = newMaxOffset + 1;
+
+					while (cursor.IsMatch(maxOffset, $"^{match.Groups[1].Value}"))
+					{
+						maxOffset++;
+					}
+				}
+
+				LogEventBuilder builder = new LogEventBuilder(cursor, maxOffset);
+				return builder.ToMatch(LogEventPriority.Normal, LogLevel.Error, KnownLogEvents.AutomationTool);
+			}
+
+			return null;
+		}
+	}
+}
+
diff --git a/Engine/Source/Programs/AutomationTool/AutomationUtils/Matchers/LogChannelEventMatcher.cs b/Engine/Source/Programs/AutomationTool/AutomationUtils/Matchers/LogChannelEventMatcher.cs
new file mode 100644
--- /dev/null
+++ b/Engine/Source/Programs/AutomationTool/AutomationUtils/Matchers/LogChannelEventMatcher.cs
@@ -0,0 +1,47 @@
+// Copyright Epic Games, Inc. All Rights Reserved.
+
+using System.Text.RegularExpressions;
+using EpicGames.Core;
+using Microsoft.Extensions.Logging;
+
+#nullable enable
+
+namespace AutomationUtils.Matchers
+{
+	/// <summary>
+	/// Matches events formatted as UE log channel output
+	/// </summary>
+	class LogChannelEventMatcher : ILogEventMatcher
+	{
+		/// <inheritdoc/>
+		public LogEventMatch? Match(ILogCursor input)
+		{
+			Match? match;
+			if(input.TryMatch(@"^(\s*)(?<channel>[a-zA-Z_][a-zA-Z0-9_]*):\s*(?<severity>Error|Warning|Display): ", out match))
+			{
+				LogEventBuilder builder = new LogEventBuilder(input);
+				builder.Annotate(match.Groups["channel"], LogEventMarkup.Channel);
+				builder.Annotate(match.Groups["severity"], LogEventMarkup.Severity);
+				
+				if (builder.Next.TryMatch(@"^\s+", out Match? indent))
+				{						
+					while (builder.Next.IsMatch($"^{indent.Value}"))
+					{
+						builder.MoveNext();
+					}					
+				}
+
+				LogLevel level = match.Groups["severity"].Value switch
+				{
+					"Error" => LogLevel.Error,
+					"Warning" => LogLevel.Warning,
+					_ => LogLevel.Information,
+				};
+
+				return builder.ToMatch(LogEventPriority.Low, level, KnownLogEvents.Engine_LogChannel);
+				
+			}
+			return null;
+		}
+	}
+}
diff --git a/Engine/Source/Programs/AutomationTool/AutomationUtils/Matchers/SourceFileLineEventMatcher.cs b/Engine/Source/Programs/AutomationTool/AutomationUtils/Matchers/SourceFileLineEventMatcher.cs
new file mode 100644
--- /dev/null
+++ b/Engine/Source/Programs/AutomationTool/AutomationUtils/Matchers/SourceFileLineEventMatcher.cs
@@ -0,0 +1,75 @@
+// Copyright Epic Games, Inc. All Rights Reserved.
+
+using System;
+using System.Text.RegularExpressions;
+using EpicGames.Core;
+using Microsoft.Extensions.Logging;
+
+#nullable enable
+
+namespace AutomationUtils.Matchers
+{
+	/// <summary>
+	/// Matches compile errors and annotates with the source file path and revision
+	/// </summary>
+	class SourceFileLineEventMatcher : ILogEventMatcher
+	{
+		const string SeverityPattern =
+			"(?<severity>ERROR|WARNING)";
+
+		const string FilePattern =
+			@"(?<file>" +
+				// optional drive letter
+				@"(?:[a-zA-Z]:)?" +
+				// any non-colon character
+				@"[^:(\s]+" +
+				// any filename character (not whitespace or slash)
+				@"[^:(\s\\/]" +
+			@")";
+
+		const string LinePattern =
+			@"(?<line>\d+)";
+
+		/// <inheritdoc/>
+		public LogEventMatch? Match(ILogCursor input)
+		{
+			Match? match;
+			if (input.TryMatch($"^\\s*{SeverityPattern}: {FilePattern}(?:\\({LinePattern}\\))?: ", out match))
+			{
+				LogLevel level = GetLogLevelFromSeverity(match);
+
+				LogEventBuilder builder = new LogEventBuilder(input);
+
+				builder.AnnotateSourceFile(match.Groups["file"], "");
+				builder.Annotate(match.Groups["severity"], LogEventMarkup.Severity);
+				builder.TryAnnotate(match.Groups["line"], LogEventMarkup.LineNumber);
+
+				EventId eventId;
+				if (input.IsMatch("copyright"))
+				{
+					eventId = KnownLogEvents.AutomationTool_MissingCopyright;
+				}
+				else
+				{
+					eventId = KnownLogEvents.AutomationTool_SourceFileLine;
+				}
+
+				return builder.ToMatch(LogEventPriority.AboveNormal, level, eventId);
+			}
+			return null;
+		}
+
+		static LogLevel GetLogLevelFromSeverity(Match match)
+		{
+			string severity = match.Groups["severity"].Value;
+			if (severity.Equals("WARNING", StringComparison.Ordinal))
+			{
+				return LogLevel.Warning;
+			}
+			else
+			{
+				return LogLevel.Error;
+			}
+		}
+	}
+}
diff --git a/Engine/Source/Programs/AutomationTool/AutomationUtils/Matchers/SystemicEventMatcher.cs b/Engine/Source/Programs/AutomationTool/AutomationUtils/Matchers/SystemicEventMatcher.cs
new file mode 100644
--- /dev/null
+++ b/Engine/Source/Programs/AutomationTool/AutomationUtils/Matchers/SystemicEventMatcher.cs
@@ -0,0 +1,30 @@
+// Copyright Epic Games, Inc. All Rights Reserved.
+
+using EpicGames.Core;
+using Microsoft.Extensions.Logging;
+
+#nullable enable
+
+namespace AutomationUtils.Matchers
+{
+	class SystemicEventMatcher : ILogEventMatcher
+	{
+		public LogEventMatch? Match(ILogCursor cursor)
+		{
+			if (cursor.IsMatch(@"LogDerivedDataCache: .*queries/writes will be limited"))
+			{
+				return new LogEventBuilder(cursor).ToMatch(LogEventPriority.High, LogLevel.Information, KnownLogEvents.Systemic_SlowDDC);
+			}
+			if (cursor.IsMatch(@"^\s*ERROR: Error: EC_OK"))
+			{
+				LogEventBuilder builder = new LogEventBuilder(cursor);
+				while (builder.Next.IsMatch(@"^\s*ERROR:\s*$"))
+				{
+					builder.MoveNext();
+				}
+				return new LogEventBuilder(cursor).ToMatch(LogEventPriority.High, LogLevel.Information, KnownLogEvents.Systemic_PdbUtil);
+			}
+			return null;
+		}
+	}
+}
diff --git a/Engine/Source/Programs/Horde/BuildHorde.xml b/Engine/Source/Programs/Horde/BuildHorde.xml
--- a/Engine/Source/Programs/Horde/BuildHorde.xml
+++ b/Engine/Source/Programs/Horde/BuildHorde.xml
@@ -20,6 +20,7 @@
 				Engine/Binaries/DotNET/EpicGames.Perforce.Native/...
 				Engine/Source/Programs/Shared/...
 				Engine/Source/Programs/Horde/...
+				Engine/Source/Programs/AutomationTool/AutomationUtils/Matchers/...
 				Engine/Source/Programs/UnrealBuildTool/Matchers/...
 				Engine/Source/ThirdParty/Perforce/...
 			</Property>
diff --git a/Engine/Source/Programs/Horde/Horde.Agent/Horde.Agent.csproj b/Engine/Source/Programs/Horde/Horde.Agent/Horde.Agent.csproj
--- a/Engine/Source/Programs/Horde/Horde.Agent/Horde.Agent.csproj
+++ b/Engine/Source/Programs/Horde/Horde.Agent/Horde.Agent.csproj
@@ -71,6 +71,9 @@
   </ItemGroup>
 
   <ItemGroup>
+    <Compile Include="..\..\AutomationTool\AutomationUtils\Matchers\*.cs">
+      <Link>Parser\Matchers\%(FileName)%(Extension)</Link>
+    </Compile>
     <Compile Include="..\..\UnrealBuildTool\Matchers\*.cs">
       <Link>Parser\Matchers\%(FileName)%(Extension)</Link>
     </Compile>
diff --git a/Engine/Source/Programs/Horde/Horde.Agent/Parser/Matchers/ContentEventMatcher.cs b/Engine/Source/Programs/Horde/Horde.Agent/Parser/Matchers/ContentEventMatcher.cs
deleted file mode 100644
diff --git a/Engine/Source/Programs/Horde/Horde.Agent/Parser/Matchers/CrashEventMatcher.cs b/Engine/Source/Programs/Horde/Horde.Agent/Parser/Matchers/CrashEventMatcher.cs
deleted file mode 100644
diff --git a/Engine/Source/Programs/Horde/Horde.Agent/Parser/Matchers/ExceptionEventMatcher.cs b/Engine/Source/Programs/Horde/Horde.Agent/Parser/Matchers/ExceptionEventMatcher.cs
deleted file mode 100644
diff --git a/Engine/Source/Programs/Horde/Horde.Agent/Parser/Matchers/ExitCodeEventMatcher.cs b/Engine/Source/Programs/Horde/Horde.Agent/Parser/Matchers/ExitCodeEventMatcher.cs
deleted file mode 100644
diff --git a/Engine/Source/Programs/Horde/Horde.Agent/Parser/Matchers/GauntletEventMatcher.cs b/Engine/Source/Programs/Horde/Horde.Agent/Parser/Matchers/GauntletEventMatcher.cs
deleted file mode 100644
diff --git a/Engine/Source/Programs/Horde/Horde.Agent/Parser/Matchers/GenericEventMatcher.cs b/Engine/Source/Programs/Horde/Horde.Agent/Parser/Matchers/GenericEventMatcher.cs
deleted file mode 100644
diff --git a/Engine/Source/Programs/Horde/Horde.Agent/Parser/Matchers/GradleEventMatcher.cs b/Engine/Source/Programs/Horde/Horde.Agent/Parser/Matchers/GradleEventMatcher.cs
deleted file mode 100644
diff --git a/Engine/Source/Programs/Horde/Horde.Agent/Parser/Matchers/LogChannelEventMatcher.cs b/Engine/Source/Programs/Horde/Horde.Agent/Parser/Matchers/LogChannelEventMatcher.cs
deleted file mode 100644
diff --git a/Engine/Source/Programs/Horde/Horde.Agent/Parser/Matchers/SourceFileLineEventMatcher.cs b/Engine/Source/Programs/Horde/Horde.Agent/Parser/Matchers/SourceFileLineEventMatcher.cs
deleted file mode 100644
diff --git a/Engine/Source/Programs/Horde/Horde.Agent/Parser/Matchers/SystemicEventMatcher.cs b/Engine/Source/Programs/Horde/Horde.Agent/Parser/Matchers/SystemicEventMatcher.cs
deleted file mode 100644
