Stage Base.ini into crash reporter pak files.

diff --git a/Engine/Source/Programs/AutomationTool/Scripts/CopyBuildToStagingDirectory.Automation.cs b/Engine/Source/Programs/AutomationTool/Scripts/CopyBuildToStagingDirectory.Automation.cs
--- a/Engine/Source/Programs/AutomationTool/Scripts/CopyBuildToStagingDirectory.Automation.cs
+++ b/Engine/Source/Programs/AutomationTool/Scripts/CopyBuildToStagingDirectory.Automation.cs
@@ -1192,6 +1192,12 @@ namespace AutomationScripts
 						}
 
 						// Add config files.
+						// TODO: Consider staging all engine ini files for parity with running in a workspace.
+						//SC.StageCrashReporterFiles(StagedFileType.UFS, DirectoryReference.Combine(SC.EngineRoot, "Config"), StageFilesSearch.TopDirectoryOnly);
+						//SC.StageCrashReporterFiles(StagedFileType.UFS, DirectoryReference.Combine(SC.EngineRoot, "Config", CrashReportPlatform.ToString()), StageFilesSearch.TopDirectoryOnly);
+						var BaseIni = FileReference.Combine(SC.EngineRoot, "Config", "Base.ini");
+						var StagedBaseIni = new StagedFileReference(BaseIni.MakeRelativeTo(Unreal.RootDirectory));
+						SC.StageCrashReporterFile(StagedFileType.UFS, BaseIni, StagedBaseIni);
 						SC.StageCrashReporterFiles(StagedFileType.UFS, DirectoryReference.Combine(SC.EngineRoot, "Programs", "CrashReportClient", "Config"), StageFilesSearch.AllDirectories);
 					}
 					// check if the game will be verifying ssl connections - if not, we can skip staging files that won't be needed
