UAT: Fix log event matchers not being registered at UAT startup.

diff --git a/Engine/Source/Programs/AutomationTool/Program.cs b/Engine/Source/Programs/AutomationTool/Program.cs
--- a/Engine/Source/Programs/AutomationTool/Program.cs
+++ b/Engine/Source/Programs/AutomationTool/Program.cs
@@ -364,6 +364,9 @@ namespace AutomationToolDriver
 			Assembly UnrealBuildToolAssembly = typeof(UnrealBuildTool.BuildVersion).Assembly;
 			Log.EventParser.AddMatchersFromAssembly(UnrealBuildToolAssembly);
 
+			Assembly AutomationUtilsAssembly = typeof(AutomationTool.CommandUtils).Assembly;
+			Log.EventParser.AddMatchersFromAssembly(AutomationUtilsAssembly);
+
 			// Enter the main program section
 			ExitCode ReturnCode = ExitCode.Success;
 			try
