UAT: Move logger matcher initialization into AutomationUtils, since it's not a hard dependency of AutomationTool.

diff --git a/Engine/Source/Programs/AutomationTool/AutomationUtils/Automation.cs b/Engine/Source/Programs/AutomationTool/AutomationUtils/Automation.cs
--- a/Engine/Source/Programs/AutomationTool/AutomationUtils/Automation.cs
+++ b/Engine/Source/Programs/AutomationTool/AutomationUtils/Automation.cs
@@ -72,6 +72,10 @@ namespace AutomationTool
 				Log.TraceVerbose("IsBuildMachine={0}", IsBuildMachine);
 				Environment.SetEnvironmentVariable("IsBuildMachine", IsBuildMachine ? "1" : "0");
 
+				// Register all the log event matchers
+				Assembly AutomationUtilsAssembly = Assembly.GetExecutingAssembly();
+				Log.EventParser.AddMatchersFromAssembly(AutomationUtilsAssembly);
+
 				// Get the path to the telemetry file, if present
 				string TelemetryFile = GlobalCommandLine.TelemetryPath;
 				JsonTracer Tracer = JsonTracer.TryRegisterAsGlobalTracer();
diff --git a/Engine/Source/Programs/AutomationTool/Program.cs b/Engine/Source/Programs/AutomationTool/Program.cs
--- a/Engine/Source/Programs/AutomationTool/Program.cs
+++ b/Engine/Source/Programs/AutomationTool/Program.cs
@@ -13,6 +13,7 @@ using EpicGames.Core;
 using System.IO;
 using System.Collections.Generic;
 using UnrealBuildBase;
+using Microsoft.Extensions.Logging;
 
 namespace AutomationToolDriver
 {
@@ -25,7 +26,7 @@ namespace AutomationToolDriver
 		/// <param name="CommandLine">Command line</param>
 		/// <param name="CurrentCommand">Recently parsed command</param>
 		/// <returns>True if the parameter has been successfully parsed.</returns>
-		private static void ParseParam(string CurrentParam, CommandInfo CurrentCommand)
+		private static void ParseParam(string CurrentParam, CommandInfo CurrentCommand, ILogger Logger)
 		{
 			if (AutomationToolCommandLine.IsParameterIgnored(CurrentParam))
             {
@@ -61,7 +62,7 @@ namespace AutomationToolDriver
 					List<string> OutAdditionalScriptDirectories = (List<string>)AutomationToolCommandLine.GetValueUnchecked(Option_ScriptDir) ?? new List<string>();
 					OutAdditionalScriptDirectories.Add(ScriptDir);
 					AutomationToolCommandLine.SetUnchecked(Option_ScriptDir, OutAdditionalScriptDirectories);
-					Log.TraceVerbose("Found additional script dir: {0}", ScriptDir);
+					Logger.LogDebug("Found additional script dir: {0}", ScriptDir);
 				}
 				else
 				{
@@ -95,7 +96,7 @@ namespace AutomationToolDriver
 				}
 				string EnvVarValue = CurrentParam.Substring(ValueStartIndex);
 
-				Log.TraceLog($"SetEnvVar {EnvVarName}={EnvVarValue}");
+				Logger.LogDebug($"SetEnvVar {EnvVarName}={EnvVarValue}");
 				Environment.SetEnvironmentVariable(EnvVarName, EnvVarValue);
 			}
 		}
@@ -229,7 +230,7 @@ namespace AutomationToolDriver
 		/// Parse the command line and create a list of commands to execute.
 		/// </summary>
 		/// <param name="Arguments">Command line</param>
-		public static void ParseCommandLine(string[] Arguments)
+		public static void ParseCommandLine(string[] Arguments, ILogger Logger)
 		{
 			AutomationToolCommandLine = new ParsedCommandLine(
 				new Dictionary<string, string>(StringComparer.InvariantCultureIgnoreCase)
@@ -261,7 +262,7 @@ namespace AutomationToolDriver
 
 			ParseProfile(ref Arguments);
 
-			Log.TraceInformation("Parsing command line: {0}", CommandLine.FormatCommandLine(Arguments));
+			Logger.LogInformation("Parsing command line: {CommandLine}", CommandLine.FormatCommandLine(Arguments));
 
 			CommandInfo CurrentCommand = null;
 			for (int Index = 0; Index < Arguments.Length; ++Index)
@@ -272,7 +273,7 @@ namespace AutomationToolDriver
 				{
 					if (Param.StartsWith("-") || Param.Contains("="))
 					{
-						ParseParam(Arguments[Index], CurrentCommand);
+						ParseParam(Arguments[Index], CurrentCommand, Logger);
 					}
 					else
 					{
@@ -286,10 +287,10 @@ namespace AutomationToolDriver
 			var Result = AutomationToolCommandLine.CommandsToExecute.Count > 0 || AutomationToolCommandLine.IsSetGlobal("-Help") || AutomationToolCommandLine.IsSetGlobal("-List");
 			if (AutomationToolCommandLine.CommandsToExecute.Count > 0)
 			{
-				Log.TraceVerbose("Found {0} scripts to execute:", AutomationToolCommandLine.CommandsToExecute.Count);
-				foreach (var Command in AutomationToolCommandLine.CommandsToExecute)
+				Logger.LogDebug("Found {NumScripts} scripts to execute:", AutomationToolCommandLine.CommandsToExecute.Count);
+				foreach (CommandInfo Command in AutomationToolCommandLine.CommandsToExecute)
 				{
-					Log.TraceVerbose("  " + Command.ToString());
+					Logger.LogDebug("  {Command}", Command.ToString());
 				}
 			}
 			else if (!Result)
@@ -313,18 +314,20 @@ namespace AutomationToolDriver
 		// Do not add [STAThread] here. It will cause deadlocks in platform automation code.
 		public static async Task<int> Main(string[] Arguments)
 		{
+			ILogger Logger = Log.Logger;
+
 			// Initialize the log system, buffering the output until we can create the log file
 			Log.AddTraceListener(StartupListener);
-			Log.TraceInformation("Starting AutomationTool...");
+			Logger.LogInformation("Starting AutomationTool...");
 
 			// Populate AutomationToolCommandLine and CommandsToExecute
 			try
 			{
-				ParseCommandLine(Arguments);
+				ParseCommandLine(Arguments, Logger);
 			}
 			catch (Exception Ex)
 			{
-				Log.TraceError("Error: " + Ex.Message);
+				Logger.LogError(Ex, "ERROR: " + Ex.Message);
 				return (int)ExitCode.Error_Arguments;
 			}
 
@@ -377,19 +380,19 @@ namespace AutomationToolDriver
 				AssemblyUtils.InstallRecursiveAssemblyResolver(PathToBinariesDotNET);
 
 				// Log the operating environment. Since we usually compile to AnyCPU, we may be executed using different system paths under WOW64.
-				Log.TraceVerbose("{2}: Running on {0} as a {1}-bit process.", RuntimePlatform.Current.ToString(), Environment.Is64BitProcess ? 64 : 32, DateTime.UtcNow.ToString("o"));
+				Logger.LogDebug("Running on {Platform} as a {Bitness}-bit process.", RuntimePlatform.Current.ToString(), Environment.Is64BitProcess ? 64 : 32);
 
 				// Log if we're running from the launcher
 				string ExecutingAssemblyLocation = Assembly.GetExecutingAssembly().Location;
 				if (string.Compare(ExecutingAssemblyLocation, Assembly.GetEntryAssembly().GetOriginalLocation(), StringComparison.OrdinalIgnoreCase) != 0)
 				{
-					Log.TraceVerbose("Executed from AutomationToolLauncher ({0})", ExecutingAssemblyLocation);
+					Logger.LogDebug("Executed from AutomationToolLauncher ({Location})", ExecutingAssemblyLocation);
 				}
-				Log.TraceVerbose("CWD={0}", Environment.CurrentDirectory);
+				Logger.LogDebug("CWD={Cwd}", Environment.CurrentDirectory);
 
 				// Log the application version
 				FileVersionInfo Version = AssemblyUtils.ExecutableVersion;
-				Log.TraceVerbose("{0} ver. {1}", Version.ProductName, Version.ProductVersion);
+				Logger.LogDebug("{ProductName} ver. {ProductVersion}", Version.ProductName, Version.ProductVersion);
 
 				bool bWaitForUATMutex = AutomationToolCommandLine.IsSetGlobal("-WaitForUATMutex");
 
@@ -398,13 +401,13 @@ namespace AutomationToolDriver
 			}
 			catch (Exception Ex)
             {
-				Log.TraceError(ExceptionUtils.FormatException(Ex));
+				Logger.LogError(Ex, "Unhandled exception: {Message}", Ex.Message);
             }
             finally
             {
 				// Write the exit code
-                Log.TraceInformation("AutomationTool executed for {0}", Timer.Elapsed.ToString("h'h 'm'm 's's'"));
-                Log.TraceInformation("AutomationTool exiting with ExitCode={0} ({1})", (int)ReturnCode, ReturnCode);
+                Logger.LogInformation("AutomationTool executed for {Time}", Timer.Elapsed.ToString("h'h 'm'm 's's'"));
+                Logger.LogInformation("AutomationTool exiting with ExitCode={ExitCode} ({ExitReason})", (int)ReturnCode, ReturnCode);
 
                 // Can't use NoThrow here because the code logs exceptions. We're shutting down logging!
                 Trace.Close();
@@ -414,7 +417,8 @@ namespace AutomationToolDriver
 
 		static async Task<ExitCode> MainProc()
 		{
-			Log.TraceInformation("Initializing script modules...");
+			ILogger Logger = Log.Logger;
+			Logger.LogInformation("Initializing script modules...");
 			var StartTime = DateTime.UtcNow;
 			string ScriptsForProject = (string)AutomationToolCommandLine.GetValueUnchecked("-ScriptsForProject");
 			List<string> AdditionalScriptDirs = (List<string>) AutomationToolCommandLine.GetValueUnchecked("-ScriptDir");
@@ -429,7 +433,7 @@ namespace AutomationToolDriver
 					Rules.RulesFileType.AutomationModule, ScriptsForProject, AdditionalScriptDirs, bForceCompile, bNoCompile, bUseBuildRecords, 
 					out bBuildSuccess, (int Count) =>
                     {
-						Log.TraceInformation($"Building {Count} projects (see Log 'Engine/Programs/AutomationTool/Saved/Logs/Log.txt' for more details)");
+						Logger.LogInformation("Building {Count} projects (see Log 'Engine/Programs/AutomationTool/Saved/Logs/Log.txt' for more details)", Count);
 					});
 
 			if (!bBuildSuccess)
@@ -456,7 +460,7 @@ namespace AutomationToolDriver
 
 			Type AutomationTools_Automation = AutomationUtilsAssembly.GetType("AutomationTool.Automation");
 			MethodInfo Automation_Process = AutomationTools_Automation.GetMethod("ProcessAsync");
-			Log.TraceInformation("Total script module initialization time: {0:0.00} s.", (DateTime.UtcNow - StartTime).TotalMilliseconds / 1000);
+			Logger.LogInformation("Total script module initialization time: {InitTime:0.00} s.", (DateTime.UtcNow - StartTime).TotalMilliseconds / 1000);
 			return await (Task<ExitCode>) Automation_Process.Invoke(null,
 				new object[] {AutomationToolCommandLine, StartupListener, ScriptModuleAssemblyPaths});
 		}
