- Fixed up the EditorOptional directories to have a different save vs input, since the save dir needs to be cleaned, but the input one cannot get cleaned

diff --git a/Engine/Source/Programs/AutomationTool/AutomationUtils/DeploymentContext.cs b/Engine/Source/Programs/AutomationTool/AutomationUtils/DeploymentContext.cs
--- a/Engine/Source/Programs/AutomationTool/AutomationUtils/DeploymentContext.cs
+++ b/Engine/Source/Programs/AutomationTool/AutomationUtils/DeploymentContext.cs
@@ -201,6 +201,11 @@ public class DeploymentContext //: ProjectParams
 	/// </summary>
 	public DirectoryReference OptionalFileStageDirectory = null;
 
+	/// <summary>
+	/// Directory to read all of the optional (ie editor only) files from (written earlier with OptionalFileStageDirectory)
+	/// </summary>
+	public DirectoryReference OptionalFileInputDirectory = null;
+
 	/// <summary>
 	/// Directory name for staged projects
 	/// </summary>
@@ -386,7 +391,8 @@ public class DeploymentContext //: ProjectParams
 		FileReference RawProjectPathOrName,
 		DirectoryReference InLocalRoot,
 		DirectoryReference BaseStageDirectory,
-		DirectoryReference OptionalStageDirectory,
+		DirectoryReference OptionalFileStageDirectory,
+		DirectoryReference OptionalFileInputDirectory,
 		DirectoryReference BaseArchiveDirectory,
 		Platform InSourcePlatform,
         Platform InTargetPlatform,
@@ -454,7 +460,8 @@ public class DeploymentContext //: ProjectParams
 			StageDirectory = DirectoryReference.Combine(BaseStageDirectory, FinalCookPlatform);
 			DebugStageDirectory = InSeparateDebugStageDirectory? DirectoryReference.Combine(BaseStageDirectory, FinalCookPlatform + "Debug") : StageDirectory;
 		}
-		OptionalFileStageDirectory = OptionalStageDirectory;
+		this.OptionalFileStageDirectory = OptionalFileStageDirectory;
+		this.OptionalFileInputDirectory = OptionalFileInputDirectory;
 
 		if (BaseArchiveDirectory != null)
 		{
diff --git a/Engine/Source/Programs/AutomationTool/AutomationUtils/ProjectParams.cs b/Engine/Source/Programs/AutomationTool/AutomationUtils/ProjectParams.cs
--- a/Engine/Source/Programs/AutomationTool/AutomationUtils/ProjectParams.cs
+++ b/Engine/Source/Programs/AutomationTool/AutomationUtils/ProjectParams.cs
@@ -805,7 +805,8 @@ namespace AutomationTool
 				this.Stage = true;
 			}
 			this.StageDirectoryParam = ParseParamValueIfNotSpecified(Command, StageDirectoryParam, "stagingdirectory", String.Empty, true);
-			this.OptionalStageDirectory = ParseParamValueIfNotSpecified(Command, OptionalStageDirectory, "optionalfilestagingdirectory", String.Empty, true);
+			this.OptionalFileStagingDirectory = ParseParamValueIfNotSpecified(Command, OptionalFileStagingDirectory, "optionalfilestagingdirectory", String.Empty, true);
+			this.OptionalFileInputDirectory = ParseParamValueIfNotSpecified(Command, OptionalFileInputDirectory, "optionalfileinputdirectory", String.Empty, true);
 			this.bCodeSign = GetOptionalParamValueIfNotSpecified(Command, CodeSign, CommandUtils.IsBuildMachine, "CodeSign", "NoCodeSign");
 			this.bTreatNonShippingBinariesAsDebugFiles = GetParamValueIfNotSpecified(Command, TreatNonShippingBinariesAsDebugFiles, false, "TreatNonShippingBinariesAsDebugFiles");
 			this.bUseExtraFlavor = GetParamValueIfNotSpecified(Command, UseExtraFlavor, false, "UseExtraFlavor");
@@ -1488,7 +1489,10 @@ namespace AutomationTool
 		public string StageDirectoryParam;
 
 		[Help("optionalfilestagingdirectory=Path", "Directory to copy the optional files to, i.e. -optionalfilestagingdirectory=C:\\StageOptional")]
-		public string OptionalStageDirectory;
+		public string OptionalFileStagingDirectory;
+
+		[Help("optionalfileinputdirectory=Path", "Directory to read the optional files from, i.e. -optionalfileinputdirectory=C:\\StageOptional")]
+		public string OptionalFileInputDirectory;
 
 		[Help("unrealexe=ExecutableName", "Name of the Unreal Editor executable, i.e. -unrealexe=UnrealEditor.exe")]
 		public string UnrealExe;
diff --git a/Engine/Source/Programs/AutomationTool/CookedEditor/CookedEditor.Automation.cs b/Engine/Source/Programs/AutomationTool/CookedEditor/CookedEditor.Automation.cs
--- a/Engine/Source/Programs/AutomationTool/CookedEditor/CookedEditor.Automation.cs
+++ b/Engine/Source/Programs/AutomationTool/CookedEditor/CookedEditor.Automation.cs
@@ -500,10 +500,10 @@ public class MakeCookedEditor : BuildCommand
 		}
 
 		// if a subclass or -makerelease didn't set ReleaseOptionalFileStageDirectory, then look in the Params for the commandline option
-		if (ReleaseOptionalFileStageDirectory == null && !string.IsNullOrEmpty(Params.OptionalStageDirectory) &&
-			File.Exists(Params.OptionalStageDirectory))
+		if (ReleaseOptionalFileStageDirectory == null && !string.IsNullOrEmpty(Params.OptionalFileInputDirectory) &&
+			Directory.Exists(Params.OptionalFileInputDirectory))
 		{
-			ReleaseOptionalFileStageDirectory = new DirectoryReference(Params.OptionalStageDirectory);
+			ReleaseOptionalFileStageDirectory = new DirectoryReference(Params.OptionalFileInputDirectory);
 		}
 
 		if (ReleaseOptionalFileStageDirectory != null)
diff --git a/Engine/Source/Programs/AutomationTool/Scripts/CopyBuildToStagingDirectory.Automation.cs b/Engine/Source/Programs/AutomationTool/Scripts/CopyBuildToStagingDirectory.Automation.cs
--- a/Engine/Source/Programs/AutomationTool/Scripts/CopyBuildToStagingDirectory.Automation.cs
+++ b/Engine/Source/Programs/AutomationTool/Scripts/CopyBuildToStagingDirectory.Automation.cs
@@ -4450,7 +4450,8 @@ namespace AutomationScripts
 				//@todo should pull StageExecutables from somewhere else if not cooked
 				var SC = new DeploymentContext(Params.RawProjectPath, Unreal.RootDirectory,
 					String.IsNullOrEmpty(StageDirectory) ? null : new DirectoryReference(StageDirectory),
-					String.IsNullOrEmpty(Params.OptionalStageDirectory) ? null : new DirectoryReference(Params.OptionalStageDirectory),
+					String.IsNullOrEmpty(Params.OptionalFileStagingDirectory) ? null : new DirectoryReference(Params.OptionalFileStagingDirectory),
+					String.IsNullOrEmpty(Params.OptionalFileInputDirectory) ? null : new DirectoryReference(Params.OptionalFileInputDirectory),
 					String.IsNullOrEmpty(ArchiveDirectory) ? null : new DirectoryReference(ArchiveDirectory),
 					Platform.Platforms[CookedDataPlatform],
 					Platform.Platforms[StagePlatform],
