- Fixing ompile errors from robomerge

diff --git a/Engine/Source/Programs/AutomationTool/CookedEditor/CookedEditor.Automation.cs b/Engine/Source/Programs/AutomationTool/CookedEditor/CookedEditor.Automation.cs
--- a/Engine/Source/Programs/AutomationTool/CookedEditor/CookedEditor.Automation.cs
+++ b/Engine/Source/Programs/AutomationTool/CookedEditor/CookedEditor.Automation.cs
@@ -1100,7 +1100,7 @@ public class MakeCookedEditor : BuildCommand
 		// cooked editor doesn't work without OptionalContent now, so always generated it, and save it somewhere that staging of the cookededitor will get
 		ReleaseOptionalFileStageDirectory = DirectoryReference.Combine(MainParams.RawProjectPath.Directory, "Saved", "CookedEditor", "OptionalData");
 		ReleaseParams.OptionalContent = true;
-		ReleaseParams.OptionalStageDirectory = ReleaseOptionalFileStageDirectory.FullName;
+		ReleaseParams.OptionalFileStagingDirectory = ReleaseOptionalFileStageDirectory.FullName;
 
 		ModifyReleaseParams(ReleaseParams);
 
