Added UAT parameter -NoZenAutoLaunch for easier debugging of Zen server

diff --git a/Engine/Source/Programs/AutomationTool/AutomationUtils/ProjectParams.cs b/Engine/Source/Programs/AutomationTool/AutomationUtils/ProjectParams.cs
--- a/Engine/Source/Programs/AutomationTool/AutomationUtils/ProjectParams.cs
+++ b/Engine/Source/Programs/AutomationTool/AutomationUtils/ProjectParams.cs
@@ -321,6 +321,7 @@ namespace AutomationTool
 			this.IoStore = InParams.IoStore;
 			this.Cook4IoStore = InParams.Cook4IoStore;
 			this.ZenStore = InParams.ZenStore;
+			this.NoZenAutoLaunch = InParams.NoZenAutoLaunch;
 			this.GenerateOptimizationData = InParams.GenerateOptimizationData;
 			this.SignPak = InParams.SignPak;
 			this.SignedPak = InParams.SignedPak;
@@ -520,6 +521,7 @@ namespace AutomationTool
 			bool? IoStore = null,
 			bool? Cook4IoStore = null,
 			bool? ZenStore = null,
+			string NoZenAutoLaunch = null,
 			bool? SkipIoStore = null,
 			bool? GenerateOptimizationData = null,
 			bool? Prereqs = null,
@@ -724,6 +726,11 @@ namespace AutomationTool
 			{
 				this.AdditionalCookerOptions += " -ZenStore";
 			}
+			this.NoZenAutoLaunch = ParseParamValueIfNotSpecified(Command, NoZenAutoLaunch, "NoZenAutoLaunch", String.Empty);
+			if (!string.IsNullOrEmpty(this.NoZenAutoLaunch) && this.Cook && !this.SkipCook)
+			{
+				this.AdditionalCookerOptions += string.Format(" -NoZenAutoLaunch={0}", this.NoZenAutoLaunch);
+			}
 			this.GenerateOptimizationData = GetParamValueIfNotSpecified(Command, GenerateOptimizationData, this.GenerateOptimizationData, "makebinaryconfig");
 			
 			this.SkipPak = GetParamValueIfNotSpecified(Command, SkipPak, this.SkipPak, "skippak");
@@ -1354,6 +1361,12 @@ namespace AutomationTool
 		[Help("zenstore", "save cooked output data to the Zen storage server")]
 		public bool ZenStore { private set; get; }
 
+		/// <summary>
+		/// Shared: URL to a running Zen server
+		/// </summary>
+		[Help("nozenautolaunch", "URL to a running Zen server")]
+		public string NoZenAutoLaunch { private set; get; }
+
 		/// <summary>
 		/// Shared: True if optimization data is generated during staging that can improve loadtimes
 		/// </summary>
@@ -3086,6 +3099,7 @@ namespace AutomationTool
 				CommandUtils.LogLog("SkipIoStore={0}", SkipIoStore);
 				CommandUtils.LogLog("Cook4IoStore={0}", Cook4IoStore);
 				CommandUtils.LogLog("ZenStore={0}", ZenStore);
+				CommandUtils.LogLog("NoZenAutoLaunch={0}", NoZenAutoLaunch);
 				CommandUtils.LogLog("SkipEncryption={0}", SkipEncryption);
 				CommandUtils.LogLog("GenerateOptimizationData={0}", GenerateOptimizationData);
 				CommandUtils.LogLog("SkipPackage={0}", SkipPackage);
