UAT: Update changelist comment to reflect latest checkin requirements.

#rnx

diff --git a/Engine/Source/Programs/AutomationTool/Scripts/FixPerforceCase.cs b/Engine/Source/Programs/AutomationTool/Scripts/FixPerforceCase.cs
--- a/Engine/Source/Programs/AutomationTool/Scripts/FixPerforceCase.cs
+++ b/Engine/Source/Programs/AutomationTool/Scripts/FixPerforceCase.cs
@@ -14,7 +14,7 @@ namespace AutomationTool
 	[Help("Files", "Pattern for files to match. Should be a full depot path with the correct case. May end with a wildcard.")]
 	class FixPerforceCase : BuildCommand
 	{
-		const string BoilerplateText = "\n\n#rb none\n#rnx";
+		const string BoilerplateText = "\n\n#rb none\n#rnx\n#preflight skip";
 
 		/// <summary>
 		/// Main entry point for the command
