Remove HoloLens platform references from AutomationTool.

#rnx

diff --git a/Engine/Platforms/Hololens/Source/Programs/AutomationTool/HoloLensPlatform.Automation.cs b/Engine/Platforms/Hololens/Source/Programs/AutomationTool/HoloLensPlatform.Automation.cs
--- a/Engine/Platforms/Hololens/Source/Programs/AutomationTool/HoloLensPlatform.Automation.cs
+++ b/Engine/Platforms/Hololens/Source/Programs/AutomationTool/HoloLensPlatform.Automation.cs
@@ -1200,12 +1200,20 @@
 		public override bool IsSupported { get { return true; } }
 		public override bool UseAbsLog { get { return false; } }
 		public override bool LaunchViaUFE { get { return false; } }
+		public override string ICUDataVersion { get { return "icudt53l"; } }
 
+
 		public override List<string> GetDebugFileExtensions()
 		{
 			return new List<string> { ".pdb", ".map" };
 		}
 
+		protected override string GetPlatformExeExtension()
+		{
+			return ".exe";
+		}
+
+
 		public override void StripSymbols(FileReference SourceFile, FileReference TargetFile)
 		{
 			bool bStripInPlace = false;
diff --git a/Engine/Platforms/Hololens/Source/Programs/AutomationTool/Scripts/BuildCMakeLib_HoloLens.Automation.cs b/Engine/Platforms/Hololens/Source/Programs/AutomationTool/Scripts/BuildCMakeLib_HoloLens.Automation.cs
new file mode 100644
--- /dev/null
+++ b/Engine/Platforms/Hololens/Source/Programs/AutomationTool/Scripts/BuildCMakeLib_HoloLens.Automation.cs
@@ -0,0 +1,26 @@
+// Copyright Epic Games, Inc. All Rights Reserved.
+
+using UnrealBuildTool;
+
+class VS2019TargetPlatform_HoloLens : BuildCMakeLib.VS2019TargetPlatform
+{
+	public override string PlatformOrGroupName => nameof(UnrealTargetPlatform.HoloLens);
+	public override string DebugDatabaseExtension => "pdb";
+	public override string DynamicLibraryExtension => "dll";
+	public override string StaticLibraryExtension => "lib";
+	public override string VariantDirectory => Architecture.ToLower();
+	public override bool IsPlatformExtension => true;
+
+	private readonly string Architecture;
+
+	public VS2019TargetPlatform_HoloLens(string Architecture = "Win64")
+	{
+		this.Architecture = Architecture;
+	}
+
+	public override string GetCMakeSetupArguments(BuildCMakeLib.TargetLib TargetLib, string TargetConfiguration)
+	{
+		return base.GetCMakeSetupArguments(TargetLib, TargetConfiguration)
+			+ string.Format(" -A {0}", Architecture);
+	}
+}
\ No newline at end of file
diff --git a/Engine/Platforms/Hololens/Source/Programs/AutomationTool/Scripts/BuildPluginCommand_HoloLens.Automation.cs b/Engine/Platforms/Hololens/Source/Programs/AutomationTool/Scripts/BuildPluginCommand_HoloLens.Automation.cs
new file mode 100644
--- /dev/null
+++ b/Engine/Platforms/Hololens/Source/Programs/AutomationTool/Scripts/BuildPluginCommand_HoloLens.Automation.cs
@@ -0,0 +1,56 @@
+// Copyright Epic Games, Inc. All Rights Reserved.
+
+using System;
+using System.Collections.Generic;
+using AutomationTool;
+using UnrealBuildTool;
+using EpicGames.Core;
+
+public class BuildPlugin_HoloLens : BuildPlugin.TargetPlatform
+{
+	const string HoloLensArchitecture = "arm64+x64";
+
+	public BuildPlugin_HoloLens()
+	{
+	}
+
+	[Obsolete("Deprecated in UE5.1; function signature changed")]
+	public override void CompilePluginWithUBT(string UBTExe, FileReference HostProjectFile, FileReference HostProjectPluginFile, PluginDescriptor Plugin, string TargetName, TargetType TargetType, UnrealTargetPlatform Platform, UnrealTargetConfiguration Configuration, List<FileReference> ManifestFileNames, string InAdditionalArgs)
+	{
+		// Make sure to save the manifests for each architecture with unique names so they don't get overwritten.
+		// This fixes packaging issues when building from binary engine releases, where the build produces a manifest for the plugin for ARM64, which
+		// then gets overwritten by the manifest for x64. Then during packaging, the plugin is referencing a manifest for the wrong architecture.
+		foreach (string Arch in HoloLensArchitecture.Split('+'))
+		{
+			FileReference ManifestFileName = FileReference.Combine(HostProjectFile.Directory, "Saved", String.Format("Manifest-{0}-{1}-{2}-{3}.xml", TargetName, Platform, Configuration, Arch));
+			ManifestFileNames.Add(ManifestFileName);
+			string Arguments = String.Format("-plugin={0} -iwyu -noubtmakefiles -manifest={1} -nohotreload", CommandUtils.MakePathSafeToUseWithCommandLine(HostProjectPluginFile.FullName), CommandUtils.MakePathSafeToUseWithCommandLine(ManifestFileName.FullName));
+			Arguments += String.Format(" -Architecture={0}", Arch);
+			if (!String.IsNullOrEmpty(InAdditionalArgs))
+			{
+				Arguments += InAdditionalArgs;
+			}
+			CommandUtils.RunUBT(CmdEnv, UBTExe, HostProjectFile, TargetName, Platform, Configuration, Arguments);
+		}
+	}
+
+	public override void CompilePluginWithUBT(FileReference UnrealBuildToolDll, FileReference HostProjectFile, FileReference HostProjectPluginFile, PluginDescriptor Plugin, string TargetName, TargetType TargetType, UnrealTargetPlatform Platform, UnrealTargetConfiguration Configuration, List<FileReference> ManifestFileNames, string InAdditionalArgs)
+	{
+		// Make sure to save the manifests for each architecture with unique names so they don't get overwritten.
+		// This fixes packaging issues when building from binary engine releases, where the build produces a manifest for the plugin for ARM64, which
+		// then gets overwritten by the manifest for x64. Then during packaging, the plugin is referencing a manifest for the wrong architecture.
+		foreach (string Arch in HoloLensArchitecture.Split('+'))
+		{
+			FileReference ManifestFileName = FileReference.Combine(HostProjectFile.Directory, "Saved", String.Format("Manifest-{0}-{1}-{2}-{3}.xml", TargetName, Platform, Configuration, Arch));
+			ManifestFileNames.Add(ManifestFileName);
+			string Arguments = String.Format("-plugin={0} -iwyu -noubtmakefiles -manifest={1} -nohotreload", CommandUtils.MakePathSafeToUseWithCommandLine(HostProjectPluginFile.FullName), CommandUtils.MakePathSafeToUseWithCommandLine(ManifestFileName.FullName));
+			Arguments += String.Format(" -Architecture={0}", Arch);
+			if (!String.IsNullOrEmpty(InAdditionalArgs))
+			{
+				Arguments += InAdditionalArgs;
+			}
+			CommandUtils.RunUBT(CmdEnv, UnrealBuildToolDll, HostProjectFile, TargetName, Platform, Configuration, Arguments);
+		}
+
+	}
+}
diff --git a/Engine/Source/Programs/AutomationTool/AutomationUtils/Platform.cs b/Engine/Source/Programs/AutomationTool/AutomationUtils/Platform.cs
--- a/Engine/Source/Programs/AutomationTool/AutomationUtils/Platform.cs
+++ b/Engine/Source/Programs/AutomationTool/AutomationUtils/Platform.cs
@@ -688,6 +688,14 @@
 		}
 
 		/// <summary>
+		/// Returns the ICU data version we use for this platform
+		/// </summary>
+		public virtual string ICUDataVersion
+		{
+			get { return "icudt64l"; }
+		}
+
+		/// <summary>
 		/// Returns true if the platform wants patches to generate a small .pak file containing the difference
 		/// of current data against a shipped pak file.
 		/// </summary>
@@ -857,7 +865,7 @@
 				return PlatformExeExtension;
 			}
 
-			if (Target == UnrealTargetPlatform.Win64 || Target == UnrealTargetPlatform.HoloLens)
+			if (Target == UnrealTargetPlatform.Win64)
 			{
 				return ".exe";
 			}
diff --git a/Engine/Source/Programs/AutomationTool/Scripts/BuildCMakeLib.Automation.cs b/Engine/Source/Programs/AutomationTool/Scripts/BuildCMakeLib.Automation.cs
--- a/Engine/Source/Programs/AutomationTool/Scripts/BuildCMakeLib.Automation.cs
+++ b/Engine/Source/Programs/AutomationTool/Scripts/BuildCMakeLib.Automation.cs
@@ -768,29 +768,6 @@
 	public override bool IsPlatformExtension => false;
 }
 
-class VS2019TargetPlatform_HoloLens : BuildCMakeLib.VS2019TargetPlatform
-{
-	public override string PlatformOrGroupName => nameof(UnrealTargetPlatform.HoloLens);
-	public override string DebugDatabaseExtension => "pdb";
-	public override string DynamicLibraryExtension => "dll";
-	public override string StaticLibraryExtension => "lib";
-	public override string VariantDirectory => Architecture.ToLower();
-	public override bool IsPlatformExtension => false;
-
-	private readonly string Architecture;
-
-	public VS2019TargetPlatform_HoloLens(string Architecture = "Win64")
-	{
-		this.Architecture = Architecture;
-	}
-
-	public override string GetCMakeSetupArguments(BuildCMakeLib.TargetLib TargetLib, string TargetConfiguration)
-	{
-		return base.GetCMakeSetupArguments(TargetLib, TargetConfiguration)
-			+ string.Format(" -A {0}", Architecture);
-	}
-}
-
 class NMakeTargetPlatform_Unix : BuildCMakeLib.NMakeTargetPlatform
 {
 	public override string PlatformOrGroupName => nameof(UnrealPlatformGroup.Unix);
diff --git a/Engine/Source/Programs/AutomationTool/Scripts/BuildPluginCommand.Automation.cs b/Engine/Source/Programs/AutomationTool/Scripts/BuildPluginCommand.Automation.cs
--- a/Engine/Source/Programs/AutomationTool/Scripts/BuildPluginCommand.Automation.cs
+++ b/Engine/Source/Programs/AutomationTool/Scripts/BuildPluginCommand.Automation.cs
@@ -23,7 +23,6 @@
 public sealed class BuildPlugin : BuildCommand
 {
 	const string AndroidArchitectures = "armv7+arm64";
-	const string HoloLensArchitecture = "arm64+x64";
 	string UnrealBuildToolDllRelativePath = @"Engine/Binaries/DotNET/UnrealBuildTool/UnrealBuildTool.dll";
 	FileReference UnrealBuildToolDll;
 
@@ -143,6 +142,54 @@
 		return FileReference.Combine(HostProjectPluginDir, PluginFile.GetFileName());
 	}
 
+	public abstract class TargetPlatform : CommandUtils
+	{
+		[Obsolete("Deprecated in UE5.1; function signature changed")]
+		public abstract void CompilePluginWithUBT(string UBTExe, FileReference HostProjectFile, FileReference HostProjectPluginFile, PluginDescriptor Plugin, string TargetName, TargetType TargetType, UnrealTargetPlatform Platform, UnrealTargetConfiguration Configuration, List<FileReference> ManifestFileNames, string InAdditionalArgs);
+
+		public abstract void CompilePluginWithUBT(FileReference UnrealBuildToolDll, FileReference HostProjectFile, FileReference HostProjectPluginFile, PluginDescriptor Plugin, string TargetName, TargetType TargetType, UnrealTargetPlatform Platform, UnrealTargetConfiguration Configuration, List<FileReference> ManifestFileNames, string InAdditionalArgs);
+
+	};
+
+	private static TargetPlatform GetTargetPlatform( UnrealTargetPlatform Platform )
+	{
+		// Grab all the non-abstract subclasses of TargetPlatform from the executing assembly.
+		var AvailablePlatformTypes = from Assembly in ScriptManager.AllScriptAssemblies
+									 from Type in Assembly.GetTypes()
+									 where !Type.IsAbstract && Type.IsAssignableTo(typeof(TargetPlatform))
+									 select Type;
+
+		var PlatformTypeMap = new Dictionary<string, Type>();
+
+		foreach (var Type in AvailablePlatformTypes)
+		{
+			int Index = Type.Name.IndexOf('_');
+			if (Index == -1)
+			{
+				throw new BuildException("Invalid BuildPluginCommand target platform type found: {0}", Type);
+			}
+
+			PlatformTypeMap.Add(Type.Name, Type);
+		}
+
+		var SelectedPlatform = $"BuildPlugin_{Platform.ToString}";
+		if (!PlatformTypeMap.ContainsKey(SelectedPlatform))
+		{
+			return null;
+		}
+
+		var SelectedType = PlatformTypeMap[SelectedPlatform];
+		TargetPlatform TargetPlatform = (TargetPlatform)Activator.CreateInstance(SelectedType);
+		if (TargetPlatform == null)
+		{
+			throw new BuildException("The target platform \"{0}\" could not be constructed.", SelectedPlatform);
+		}
+
+		return TargetPlatform;
+	}
+
+
+
 	[Obsolete("Deprecated in UE5.1; function signature changed")]
 	public static FileReference[] CompilePlugin(string UBTExe, FileReference HostProjectFile, FileReference HostProjectPluginFile, PluginDescriptor Plugin, List<UnrealTargetPlatform> HostPlatforms, List<UnrealTargetPlatform> TargetPlatforms, string AdditionalArgs = "")
 	{
@@ -248,23 +295,10 @@
 		// Add these modules to the build agenda
 		if(bCompilePlatform)
 		{
-			if (Platform == UnrealTargetPlatform.HoloLens)
+			TargetPlatform TargetPlatform = GetTargetPlatform(Platform);
+			if (TargetPlatform != null)
 			{
-				// Make sure to save the manifests for each architecture with unique names so they don't get overwritten.
-				// This fixes packaging issues when building from binary engine releases, where the build produces a manifest for the plugin for ARM64, which
-				// then gets overwritten by the manifest for x64. Then during packaging, the plugin is referencing a manifest for the wrong architecture.
-				foreach (string Arch in HoloLensArchitecture.Split('+'))
-				{
-					FileReference ManifestFileName = FileReference.Combine(HostProjectFile.Directory, "Saved", String.Format("Manifest-{0}-{1}-{2}-{3}.xml", TargetName, Platform, Configuration, Arch));
-					ManifestFileNames.Add(ManifestFileName);
-					string Arguments = String.Format("-plugin={0} -iwyu -noubtmakefiles -manifest={1} -nohotreload", CommandUtils.MakePathSafeToUseWithCommandLine(HostProjectPluginFile.FullName), CommandUtils.MakePathSafeToUseWithCommandLine(ManifestFileName.FullName));
-					Arguments += String.Format(" -Architecture={0}", Arch);
-					if (!String.IsNullOrEmpty(InAdditionalArgs))
-					{
-						Arguments += InAdditionalArgs;
-					}
-					CommandUtils.RunUBT(CmdEnv, UBTExe, HostProjectFile, TargetName, Platform, Configuration, Arguments);
-				}
+				TargetPlatform.CompilePluginWithUBT(UBTExe, HostProjectFile, HostProjectPluginFile, Plugin, TargetName, TargetType, Platform, Configuration, ManifestFileNames, InAdditionalArgs );
 			}
 			else
 			{
@@ -308,23 +342,11 @@
 		// Add these modules to the build agenda
 		if(bCompilePlatform)
 		{
-			if (Platform == UnrealTargetPlatform.HoloLens)
+			TargetPlatform TargetPlatform = GetTargetPlatform(Platform);
+
+			if (TargetPlatform != null)
 			{
-				// Make sure to save the manifests for each architecture with unique names so they don't get overwritten.
-				// This fixes packaging issues when building from binary engine releases, where the build produces a manifest for the plugin for ARM64, which
-				// then gets overwritten by the manifest for x64. Then during packaging, the plugin is referencing a manifest for the wrong architecture.
-				foreach (string Arch in HoloLensArchitecture.Split('+'))
-				{
-					FileReference ManifestFileName = FileReference.Combine(HostProjectFile.Directory, "Saved", String.Format("Manifest-{0}-{1}-{2}-{3}.xml", TargetName, Platform, Configuration, Arch));
-					ManifestFileNames.Add(ManifestFileName);
-					string Arguments = String.Format("-plugin={0} -iwyu -noubtmakefiles -manifest={1} -nohotreload", CommandUtils.MakePathSafeToUseWithCommandLine(HostProjectPluginFile.FullName), CommandUtils.MakePathSafeToUseWithCommandLine(ManifestFileName.FullName));
-					Arguments += String.Format(" -Architecture={0}", Arch);
-					if (!String.IsNullOrEmpty(InAdditionalArgs))
-					{
-						Arguments += InAdditionalArgs;
-					}
-					CommandUtils.RunUBT(CmdEnv, UnrealBuildToolDll, HostProjectFile, TargetName, Platform, Configuration, Arguments);
-				}
+				TargetPlatform.CompilePluginWithUBT(UnrealBuildToolDll, HostProjectFile, HostProjectPluginFile, Plugin, TargetName, TargetType, Platform, Configuration, ManifestFileNames, InAdditionalArgs);
 			}
 			else
 			{
diff --git a/Engine/Source/Programs/AutomationTool/Scripts/CopyBuildToStagingDirectory.Automation.cs b/Engine/Source/Programs/AutomationTool/Scripts/CopyBuildToStagingDirectory.Automation.cs
--- a/Engine/Source/Programs/AutomationTool/Scripts/CopyBuildToStagingDirectory.Automation.cs
+++ b/Engine/Source/Programs/AutomationTool/Scripts/CopyBuildToStagingDirectory.Automation.cs
@@ -889,13 +889,6 @@
 					SC.PlatformCookDir = CookOutputDir;
 					SC.MetadataDir = DirectoryReference.Combine(CookOutputDir, SC.ShortProjectName, "Metadata");
 
-					// Work out which ICU data version we use for this platform
-					var ICUDataVersion = "icudt64l";
-					if (SC.StageTargetPlatform.PlatformType == UnrealTargetPlatform.HoloLens)
-					{
-						ICUDataVersion = "icudt53l";
-					}
-
 					// Initialize internationalization preset.
 					string InternationalizationPreset = GetInternationalizationPreset(Params, PlatformGameConfig);
 
@@ -916,6 +909,7 @@
 					}
 
 					// Stage ICU internationalization data from Engine.
+					var ICUDataVersion = SC.StageTargetPlatform.ICUDataVersion;
 					SC.StageFiles(StagedFileType.UFS, DirectoryReference.Combine(SC.LocalRoot, "Engine", "Content", "Internationalization", InternationalizationPreset, ICUDataVersion), StageFilesSearch.AllDirectories, new StagedDirectoryReference(String.Format("Engine/Content/Internationalization/{0}", ICUDataVersion)));
 
 					// Engine ufs (content)
