UAT: Fix .o files not being staged to the correct pakchunk

diff --git a/Engine/Source/Programs/AutomationTool/Scripts/CopyBuildToStagingDirectory.Automation.cs b/Engine/Source/Programs/AutomationTool/Scripts/CopyBuildToStagingDirectory.Automation.cs
--- a/Engine/Source/Programs/AutomationTool/Scripts/CopyBuildToStagingDirectory.Automation.cs
+++ b/Engine/Source/Programs/AutomationTool/Scripts/CopyBuildToStagingDirectory.Automation.cs
@@ -3677,7 +3677,8 @@ namespace AutomationScripts
 
 					string OriginalFilename = StagingFile.Key;
 					string NoExtension = CombinePaths(Path.GetDirectoryName(OriginalFilename), Path.GetFileNameWithoutExtension(OriginalFilename));
-					if(Path.GetExtension(NoExtension) == ".m")
+					string AdditionalExtension = Path.GetExtension(NoExtension);
+					if (AdditionalExtension == ".m" || AdditionalExtension == ".o")
 					{
 						// Hack around .m.ubulk files having a double extension
 						NoExtension = CombinePaths(Path.GetDirectoryName(OriginalFilename), Path.GetFileNameWithoutExtension(NoExtension));
