Improve errors reported to Automated Test Session
* Support multiple devices on a test (ie: network tests)
* Support log attached to device app instance
* Make Incomplete run flagged as failure when session is finishing
* Critical failure tag as crash in the test report

diff --git a/Engine/Source/Developer/AutomationController/Private/AutomationControllerManager.h b/Engine/Source/Developer/AutomationController/Private/AutomationControllerManager.h
--- a/Engine/Source/Developer/AutomationController/Private/AutomationControllerManager.h
+++ b/Engine/Source/Developer/AutomationController/Private/AutomationControllerManager.h
@@ -37,7 +37,7 @@ public:
 	EAutomationState State;
 
 	UPROPERTY()
-	FString DeviceInstance;
+	TArray<FString> DeviceInstance;
 
 	UPROPERTY()
 	float Duration;
diff --git a/Engine/Source/Developer/AutomationController/Private/AutomationControllerManger.cpp b/Engine/Source/Developer/AutomationController/Private/AutomationControllerManger.cpp
--- a/Engine/Source/Developer/AutomationController/Private/AutomationControllerManger.cpp
+++ b/Engine/Source/Developer/AutomationController/Private/AutomationControllerManger.cpp
@@ -534,7 +534,6 @@ void FAutomationControllerManager::CollectTestResults(TSharedPtr<IAutomationRepo
 		FAutomatedTestResult& TestResult = JsonTestPassResults.GetTestResult(Report);
 		TestResult.SetEvents(Results.GetEntries(), Results.GetWarningTotal(), Results.GetErrorTotal());
 		TestResult.SetArtifacts(Results.Artifacts);
-		TestResult.DeviceInstance = Results.GameInstance;
 		TestResult.Duration = Results.Duration;
 
 		JsonTestPassResults.TotalDuration += Results.Duration;
@@ -724,25 +723,16 @@ void FAutomationControllerManager::ExecuteNextTask( int32 ClusterIndex, OUT bool
 					TArray<FMessageAddress> DeviceAddresses = DeviceClusterManager.GetDevicesReservedForTest(ClusterIndex, NextTest);
 					if (DeviceAddresses.Num() == NextTest->GetNumParticipantsRequired())
 					{
+						TArray<FString> GameInstances;
+
 						// Send it to each device
 						for (int32 AddressIndex = 0; AddressIndex < DeviceAddresses.Num(); ++AddressIndex)
 						{
 							FAutomationTestResults TestResults;
-							TSharedPtr< IAutomationReport > Test = TestsRunThisPass[AddressIndex];
-
-							UE_LOG(LogAutomationController, Display, AutomationTestStarting, *Test->GetDisplayName(), *Test->GetFullTestPath());
-							TestResults.State = EAutomationState::InProcess;
-
-							if (JsonTestPassResults.IsRequired && bResumeRunTest)
-							{
-								JsonTestPassResults.UpdateTestResultStatus(NextTest, TestResults.State);
-								// Save the whole pass report so that if the next test triggers a critical failure we are not left with no pass report and we can resume.
-								GenerateJsonTestPassSummary(JsonTestPassResults);
-							}
-
 							TestResults.GameInstance = DeviceClusterManager.GetClusterDeviceName(ClusterIndex, DeviceIndex);
+							TestResults.State = EAutomationState::InProcess;
+							GameInstances.Add(TestResults.GameInstance);
 							NextTest->SetResults(ClusterIndex, CurrentTestPass, TestResults);
-							NextTest->ResetNetworkCommandResponses();
 
 							// Mark the device as busy
 							FMessageAddress DeviceAddress = DeviceAddresses[AddressIndex];
@@ -755,6 +745,19 @@ void FAutomationControllerManager::ExecuteNextTask( int32 ClusterIndex, OUT bool
 							// Add a test so we can check later if the device is still active
 							TestRunningArray.Add(FTestRunningInfo(DeviceAddress));
 						}
+
+						UE_LOG(LogAutomationController, Display, AutomationTestStarting, *NextTest->GetDisplayName(), *NextTest->GetFullTestPath());
+
+						if (JsonTestPassResults.IsRequired && bResumeRunTest)
+						{
+							JsonTestPassResults.UpdateTestResultStatus(NextTest, EAutomationState::InProcess);
+							FAutomatedTestResult& TestResult = JsonTestPassResults.GetTestResult(NextTest);
+							TestResult.DeviceInstance = GameInstances;
+							// Save the whole pass report so that if the next test triggers a critical failure we are not left with no pass report and we can resume.
+							GenerateJsonTestPassSummary(JsonTestPassResults);
+						}
+
+						NextTest->ResetNetworkCommandResponses();
 					}
 				}
 			}
@@ -777,7 +780,6 @@ void FAutomationControllerManager::ExecuteNextTask( int32 ClusterIndex, OUT bool
 			{
 				FAutomationTestResults TestResults;
 				TestResults.State = EAutomationState::Skipped;
-				TestResults.GameInstance = DeviceClusterManager.GetClusterDeviceName(ClusterIndex, 0);
 				TestResults.AddEvent(FAutomationEvent(EAutomationEventType::Warning, FString::Printf(TEXT("Skipped because the test needed %d devices to participate, Only had %d available."), CurrentTest->GetNumParticipantsRequired(), DeviceClusterManager.GetNumDevicesInCluster(ClusterIndex))));
 
 				CurrentTest->SetResults(ClusterIndex, CurrentTestPass, TestResults);
diff --git a/Engine/Source/Developer/AutomationController/Private/AutomationDeviceClusterManager.h b/Engine/Source/Developer/AutomationController/Private/AutomationDeviceClusterManager.h
--- a/Engine/Source/Developer/AutomationController/Private/AutomationDeviceClusterManager.h
+++ b/Engine/Source/Developer/AutomationController/Private/AutomationDeviceClusterManager.h
@@ -56,6 +56,10 @@ public:
 	/** The name of the current RHI */
 	UPROPERTY()
 	FString RHI;
+
+	/** The path to the application instance log */
+	UPROPERTY()
+	FString AppInstanceLog;
 };
 
 /**
diff --git a/Engine/Source/Programs/AutomationTool/Gauntlet/Framework/Utils/Gauntlet.HordeReport.cs b/Engine/Source/Programs/AutomationTool/Gauntlet/Framework/Utils/Gauntlet.HordeReport.cs
--- a/Engine/Source/Programs/AutomationTool/Gauntlet/Framework/Utils/Gauntlet.HordeReport.cs
+++ b/Engine/Source/Programs/AutomationTool/Gauntlet/Framework/Utils/Gauntlet.HordeReport.cs
@@ -473,7 +473,7 @@ namespace Gauntlet
 						ConvertedTestResult.TestDisplayName = InTestResult.TestDisplayName;
 						ConvertedTestResult.FullTestPath = InTestResult.FullTestPath;
 						ConvertedTestResult.State = InTestResult.State;
-						ConvertedTestResult.DeviceInstance = InTestResult.DeviceInstance;
+						ConvertedTestResult.DeviceInstance = InTestResult.DeviceInstance.FirstOrDefault();
 						Guid TestGuid = Guid.NewGuid();
 						ConvertedTestResult.ArtifactName = TestGuid + ".json";
 						InTestResult.ArtifactName = ConvertedTestResult.ArtifactName;
@@ -563,7 +563,7 @@ namespace Gauntlet
 				public string TestUID { get; set; }
 				public string Suite { get; set; }
 				public TestStateType State { get; set; }
-				public string DeviceAppInstanceName { get; set; }
+				public List<string> DeviceAppInstanceName { get; set; }
 				public uint ErrorCount { get; set; }
 				public uint WarningCount { get; set; }
 				public string ErrorHashAggregate { get; set; }
@@ -572,13 +572,13 @@ namespace Gauntlet
 				// not part of json output
 				public List<TestEvent> Events;
 				public SortedSet<string> ErrorHashes;
-				public TestResult(string InName, string InSuite, TestStateType InState, string InDevice, string InDateTime)
+				public TestResult(string InName, string InSuite, TestStateType InState, List<string> InDevices, string InDateTime)
 				{
 					Name = InName;
 					TestUID = GenerateHash(Name);
 					Suite = InSuite;
 					State = InState;
-					DeviceAppInstanceName = InDevice;
+					DeviceAppInstanceName = InDevices;
 					ErrorHashes = new SortedSet<string>();
 					ErrorHashAggregate = "";
 					DateTime = InDateTime;
@@ -640,9 +640,9 @@ namespace Gauntlet
 					Tests = new Dictionary<string, TestResult>();
 					TestResultsTestDataUID = Guid.NewGuid().ToString();
 				}
-				public TestResult NewTestResult(string InName, string InSuite, TestStateType InState, string InDevice, string InDateTime)
+				public TestResult NewTestResult(string InName, string InSuite, TestStateType InState, List<string> InDevices, string InDateTime)
 				{
-					TestResult NewItem = new TestResult(InName, InSuite, InState, InDevice, InDateTime);
+					TestResult NewItem = new TestResult(InName, InSuite, InState, InDevices, InDateTime);
 					Tests[NewItem.TestUID] = NewItem;
 					return NewItem;
 				}
@@ -651,11 +651,13 @@ namespace Gauntlet
 			{
 				public string Name { get; set; }
 				public string AppInstanceName { get; set; }
+				public string AppInstanceLog { get; set; }
 				public Dictionary<string, string> Metadata { get; set; }
-				public Device(string InName, string InInstance)
+				public Device(string InName, string InInstance, string InInstanceLog)
 				{
 					Name = InName;
 					AppInstanceName = InInstance;
+					AppInstanceLog = InInstanceLog;
 					Metadata = new Dictionary<string, string>();
 				}
 				public void SetMetaData(string Key, string Value)
@@ -783,9 +785,9 @@ namespace Gauntlet
 
 			private string CurrentTestUID;
 
-			private Device NewDevice(string InName, string InInstance)
+			private Device NewDevice(string InName, string InInstance, string InInstanceLog = null)
 			{
-				Device NewItem = new Device(InName, InInstance);
+				Device NewItem = new Device(InName, InInstance, InInstanceLog);
 				Devices.Add(NewItem);
 				return NewItem;
 			}
@@ -813,27 +815,56 @@ namespace Gauntlet
 				return TestSessionInfo.Tests[CurrentTestUID];
 			}
 
-			public void SetTest(string InName, string InSuite, string InDevice, string InDateTime = null)
+			/// <summary>
+			/// Add a Test entry and set it as current test
+			/// </summary>
+			/// <param name="InName"></param>
+			/// <param name="InSuite"></param>
+			/// <param name="InDevice"></param>
+			/// <param name="InDateTime"></param>
+			public void SetTest(string InName, string InSuite, List<string> InDevices, string InDateTime = null)
 			{
-				TestResult Test = TestSessionInfo.NewTestResult(InName, InSuite, TestStateType.Unknown, InDevice, InDateTime ?? DateTime.Now.ToString("yyyy.MM.dd-HH.mm.ss"));
+				TestResult Test = TestSessionInfo.NewTestResult(InName, InSuite, TestStateType.Unknown, InDevices, InDateTime);
 				CurrentTestUID = Test.TestUID;
 				// Index the test result data.
 				TestResults[Test.TestUID] = new TestResultData(Test);
 			}
+			/// <summary>
+			/// Set current test state
+			/// </summary>
+			/// <param name="InState"></param>
 			public void SetTestState(TestStateType InState)
 			{
 				TestResult CurrentTest = GetCurrentTest();
 				CurrentTest.State = InState;
 			}
+			/// <summary>
+			/// Set current test time elapse in second
+			/// </summary>
+			/// <param name="InTimeElapseSec"></param>
 			public void SetTestTimeElapseSec(float InTimeElapseSec)
 			{
 				TestResult CurrentTest = GetCurrentTest();
 				CurrentTest.TimeElapseSec = InTimeElapseSec;
 			}
+			/// <summary>
+			/// Add event to current test
+			/// </summary>
+			/// <param name="InType"></param>
+			/// <param name="InMessage"></param>
+			/// <param name="InContext"></param>
 			public override void AddEvent(EventType InType, string InMessage, object InContext = null)
 			{
 				AddEvent(InType, InMessage, "gauntlet", InContext == null ? null : InContext.ToString());
 			}
+			/// <summary>
+			/// Overload of AddEvent, add even to current test with date time
+			/// </summary>
+			/// <param name="InType"></param>
+			/// <param name="InMessage"></param>
+			/// <param name="InTag"></param>
+			/// <param name="InContext"></param>
+			/// <param name="InDateTime"></param>
 			public void AddEvent(EventType InType, string InMessage, string InTag, string InContext = null, string InDateTime = null)
 			{
 				TestResult CurrentTest = GetCurrentTest();
@@ -843,6 +874,13 @@ namespace Gauntlet
 					IndexTestError(Event.Hash, CurrentTest.TestUID, InMessage, InTag);
 				}
 			}
+			/// <summary>
+			/// Add artifact to last test event
+			/// </summary>
+			/// <param name="InTag"></param>
+			/// <param name="InFilePath"></param>
+			/// <param name="InReferencePath"></param>
+			/// <returns></returns>
 			public bool AddArtifactToLastEvent(string InTag, string InFilePath, string InReferencePath = null)
 			{
 				TestEvent LastEvent = GetCurrentTest().GetLastEvent();
@@ -891,7 +929,7 @@ namespace Gauntlet
 				{
 					foreach (UnrealAutomationDevice InDevice in InTestPassResults.Devices)
 					{
-						Device ConvertedDevice = OutTestPassResults.NewDevice(InDevice.DeviceName, InDevice.Instance);
+						Device ConvertedDevice = OutTestPassResults.NewDevice(InDevice.DeviceName, InDevice.Instance, InDevice.AppInstanceLog);
 						ConvertedDevice.SetMetaData("platform", InDevice.Platform);
 						ConvertedDevice.SetMetaData("os_version", InDevice.OSVersion);
 						ConvertedDevice.SetMetaData("model", InDevice.Model);
@@ -908,6 +946,12 @@ namespace Gauntlet
 				{
 					foreach (UnrealAutomatedTestResult InTestResult in InTestPassResults.Tests)
 					{
+						string TestDateTime = InTestResult.DateTime;
+						if (TestDateTime == "0001.01.01-00.00.00")
+						{
+							// Special case: when UE Test DateTime is set this way, it means the value was 0 or null before getting converted to json.
+							TestDateTime = null;
+						}
 						OutTestPassResults.SetTest(InTestResult.FullTestPath, InSuite, InTestResult.DeviceInstance, InTestResult.DateTime);
 						foreach (UnrealAutomationEntry InTestEntry in InTestResult.Entries)
 						{
@@ -930,6 +974,12 @@ namespace Gauntlet
 							{
 								Context = string.Format("{0} @line:{1}", InTestEntry.Filename, InTestEntry.LineNumber.ToString());
 							}
+							// Tag critical failure as crash
+							if (InTestEntry.Event.Type == EventType.Error && Tag == "entry" && InTestEntry.Event.Message.Contains("critical failure"))
+							{
+								Tag = "crash";
+							}
+
 							OutTestPassResults.AddEvent(InTestEntry.Event.Type, InTestEntry.Event.Message, Tag, Context, InTestEntry.Timestamp);
 							bool IsInfo = InTestEntry.Event.Type == EventType.Info;
 							// Add Artifacts
@@ -999,7 +1049,7 @@ namespace Gauntlet
 								{
 									foreach(var Item in FailedToAttached)
 									{
-										OutTestPassResults.AddEvent(EventType.Warning, string.Format("Failed to attached Artifact {0}.", FailedToAttached));
+										OutTestPassResults.AddEvent(EventType.Warning, string.Format("Failed to attached Artifact {0}.", Item));
 									}
 								}
 							}
diff --git a/Engine/Source/Programs/AutomationTool/Gauntlet/Unreal/Base/Gauntlet.UnrealTestNode.cs b/Engine/Source/Programs/AutomationTool/Gauntlet/Unreal/Base/Gauntlet.UnrealTestNode.cs
--- a/Engine/Source/Programs/AutomationTool/Gauntlet/Unreal/Base/Gauntlet.UnrealTestNode.cs
+++ b/Engine/Source/Programs/AutomationTool/Gauntlet/Unreal/Base/Gauntlet.UnrealTestNode.cs
@@ -1264,6 +1264,8 @@ namespace Gauntlet
 				string HordeArtifactPath = GetCachedConfiguration().HordeArtifactPath;
 				Log.Verbose("Reading json Unreal Automated test report from {0}", JsonReportPath);
 				UnrealAutomatedTestPassResults JsonTestPassResults = UnrealAutomatedTestPassResults.LoadFromJson(JsonReportPath);
+				var MainRole = GetCachedConfiguration().GetMainRequiredRole();
+				string LastTestWithCriticalFailure = string.Empty;
 				if (JsonTestPassResults.InProcess > 0)
 				{
 					// The test pass did not run completely
@@ -1280,14 +1282,31 @@ namespace Gauntlet
 								break;
 							}
 						}
-						if (FatalError != null)
+						var Test = JsonTestPassResults.Tests.FirstOrDefault((T => T.State == TestStateType.InProcess));
+						if (!String.IsNullOrEmpty(Test.TestDisplayName))
 						{
-							var Test = JsonTestPassResults.Tests.FirstOrDefault((T => T.State == TestStateType.InProcess));
-							if (!String.IsNullOrEmpty(Test.TestDisplayName))
+							LastTestWithCriticalFailure = Test.FullTestPath;
+							string ErrorMessage = "Engine encountered a critical failure. \n";
+							if (FatalError != null)
 							{
-								Test.AddError("Engine encountered a critical failure. \n" + FatalError.FormatForLog());
-								JsonTestPassResults.WriteToJson(JsonReportPath);
+								ErrorMessage += FatalError.FormatForLog();
+							}
+							else
+							{
+								ErrorMessage += "No callstack found in the log.";
+							}
+							Test.AddError(ErrorMessage);
+							if (Retries >= MaxRetries || JsonTestPassResults.NotRun == 0)
+							{
+								// Setting the test as fail because no retry will be done anymore.
+								// The InProcess state won't be used for pass resume
+								Test.State = TestStateType.Fail;
+								if (Retries >= MaxRetries)
+								{
+									Test.AddWarning(string.Format("Session reached maximum of retries({0}) to resume on critical failure!", Retries));
+								}
 							}
+							JsonTestPassResults.WriteToJson(JsonReportPath);
 						}
 					}
 				}
@@ -1312,7 +1331,13 @@ namespace Gauntlet
 								{
 									string LogName = Path.GetFullPath(Artifact.LogPath).Replace(Path.GetFullPath(Context.Options.LogDir), "").TrimStart(Path.DirectorySeparatorChar);
 									TempReport.AttachArtifact(Artifact.LogPath, LogName);
+									// Reference last run instance log
+									if(Artifact.SessionRole.RoleType == MainRole.Type)
+									{
+										JsonTestPassResults.Devices.Last().AppInstanceLog = LogName.Replace("\\", "/");
+									}
 								}
+								JsonTestPassResults.WriteToJson(JsonReportPath);
 							}
 							// Discard the report as we are going to do another pass.
 							return null;
@@ -1320,6 +1345,18 @@ namespace Gauntlet
 						else
 						{
 							Log.Error("Reach maximum of retries({0}) to resume on critical failure!", Retries);
+							// Adding a note to the report about why the not-run are not going to be run
+							string Message = string.Format("Session reached maximum of retries({0}) to resume on critical failure!", Retries);
+							if (!string.IsNullOrEmpty(LastTestWithCriticalFailure))
+							{
+								Message += string.Format(" \nLast critical failure was caught on {0}", LastTestWithCriticalFailure);
+							}
+							var NotRunTests = JsonTestPassResults.Tests.Where((T => T.State == TestStateType.NotRun));
+							foreach(var Test in NotRunTests)
+							{
+								Test.AddInfo(Message);
+							}
+							JsonTestPassResults.WriteToJson(JsonReportPath);
 						}
 					}
 				}
@@ -1337,9 +1374,9 @@ namespace Gauntlet
 				HordeTestPassResults.PreFlightChange = GetCachedConfiguration().PreFlightChange;
 				// Metadata
 				// With UE Test Automation, we care only for one role.
-				var MainRole = new List<UnrealTestRole>() { GetCachedConfiguration().GetMainRequiredRole() };
-				SetReportMetadata(HordeTestPassResults, MainRole);
-				SetReportMetadata(CopyTestPassResults, MainRole); // Set the metadata to the old report too
+				var RoleList = new List<UnrealTestRole>() { MainRole };
+				SetReportMetadata(HordeTestPassResults, RoleList);
+				SetReportMetadata(CopyTestPassResults, RoleList); // Set the metadata to the old report too
 				// Attached test Artifacts
 				if (SessionArtifacts != null)
 				{
@@ -1347,6 +1384,11 @@ namespace Gauntlet
 					{
 						string LogName = Path.GetFullPath(Artifact.LogPath).Replace(Path.GetFullPath(Context.Options.LogDir), "").TrimStart(Path.DirectorySeparatorChar);
 						HordeTestPassResults.AttachArtifact(Artifact.LogPath, LogName);
+						// Reference last run instance log
+						if (Artifact.SessionRole.RoleType == MainRole.Type)
+						{
+							HordeTestPassResults.Devices.Last().AppInstanceLog = LogName.Replace("\\", "/");
+						}
 					}
 				}
 				return HordeTestPassResults;
diff --git a/Engine/Source/Programs/AutomationTool/Gauntlet/Unreal/Utils/Gauntlet.UnrealTestResults.cs b/Engine/Source/Programs/AutomationTool/Gauntlet/Unreal/Utils/Gauntlet.UnrealTestResults.cs
--- a/Engine/Source/Programs/AutomationTool/Gauntlet/Unreal/Utils/Gauntlet.UnrealTestResults.cs
+++ b/Engine/Source/Programs/AutomationTool/Gauntlet/Unreal/Utils/Gauntlet.UnrealTestResults.cs
@@ -21,6 +21,7 @@ namespace Gauntlet
 		public int RAMInGB { get; set; }
 		public string RenderMode { get; set; }
 		public string RHI { get; set; }
+		public string AppInstanceLog { get; set; }
 	}
 	public class UnrealAutomationComparisonFiles
 	{
@@ -83,7 +84,7 @@ namespace Gauntlet
 		public string FullTestPath { get; set; }
 		public string ArtifactName { get; set; }
 		public TestStateType State { get; set; }
-		public string DeviceInstance { get; set; }
+		public List<string> DeviceInstance { get; set; }
 		public float Duration { get; set; }
 		public string DateTime { get; set; }
 		public int Warnings { get; set; }
