Add failure to check in a RoboMerge gate file due to gate being locked to systemic errors.

#rnx

diff --git a/Engine/Source/Programs/AutomationTool/AutomationUtils/Matchers/SystemicEventMatcher.cs b/Engine/Source/Programs/AutomationTool/AutomationUtils/Matchers/SystemicEventMatcher.cs
--- a/Engine/Source/Programs/AutomationTool/AutomationUtils/Matchers/SystemicEventMatcher.cs
+++ b/Engine/Source/Programs/AutomationTool/AutomationUtils/Matchers/SystemicEventMatcher.cs
@@ -24,6 +24,10 @@ namespace AutomationUtils.Matchers
 				}
 				return new LogEventBuilder(cursor).ToMatch(LogEventPriority.High, LogLevel.Information, KnownLogEvents.Systemic_PdbUtil);
 			}
+			if (cursor.IsMatch(@"RoboMerge\/gates.*already locked on Commit Server by buildmachine"))
+			{
+				return new LogEventBuilder(cursor).ToMatch(LogEventPriority.Low, LogLevel.Information, KnownLogEvents.Systemic_RoboMergeGateLocked);
+			}
 			return null;
 		}
 	}
diff --git a/Engine/Source/Programs/Shared/EpicGames.Core/KnownLogEvents.cs b/Engine/Source/Programs/Shared/EpicGames.Core/KnownLogEvents.cs
--- a/Engine/Source/Programs/Shared/EpicGames.Core/KnownLogEvents.cs
+++ b/Engine/Source/Programs/Shared/EpicGames.Core/KnownLogEvents.cs
@@ -170,6 +170,8 @@ namespace EpicGames.Core
 		/// </summary>
 		public static EventId Systemic_MSBuild { get; } = new EventId(750);
 
+		public static EventId Systemic_RoboMergeGateLocked { get; } = new EventId(760);
+
 		/// <summary>
 		/// Horde error codes
 		/// </summary>
