fix platform name lookup.

#rnx

diff --git a/Engine/Source/Programs/AutomationTool/Scripts/BuildPluginCommand.Automation.cs b/Engine/Source/Programs/AutomationTool/Scripts/BuildPluginCommand.Automation.cs
--- a/Engine/Source/Programs/AutomationTool/Scripts/BuildPluginCommand.Automation.cs
+++ b/Engine/Source/Programs/AutomationTool/Scripts/BuildPluginCommand.Automation.cs
@@ -172,7 +172,7 @@
 			PlatformTypeMap.Add(Type.Name, Type);
 		}
 
-		var SelectedPlatform = $"BuildPlugin_{Platform.ToString}";
+		var SelectedPlatform = $"BuildPlugin_{Platform.ToString()}";
 		if (!PlatformTypeMap.ContainsKey(SelectedPlatform))
 		{
 			return null;
