Registers matchers from UBT into log event parser in UAT, for times where we build inside UAT.

diff --git a/Engine/Source/Programs/AutomationTool/AutomationUtils/Automation.cs b/Engine/Source/Programs/AutomationTool/AutomationUtils/Automation.cs
--- a/Engine/Source/Programs/AutomationTool/AutomationUtils/Automation.cs
+++ b/Engine/Source/Programs/AutomationTool/AutomationUtils/Automation.cs
@@ -76,6 +76,9 @@
 				Assembly AutomationUtilsAssembly = Assembly.GetExecutingAssembly();
 				Log.EventParser.AddMatchersFromAssembly(AutomationUtilsAssembly);
 
+				Assembly UnrealBuildToolAssembly = typeof(UnrealBuildTool.BuildVersion).Assembly;
+				Log.EventParser.AddMatchersFromAssembly(UnrealBuildToolAssembly);
+
 				// Get the path to the telemetry file, if present
 				string TelemetryFile = GlobalCommandLine.TelemetryPath;
 				JsonTracer Tracer = JsonTracer.TryRegisterAsGlobalTracer();
