- Added Mac device to Turnkey

diff --git a/Engine/Source/Programs/AutomationTool/Mac/MacPlatform.Automation.cs b/Engine/Source/Programs/AutomationTool/Mac/MacPlatform.Automation.cs
--- a/Engine/Source/Programs/AutomationTool/Mac/MacPlatform.Automation.cs
+++ b/Engine/Source/Programs/AutomationTool/Mac/MacPlatform.Automation.cs
@@ -1,4 +1,4 @@
-// Copyright Epic Games, Inc. All Rights Reserved.
+﻿// Copyright Epic Games, Inc. All Rights Reserved.
 using System;
 using System.Collections.Generic;
 using System.Linq;
@@ -47,7 +47,21 @@
 		}		
 	}
 
+	public override DeviceInfo[] GetDevices()
+	{
+		List<DeviceInfo> Devices = new List<DeviceInfo>();
 
+		if (HostPlatform.Current.HostEditorPlatform == TargetPlatformType)
+		{
+			DeviceInfo LocalMachine = new DeviceInfo(TargetPlatformType, Environment.MachineName, Environment.MachineName,
+				Environment.OSVersion.Version.ToString(), "Computer", true, true);
+
+			Devices.Add(LocalMachine);
+		}
+
+		return Devices.ToArray();
+	}
+
 	public override string GetCookPlatform(bool bDedicatedServer, bool bIsClientOnly)
 	{
 		const string NoEditorCookPlatform = "Mac";
