- Added a ctrl-c handler to AutomationTool, which helps keep zombie dotnet processes, at least on Mac, maybe Linux

diff --git a/Engine/Source/Programs/AutomationTool/Program.cs b/Engine/Source/Programs/AutomationTool/Program.cs
--- a/Engine/Source/Programs/AutomationTool/Program.cs
+++ b/Engine/Source/Programs/AutomationTool/Program.cs
@@ -367,6 +367,17 @@
 			Assembly UnrealBuildToolAssembly = typeof(UnrealBuildTool.BuildVersion).Assembly;
 			Log.EventParser.AddMatchersFromAssembly(UnrealBuildToolAssembly);
 
+			// when running frmo RunUAT.sh (Mac/Linux) we need to install a Ctrl-C handler, or hitting Ctrl-C from a terminal
+			// can leave dotnet process in a zombie state (some order of process destruction is failing)
+			// by putting this in, the Ctrl-C may not be handled immediately, but it shouldn't leave a zombie process
+			if (OperatingSystem.IsMacOS() || OperatingSystem.IsLinux())
+			{
+				Console.CancelKeyPress += delegate
+				{
+					Console.WriteLine("Dying....");
+				};
+			}
+
 			// Enter the main program section
 			ExitCode ReturnCode = ExitCode.Success;
 			try
