[Automation] Fix for correct dest path for ADB

#tests Ran ReplayRun preflight and the job no longer had issues with build copy

diff --git a/Engine/Source/Programs/AutomationTool/Gauntlet/Platform/Android/Gauntlet.TargetDeviceAndroid.cs b/Engine/Source/Programs/AutomationTool/Gauntlet/Platform/Android/Gauntlet.TargetDeviceAndroid.cs
--- a/Engine/Source/Programs/AutomationTool/Gauntlet/Platform/Android/Gauntlet.TargetDeviceAndroid.cs
+++ b/Engine/Source/Programs/AutomationTool/Gauntlet/Platform/Android/Gauntlet.TargetDeviceAndroid.cs
@@ -992,8 +992,7 @@ namespace Gauntlet
 					Match OBBMatch = Regex.Match(DestFile, @"\.(\d+)\.com.*\.obb");
 					if (OBBMatch.Success)
 					{
-						string NewFileName = DestFile.Replace(OBBMatch.Groups[1].ToString(), PackageVersion);
-						DestPath = DestPath.Replace(DestFile, NewFileName);
+						DestPath = StorageLocation + "/obb/" + Build.AndroidPackageName + "/" + Path.GetFileName(SrcPath).Replace(".Client.obb", ".obb").Replace(OBBMatch.Groups[1].ToString(), PackageVersion);
 					}
 
 					DestPath = Regex.Replace(DestPath, "%STORAGE%", StorageLocation, RegexOptions.IgnoreCase);
