- `WhitelistConfigFiles` is now enough to allow for staging individual configuration files from Restricted folders, without having to also specify also `WhitelistDirectories` for that folder as well.
- `WhitelistDirectories` now works properly for restricted directories in distributed builds.

diff --git a/Engine/Source/Programs/AutomationTool/CookedEditor/CookedEditor.Automation.cs b/Engine/Source/Programs/AutomationTool/CookedEditor/CookedEditor.Automation.cs
--- a/Engine/Source/Programs/AutomationTool/CookedEditor/CookedEditor.Automation.cs
+++ b/Engine/Source/Programs/AutomationTool/CookedEditor/CookedEditor.Automation.cs
@@ -16,7 +16,6 @@ public class ConfigHelper
 	private string SharedConfigSection;
 	private ConfigHierarchy GameConfig;
 
-
 	public ConfigHelper(UnrealTargetPlatform Platform, FileReference ProjectFile, bool bIsCookedCooker)
 	{
 		SharedConfigSection = "CookedEditorSettings";
@@ -302,10 +301,17 @@ public class ModifyStageContext
 
 	private void HandleRestrictedFiles(DeploymentContext SC, ref Dictionary<StagedFileReference, FileReference> Files)
 	{
+		// If a directory has been specifically whitelisted, do not omit the files wihin it from the staging process.
 		if (bIsForExternalDistribution)
 		{
+			Int32 OrigNumFiles = Files.Count();
 			// remove entries where any restricted folder names are in the name remapped path (if we remap from NFL to non-NFL, then we don't remove it)
-			Files = Files.Where(x => !SC.RestrictedFolderNames.Any(y => Project.ApplyDirectoryRemap(SC, x.Key).ContainsName(y))).ToDictionary(x => x.Key, x => x.Value);
+			// If a configuration file has been explicitly allowed, do not omit it either.
+			Files = Files.Where(x => SC.ConfigFilesAllowList.Contains(x.Key) || SC.DirectoriesAllowList.Contains(x.Key.Directory) || !SC.RestrictedFolderNames.Any(y => Project.ApplyDirectoryRemap(SC, x.Key).ContainsName(y))).ToDictionary(x => x.Key, x => x.Value);
+			if (OrigNumFiles != Files.Count())
+			{
+				Log.TraceInformationOnce("Some files were not staged since they have restricted folder names in the remapped path.");
+			}
 		}
 		else
 		{
@@ -1251,4 +1257,4 @@ public class MakeCookedEditor : BuildCommand
 		// convert to array type
 		return KeyValues;
 	}
-}
\ No newline at end of file
+}
diff --git a/Engine/Source/Programs/AutomationTool/Scripts/CopyBuildToStagingDirectory.Automation.cs b/Engine/Source/Programs/AutomationTool/Scripts/CopyBuildToStagingDirectory.Automation.cs
--- a/Engine/Source/Programs/AutomationTool/Scripts/CopyBuildToStagingDirectory.Automation.cs
+++ b/Engine/Source/Programs/AutomationTool/Scripts/CopyBuildToStagingDirectory.Automation.cs
@@ -1338,6 +1338,7 @@
 				RestrictedFiles.UnionWith(SC.CrashReporterUFSFiles.Keys.Where(x => x.ContainsName(RestrictedName)));
 			}
 			RestrictedFiles.RemoveWhere(RestrictedFile => SC.DirectoriesAllowList.Any(TestDirectory => RestrictedFile.Directory.IsUnderDirectory(TestDirectory)));
+			RestrictedFiles.RemoveWhere(RestrictedFile => SC.ConfigFilesAllowList.Contains(RestrictedFile));
 
 			if (RestrictedFiles.Count > 0)
 			{
