Fixed warning regarding usage of banned words.

diff --git a/Engine/Source/Programs/AutomationTool/CookedEditor/CookedEditor.Automation.cs b/Engine/Source/Programs/AutomationTool/CookedEditor/CookedEditor.Automation.cs
--- a/Engine/Source/Programs/AutomationTool/CookedEditor/CookedEditor.Automation.cs
+++ b/Engine/Source/Programs/AutomationTool/CookedEditor/CookedEditor.Automation.cs
@@ -301,7 +301,7 @@ public class ModifyStageContext
 
 	private void HandleRestrictedFiles(DeploymentContext SC, ref Dictionary<StagedFileReference, FileReference> Files)
 	{
-		// If a directory has been specifically whitelisted, do not omit the files wihin it from the staging process.
+		// If a directory has been specifically allowed, do not omit the files wihin it from the staging process.
 		if (bIsForExternalDistribution)
 		{
 			Int32 OrigNumFiles = Files.Count();
