BuildCMakeLib: Use Intermediate as temp build directory instead of Build, as it won't conflict with preexisting files or folders

#rnx
#rn none

diff --git a/Engine/Source/Programs/AutomationTool/Scripts/BuildCMakeLib.Automation.cs b/Engine/Source/Programs/AutomationTool/Scripts/BuildCMakeLib.Automation.cs
--- a/Engine/Source/Programs/AutomationTool/Scripts/BuildCMakeLib.Automation.cs
+++ b/Engine/Source/Programs/AutomationTool/Scripts/BuildCMakeLib.Automation.cs
@@ -152,7 +152,7 @@
 		}
 
 		protected DirectoryReference GetProjectsDirectory(TargetLib TargetLib, string TargetConfiguration) =>
-			DirectoryReference.Combine(GetTargetLibRootDirectory(TargetLib), "Build",
+			DirectoryReference.Combine(GetTargetLibRootDirectory(TargetLib), "Intermediate",
 				IsPlatformExtension ? "" : PlatformOrGroupName,
 				VariantDirectory ?? "",
 				SeparateProjectPerConfig ? TargetLib.BuildMap[TargetConfiguration] : "");
@@ -375,7 +375,7 @@
 		{
 			if (string.IsNullOrEmpty(TargetConfiguration))
 			{
-				InternalUtils.SafeDeleteDirectory(DirectoryReference.Combine(GetTargetLibRootDirectory(TargetLib), "Build").FullName);
+				InternalUtils.SafeDeleteDirectory(DirectoryReference.Combine(GetTargetLibRootDirectory(TargetLib), "Intermediate").FullName);
 			}
 			else
 			{
