Reformat shader compilation errors so that Horde can include them in issues.

diff --git a/Engine/Source/Programs/AutomationTool/AutomationUtils/Matchers/ContentEventMatcher.cs b/Engine/Source/Programs/AutomationTool/AutomationUtils/Matchers/ContentEventMatcher.cs
--- a/Engine/Source/Programs/AutomationTool/AutomationUtils/Matchers/ContentEventMatcher.cs
+++ b/Engine/Source/Programs/AutomationTool/AutomationUtils/Matchers/ContentEventMatcher.cs
@@ -35,13 +35,18 @@ namespace AutomationUtils.Matchers
 				Match? match = s_pattern.Match(input.CurrentLine!);
 				if (match.Success)
 				{
-					LogEventBuilder builder = new LogEventBuilder(input);
+					LogEventBuilder builder = new LogEventBuilder(input.Hanging());
 
 					builder.Annotate(match.Groups["channel"], LogEventMarkup.Channel);
 					builder.Annotate(match.Groups["severity"], LogEventMarkup.Severity);
 					builder.AnnotateAsset(match.Groups["asset"]);
 					builder.Annotate(match.Groups["message"], LogEventMarkup.Message);
 
+					while(builder.Current.TryGetLine(1, out _))
+					{
+						builder.MoveNext();
+					}
+
 					return builder.ToMatch(LogEventPriority.AboveNormal, GetLogLevelFromSeverity(match), KnownLogEvents.Engine_AssetLog);
 				}
 			}
