Remove direct access to log buffer lines; use standard methods to determine whether lines are valid or not.

diff --git a/Engine/Source/Programs/AutomationTool/AutomationUtils/Matchers/LocalizationEventMatcher.cs b/Engine/Source/Programs/AutomationTool/AutomationUtils/Matchers/LocalizationEventMatcher.cs
--- a/Engine/Source/Programs/AutomationTool/AutomationUtils/Matchers/LocalizationEventMatcher.cs
+++ b/Engine/Source/Programs/AutomationTool/AutomationUtils/Matchers/LocalizationEventMatcher.cs
@@ -27,8 +27,8 @@ namespace AutomationUtils.Matchers
 		/// <inheritdoc/>
 		public LogEventMatch? Match(ILogCursor input)
 		{
-			Match? match = s_pattern.Match(input.CurrentLine!);
-			if (match.Success)
+			Match? match;
+			if (input.TryMatch(s_pattern, out match))
 			{
 				LogEventBuilder builder = new LogEventBuilder(input.Hanging());
 				builder.Annotate(match.Groups["channel"], LogEventMarkup.Channel);
diff --git a/Engine/Source/Programs/AutomationTool/AutomationUtils/Matchers/LogChannelEventMatcher.cs b/Engine/Source/Programs/AutomationTool/AutomationUtils/Matchers/LogChannelEventMatcher.cs
--- a/Engine/Source/Programs/AutomationTool/AutomationUtils/Matchers/LogChannelEventMatcher.cs
+++ b/Engine/Source/Programs/AutomationTool/AutomationUtils/Matchers/LogChannelEventMatcher.cs
@@ -25,8 +25,8 @@ namespace AutomationUtils.Matchers
 		/// <inheritdoc/>
 		public LogEventMatch? Match(ILogCursor input)
 		{
-			Match? match = s_pattern.Match(input.CurrentLine!);
-			if (match.Success)
+			Match? match;
+			if (input.TryMatch(s_pattern, out match))
 			{
 				LogEventBuilder builder = new LogEventBuilder(input.Hanging());
 				builder.Annotate(match.Groups["channel"], LogEventMarkup.Channel);
