UAT: Output a structured log file from UAT.

diff --git a/Engine/Source/Programs/AutomationTool/AutomationUtils/Automation.cs b/Engine/Source/Programs/AutomationTool/AutomationUtils/Automation.cs
--- a/Engine/Source/Programs/AutomationTool/AutomationUtils/Automation.cs
+++ b/Engine/Source/Programs/AutomationTool/AutomationUtils/Automation.cs
@@ -104,13 +104,8 @@ namespace AutomationTool
 				CommandUtils.InitCommandEnvironment();
 
 				// Create the log file, and flush the startup listener to it
-				TraceListener LogTraceListener = LogUtils.AddLogFileListener(CommandUtils.CmdEnv.LogFolder,
-					CommandUtils.CmdEnv.FinalLogFolder);
-				StartupListener.CopyTo(LogTraceListener);
-				Trace.Listeners.Remove(StartupListener);
-				
-				Log.TraceInformation($"Log location: {LogUtils.LogFileName}");
-				if (!String.Equals(LogUtils.FinalLogFileName, LogUtils.LogFileName))
+				LogUtils.AddLogFileListener(new DirectoryReference(CommandUtils.CmdEnv.LogFolder), new DirectoryReference(CommandUtils.CmdEnv.FinalLogFolder));
+				if (LogUtils.FinalLogFileName != LogUtils.LogFileName)
 				{
 					Log.TraceInformation($"Final log location: {LogUtils.FinalLogFileName}");	
 				}
diff --git a/Engine/Source/Programs/AutomationTool/AutomationUtils/LogUtils.cs b/Engine/Source/Programs/AutomationTool/AutomationUtils/LogUtils.cs
--- a/Engine/Source/Programs/AutomationTool/AutomationUtils/LogUtils.cs
+++ b/Engine/Source/Programs/AutomationTool/AutomationUtils/LogUtils.cs
@@ -1,4 +1,4 @@
-﻿// Copyright Epic Games, Inc. All Rights Reserved.
+// Copyright Epic Games, Inc. All Rights Reserved.
 using System;
 using System.Collections.Generic;
 using System.Linq;
@@ -7,6 +7,7 @@ using System.IO;
 using System.Runtime.CompilerServices;
 using System.Diagnostics;
 using UnrealBuildTool;
+using EpicGames.Core;
 
 namespace AutomationTool
 {
@@ -15,60 +16,54 @@ namespace AutomationTool
 		/// <summary>
 		/// Stores the log filename.
 		/// </summary>
-		public static string LogFileName
+		public static FileReference LogFileName
 		{
 			get;
-			set;
+			private set;
 		}
 
         /// <summary>
 		/// Stores the final log filename. The build system uses this to display the network path that the log will be copied to once builds complete.
         /// </summary>
-		public static string FinalLogFileName
+		public static FileReference FinalLogFileName
         {
 			get;
-			set;
+			private set;
         }
 
-        /// <summary>
-        /// Creates the TraceListener used for file logging.
-        /// We cannot simply use a TextWriterTraceListener because we need more flexibility when the file cannot be created.
-        /// TextWriterTraceListener lazily creates the file, silently failing when it cannot.
-        /// </summary>
-        /// <returns>The newly created TraceListener, or null if it could not be created.</returns>
-        public static TraceListener AddLogFileListener(string LogFolder, string FinalLogFolder)
-        {
+		/// <summary>
+		/// Creates the TraceListener used for file logging.
+		/// We cannot simply use a TextWriterTraceListener because we need more flexibility when the file cannot be created.
+		/// TextWriterTraceListener lazily creates the file, silently failing when it cannot.
+		/// </summary>
+		/// <returns>The newly created TraceListener, or null if it could not be created.</returns>
+		public static void AddLogFileListener(DirectoryReference LogFolder, DirectoryReference FinalLogFolder)
+		{
 			int NumAttempts = 0;
-			for(;;)
-            {
-                try
-                {
-                    // We do not need to set AutoFlush on the StreamWriter because we set Trace.AutoFlush, which calls it for us.
-                    // Not only would this be redundant, StreamWriter AutoFlush does not flush the encoder, while a direct call to 
-                    // StreamWriter.Flush() will, which is what the Trace system with AutoFlush = true will do.
-                    // Internally, FileStream constructor opens the file with good arguments for writing to log files.
+			for (; ; )
+			{
+				try
+				{
 					string Name = (NumAttempts == 0) ? "Log.txt" : string.Format("Log_{0}.txt", NumAttempts + 1);
-	                string LogFileName = CommandUtils.CombinePaths(LogFolder, Name);
 
-                    TextWriterTraceListener Listener = new TextWriterTraceListener(new StreamWriter(LogFileName), "AutomationFileLogListener");
-					Trace.Listeners.Add(Listener);
+					FileReference LogFileName = FileReference.Combine(LogFolder, Name);
+					Log.AddFileWriterWithoutBackup("AutomationFileLogListener", LogFileName);
 
 					LogUtils.LogFileName = LogFileName;
-					LogUtils.FinalLogFileName = CommandUtils.CombinePaths(FinalLogFolder, Name);;
-
-					return Listener;
-                }
-                catch (Exception Ex)
-                {
-					if(NumAttempts++ >= 10)
-                    {
+					LogUtils.FinalLogFileName = FileReference.Combine(FinalLogFolder, Name);
+					break;
+				}
+				catch (Exception Ex)
+				{
+					if (NumAttempts++ >= 10)
+					{
 						throw new AutomationException(Ex, "Unable to create log file after {0} attempts.", NumAttempts);
 					}
-                    }
-                }
-        }
+				}
+			}
+		}
 
-        /// <summary>
+		/// <summary>
 		/// Dumps exception info to log.
 		/// @todo: Remove this function as it doesn't do a good job printing the exception information.
 		/// </summary>
diff --git a/Engine/Source/Programs/Shared/EpicGames.Core/Log.cs b/Engine/Source/Programs/Shared/EpicGames.Core/Log.cs
--- a/Engine/Source/Programs/Shared/EpicGames.Core/Log.cs
+++ b/Engine/Source/Programs/Shared/EpicGames.Core/Log.cs
@@ -204,8 +204,6 @@ namespace EpicGames.Core
 
 		/// <summary>
 		/// Adds a trace listener that writes to a log file.
-		/// If Log.DuplicateLogFiles is true, two files will be created - one with the requested name,
-		/// another with a timestamp appended before any extension.
 		/// If a StartupTraceListener was in use, this function will copy its captured data to the log file(s)
 		/// and remove the startup listener from the list of registered listeners.
 		/// </summary>
@@ -240,9 +238,9 @@ namespace EpicGames.Core
 				// Ensure that the backup gets a unique name, in the extremely unlikely case that UBT was run twice during
 				// the same second.
 				DateTime fileTime = File.GetCreationTimeUtc(outputFile.FullName);
-					
+
 				FileReference backupFile;
-				for (;;)
+				for (; ; )
 				{
 					string timestamp = $"{fileTime:yyyy.MM.dd-HH.mm.ss}";
 					backupFile = FileReference.Combine(outputFile.Directory,
@@ -257,7 +255,19 @@ namespace EpicGames.Core
 
 				FileReference.Move(outputFile, backupFile);
 			}
-			
+
+			AddFileWriterWithoutBackup(name, outputFile);
+		}
+
+		/// <summary>
+		/// Adds a trace listener that writes to a log file.
+		/// If a StartupTraceListener was in use, this function will copy its captured data to the log file(s)
+		/// and remove the startup listener from the list of registered listeners.
+		/// </summary>
+		/// <param name="outputFile">The file to write to</param>
+		/// <returns>The created trace listener</returns>
+		public static void AddFileWriterWithoutBackup(string name, FileReference outputFile)
+		{
 			TextWriterTraceListener firstTextWriter = DefaultLogger.AddFileWriter(name, outputFile);
 			
 			// find the StartupTraceListener in the listeners that was added early on
@@ -461,7 +471,7 @@ namespace EpicGames.Core
 		/// </summary>
 		/// <param name="ex">The exception to display</param>
 		/// <param name="logFileName">The log filename to display, if any</param>
-		public static void WriteException(Exception ex, string? logFileName)
+		public static void WriteException(Exception ex, FileReference? logFileName)
 		{
 			string logSuffix = (logFileName == null) ? "" : String.Format("\n(see {0} for full exception trace)", logFileName);
 			TraceLog("==============================================================================");
@@ -1170,6 +1180,9 @@ namespace EpicGames.Core
 					}
 				}
 
+				JsonLogEvent jsonLogEvent = JsonLogEvent.FromLoggerState(logLevel, eventId, state, exception, formatter);
+				_eventChannel.Writer.TryWrite(jsonLogEvent);
+
 				// Handle the console output separately; we format things differently
 				if (logLevel >= LogLevel.Information)
 				{
@@ -1191,9 +1204,6 @@ namespace EpicGames.Core
 					}
 					try
 					{
-						JsonLogEvent jsonLogEvent = JsonLogEvent.FromLoggerState(logLevel, eventId, state, exception, formatter);
-						_eventChannel.Writer.TryWrite(jsonLogEvent);
-
 						if (WriteJsonToStdOut)
 						{
 							Console.WriteLine(Encoding.UTF8.GetString(jsonLogEvent.Data.Span));
