BuildCMakeLib: Default output binaries to lib root directory, instead of Engine\Binaries

#rnx

diff --git a/Engine/Source/Programs/AutomationTool/Scripts/BuildCMakeLib.Automation.cs b/Engine/Source/Programs/AutomationTool/Scripts/BuildCMakeLib.Automation.cs
--- a/Engine/Source/Programs/AutomationTool/Scripts/BuildCMakeLib.Automation.cs
+++ b/Engine/Source/Programs/AutomationTool/Scripts/BuildCMakeLib.Automation.cs
@@ -21,7 +21,7 @@ using UnrealBuildBase;
 [Help("TargetPlatform", "Specify the name of the target platform to build (eg. -TargetPlatform=IOS).")]
 [Help("TargetArchitecture", "Specify the name of the target architecture to build (eg. -TargetArchitecture=x86_64).")]
 [Help("TargetConfigs", "Specify a list of configurations to build, separated by '+' characters (eg. -TargetConfigs=release+debug). Default is release+debug.")]
-[Help("BinOutputPath", "Override the path to output binaries to. (eg. -BinOutputPath=bin). Default is Binaries.")]
+[Help("BinOutputPath", "Override the path to output binaries to. (eg. -BinOutputPath=bin). Default is empty.")]
 [Help("LibOutputPath", "Override the path to output libraries to. (eg. -LibOutputPath=lib). Default is empty.")]
 [Help("CMakeGenerator", "Specify the CMake generator to use.")]
 [Help("CMakeProjectIncludeFile", "Specify the name of the CMake project include file to use, first looks in current directory then looks in global directory.")]
@@ -189,7 +189,7 @@ public sealed class BuildCMakeLib : BuildCommand
 
 		protected DirectoryReference GetOutputBinaryDirectory(TargetLib TargetLib, string TargetConfiguration)
 		{
-			return DirectoryReference.Combine(PlatformEngineRoot, TargetLib.BinOutputPath, IsPlatformExtension ? "" : PlatformOrGroupName, VariantDirectory ?? "", TargetConfiguration ?? "");
+			return DirectoryReference.Combine(GetTargetLibRootDirectory(TargetLib), TargetLib.BinOutputPath, IsPlatformExtension ? "" : PlatformOrGroupName, VariantDirectory ?? "", TargetConfiguration ?? "");
 		}
 
 		public abstract string PlatformOrGroupName { get; }
