Added HLOD & Minimap build scripts
Unified "-AutoSubmit" option

diff --git a/Engine/Source/Editor/UnrealEd/Private/WorldPartition/WorldPartitionBuilder.cpp b/Engine/Source/Editor/UnrealEd/Private/WorldPartition/WorldPartitionBuilder.cpp
--- a/Engine/Source/Editor/UnrealEd/Private/WorldPartition/WorldPartitionBuilder.cpp
+++ b/Engine/Source/Editor/UnrealEd/Private/WorldPartition/WorldPartitionBuilder.cpp
@@ -48,7 +48,7 @@ FCellInfo::FCellInfo()
 UWorldPartitionBuilder::UWorldPartitionBuilder(const FObjectInitializer& ObjectInitializer)
 	: Super(ObjectInitializer)
 {
-	bSubmit = FParse::Param(FCommandLine::Get(), TEXT("Submit"));
+	bAutoSubmit = FParse::Param(FCommandLine::Get(), TEXT("AutoSubmit"));
 }
 
 bool UWorldPartitionBuilder::RunBuilder(UWorld* World)
diff --git a/Engine/Source/Editor/UnrealEd/Private/WorldPartition/WorldPartitionHLODsBuilder.cpp b/Engine/Source/Editor/UnrealEd/Private/WorldPartition/WorldPartitionHLODsBuilder.cpp
--- a/Engine/Source/Editor/UnrealEd/Private/WorldPartition/WorldPartitionHLODsBuilder.cpp
+++ b/Engine/Source/Editor/UnrealEd/Private/WorldPartition/WorldPartitionHLODsBuilder.cpp
@@ -256,7 +256,7 @@ bool UWorldPartitionHLODsBuilder::ValidateParams() const
 		}
 	}
 
-	if (ShouldRunStep(EHLODBuildStep::HLOD_Finalize) && bSubmit && !ISourceControlModule::Get().GetProvider().IsEnabled())
+	if (ShouldRunStep(EHLODBuildStep::HLOD_Finalize) && bAutoSubmit && !ISourceControlModule::Get().GetProvider().IsEnabled())
 	{
 		UE_LOG(LogWorldPartitionHLODsBuilder, Error, TEXT("Submit requires that a valid source control provider is enabled, exiting..."), *BuildManifest);
 		return false;
@@ -311,7 +311,7 @@ bool UWorldPartitionHLODsBuilder::RunInternal(UWorld* World, const FCellInfo& In
 		bRet = DeleteHLODActors();
 	}
 
-	if (bRet && ShouldRunStep(EHLODBuildStep::HLOD_Finalize) && bSubmit)
+	if (bRet && ShouldRunStep(EHLODBuildStep::HLOD_Finalize) && bAutoSubmit)
 	{
 		bRet = SubmitHLODActors();
 	}
diff --git a/Engine/Source/Editor/UnrealEd/Private/WorldPartition/WorldPartitionMiniMapBuilder.cpp b/Engine/Source/Editor/UnrealEd/Private/WorldPartition/WorldPartitionMiniMapBuilder.cpp
--- a/Engine/Source/Editor/UnrealEd/Private/WorldPartition/WorldPartitionMiniMapBuilder.cpp
+++ b/Engine/Source/Editor/UnrealEd/Private/WorldPartition/WorldPartitionMiniMapBuilder.cpp
@@ -39,8 +39,6 @@ UWorldPartitionMiniMapBuilder::UWorldPartitionMiniMapBuilder(const FObjectInitia
 
 bool UWorldPartitionMiniMapBuilder::PreRun(UWorld* World, FPackageSourceControlHelper& PackageHelper)
 {
-	bAutoSubmit = FParse::Param(FCommandLine::Get(), TEXT("AutoSubmit"));
-	
 	if (WorldMiniMap == nullptr)
 	{
 		WorldMiniMap = FWorldPartitionMiniMapHelper::GetWorldPartitionMiniMap(World, true);
diff --git a/Engine/Source/Editor/UnrealEd/Public/WorldPartition/WorldPartitionBuilder.h b/Engine/Source/Editor/UnrealEd/Public/WorldPartition/WorldPartitionBuilder.h
--- a/Engine/Source/Editor/UnrealEd/Public/WorldPartition/WorldPartitionBuilder.h
+++ b/Engine/Source/Editor/UnrealEd/Public/WorldPartition/WorldPartitionBuilder.h
@@ -88,5 +88,5 @@ protected:
 	bool bLoadNonDynamicDataLayers = true;
 	bool bLoadInitiallyActiveDataLayers = true;
 
-	bool bSubmit = true;
+	bool bAutoSubmit = false;
 };
\ No newline at end of file
diff --git a/Engine/Source/Editor/UnrealEd/Public/WorldPartition/WorldPartitionMiniMapBuilder.h b/Engine/Source/Editor/UnrealEd/Public/WorldPartition/WorldPartitionMiniMapBuilder.h
--- a/Engine/Source/Editor/UnrealEd/Public/WorldPartition/WorldPartitionMiniMapBuilder.h
+++ b/Engine/Source/Editor/UnrealEd/Public/WorldPartition/WorldPartitionMiniMapBuilder.h
@@ -43,5 +43,4 @@ private:
 	TObjectPtr<AWorldPartitionMiniMap> WorldMiniMap;
 
 	FBox EditorBounds;
-	bool bAutoSubmit;
 };
\ No newline at end of file
diff --git a/Engine/Source/Programs/AutomationTool/Scripts/WorldPartitionBuilder.Automation.cs b/Engine/Source/Programs/AutomationTool/Scripts/WorldPartitionBuilder.Automation.cs
--- a/Engine/Source/Programs/AutomationTool/Scripts/WorldPartitionBuilder.Automation.cs
+++ b/Engine/Source/Programs/AutomationTool/Scripts/WorldPartitionBuilder.Automation.cs
@@ -7,6 +7,12 @@ using System.Threading;
 
 namespace AutomationScripts.Automation
 {
+	[Help("Execute a World Partition builder")]
+	[Help("Builder=<Name>", "Name of the builder to run")]
+	[Help("CommandletArgs=<value>", "Arguments to provide to the builder commandlet")]
+	[Help("Submit", "If the files modified by the builder should be submitted at the end of the process")]
+	[Help("ShelveUser=<Name>", "If provided (along with -ShelveWorkspace, modified files will be shelved for the P4 User in the specified Workspace.")]
+	[Help("ShelveWorkspace=<Name>", "If provided (along with -ShelveUser, modified files will be shelved for the P4 User in the specified Workspace.")]
 	public class WorldPartitionBuilder : BuildCommand
 	{
 		public override void ExecuteBuild()
@@ -14,13 +20,13 @@ namespace AutomationScripts.Automation
 			string Builder = ParseRequiredStringParam("Builder");
 			string CommandletArgs = ParseOptionalStringParam("CommandletArgs");
 
-			bool bSubmit = ParseParam("Submit");
+			bool bSubmitResult = ParseParam("Submit");
 
 			string ShelveUser = ParseOptionalStringParam("ShelveUser");
 			string ShelveWorkspace = ParseOptionalStringParam("ShelveWorkspace");
 			bool bShelveResult = !String.IsNullOrEmpty(ShelveUser) && !String.IsNullOrEmpty(ShelveWorkspace);
 
-			if (!P4Enabled && (bSubmit || bShelveResult))
+			if (!P4Enabled && (bSubmitResult || bShelveResult))
 			{
 				LogError("P4 required to submit or shelve build results");
 				return;
@@ -28,15 +34,17 @@ namespace AutomationScripts.Automation
 
 			CommandletArgs = "-Builder=" + Builder + " " + CommandletArgs;
 
-			if (bSubmit)
+			if (bSubmitResult)
 			{
-				CommandletArgs += " -Submit";
+				CommandletArgs += " -AutoSubmit";
 			}
 
 			string EditorExe = "UnrealEditor-Cmd.exe";
 			EditorExe = AutomationTool.HostPlatform.Current.GetUnrealExePath(EditorExe);
 
 			FileReference ProjectPath = ParseProjectParam();
+
+			// Execute the commandlet - Will throw an exception on failures
 			RunCommandlet(ProjectPath, EditorExe, "WorldPartitionBuilderCommandlet", CommandletArgs);
 			
 			if (bShelveResult)
