CreatePlatformExtension UAT command adds 'HasExplicitPlatforms' if it removes the last existing platform from a plugin platform list.

#rnx

diff --git a/Engine/Source/Programs/AutomationTool/Scripts/CreatePlatformExtension.cs b/Engine/Source/Programs/AutomationTool/Scripts/CreatePlatformExtension.cs
--- a/Engine/Source/Programs/AutomationTool/Scripts/CreatePlatformExtension.cs
+++ b/Engine/Source/Programs/AutomationTool/Scripts/CreatePlatformExtension.cs
@@ -347,9 +347,11 @@ public class CreatePlatformExtension : BuildCommand
 						}
 
 						// remove platform from parent plugin references
-						if (bHasPlatform && ParentModuleDesc.PlatformAllowList != null)
+						if (bHasPlatform && ParentModuleDesc.PlatformAllowList != null && ParentModuleDesc.PlatformAllowList.Contains(Platform) )
 						{
-							bParentPluginDirty |= ParentModuleDesc.PlatformAllowList.Remove(Platform);
+							ParentModuleDesc.PlatformAllowList.Remove(Platform);
+							ParentModuleDesc.bHasExplicitPlatforms |= (ParentModuleDesc.PlatformAllowList.Count == 0); // an empty list is interpreted as 'all platforms are allowed' otherwise
+							bParentPluginDirty = true;
 						}
 						if (bHasPlatform && ParentModuleDesc.PlatformDenyList != null)
 						{
@@ -384,6 +386,7 @@ public class CreatePlatformExtension : BuildCommand
 						if (bHasPlatform && PlatformArrayContainsPlatform( ParentPluginDesc.PlatformAllowList, Platform ) )
 						{
 							ParentPluginDesc.PlatformAllowList = ParentPluginDesc.PlatformAllowList.Where( X => !X.Equals(Platform.ToString() ) ).ToArray();
+							ParentPluginDesc.bHasExplicitPlatforms |= (ParentPluginDesc.PlatformAllowList.Length == 0); // an empty list is interpreted as "all platforms" otherwise
 							bParentPluginDirty = true;
 						}
 						if (bHasPlatform && PlatformArrayContainsPlatform( ParentPluginDesc.PlatformDenyList, Platform ) )
@@ -398,9 +401,11 @@ public class CreatePlatformExtension : BuildCommand
 				ChildPlugin.WriteObjectEnd();
 
 				// remove platform from parent plugin, if necessary
-				if (bHasPlatform && ParentPlugin.SupportedTargetPlatforms != null)
+				if (bHasPlatform && ParentPlugin.SupportedTargetPlatforms != null && ParentPlugin.SupportedTargetPlatforms.Contains(Platform) )
 				{
-					bParentPluginDirty |= ParentPlugin.SupportedTargetPlatforms.Remove(Platform);
+					ParentPlugin.SupportedTargetPlatforms.Remove(Platform);
+					ParentPlugin.bHasExplicitPlatforms |= (ParentPlugin.SupportedTargetPlatforms.Count == 0); // an empty list is interpreted as "all platforms" otherwise
+					bParentPluginDirty = true;
 				}
 			}
 			AddNewFile(FinalFileName);
