Remove loading of ignore patterns from UAT/UBT. These can cause a large CPU load when a lot of warnings and errors are present, which should be masked out in a more efficient and targetted way.

diff --git a/Engine/Source/Programs/AutomationTool/AutomationUtils/Automation.cs b/Engine/Source/Programs/AutomationTool/AutomationUtils/Automation.cs
--- a/Engine/Source/Programs/AutomationTool/AutomationUtils/Automation.cs
+++ b/Engine/Source/Programs/AutomationTool/AutomationUtils/Automation.cs
@@ -79,9 +79,6 @@ namespace AutomationTool
 				Assembly UnrealBuildToolAssembly = typeof(UnrealBuildTool.BuildVersion).Assembly;
 				Log.EventParser.AddMatchersFromAssembly(UnrealBuildToolAssembly);
 
-				// Read all the ignore patterns
-				await Log.EventParser.ReadIgnorePatternsAsync(Unreal.EngineDirectory);
-
 				// Get the path to the telemetry file, if present
 				string TelemetryFile = GlobalCommandLine.TelemetryPath;
 				JsonTracer Tracer = JsonTracer.TryRegisterAsGlobalTracer();
diff --git a/Engine/Source/Programs/UnrealBuildTool/UnrealBuildTool.cs b/Engine/Source/Programs/UnrealBuildTool/UnrealBuildTool.cs
--- a/Engine/Source/Programs/UnrealBuildTool/UnrealBuildTool.cs
+++ b/Engine/Source/Programs/UnrealBuildTool/UnrealBuildTool.cs
@@ -506,9 +506,6 @@ namespace UnrealBuildTool
 				// Add all the default event matchers from the UBT assembly
 				Log.EventParser.AddMatchersFromAssembly(Assembly.GetExecutingAssembly());
 
-				// Read all the ignore patterns
-				await Log.EventParser.ReadIgnorePatternsAsync(Unreal.EngineDirectory);
-
 				// Configure the progress writer
 				ProgressWriter.bWriteMarkup = Options.bWriteProgressMarkup;
 
