Add the ability for engine plugins to retry copying stuff.

#rnx

diff --git a/Engine/Source/Programs/AutomationTool/AutomationUtils/CommandUtils.cs b/Engine/Source/Programs/AutomationTool/AutomationUtils/CommandUtils.cs
--- a/Engine/Source/Programs/AutomationTool/AutomationUtils/CommandUtils.cs
+++ b/Engine/Source/Programs/AutomationTool/AutomationUtils/CommandUtils.cs
@@ -1793,12 +1793,13 @@ namespace AutomationTool
 		/// <param name="TargetDir">Target directory</param>
 		/// <param name="RelativePaths">Paths relative to the source directory to copy</param>
 		/// <param name="MaxThreads">Maximum number of threads to create</param>
+		/// <param name="bRetry"></param>
 		/// <returns>List of filenames copied to the target directory</returns>
-		public static List<string> ThreadedCopyFiles(string SourceDir, string TargetDir, List<string> RelativePaths, int MaxThreads = 64)
+		public static List<string> ThreadedCopyFiles(string SourceDir, string TargetDir, List<string> RelativePaths, int MaxThreads = 64, bool bRetry = false)
 		{
             var SourceFileNames = RelativePaths.Select(RelativePath => CommandUtils.CombinePaths(SourceDir, RelativePath)).ToList();
             var TargetFileNames = RelativePaths.Select(RelativePath => CommandUtils.CombinePaths(TargetDir, RelativePath)).ToList();
-			CommandUtils.ThreadedCopyFiles(SourceFileNames, TargetFileNames, MaxThreads);
+			CommandUtils.ThreadedCopyFiles(SourceFileNames, TargetFileNames, MaxThreads, bRetry: bRetry);
 			return TargetFileNames;
 		}
 
@@ -1810,14 +1811,15 @@ namespace AutomationTool
 		/// <param name="Filter">Filter which selects files from the source directory to copy</param>
 		/// <param name="bIgnoreSymlinks">Whether to ignore symlinks during the copy</param>
 		/// <param name="MaxThreads">Maximum number of threads to create</param>
+		/// <param name="bRetry"></param>
 		/// <returns>List of filenames copied to the target directory</returns>
-		public static List<string> ThreadedCopyFiles(string SourceDir, string TargetDir, FileFilter Filter, bool bIgnoreSymlinks, int MaxThreads = 64)
+		public static List<string> ThreadedCopyFiles(string SourceDir, string TargetDir, FileFilter Filter, bool bIgnoreSymlinks, int MaxThreads = 64, bool bRetry = false)
 		{
 			// Filter all the relative paths
 			LogInformation("Applying filter to {0}...", SourceDir);
 			DirectoryReference SourceDirRef = new DirectoryReference(SourceDir);
 			var RelativePaths = Filter.ApplyToDirectory(SourceDirRef, bIgnoreSymlinks).Select(x => x.MakeRelativeTo(SourceDirRef)).ToList();
-			return ThreadedCopyFiles(SourceDir, TargetDir, RelativePaths);
+			return ThreadedCopyFiles(SourceDir, TargetDir, RelativePaths, MaxThreads, bRetry: bRetry);
 		}
 
 		/// <summary>
