Gauntlet - Catch callstack line that does not start by "[callstack]" (on Mac)
* + When looking for fatal error, keep last callstack with the most lines.
* + Catch "critical error" callstack report

#rnx

diff --git a/Engine/Source/Programs/AutomationTool/Gauntlet/Unreal/Utils/Gauntlet.UnrealLogParser.cs b/Engine/Source/Programs/AutomationTool/Gauntlet/Unreal/Utils/Gauntlet.UnrealLogParser.cs
--- a/Engine/Source/Programs/AutomationTool/Gauntlet/Unreal/Utils/Gauntlet.UnrealLogParser.cs
+++ b/Engine/Source/Programs/AutomationTool/Gauntlet/Unreal/Utils/Gauntlet.UnrealLogParser.cs
@@ -543,11 +543,13 @@ namespace Gauntlet
 		/// <returns></returns>
 		public UnrealLog.CallstackMessage GetFatalError()
 		{
-			string[] ErrorMsgMatches = new string[] { "(Fatal Error:.+)", "(Assertion Failed:.+)", "(Unhandled Exception:.+)", "(LowLevelFatalError.+)" };
+			string[] ErrorMsgMatches = new string[] { @"(Fatal Error:.+)", @"Critical error: =+[\s\n]+(?:.+?\s*Error:\s*)?(.+)", @"(Assertion Failed:.+)", @"(Unhandled Exception:.+)", @"(LowLevelFatalError.+)" };
 
 			var Traces = ParseTracedErrors(ErrorMsgMatches);
 
-			return Traces.Count() > 0 ? Traces.First() : null;
+			// Keep the one with the most information.
+			Traces.OrderBy(T => T.Callstack.Length);
+			return Traces.Count() > 0 ? Traces.Last() : null;
 		}
 
 		/// <summary>
@@ -668,10 +670,11 @@ namespace Gauntlet
 					if (MsgMatch.Success)
 					{
 						string MsgString = MsgMatch.Groups[1].ToString();
-						if (MsgString.Contains("[Callstack]") == false)
+						if (!MsgMatch.Groups[0].ToString().Contains("\n") /* avoid a bug where .+ match \n */
+							&& !string.IsNullOrEmpty(MsgString)
+							&& Regex.Match(MsgString, @"0[xX][0-9A-f]{8,16}").Success == false)
 						{
-							// Put the message first. e,g "foo != false: Assert [File:d:\UE\foomanager.cpp:10]
-							NewTrace.Message = MsgMatch.Groups[1].ToString() + ": " + NewTrace.Message;
+							NewTrace.Message = NewTrace.Message + "\n" + MsgString;
 						}
 					}
 
@@ -742,7 +745,7 @@ namespace Gauntlet
 						// filename is optional, must be in [file]
 						// Note - Ubreal callstacks are always meant to omit all three with placeholders for missing values, but
 						// we'll assume that may not happen...
-						Match CSMatch = Regex.Match(Line, @"\[Callstack\]\s+(0[xX][0-9A-f]{8,16})(?:\s+([^\s]+))?(?:\s+\[(.*?)\])?", RegexOptions.IgnoreCase);
+						Match CSMatch = Regex.Match(Line, @"(0[xX][0-9A-f]{8,16})(?:\s+([^\s]+))?(?:\s+\[(.*?)\])?", RegexOptions.IgnoreCase);
 
 						if (CSMatch.Success)
 						{
@@ -752,7 +755,7 @@ namespace Gauntlet
 
 							if (string.IsNullOrEmpty(File))
 							{
-								File = "[Unknown File]";
+								File = "Unknown File";
 							}
 
 							// Remove any exe
@@ -763,7 +766,7 @@ namespace Gauntlet
 								Func = Func.Substring(Func.IndexOf(StripFrom) + StripFrom.Length);
 							}
 
-							string NewLine = string.Format("{0} {1} {2}", Address, Func, File);
+							string NewLine = string.Format("{0} {1} [{2}]", Address, Func, File);
 
 							Backtrace.Add(NewLine);
 
