Gauntlet - Properly report session timeout in Automated Test Session
* + avoid flooding the console output with empty lines when reporting tests that did not run.

#rnx

diff --git a/Engine/Source/Programs/AutomationTool/Gauntlet/Unreal/Automation/UE.Automation.cs b/Engine/Source/Programs/AutomationTool/Gauntlet/Unreal/Automation/UE.Automation.cs
--- a/Engine/Source/Programs/AutomationTool/Gauntlet/Unreal/Automation/UE.Automation.cs
+++ b/Engine/Source/Programs/AutomationTool/Gauntlet/Unreal/Automation/UE.Automation.cs
@@ -732,14 +732,23 @@ namespace UE
 					}
 					else if (IncompleteTests.Count() > 0)
 					{
-						MB.H3("The following test(s) were incomplete:");
-
-						foreach (UnrealAutomatedTestResult Result in IncompleteTests)
+						var InProcess = IncompleteTests.Where(T => T.State == TestStateType.InProcess);
+						if (InProcess.Any())
+						{
+							MB.H3("The following test(s) were incomplete:");
+							foreach (UnrealAutomatedTestResult Result in InProcess)
+							{
+								MB.H4(string.Format("Error: Test '{0}' did not complete", Result.TestDisplayName));
+								MB.Paragraph("FullName: " + Result.FullTestPath);
+								MB.UnorderedList(CapErrorOrWarningList(Result.ErrorEvents.Select(E => ErrorLineForHorde(E.Message)).Distinct()));
+								MB.UnorderedList(CapErrorOrWarningList(Result.WarningEvents.Select(E => WarningLineForHorde(E.Message)).Distinct()));
+							}
+						}
+						var NotRun = IncompleteTests.Where(T => T.State == TestStateType.NotRun);
+						if (NotRun.Any())
 						{
-							MB.H4(string.Format("Error: Test '{0}' did not run or complete", Result.TestDisplayName));
-							MB.Paragraph("FullName: " + Result.FullTestPath);
-							MB.UnorderedList(CapErrorOrWarningList(Result.ErrorEvents.Select(E => ErrorLineForHorde(E.Message)).Distinct()));
-							MB.UnorderedList(CapErrorOrWarningList(Result.WarningEvents.Select(E => WarningLineForHorde(E.Message)).Distinct()));
+							MB.H3("The following test(s) were not run:");
+							MB.UnorderedList(NotRun.Select(T => T.FullTestPath));
 						}
 					}
 				}
diff --git a/Engine/Source/Programs/AutomationTool/Gauntlet/Unreal/Base/Gauntlet.UnrealTestNode.cs b/Engine/Source/Programs/AutomationTool/Gauntlet/Unreal/Base/Gauntlet.UnrealTestNode.cs
--- a/Engine/Source/Programs/AutomationTool/Gauntlet/Unreal/Base/Gauntlet.UnrealTestNode.cs
+++ b/Engine/Source/Programs/AutomationTool/Gauntlet/Unreal/Base/Gauntlet.UnrealTestNode.cs
@@ -1258,6 +1258,7 @@ namespace Gauntlet
 		/// <returns>ITestReport</returns>
 		public virtual ITestReport CreateUnrealEngineTestPassReport(string UnrealAutomatedTestReportPath, string ReportURL)
 		{
+			bool HasTimeout = RoleResults != null && RoleResults.Where(R => R.ProcessResult == UnrealProcessResult.TimeOut).Any();
 			string JsonReportPath = Path.Combine(UnrealAutomatedTestReportPath, "index.json");
 			if (File.Exists(JsonReportPath))
 			{
@@ -1286,14 +1287,22 @@ namespace Gauntlet
 						if (!String.IsNullOrEmpty(Test.TestDisplayName))
 						{
 							LastTestWithCriticalFailure = Test.FullTestPath;
-							string ErrorMessage = "Engine encountered a critical failure. \n";
-							if (FatalError != null)
+							string ErrorMessage;
+							if (HasTimeout)
 							{
-								ErrorMessage += FatalError.FormatForLog();
+								ErrorMessage = String.Format("Session reached timeout after {0} seconds.", MaxDuration);
 							}
 							else
 							{
-								ErrorMessage += "No callstack found in the log.";
+								ErrorMessage = "Engine encountered a critical failure. \n";
+								if (FatalError != null)
+								{
+									ErrorMessage += FatalError.FormatForLog();
+								}
+								else
+								{
+									ErrorMessage += "No callstack found in the log.";
+								}
 							}
 							Test.AddError(ErrorMessage);
 							if (Retries >= MaxRetries || JsonTestPassResults.NotRun == 0)
@@ -1314,7 +1323,6 @@ namespace Gauntlet
 				{
 					// The test pass did not run at all
 					Log.Verbose("Found not-run tests: {0}", JsonTestPassResults.NotRun);
-					bool HasTimeout = RoleResults != null && RoleResults.Where(R => R.ProcessResult == UnrealProcessResult.TimeOut).Any();
 					if (GetCachedConfiguration().ResumeOnCriticalFailure && !HasTimeout)
 					{
 						if (Retries < MaxRetries)
@@ -1359,6 +1367,17 @@ namespace Gauntlet
 							JsonTestPassResults.WriteToJson(JsonReportPath);
 						}
 					}
+					else if(HasTimeout)
+					{
+						// Adding a note to the report about why the not-run are not going to be run
+						string Message = string.Format("Session reached timeout after {0} seconds.", MaxDuration);
+						var NotRunTests = JsonTestPassResults.Tests.Where((T => T.State == TestStateType.NotRun));
+						foreach (var Test in NotRunTests)
+						{
+							Test.AddInfo(Message);
+						}
+						JsonTestPassResults.WriteToJson(JsonReportPath);
+					}
 				}
 				string HordeTestDataKey = string.IsNullOrEmpty(GetCachedConfiguration().HordeTestDataKey) ? Name + " " + Context.ToString() : GetCachedConfiguration().HordeTestDataKey;
 				GetCachedConfiguration().HordeTestDataKey = HordeTestDataKey;
