Removes deprecated automation dashboard submission functions
#pf 629f5bbff73a9b013dbb47c5

diff --git a/Engine/Source/Programs/AutomationTool/Gauntlet/Unreal/Base/Gauntlet.UnrealTestNode.cs b/Engine/Source/Programs/AutomationTool/Gauntlet/Unreal/Base/Gauntlet.UnrealTestNode.cs
--- a/Engine/Source/Programs/AutomationTool/Gauntlet/Unreal/Base/Gauntlet.UnrealTestNode.cs
+++ b/Engine/Source/Programs/AutomationTool/Gauntlet/Unreal/Base/Gauntlet.UnrealTestNode.cs
@@ -243,11 +243,6 @@ namespace Gauntlet
 			TestNodeEvents.Add(InEvent);
 		}
 
-		/// <summary>
-		/// Whether we submit to the dashboard
-		/// </summary>
-		public virtual bool ShouldSubmitDashboardResult { get { return CommandUtils.IsBuildMachine; } }
-
 		/// <summary>
 		/// Helper class that turns our wishes into reality
 		/// </summary>
@@ -1036,14 +1031,11 @@ namespace Gauntlet
 
 			try
 			{
-				// Check if the deprecated signature is overriden, call it anyway if that the case.
-				var DeprecatedSignature = new[] { typeof(TestResult), typeof(UnrealTestContext), typeof(UnrealBuildSource), typeof(IEnumerable<UnrealRoleArtifacts>), typeof(string) };
-				if (Utils.InterfaceHelpers.HasOverriddenMethod(this.GetType(), "SubmitToDashboard", DeprecatedSignature))
-				{
-					SubmitToDashboard(GetTestResult(), Context, Context.BuildInfo, SessionArtifacts, ArtifactPath);
+				if (Report != null)
+				{ 
+					SubmitToDashboard(Report); 
 				}
 
-				if (Report != null) { SubmitToDashboard(Report); }
 			}
 			catch (Exception Ex)
 			{
@@ -1441,19 +1433,6 @@ namespace Gauntlet
 			}
 		}
 
-		/// <summary>
-		/// DEPRECATED Optional function that is called on test completion and gives an opportunity to submit a report to a Dashboard
-		/// </summary>
-		/// <param name="Result"></param>
-		/// <param name="Context"></param>
-		/// <param name="Build"></param>
-		/// <param name="Artifacts"></param>
-		/// <param name="ArtifactPath"></param>
-		/// <returns></returns>
-		public virtual void SubmitToDashboard(TestResult Result, UnrealTestContext Context, UnrealBuildSource Build, IEnumerable<UnrealRoleArtifacts> Artifacts, string ArtifactPath)
-		{
-		}
-
 		/// <summary>
 		/// Optional function that is called on test completion and gives an opportunity to submit a report to a Dashboard
 		/// </summary>
