LowLevelTest fixes:
  - add virtual devices to the pool, for parity with RunUnreal.
  - release devices when the session is destroyed.

diff --git a/Engine/Source/Programs/AutomationTool/LowLevelTests/RunLowLevelTests.cs b/Engine/Source/Programs/AutomationTool/LowLevelTests/RunLowLevelTests.cs
--- a/Engine/Source/Programs/AutomationTool/LowLevelTests/RunLowLevelTests.cs
+++ b/Engine/Source/Programs/AutomationTool/LowLevelTests/RunLowLevelTests.cs
@@ -131,6 +131,7 @@ namespace LowLevelTests
 
 			DevicePool.Instance.SetLocalOptions(Options.TempDir, Options.Parallel > 1, Options.DeviceURL);
 			DevicePool.Instance.AddLocalDevices(1);
+			DevicePool.Instance.AddVirtualDevices(2);
 
 			if (!string.IsNullOrEmpty(Options.Device))
 			{
@@ -336,6 +337,8 @@ namespace LowLevelTests
 				{
 					IDeviceUsageReporter.RecordEnd(Instance.Device.Name, (UnrealTargetPlatform)Instance.Device.Platform, IDeviceUsageReporter.EventType.Test, IDeviceUsageReporter.EventState.Failure);
 				}
+				
+				UnrealDeviceReservation?.ReleaseDevices();				
 			}
 		}
 	}
