Add retrying to potential network share failure points in temp storage.

#rnx

diff --git a/Engine/Source/Programs/AutomationTool/BuildGraph/TempStorage.cs b/Engine/Source/Programs/AutomationTool/BuildGraph/TempStorage.cs
--- a/Engine/Source/Programs/AutomationTool/BuildGraph/TempStorage.cs
+++ b/Engine/Source/Programs/AutomationTool/BuildGraph/TempStorage.cs
@@ -16,6 +16,7 @@ using AutomationTool;
 using EpicGames.Core;
 using OpenTracing;
 using OpenTracing.Util;
+using Microsoft.Extensions.Logging;
 
 namespace AutomationTool
 {
@@ -379,10 +380,10 @@ namespace AutomationTool
 		/// <summary>
 		/// Load a manifest from disk
 		/// </summary>
-		/// <param name="File">File to load</param>
-		static public TempStorageManifest Load(FileReference File)
+		/// <param name="file">File to load</param>
+		static public TempStorageManifest Load(FileReference file)
 		{
-			using(StreamReader Reader = new StreamReader(File.FullName))
+			using (StreamReader Reader = new(file.FullName))
 			{
 				return (TempStorageManifest)Serializer.Deserialize(Reader);
 			}
@@ -659,8 +660,9 @@ namespace AutomationTool
 
 					// Check the manifests are identical, byte by byte
 					byte[] LocalManifestBytes = File.ReadAllBytes(LocalFileListLocation.FullName);
-					byte[] SharedManifestBytes = File.ReadAllBytes(SharedFileListLocation.FullName);
-					if(!LocalManifestBytes.SequenceEqual(SharedManifestBytes))
+					byte[] SharedManifestBytes = null;
+					PerformActionWithRetries(() => SharedManifestBytes = File.ReadAllBytes(SharedFileListLocation.FullName), 3, TimeSpan.FromSeconds(1));
+					if (!LocalManifestBytes.SequenceEqual(SharedManifestBytes))
 					{
 						return false;
 					}
@@ -693,7 +695,8 @@ namespace AutomationTool
 
 					// Check the manifests are identical, byte by byte
 					byte[] LocalManifestBytes = File.ReadAllBytes(LocalManifestFile.FullName);
-					byte[] SharedManifestBytes = File.ReadAllBytes(SharedManifestFile.FullName);
+					byte[] SharedManifestBytes = null;
+					PerformActionWithRetries(() => SharedManifestBytes = File.ReadAllBytes(SharedManifestFile.FullName), 3, TimeSpan.FromSeconds(1));
 					if(!LocalManifestBytes.SequenceEqual(SharedManifestBytes))
 					{
 						return false;
@@ -702,9 +705,6 @@ namespace AutomationTool
 
 				// Read the manifest and check the files
 				TempStorageManifest LocalManifest = TempStorageManifest.Load(LocalManifestFile);
-
-
-
 				if(LocalManifest.Files.Any(x => !x.Compare(RootDir)))
 				{
 					return false;
@@ -742,36 +742,26 @@ namespace AutomationTool
 				int Attempts = 5;
 				FileReference SharedFileListLocation;
 				SharedFileListLocation = GetTaggedFileListLocation(SharedDir, NodeName, TagName);
-				while (Attempts-- > 0)
+				PerformActionWithRetries(() =>
 				{
 					if (!FileReference.Exists(SharedFileListLocation))
 					{
-						if (Attempts == 0)
-						{
-							throw new AutomationException("Missing local or shared file list - {0}", SharedFileListLocation.FullName);
-						}
-
-						Thread.Sleep(TimeSpan.FromSeconds(5));
-						continue;
+						throw new AutomationException("Missing local or shared file list - {0}", SharedFileListLocation.FullName);
 					}
+				}, Attempts, TimeSpan.FromSeconds(5));
 
-					try
+				try
+				{
+					PerformActionWithRetries(() =>
 					{
 						// Read the shared manifest
 						CommandUtils.LogInformation("Copying shared tag set from {0} to {1}", SharedFileListLocation.FullName, LocalFileListLocation.FullName);
 						FileList = TempStorageFileList.Load(SharedFileListLocation);
-					}
-					catch (IOException)
-					{
-						if (Attempts == 0)
-						{
-							throw new AutomationException("Local or shared file list {0} was found but failed to be read", SharedFileListLocation.FullName);
-						}
-
-						continue;
-					}
-
-					break;
+					}, Attempts, TimeSpan.FromSeconds(5));
+				}
+				catch
+				{
+					throw new AutomationException("Local or shared file list {0} was found but failed to be read", SharedFileListLocation.FullName);
 				}
 
 				// Save the manifest locally
@@ -799,10 +789,20 @@ namespace AutomationTool
 			if(SharedDir != null && bWriteToSharedStorage)
 			{
 				FileReference SharedFileListLocation = GetTaggedFileListLocation(SharedDir, NodeName, TagName);
-				CommandUtils.LogInformation("Saving file list to {0} and {1}", LocalFileListLocation.FullName, SharedFileListLocation.FullName);
 
-				DirectoryReference.CreateDirectory(SharedFileListLocation.Directory);
-				FileList.Save(SharedFileListLocation);
+				try
+				{
+					PerformActionWithRetries(() =>
+					{
+						CommandUtils.LogInformation("Saving file list to {0} and {1}", LocalFileListLocation.FullName, SharedFileListLocation.FullName);
+						DirectoryReference.CreateDirectory(SharedFileListLocation.Directory);
+						FileList.Save(SharedFileListLocation);
+					}, 3, TimeSpan.FromSeconds(5));
+				}
+				catch (Exception ex)
+				{
+					throw new AutomationException("Failed to save file list {0} to {1}, exception: {2}", LocalFileListLocation, SharedFileListLocation, ex);
+				}
 			}
 			else
 			{
@@ -848,7 +848,7 @@ namespace AutomationTool
 
 					// Save the shared manifest
 					CommandUtils.LogInformation("Saving shared manifest to {0}", SharedManifestFile.FullName);
-					Manifest.Save(SharedManifestFile);
+					PerformActionWithRetries(() => Manifest.Save(SharedManifestFile), 3, TimeSpan.FromSeconds(5));
 				}
 
 				// Save the local manifest
@@ -882,7 +882,7 @@ namespace AutomationTool
 				bool bLocal = FileReference.Exists(LocalManifestFile);
 
 				// Read the manifest, either from local storage or shared storage
-				TempStorageManifest Manifest;
+				TempStorageManifest Manifest = null;
 				if(bLocal)
 				{
 					CommandUtils.LogInformation("Reading shared manifest from {0}", LocalManifestFile.FullName);
@@ -907,7 +907,7 @@ namespace AutomationTool
 
 					// Read the shared manifest
 					CommandUtils.LogInformation("Copying shared manifest from {0} to {1}", SharedManifestFile.FullName, LocalManifestFile.FullName);
-					Manifest = TempStorageManifest.Load(SharedManifestFile);
+					PerformActionWithRetries(() => Manifest = TempStorageManifest.Load(SharedManifestFile), 3, TimeSpan.FromSeconds(5));
 
 					// Unzip all the build products
 					DirectoryReference SharedNodeDir = GetDirectoryForNode(SharedDir, NodeName);
@@ -947,6 +947,26 @@ namespace AutomationTool
 			}
 		}
 
+		static void PerformActionWithRetries(Action retryAction, int retryCount, TimeSpan waitTime)
+		{
+			while (retryCount-- > 0)
+			{
+				try
+				{
+					retryAction();
+				}
+				catch
+				{
+					if (retryCount == 0)
+					{
+						throw;
+					}
+
+					Thread.Sleep(waitTime);
+				}
+			}
+		}
+
 		/// <summary>
 		/// Zips a set of files (that must be rooted at the given RootDir) to a set of zip files in the given OutputDir. The files will be prefixed with the given basename.
 		/// </summary>
@@ -960,7 +980,7 @@ namespace AutomationTool
 		/// This function tries to zip the files in parallel as fast as it can. It makes no guarantees about how many zip files will be created or which files will be in which zip,
 		/// but it does try to reasonably balance the file sizes.
 		/// </remarks>
-		private static FileInfo[] ParallelZipFiles(FileInfo[] InputFiles, DirectoryReference RootDir, DirectoryReference OutputDir, DirectoryReference StagingDir, string ZipBaseName)
+		static FileInfo[] ParallelZipFiles(FileInfo[] InputFiles, DirectoryReference RootDir, DirectoryReference OutputDir, DirectoryReference StagingDir, string ZipBaseName)
 		{
 			// First get the sizes of all the files. We won't parallelize if there isn't enough data to keep the number of zips down.
 			var FilesInfo = InputFiles
@@ -969,11 +989,11 @@ namespace AutomationTool
 
 			// Profiling results show that we can zip 100MB quite fast and it is not worth parallelizing that case and creating a bunch of zips that are relatively small.
 			const long MinFileSizeToZipInParallel = 1024 * 1024 * 100L;
-			var bZipInParallel = FilesInfo.Sum(FileInfo => FileInfo.FileSize) >= MinFileSizeToZipInParallel;
+			bool bZipInParallel = FilesInfo.Sum(FileInfo => FileInfo.FileSize) >= MinFileSizeToZipInParallel;
 
 			// order the files in descending order so our threads pick up the biggest ones first.
 			// We want to end with the smaller files to more effectively fill in the gaps
-			var FilesToZip = new ConcurrentQueue<FileReference>(FilesInfo.OrderByDescending(FileInfo => FileInfo.FileSize).Select(FileInfo => FileInfo.File));
+			ConcurrentQueue<FileReference> FilesToZip = new(FilesInfo.OrderByDescending(FileInfo => FileInfo.FileSize).Select(FileInfo => FileInfo.File));
 
 			ConcurrentBag<FileInfo> ZipFiles = new ConcurrentBag<FileInfo>();
 
@@ -989,7 +1009,7 @@ namespace AutomationTool
 				select new Thread((object indexObject) =>
 				{
 					int index = (int)indexObject;
-					var ZipFileName = FileReference.Combine(ZipDir, string.Format("{0}{1}.zip", ZipBaseName, bZipInParallel ? "-" + index.ToString("00") : ""));
+					FileReference ZipFileName = FileReference.Combine(ZipDir, string.Format("{0}{1}.zip", ZipBaseName, bZipInParallel ? "-" + index.ToString("00") : ""));
 					// don't create the zip unless we have at least one file to add
 					FileReference File;
 					if (FilesToZip.TryDequeue(out File))
@@ -997,7 +1017,7 @@ namespace AutomationTool
 						try
 						{
 							// Create one zip per thread using the given basename
-							using (var ZipArchive = ZipFile.Open(ZipFileName.FullName, ZipArchiveMode.Create))
+							using (ZipArchive ZipArchive = ZipFile.Open(ZipFileName.FullName, ZipArchiveMode.Create))
 							{
 								// pull from the queue until we are out of files.
 								do
@@ -1055,7 +1075,7 @@ namespace AutomationTool
 		/// <remarks>
 		/// The code is expected to be the used as the symmetrical inverse of <see cref="ParallelZipFiles"/>, but could be used independently, as long as the files in the zip do not overlap.
 		/// </remarks>
-		private static void ParallelUnzipFiles(FileInfo[] ZipFiles, DirectoryReference RootDir)
+		static void ParallelUnzipFiles(FileInfo[] ZipFiles, DirectoryReference RootDir)
 		{
 			Parallel.ForEach(ZipFiles,
 				(ZipFile) =>
@@ -1080,7 +1100,7 @@ namespace AutomationTool
 								foreach (ZipArchiveEntry Entry in ZipArchive.Entries)
 								{
 									// Use CommandUtils.CombinePaths to ensure directory separators get converted correctly.
-									var ExtractedFilename = CommandUtils.CombinePaths(RootDir.FullName, Entry.FullName);
+									string ExtractedFilename = CommandUtils.CombinePaths(RootDir.FullName, Entry.FullName);
 
 									// Skip this if it's already been extracted.
 									if (ExtractedPaths.Contains(ExtractedFilename))
