Added support for sessionlabel to UAT. -sessionlabel was not being passed through from UAT commandline to the exe commandline.

diff --git a/Engine/Source/Programs/AutomationTool/AutomationUtils/ProjectParams.cs b/Engine/Source/Programs/AutomationTool/AutomationUtils/ProjectParams.cs
--- a/Engine/Source/Programs/AutomationTool/AutomationUtils/ProjectParams.cs
+++ b/Engine/Source/Programs/AutomationTool/AutomationUtils/ProjectParams.cs
@@ -413,6 +413,7 @@ namespace AutomationTool
 			this.Trace = InParams.Trace;
 			this.TraceHost = InParams.TraceHost;
 			this.TraceFile = InParams.TraceFile;
+			this.SessionLabel = InParams.SessionLabel;
 		}
 
 		/// <summary>
@@ -573,6 +574,7 @@ namespace AutomationTool
 			string Trace = null,
 			string TraceHost = null,
 			string TraceFile = null,
+			string SessionLabel = null,
 			ParamList<string> InMapsToRebuildLightMaps = null,
             ParamList<string> InMapsToRebuildHLOD = null,
             ParamList<string> TitleID = null
@@ -663,7 +665,7 @@ namespace AutomationTool
 			this.StageBaseReleasePaks = GetParamValueIfNotSpecified(Command, StageBaseReleasePaks, this.StageBaseReleasePaks, "StageBaseReleasePaks");
 			this.DiscVersion = ParseParamValueIfNotSpecified(Command, DiscVersion, "DiscVersion", String.Empty);
 			this.AdditionalCookerOptions = ParseParamValueIfNotSpecified(Command, AdditionalCookerOptions, "AdditionalCookerOptions", String.Empty);
-			
+		
 			DLCName = ParseParamValueIfNotSpecified(Command, DLCName, "DLCName", String.Empty);
 			if (!String.IsNullOrEmpty(DLCName))
 			{
@@ -976,7 +978,15 @@ namespace AutomationTool
 				}
 			}
 
-			if (ClientConfigsToBuild == null)
+			SessionLabel = Command.ParseParamValue("sessionlabel");
+
+			if (SessionLabel!=null)
+			{ 
+				this.SessionLabel += "-sessionlabel";
+				this.SessionLabel += "=" + SessionLabel;	
+			}
+
+				if (ClientConfigsToBuild == null)
 			{
 				if (Command != null)
 				{
@@ -2221,6 +2231,9 @@ namespace AutomationTool
 		[Help("tracefile", "The file where the trace will be recorded")]
 		public string TraceFile { get; set; }
 
+		[Help("sessionlabel", "A label to pass to analytics")]
+		public string SessionLabel { get; set; }
+
 		private List<SingleTargetProperties> DetectedTargets;
 		private Dictionary<UnrealTargetPlatform, ConfigHierarchy> LoadedEngineConfigs;
 		private Dictionary<UnrealTargetPlatform, ConfigHierarchy> LoadedGameConfigs;
diff --git a/Engine/Source/Programs/AutomationTool/Scripts/RunProjectCommand.Automation.cs b/Engine/Source/Programs/AutomationTool/Scripts/RunProjectCommand.Automation.cs
--- a/Engine/Source/Programs/AutomationTool/Scripts/RunProjectCommand.Automation.cs
+++ b/Engine/Source/Programs/AutomationTool/Scripts/RunProjectCommand.Automation.cs
@@ -665,6 +665,10 @@ namespace AutomationScripts
 			{
 				TempCmdLine += Params.TraceFile + " ";
 			}
+			if (!String.IsNullOrEmpty(Params.SessionLabel))
+			{
+				TempCmdLine += Params.SessionLabel + " ";
+			}
 
 			TempCmdLine += SC.StageTargetPlatform.GetLaunchExtraCommandLine(Params);
 
