Horde: Restore functionality to ignore systemic events downstream of specific fingerprinted events.

diff --git a/Engine/Source/Programs/AutomationTool/AutomationUtils/Matchers/SystemicEventMatcher.cs b/Engine/Source/Programs/AutomationTool/AutomationUtils/Matchers/SystemicEventMatcher.cs
--- a/Engine/Source/Programs/AutomationTool/AutomationUtils/Matchers/SystemicEventMatcher.cs
+++ b/Engine/Source/Programs/AutomationTool/AutomationUtils/Matchers/SystemicEventMatcher.cs
@@ -10,6 +10,8 @@ namespace AutomationUtils.Matchers
 {
 	class SystemicEventMatcher : ILogEventMatcher
 	{
+		static readonly Regex s_automationTool = new Regex(@"^\s*(?:BUILD FAILED|AutomationTool exiting with ExitCode=[1-9])");
+
 		static readonly Regex s_ddc = new Regex(@"LogDerivedDataCache: .*queries/writes will be limited");
 
 		static readonly Regex s_pdbUtil = new Regex(@"^\s*ERROR: Error: EC_OK");
@@ -19,6 +21,10 @@ namespace AutomationUtils.Matchers
 
 		public LogEventMatch? Match(ILogCursor cursor)
 		{
+			if (cursor.IsMatch(s_automationTool))
+			{
+				return new LogEventBuilder(cursor).ToMatch(LogEventPriority.High, LogLevel.Error, KnownLogEvents.Systemic_AutomationTool);
+			}
 			if (cursor.IsMatch(s_ddc))
 			{
 				return new LogEventBuilder(cursor).ToMatch(LogEventPriority.High, LogLevel.Information, KnownLogEvents.Systemic_SlowDDC);
diff --git a/Engine/Source/Programs/Horde/Horde.Build/Issues/IssueEvent.cs b/Engine/Source/Programs/Horde/Horde.Build/Issues/IssueEvent.cs
--- a/Engine/Source/Programs/Horde/Horde.Build/Issues/IssueEvent.cs
+++ b/Engine/Source/Programs/Horde/Horde.Build/Issues/IssueEvent.cs
@@ -63,6 +63,15 @@ namespace Horde.Build.Issues
 			EventData = stepEventData;
 		}
 
+		/// <summary>
+		/// Tests whether this is a systemic event id
+		/// </summary>
+		/// <returns>True if this is a systemic event id</returns>
+		public bool IsSystemic()
+		{
+			return EventId.HasValue && (EventId.Value.Id >= KnownLogEvents.Systemic.Id && EventId.Value.Id <= KnownLogEvents.Systemic_Max.Id);
+		}
+
 		/// <inheritdoc/>
 		public override string ToString()
 		{
diff --git a/Engine/Source/Programs/Shared/EpicGames.Core/KnownLogEvents.cs b/Engine/Source/Programs/Shared/EpicGames.Core/KnownLogEvents.cs
--- a/Engine/Source/Programs/Shared/EpicGames.Core/KnownLogEvents.cs
+++ b/Engine/Source/Programs/Shared/EpicGames.Core/KnownLogEvents.cs
@@ -210,6 +210,11 @@ namespace EpicGames.Core
 		/// </summary>
 		public static EventId Systemic_LogEventMatcher { get; } = new EventId(763);
 
+		/// <summary>
+		/// Generic "Build failed" messages from UAT
+		/// </summary>
+		public static EventId Systemic_AutomationTool = new EventId(770);
+
 		/// <summary>
 		/// Maximum code for systemic events. Add new events in the 700-799 range.
 		/// </summary>
