Added extra provisioning debugging

diff --git a/Engine/Source/Programs/AutomationTool/IOS/IOSPlatform.Automation.cs b/Engine/Source/Programs/AutomationTool/IOS/IOSPlatform.Automation.cs
--- a/Engine/Source/Programs/AutomationTool/IOS/IOSPlatform.Automation.cs
+++ b/Engine/Source/Programs/AutomationTool/IOS/IOSPlatform.Automation.cs
@@ -1155,6 +1155,8 @@
 							idx += "<string>".Length;
 							UUID = AllText.Substring(idx, AllText.IndexOf("</string>", idx) - idx);
 							Arguments += " PROVISIONING_PROFILE_SPECIFIER=" + UUID;
+
+							LogInformation("Extracted Provision UUID {0} from {1}", UUID, Provision);
 						}
 					}
 				}
