Fix TempStorage Manifest XML Writing

dotnet6 xml serialization nolonger emits new lines, change to using XMLWriter to enable newlines

#rnx

diff --git a/Engine/Source/Programs/AutomationTool/BuildGraph/TempStorage.cs b/Engine/Source/Programs/AutomationTool/BuildGraph/TempStorage.cs
--- a/Engine/Source/Programs/AutomationTool/BuildGraph/TempStorage.cs
+++ b/Engine/Source/Programs/AutomationTool/BuildGraph/TempStorage.cs
@@ -1,5 +1,8 @@
 // Copyright Epic Games, Inc. All Rights Reserved.
 
+using EpicGames.Core;
+using OpenTracing;
+using OpenTracing.Util;
 using System;
 using System.Collections.Concurrent;
 using System.Collections.Generic;
@@ -7,16 +10,10 @@ using System.Diagnostics;
 using System.IO;
 using System.IO.Compression;
 using System.Linq;
-using System.Text;
 using System.Threading;
 using System.Threading.Tasks;
+using System.Xml;
 using System.Xml.Serialization;
-using UnrealBuildTool;
-using AutomationTool;
-using EpicGames.Core;
-using OpenTracing;
-using OpenTracing.Util;
-using Microsoft.Extensions.Logging;
 
 namespace AutomationTool
 {
@@ -397,7 +394,11 @@ namespace AutomationTool
 		{
 			using(StreamWriter Writer = new StreamWriter(File.FullName))
 			{
-				Serializer.Serialize(Writer, this);
+				XmlWriterSettings WriterSettings = new() { Indent = true };
+				using (XmlWriter XMLWriter = XmlWriter.Create(Writer, WriterSettings))
+				{
+					Serializer.Serialize(XMLWriter, this);
+				}
 			}
 		}
 	}
@@ -485,9 +486,13 @@ namespace AutomationTool
 		/// <param name="File">File to save</param>
 		public void Save(FileReference File)
 		{
-			using(StreamWriter Writer = new StreamWriter(File.FullName))
+			using (StreamWriter Writer = new StreamWriter(File.FullName))
 			{
-				Serializer.Serialize(Writer, this);
+				XmlWriterSettings WriterSettings = new() { Indent = true };
+				using (XmlWriter XMLWriter = XmlWriter.Create(Writer, WriterSettings))
+				{
+					Serializer.Serialize(XMLWriter, this);
+				}
 			}
 		}
 
