Update copy task to exit gracefully if there aren't any files to copy when overwrite=false

diff --git a/Engine/Source/Programs/AutomationTool/BuildGraph/Tasks/CopyTask.cs b/Engine/Source/Programs/AutomationTool/BuildGraph/Tasks/CopyTask.cs
--- a/Engine/Source/Programs/AutomationTool/BuildGraph/Tasks/CopyTask.cs
+++ b/Engine/Source/Programs/AutomationTool/BuildGraph/Tasks/CopyTask.cs
@@ -136,15 +136,22 @@ namespace AutomationTool.Tasks
 			//  If we're not overwriting, remove any files where the destination file already exists.
 			if (!Overwrite)
 			{
-				TargetFileToSourceFile = TargetFileToSourceFile.Where(File =>
+				Dictionary<FileReference, FileReference> FilteredTargetToSourceFile = new Dictionary<FileReference, FileReference>();
+				foreach (KeyValuePair<FileReference, FileReference> File in TargetFileToSourceFile)
 				{
 					if (FileReference.Exists(File.Key))
 					{
 						CommandUtils.LogInformation("Not copying existing file {0}", File.Key);
-						return false;
+						continue;
 					}
-					return true;
-				}).ToDictionary(Pair => Pair.Key, Pair => Pair.Value);
+					FilteredTargetToSourceFile.Add(File.Key, File.Value);
+				}
+				if(FilteredTargetToSourceFile.Count == 0)
+				{
+					CommandUtils.LogWarning("All files already exist, exiting early.");
+					return Task.CompletedTask;
+				}
+				TargetFileToSourceFile = FilteredTargetToSourceFile;
 			}
 
 			// If the target is on a network share, retry creating the first directory until it succeeds
