CopyBuildToStaging.cs: Add error for staging without IoStore after cooking with ZenStore

#rnx

diff --git a/Engine/Source/Programs/AutomationTool/Scripts/CopyBuildToStagingDirectory.Automation.cs b/Engine/Source/Programs/AutomationTool/Scripts/CopyBuildToStagingDirectory.Automation.cs
--- a/Engine/Source/Programs/AutomationTool/Scripts/CopyBuildToStagingDirectory.Automation.cs
+++ b/Engine/Source/Programs/AutomationTool/Scripts/CopyBuildToStagingDirectory.Automation.cs
@@ -4065,6 +4065,14 @@ namespace AutomationScripts
 		{
 			if (ShouldCreatePak(Params, SC))
 			{
+				if (!ShouldCreateIoStoreContainerFiles(Params, SC))
+				{
+					FileReference ZenFileManifest = FileReference.Combine(SC.MetadataDir, "zenfs.manifest");
+					if (FileReference.Exists(ZenFileManifest))
+					{
+						throw new AutomationException(String.Format("A Zen file manifest must not exist when staging without IoStore. Did not expected to find {0}. Ensure that legacy cooking with ZenStore disabled (bUseZenStore=false) was successful.", ZenFileManifest.FullName));
+					}
+				}
 				if (SC.CrashReporterUFSFiles.Count > 0)
 				{
 					CreatePakForCrashReporter(Params, SC);
