Build Automation Fixes.
Fix DLCOverrideCookedSubDir  param being used for DLCOverrideStagedSubDir if DLCOverrideStagedSubDir was not specified via commandline.
Don't allow TargetPlatformDescriptor.CookFlavor to be null - it needs to be empty string to be a valid key.

diff --git a/Engine/Source/Programs/AutomationTool/AutomationUtils/Platform.cs b/Engine/Source/Programs/AutomationTool/AutomationUtils/Platform.cs
--- a/Engine/Source/Programs/AutomationTool/AutomationUtils/Platform.cs
+++ b/Engine/Source/Programs/AutomationTool/AutomationUtils/Platform.cs
@@ -80,7 +80,7 @@
 		public TargetPlatformDescriptor(UnrealTargetPlatform InType, string InCookFlavor)
 		{
 			Type = InType;
-			CookFlavor = InCookFlavor;
+			CookFlavor = InCookFlavor ?? "";
 		}
 
 		public override string ToString()
diff --git a/Engine/Source/Programs/AutomationTool/AutomationUtils/ProjectParams.cs b/Engine/Source/Programs/AutomationTool/AutomationUtils/ProjectParams.cs
--- a/Engine/Source/Programs/AutomationTool/AutomationUtils/ProjectParams.cs
+++ b/Engine/Source/Programs/AutomationTool/AutomationUtils/ProjectParams.cs
@@ -694,7 +694,7 @@ namespace AutomationTool
 			this.DLCPakPluginFile = GetParamValueIfNotSpecified(Command, DLCPakPluginFile, this.DLCPakPluginFile, "DLCPakPluginFile");
 			this.DLCActLikePatch = GetParamValueIfNotSpecified(Command, DLCActLikePatch, this.DLCActLikePatch, "DLCActLikePatch");
 			this.DLCOverrideCookedSubDir = ParseParamValueIfNotSpecified(Command, DLCOverrideCookedSubDir, "DLCOverrideCookedSubDir", null);
-			this.DLCOverrideStagedSubDir = ParseParamValueIfNotSpecified(Command, DLCOverrideCookedSubDir, "DLCOverrideStagedSubDir", null);
+			this.DLCOverrideStagedSubDir = ParseParamValueIfNotSpecified(Command, DLCOverrideStagedSubDir, "DLCOverrideStagedSubDir", null);
 
 			this.SkipCook = GetParamValueIfNotSpecified(Command, SkipCook, this.SkipCook, "skipcook");
 			if (this.SkipCook)
