Automation can override the RHI & other settings for all Windows-based platforms.

#rnx

diff --git a/Engine/Source/Programs/AutomationTool/Gauntlet/Unreal/Automation/UE.Automation.cs b/Engine/Source/Programs/AutomationTool/Gauntlet/Unreal/Automation/UE.Automation.cs
--- a/Engine/Source/Programs/AutomationTool/Gauntlet/Unreal/Automation/UE.Automation.cs
+++ b/Engine/Source/Programs/AutomationTool/Gauntlet/Unreal/Automation/UE.Automation.cs
@@ -241,7 +241,7 @@ namespace UE
 				}
 
 				// Options specific to windows
-				if (ConfigRole.Platform == UnrealTargetPlatform.Win64)
+				if (ConfigRole.Platform != null && ((UnrealTargetPlatform)ConfigRole.Platform).IsInGroup(UnrealPlatformGroup.Windows))
 				{
 					if (PreferNvidia)
 					{
@@ -666,7 +666,7 @@ namespace UE
 							if (Report != null)
 							{
 								var MainRolePlatform = Context.GetRoleContext(Config.GetMainRequiredRole().Type).Platform;
-								Report.SetMetadata("RHI", string.IsNullOrEmpty(Config.RHI) || MainRolePlatform != UnrealTargetPlatform.Win64 ? "default" : Config.RHI.ToLower());
+								Report.SetMetadata("RHI", string.IsNullOrEmpty(Config.RHI) || !MainRolePlatform.IsInGroup(UnrealPlatformGroup.Windows) ? "default" : Config.RHI.ToLower());
 							}
 						}
 					}
