Safe guard against massive lists of modified files causing a hang in AutomationTool.

#rnx

diff --git a/Engine/Source/Programs/AutomationTool/BuildGraph/BuildGraph.cs b/Engine/Source/Programs/AutomationTool/BuildGraph/BuildGraph.cs
--- a/Engine/Source/Programs/AutomationTool/BuildGraph/BuildGraph.cs
+++ b/Engine/Source/Programs/AutomationTool/BuildGraph/BuildGraph.cs
@@ -1246,7 +1246,17 @@ namespace AutomationTool
 			}
 			if(ModifiedFiles.Count > 0)
 			{
-				throw new AutomationException("Build {0} from a previous step have been modified:\n{1}", (ModifiedFiles.Count == 1)? "product" : "products", String.Join("\n", ModifiedFiles.Select(x => x.Value)));
+				string modifiedFileList = "";
+				if (ModifiedFiles.Count < 100)
+				{
+					modifiedFileList = String.Join("\n", ModifiedFiles.Select(x => x.Value));
+				}
+				else
+				{
+					modifiedFileList = String.Join("\n", ModifiedFiles.Take(100).Select(x => x.Value));
+					modifiedFileList += $"{Environment.NewLine}And {ModifiedFiles.Count - 100} more.";
+				}
+				throw new AutomationException("Build {0} from a previous step have been modified:\n{1}", (ModifiedFiles.Count == 1)? "product" : "products", modifiedFileList);
 			}
 
 			// Determine all the output files which are required to be copied to temp storage (because they're referenced by nodes in another agent)
