[Automation] Use SrcFile as the regex match to ensure we copy to the correct destination path on Android

#tests preflight replay run and ensure it copies the build and launches on android devices.

diff --git a/Engine/Source/Programs/AutomationTool/Gauntlet/Platform/Android/Gauntlet.TargetDeviceAndroid.cs b/Engine/Source/Programs/AutomationTool/Gauntlet/Platform/Android/Gauntlet.TargetDeviceAndroid.cs
--- a/Engine/Source/Programs/AutomationTool/Gauntlet/Platform/Android/Gauntlet.TargetDeviceAndroid.cs
+++ b/Engine/Source/Programs/AutomationTool/Gauntlet/Platform/Android/Gauntlet.TargetDeviceAndroid.cs
@@ -982,17 +982,17 @@ namespace Gauntlet
 
 				Build.FilesToInstall.Keys.ToList().ForEach(K =>
 				{
-
 					string SrcPath = K;
-					string DestPath = Build.FilesToInstall[K];
+					string SrcFile = Path.GetFileName(SrcPath);
 
+					string DestPath = Build.FilesToInstall[K];
 					string DestFile = Path.GetFileName(DestPath);
 
 					// If we installed a new APK we need to change the package version
-					Match OBBMatch = Regex.Match(DestFile, @"\.(\d+)\.com.*\.obb");
+					Match OBBMatch = Regex.Match(SrcFile, @"\.(\d+)\.com.*\.obb");
 					if (OBBMatch.Success)
 					{
-						DestPath = StorageLocation + "/obb/" + Build.AndroidPackageName + "/" + Path.GetFileName(SrcPath).Replace(".Client.obb", ".obb").Replace(OBBMatch.Groups[1].ToString(), PackageVersion);
+						DestPath = StorageLocation + "/obb/" + Build.AndroidPackageName + "/" + SrcFile.Replace(".Client.obb", ".obb").Replace(OBBMatch.Groups[1].ToString(), PackageVersion);
 					}
 
 					DestPath = Regex.Replace(DestPath, "%STORAGE%", StorageLocation, RegexOptions.IgnoreCase);
