Exclude CSVTools from CheckCSProjectDotNetVersion

Standalone, not part of any build process.

#rnx

diff --git a/Engine/Source/Programs/AutomationTool/Scripts/CheckCsprojDotNetVersion.cs b/Engine/Source/Programs/AutomationTool/Scripts/CheckCsprojDotNetVersion.cs
--- a/Engine/Source/Programs/AutomationTool/Scripts/CheckCsprojDotNetVersion.cs
+++ b/Engine/Source/Programs/AutomationTool/Scripts/CheckCsprojDotNetVersion.cs
@@ -39,7 +39,8 @@ class CheckCsprojDotNetVersion : BuildCommand
                 CsProj.GetFileNameWithoutExtension().ToLower().Contains("_mono") ||
                 CsProj.GetFileNameWithoutExtension().ToLower().Contains("unrealvs") ||
 				CsProj.ContainsName("DatasmithRevitExporter", EngineDir) ||
-				CsProj.ContainsName("DatasmithNavisworksExporter", EngineDir))
+				CsProj.ContainsName("DatasmithNavisworksExporter", EngineDir) ||
+				CsProj.ContainsName("CSVTools", EngineDir))
 
             {
                 continue;
