Fix UnrealTestResults output not being camel case when resaving.
+ Move the critical failure check from UnrealTestNode to UE.AutomationNodeBase to make it work for non-Horde specific test report.

#rnx

diff --git a/Engine/Source/Programs/AutomationTool/Gauntlet/Unreal/Automation/UE.Automation.cs b/Engine/Source/Programs/AutomationTool/Gauntlet/Unreal/Automation/UE.Automation.cs
--- a/Engine/Source/Programs/AutomationTool/Gauntlet/Unreal/Automation/UE.Automation.cs
+++ b/Engine/Source/Programs/AutomationTool/Gauntlet/Unreal/Automation/UE.Automation.cs
@@ -565,7 +565,7 @@ namespace UE
 					}
 				}
 				
-				if(TestPassResults == null)
+				if(TestPassResults == null && InLog != null)
 				{
 					// Parse automaton info from the log then
 					TestPassResults = new UnrealAutomatedTestPassResults();
@@ -599,6 +599,127 @@ namespace UE
 			return TestPassResults;
 		}
 
+		/// <summary>
+		/// Look for critical failure during the test session and update the test states
+		/// </summary>
+		/// <param name="JsonTestPassResults"></param>
+		private void UpdateTestStateOnCriticalFailure(UnrealAutomatedTestPassResults JsonTestPassResults)
+		{
+			bool HasTimeout = RoleResults != null && RoleResults.Where(R => R.ProcessResult == UnrealProcessResult.TimeOut).Any();
+			string HordeArtifactPath = GetConfiguration().HordeArtifactPath;
+			var MainRole = GetConfiguration().GetMainRequiredRole();
+			string LastTestWithCriticalFailure = string.Empty;
+			if (JsonTestPassResults.InProcess > 0)
+			{
+				// The test pass did not run completely
+				Log.Verbose("Found in-process tests: {0}", JsonTestPassResults.InProcess);
+				// Get any critical error and push it to json report and resave it.
+				if (RoleResults != null)
+				{
+					UnrealLog.CallstackMessage FatalError = null;
+					foreach (var Result in RoleResults)
+					{
+						if (Result.LogSummary.FatalError != null)
+						{
+							FatalError = Result.LogSummary.FatalError;
+							break;
+						}
+					}
+					var Test = JsonTestPassResults.Tests.FirstOrDefault((T => T.State == TestStateType.InProcess));
+					if (!String.IsNullOrEmpty(Test.TestDisplayName))
+					{
+						LastTestWithCriticalFailure = Test.FullTestPath;
+						string ErrorMessage;
+						if (HasTimeout)
+						{
+							ErrorMessage = String.Format("Session reached timeout after {0} seconds.", MaxDuration);
+						}
+						else
+						{
+							ErrorMessage = "Engine encountered a critical failure. \n";
+							if (FatalError != null)
+							{
+								ErrorMessage += FatalError.FormatForLog();
+							}
+							else
+							{
+								ErrorMessage += "No callstack found in the log.";
+							}
+						}
+						Test.AddError(ErrorMessage);
+						if (!CanRetry() || JsonTestPassResults.NotRun == 0)
+						{
+							// Setting the test as fail because no retry will be done anymore.
+							// The InProcess state won't be used for pass resume
+							Test.State = TestStateType.Fail;
+							if (!CanRetry())
+							{
+								Test.AddWarning(string.Format("Session reached maximum of retries({0}) to resume on critical failure!", Retries));
+							}
+						}
+						JsonTestPassResults.WriteToJson();
+					}
+				}
+			}
+			if (JsonTestPassResults.NotRun > 0)
+			{
+				// The test pass did not run at all
+				Log.Verbose("Found not-run tests: {0}", JsonTestPassResults.NotRun);
+				if (GetConfiguration().ResumeOnCriticalFailure && !HasTimeout)
+				{
+					// Reschedule test to resume from last 'in-process' test.
+					if (SetToRetryIfPossible())
+					{
+						// Attach current artifacts to Horde output
+						if (SessionArtifacts != null && !string.IsNullOrEmpty(HordeArtifactPath))
+						{
+							HordeReport.SimpleTestReport TempReport = new HordeReport.SimpleTestReport();
+							TempReport.SetOutputArtifactPath(HordeArtifactPath);
+							foreach (UnrealRoleArtifacts Artifact in SessionArtifacts)
+							{
+								string LogName = Path.GetFullPath(Artifact.LogPath).Replace(Path.GetFullPath(Context.Options.LogDir), "").TrimStart(Path.DirectorySeparatorChar);
+								TempReport.AttachArtifact(Artifact.LogPath, LogName);
+								// Reference last run instance log
+								if (Artifact.SessionRole.RoleType == MainRole.Type)
+								{
+									JsonTestPassResults.Devices.Last().AppInstanceLog = LogName.Replace("\\", "/");
+								}
+							}
+							JsonTestPassResults.WriteToJson();
+						}
+						// Discard the report as we are going to do another pass.
+					}
+					else
+					{
+						Log.Error("Reach maximum of retries({0}) to resume on critical failure!", Retries);
+						// Adding a note to the report about why the not-run are not going to be run
+						string Message = string.Format("Session reached maximum of retries({0}) to resume on critical failure!", Retries);
+						if (!string.IsNullOrEmpty(LastTestWithCriticalFailure))
+						{
+							Message += string.Format(" \nLast critical failure was caught on {0}", LastTestWithCriticalFailure);
+						}
+						var NotRunTests = JsonTestPassResults.Tests.Where((T => T.State == TestStateType.NotRun));
+						foreach (var Test in NotRunTests)
+						{
+							Test.AddInfo(Message);
+						}
+						JsonTestPassResults.WriteToJson();
+					}
+				}
+				else if (HasTimeout)
+				{
+					// Adding a note to the report about why the not-run are not going to be run
+					string Message = string.Format("Session reached timeout after {0} seconds.", MaxDuration);
+					var NotRunTests = JsonTestPassResults.Tests.Where((T => T.State == TestStateType.NotRun));
+					foreach (var Test in NotRunTests)
+					{
+						Test.AddInfo(Message);
+					}
+					JsonTestPassResults.WriteToJson();
+				}
+			}
+		}
+
 		/// <summary>
 		/// Override GetExitCodeAndReason to provide additional checking of success / failure based on what occurred
 		/// </summary>
@@ -646,9 +767,14 @@ namespace UE
 		public override ITestReport CreateReport(TestResult Result)
 		{
 			ITestReport Report = null;
-			if (GetConfiguration() is AutomationTestConfig)
+			if (GetConfiguration() is AutomationTestConfig Config)
 			{
-				var Config = GetConfiguration() as AutomationTestConfig;
+				// Parse the json test pass results and look for critical failure to add to report
+				var TestPassResults = GetTestPassResults(null);
+				if (TestPassResults != null)
+				{
+					UpdateTestStateOnCriticalFailure(TestPassResults);
+				}
 				// Save test result data for Horde build system
 				bool WriteTestResultsForHorde = Config.WriteTestResultsForHorde;
 				if (WriteTestResultsForHorde)
diff --git a/Engine/Source/Programs/AutomationTool/Gauntlet/Unreal/Base/Gauntlet.UnrealTestNode.cs b/Engine/Source/Programs/AutomationTool/Gauntlet/Unreal/Base/Gauntlet.UnrealTestNode.cs
--- a/Engine/Source/Programs/AutomationTool/Gauntlet/Unreal/Base/Gauntlet.UnrealTestNode.cs
+++ b/Engine/Source/Programs/AutomationTool/Gauntlet/Unreal/Base/Gauntlet.UnrealTestNode.cs
@@ -264,8 +264,21 @@ namespace Gauntlet
 
 		static protected DateTime SessionStartTime = DateTime.MinValue;
 
-		private int Retries = 0;
+		public int Retries { get; private set; } = 0;
 		private int MaxRetries = 3;
+		public virtual bool CanRetry() { return Retries < MaxRetries; }
+		public virtual bool SetToRetryIfPossible()
+		{
+			if (!CanRetry()) { return false; }
+			if (IsSetToRetry()) { return true; }
+			Retries++;
+			SetUnrealTestResult(TestResult.WantRetry);
+			return true;
+		}
+		public virtual bool IsSetToRetry()
+		{
+			return GetTestResult() == TestResult.WantRetry;
+		}
 
 		/// <summary>
 		/// Standard semantic versioning for tests. Should be overwritten within individual tests, and individual test maintainers
@@ -1250,7 +1263,7 @@ namespace Gauntlet
 		/// <returns>ITestReport</returns>
 		public virtual ITestReport CreateUnrealEngineTestPassReport(string UnrealAutomatedTestReportPath, string ReportURL)
 		{
-			bool HasTimeout = RoleResults != null && RoleResults.Where(R => R.ProcessResult == UnrealProcessResult.TimeOut).Any();
+			if(IsSetToRetry()) { return null; }
 			string JsonReportPath = Path.Combine(UnrealAutomatedTestReportPath, "index.json");
 			if (File.Exists(JsonReportPath))
 			{
@@ -1258,119 +1271,6 @@ namespace Gauntlet
 				Log.Verbose("Reading json Unreal Automated test report from {0}", JsonReportPath);
 				UnrealAutomatedTestPassResults JsonTestPassResults = UnrealAutomatedTestPassResults.LoadFromJson(JsonReportPath);
 				var MainRole = GetCachedConfiguration().GetMainRequiredRole();
-				string LastTestWithCriticalFailure = string.Empty;
-				if (JsonTestPassResults.InProcess > 0)
-				{
-					// The test pass did not run completely
-					Log.Verbose("Found in-process tests: {0}", JsonTestPassResults.InProcess);
-					// Get any critical error and push it to json report and resave it.
-					if (RoleResults != null)
-					{
-						UnrealLog.CallstackMessage FatalError = null;
-						foreach (UnrealRoleResult Result in RoleResults)
-						{
-							if (Result.LogSummary.FatalError != null)
-							{
-								FatalError = Result.LogSummary.FatalError;
-								break;
-							}
-						}
-						var Test = JsonTestPassResults.Tests.FirstOrDefault((T => T.State == TestStateType.InProcess));
-						if (!String.IsNullOrEmpty(Test.TestDisplayName))
-						{
-							LastTestWithCriticalFailure = Test.FullTestPath;
-							string ErrorMessage;
-							if (HasTimeout)
-							{
-								ErrorMessage = String.Format("Session reached timeout after {0} seconds.", MaxDuration);
-							}
-							else
-							{
-								ErrorMessage = "Engine encountered a critical failure. \n";
-								if (FatalError != null)
-								{
-									ErrorMessage += FatalError.FormatForLog();
-								}
-								else
-								{
-									ErrorMessage += "No callstack found in the log.";
-								}
-							}
-							Test.AddError(ErrorMessage);
-							if (Retries >= MaxRetries || JsonTestPassResults.NotRun == 0)
-							{
-								// Setting the test as fail because no retry will be done anymore.
-								// The InProcess state won't be used for pass resume
-								Test.State = TestStateType.Fail;
-								if (Retries >= MaxRetries)
-								{
-									Test.AddWarning(string.Format("Session reached maximum of retries({0}) to resume on critical failure!", Retries));
-								}
-							}
-							JsonTestPassResults.WriteToJson(JsonReportPath);
-						}
-					}
-				}
-				if (JsonTestPassResults.NotRun > 0)
-				{
-					// The test pass did not run at all
-					Log.Verbose("Found not-run tests: {0}", JsonTestPassResults.NotRun);
-					if (GetCachedConfiguration().ResumeOnCriticalFailure && !HasTimeout)
-					{
-						if (Retries < MaxRetries)
-						{
-							Retries++;
-							// Reschedule test to resume from last 'in-process' test.
-							SetUnrealTestResult(TestResult.WantRetry);
-							// Attach current artifacts to Horde output
-							if (SessionArtifacts != null)
-							{
-								HordeReport.SimpleTestReport TempReport = new HordeReport.SimpleTestReport();
-								TempReport.SetOutputArtifactPath(HordeArtifactPath);
-								foreach (UnrealRoleArtifacts Artifact in SessionArtifacts)
-								{
-									string LogName = Path.GetFullPath(Artifact.LogPath).Replace(Path.GetFullPath(Context.Options.LogDir), "").TrimStart(Path.DirectorySeparatorChar);
-									TempReport.AttachArtifact(Artifact.LogPath, LogName);
-									// Reference last run instance log
-									if(Artifact.SessionRole.RoleType == MainRole.Type)
-									{
-										JsonTestPassResults.Devices.Last().AppInstanceLog = LogName.Replace("\\", "/");
-									}
-								}
-								JsonTestPassResults.WriteToJson(JsonReportPath);
-							}
-							// Discard the report as we are going to do another pass.
-							return null;
-						}
-						else
-						{
-							Log.Error("Reach maximum of retries({0}) to resume on critical failure!", Retries);
-							// Adding a note to the report about why the not-run are not going to be run
-							string Message = string.Format("Session reached maximum of retries({0}) to resume on critical failure!", Retries);
-							if (!string.IsNullOrEmpty(LastTestWithCriticalFailure))
-							{
-								Message += string.Format(" \nLast critical failure was caught on {0}", LastTestWithCriticalFailure);
-							}
-							var NotRunTests = JsonTestPassResults.Tests.Where((T => T.State == TestStateType.NotRun));
-							foreach(var Test in NotRunTests)
-							{
-								Test.AddInfo(Message);
-							}
-							JsonTestPassResults.WriteToJson(JsonReportPath);
-						}
-					}
-					else if(HasTimeout)
-					{
-						// Adding a note to the report about why the not-run are not going to be run
-						string Message = string.Format("Session reached timeout after {0} seconds.", MaxDuration);
-						var NotRunTests = JsonTestPassResults.Tests.Where((T => T.State == TestStateType.NotRun));
-						foreach (var Test in NotRunTests)
-						{
-							Test.AddInfo(Message);
-						}
-						JsonTestPassResults.WriteToJson(JsonReportPath);
-					}
-				}
 				string HordeTestDataKey = string.IsNullOrEmpty(GetCachedConfiguration().HordeTestDataKey) ? Name + " " + Context.ToString() : GetCachedConfiguration().HordeTestDataKey;
 				GetCachedConfiguration().HordeTestDataKey = HordeTestDataKey;
 				// Convert test results for Horde
diff --git a/Engine/Source/Programs/AutomationTool/Gauntlet/Unreal/Utils/Gauntlet.UnrealTestResults.cs b/Engine/Source/Programs/AutomationTool/Gauntlet/Unreal/Utils/Gauntlet.UnrealTestResults.cs
--- a/Engine/Source/Programs/AutomationTool/Gauntlet/Unreal/Utils/Gauntlet.UnrealTestResults.cs
+++ b/Engine/Source/Programs/AutomationTool/Gauntlet/Unreal/Utils/Gauntlet.UnrealTestResults.cs
@@ -6,41 +6,64 @@ using System.Collections.Generic;
 using System.Text.Json;
 using System.Text.Json.Serialization;
 using System.Linq;
+using AutomationTool;
 
 namespace Gauntlet
 {
 	public class UnrealAutomationDevice
 	{
+		[JsonPropertyName("deviceName")]
 		public string DeviceName { get; set; }
+		[JsonPropertyName("instance")] 
 		public string Instance { get; set; }
+		[JsonPropertyName("platform")]
 		public string Platform { get; set; }
+		[JsonPropertyName("oSversion")]
 		public string OSVersion { get; set; }
+		[JsonPropertyName("model")]
 		public string Model { get; set; }
+		[JsonPropertyName("gPU")]
 		public string GPU { get; set; }
+		[JsonPropertyName("cPUModel")]
 		public string CPUModel { get; set; }
+		[JsonPropertyName("rAMInGB")]
 		public int RAMInGB { get; set; }
+		[JsonPropertyName("renderMode")]
 		public string RenderMode { get; set; }
+		[JsonPropertyName("rHI")]
 		public string RHI { get; set; }
+		[JsonPropertyName("appInstanceLog")]
 		public string AppInstanceLog { get; set; }
 	}
 	public class UnrealAutomationComparisonFiles
 	{
+		[JsonPropertyName("difference")]
 		public string Difference { get; set; }
+		[JsonPropertyName("approved")]
 		public string Approved { get; set; }
+		[JsonPropertyName("unapproved")]
 		public string Unapproved { get; set; }
 	}
 	public class UnrealAutomationArtifact
 	{
+		[JsonPropertyName("id")]
 		public string Id { get; set; }
+		[JsonPropertyName("name")]
 		public string Name { get; set; }
+		[JsonPropertyName("type")]
 		public string Type { get; set; }
+		[JsonPropertyName("files")]
 		public UnrealAutomationComparisonFiles Files { get; set; }
 	}
 	public class UnrealAutomationEvent
 	{
+		[JsonPropertyName("type")]
 		public EventType Type { get; set; }
+		[JsonPropertyName("message")]
 		public string Message { get; set; }
+		[JsonPropertyName("context")]
 		public string Context { get; set; }
+		[JsonPropertyName("artifact")]
 		public string Artifact { get; set; }
 
 		public UnrealAutomationEvent()
@@ -73,23 +96,38 @@ namespace Gauntlet
 	}
 	public class UnrealAutomationEntry
 	{
+		[JsonPropertyName("event")]
 		public UnrealAutomationEvent Event { get; set; }
+		[JsonPropertyName("filename")]
 		public string Filename { get; set; }
+		[JsonPropertyName("lineNumber")]
 		public int LineNumber { get; set; }
+		[JsonPropertyName("timeStamp")]
 		public string Timestamp { get; set; }
 	}
 	public class UnrealAutomatedTestResult
 	{
+		[JsonPropertyName("testDisplayName")]
 		public string TestDisplayName { get; set; }
+		[JsonPropertyName("fullTestPath")]
 		public string FullTestPath { get; set; }
+		[JsonPropertyName("artifactName")]
 		public string ArtifactName { get; set; }
+		[JsonPropertyName("state")]
 		public TestStateType State { get; set; }
+		[JsonPropertyName("deviceInstance")]
 		public List<string> DeviceInstance { get; set; }
+		[JsonPropertyName("duration")]
 		public float Duration { get; set; }
+		[JsonPropertyName("dateTime")]
 		public string DateTime { get; set; }
+		[JsonPropertyName("warnings")]
 		public int Warnings { get; set; }
+		[JsonPropertyName("errors")]
 		public int Errors { get; set; }
+		[JsonPropertyName("artifacts")]
 		public List<UnrealAutomationArtifact> Artifacts { get; set; }
+		[JsonPropertyName("entries")]
 		public List<UnrealAutomationEntry> Entries { get; set; }
 
 		public UnrealAutomatedTestResult()
@@ -196,18 +234,36 @@ namespace Gauntlet
 	}
 	public class UnrealAutomatedTestPassResults
 	{
+		[JsonPropertyName("devices")]
 		public List<UnrealAutomationDevice> Devices { get; set; }
+		[JsonPropertyName("reportCreatedOn")]
 		public string ReportCreatedOn { get; set; }
+		[JsonPropertyName("succeeded")]
 		public int Succeeded { get; set; }
+		[JsonPropertyName("succeededWithWarnings")]
 		public int SucceededWithWarnings { get; set; }
+		[JsonPropertyName("failed")]
 		public int Failed { get; set; }
+		[JsonPropertyName("notRun")]
 		public int NotRun { get; set; }
+		[JsonPropertyName("inProcess")]
 		public int InProcess { get; set; }
+		[JsonPropertyName("totalDuration")]
 		public float TotalDuration { get; set; }
+		[JsonPropertyName("comparisonExported")]
 		public bool ComparisonExported { get; set; }
+		[JsonPropertyName("comparisonExportDirectory")]
 		public string ComparisonExportDirectory { get; set; }
+		[JsonPropertyName("tests")]
 		public List<UnrealAutomatedTestResult> Tests { get; set; }
 
+		private string FilePath;
+
+		public void SetFilePath(string InFilePath)
+		{
+			FilePath = InFilePath;
+		}
+
 		public UnrealAutomatedTestPassResults()
 		{
 			Devices = new List<UnrealAutomationDevice>();
@@ -266,15 +322,25 @@ namespace Gauntlet
 			};
 			string JsonString = File.ReadAllText(FilePath);
 			UnrealAutomatedTestPassResults JsonTestPassResults = JsonSerializer.Deserialize<UnrealAutomatedTestPassResults>(JsonString, Options);
+			JsonTestPassResults.SetFilePath(FilePath);
 			return JsonTestPassResults;
 		}
 
 		/// <summary>
 		/// Write json data into a file
 		/// </summary>
-		/// <param name="FilePath"></param>
-		public void WriteToJson(string FilePath)
+		/// <param name="InFilePath"></param>
+		public void WriteToJson(string InFilePath = null)
 		{
+			if (!string.IsNullOrEmpty(InFilePath))
+			{
+				SetFilePath(InFilePath);
+			}
+			if (string.IsNullOrEmpty(FilePath))
+			{
+				throw new AutomationException("Can't Write to Json. FilePath is not specified.");
+			}
+
 			string OutputTestDataDir = Path.GetDirectoryName(FilePath);
 			if (!Directory.Exists(OutputTestDataDir))
 			{
