Implement Fetch Data from telemetry database for post run analytics

diff --git a/Engine/Source/Programs/AutomationTool/Gauntlet/Framework/Base/Gauntlet.Database.cs b/Engine/Source/Programs/AutomationTool/Gauntlet/Framework/Base/Gauntlet.Database.cs
--- a/Engine/Source/Programs/AutomationTool/Gauntlet/Framework/Base/Gauntlet.Database.cs
+++ b/Engine/Source/Programs/AutomationTool/Gauntlet/Framework/Base/Gauntlet.Database.cs
@@ -6,6 +6,8 @@ using System.Linq;
 using System.Data;
 using System.Collections.Generic;
 using MySql.Data.MySqlClient;
+using System.Text.Json;
+using System.Text.RegularExpressions;
 using AutomationTool;
 
 namespace Gauntlet
@@ -37,6 +39,12 @@ namespace Gauntlet
 		/// <param name="Query"></param>
 		/// <returns></returns>
 		DataSet ExecuteQuery(string Query);
+		/// <summary>
+		/// Fetch collection of object from target Database, uses Conditions to narrow down pool of data.
+		/// </summary>
+		/// <param name="Conditions"></param>
+		/// <returns></returns>
+		DataSet FetchData(ITelemetryContext Conditions);
 	}
 
 	/// <summary>
@@ -91,7 +99,7 @@ namespace Gauntlet
 
 		public override string ToString()
 		{
-			return Config.GetConfigValue("Server");
+			return string.Format("{0} with config '{1}'", Config.GetConfigValue("Server"), Config.GetType().FullName);
 		}
 
 		public DataSet ExecuteQuery(string SqlQuery)
@@ -104,10 +112,11 @@ namespace Gauntlet
 			foreach (var Chunk in ChunkIt(Rows))
 			{
 				string SqlQuery = string.Format(
-					"INSERT INTO `{0}`.{1} ({2}) VALUES {3}",
+					"INSERT INTO `{0}`.{1} ({2}) VALUES {3};SELECT LAST_INSERT_ID()",
 					Config.DatabaseName, Table, string.Join(", ", Columns), string.Join(", ", Chunk.Select(R => string.Format("({0})", string.Join(", ", R.Select(V => string.Format("'{0}'", V))))))
 				);
-				if(MySqlHelper.ExecuteScalar(Config.ConfigString, SqlQuery) == null)
+
+				if (MySqlHelper.ExecuteScalar(Config.ConfigString, SqlQuery) == null)
 				{
 					return false;
 				}
@@ -161,6 +170,11 @@ namespace Gauntlet
 			return Config.PostSubmitQuery(this, DataRows, Context);
 		}
 
+		public DataSet FetchData(ITelemetryContext Context)
+		{
+			return MySqlHelper.ExecuteDataset(Config.ConfigString, Config.GetFetchDataQuery(Context));
+		}
+
 	}
 
 	public abstract class MySQLConfig<Data> : IDatabaseConfig<Data>, IMySQLConfig<Data> where Data : class
@@ -270,7 +284,160 @@ namespace Gauntlet
 		/// </summary>
 		/// <returns></returns>
 		public abstract IEnumerable<string> FormatDataForTable(Data InData, ITelemetryContext InContext);
+		/// <summary>
+		/// Return the query to fetch data
+		/// </summary>
+		/// <param name="InContext"></param>
+		/// <returns></returns>
+		public virtual string GetFetchDataQuery(ITelemetryContext InContext)
+        {
+			return null;
+        }
+	}
+
+	public class JsonHttpRequestDriver<Data> : IDatabaseDriver<Data> where Data : class
+	{
+		protected JsonHttpRequestConfig<Data> Config;
+		public JsonHttpRequestDriver(JsonHttpRequestConfig<Data> InConfig)
+		{
+			Config = InConfig;
+			if (string.IsNullOrEmpty(Config.Host))
+			{
+				throw new AutomationException(string.Format("Database Driver '{0}' host not configured.", this.GetType().FullName));
+			}
+			if (Config.Auth == null)
+			{
+				throw new AutomationException(string.Format("Database Driver '{0}' authorization header not configured.", this.GetType().FullName));
+			}
+		}
+		public virtual bool SubmitDataItems(IEnumerable<Data> DataItems, ITelemetryContext Context)
+		{
+			var Rows = Config.FormatData(DataItems, Context);
+			foreach (var Chunk in ChunkIt(Rows))
+			{
+				string JsonString = JsonSerializer.Serialize(Chunk);
+				var Result = ExecuteQuery(JsonString);
+				if (Result == null)
+				{
+					return false;
+				}
+			}
+
+			return true;
+		}
+		private IEnumerable<IEnumerable<T>> ChunkIt<T>(IEnumerable<T> ToChunk, int ChunkSize = 1000)
+		{
+			return ToChunk.Select((v, i) => new { Value = v, Index = i }).GroupBy(x => x.Index / ChunkSize).Select(g => g.Select(x => x.Value));
+		}
+		public virtual DataSet ExecuteQuery(string Query)
+		{
+			// Send data
+			var Auth = HttpRequest.Authentication(Config.Auth.target, Config.Auth.token);
+			var Connection = new HttpRequest.Connection(Config.Host, Auth);
+			var Response = Connection.PostJson(Config.Path, Query);
+			if (!Response.IsSuccessStatusCode)
+			{
+				Log.Error("Error sending data to host. Status={0}.\n{1}", Response.StatusCode, Response.Content);
+				return null;
+			}
+			// Format response
+			var Data = new DataSet();
+			DataTable Table = new DataTable();
+			Table.Columns.Add("Status", typeof(int));
+			Table.Columns.Add("Content", typeof(string));
+			Table.Rows.Add(new object[] { Response.StatusCode, Response.Content });
+			Table.AcceptChanges();
+			Data.Tables.Add(Table);
+			return Data;
+		}
+		public DataSet FetchData(ITelemetryContext Conditions)
+		{
+			return new DataSet();
+		}
+		public override string ToString()
+		{
+			return string.Format("{0}/{1} with config '{2}'", Config.Host, Config.Path ?? "", Config.GetType().FullName);
+		}
+	}
+
+	public abstract class JsonHttpRequestConfig<Data> : IDatabaseConfig<Data>where Data : class
+	{
+		public string Host;
+		public string Path;
+		public JsonHttpRequestConfig.Auth Auth;
+		public virtual void LoadConfig(string ConfigFilePath)
+		{
+			JsonHttpRequestConfig Config = null;
+			Host = string.Empty;
+			Path = string.Empty;
+			Auth = null;
+
+			if (File.Exists(ConfigFilePath))
+			{
+				try
+				{
+					Config = JsonHttpRequestConfig.LoadFromFile(ConfigFilePath);
+					Log.Info("Found Http Request connection config from file.");
+				}
+				catch (Exception Ex)
+				{
+					Log.Warning("Properly found config file, but couldn't read a valid json schema.\n{0}", Ex);
+					return;
+				}
+			}
+			else
+			{
+				Log.Error("Could not find Http Request connection config file at '{0}'.", ConfigFilePath);
+				return;
+			}
+
+			var Url = new Regex(@"^(https?://[^/]+/?)(.+)?$");
+			var Match = Url.Match(Config.endpoint);
+			var Success = Match.Success;
+			if (!Success)
+			{
+				Log.Warning("Properly found config file, but couldn't parse http endpoint.");
+				return;
+			}
 
+			Host = Match.Groups[1].Value;
+			Path = Match.Groups[2].Value;
+			Auth = Config.authorization;
+		}
+
+		/// <summary>
+		/// Format the data for target data type
+		/// </summary>
+		/// <returns></returns>
+		public abstract IEnumerable<object> FormatData(IEnumerable<Data> DataItems, ITelemetryContext Context);
+
+		public virtual IDatabaseDriver<Data> GetDriver()
+		{
+			JsonHttpRequestDriver<Data> Driver = new JsonHttpRequestDriver<Data>(this);
+			return Driver;
+		}
+	}
+
+	public class JsonHttpRequestConfig
+	{
+		public class Auth
+		{
+			public string target { get; set; }
+			public string token { get; set; }
+		}
+
+		public string endpoint { get; set; }
+		public Auth authorization { get; set; }
+
+		public static JsonHttpRequestConfig LoadFromFile(string FilePath)
+		{
+			JsonSerializerOptions Options = new JsonSerializerOptions
+			{
+				PropertyNameCaseInsensitive = true
+			};
+			string JsonString = File.ReadAllText(FilePath);
+			return JsonSerializer.Deserialize<JsonHttpRequestConfig>(JsonString, Options);
+		}
 	}
 
 	public class DatabaseConfigManager<Data> where Data : class
diff --git a/Engine/Source/Programs/AutomationTool/Gauntlet/Framework/Utils/Gauntlet.HttpRequest.cs b/Engine/Source/Programs/AutomationTool/Gauntlet/Framework/Utils/Gauntlet.HttpRequest.cs
new file mode 100644
--- /dev/null
+++ b/Engine/Source/Programs/AutomationTool/Gauntlet/Framework/Utils/Gauntlet.HttpRequest.cs
@@ -0,0 +1,54 @@
+// Copyright Epic Games, Inc. All Rights Reserved.
+
+using System;
+using System.Net;
+using System.Text;
+using System.Net.Http;
+using System.Net.Http.Headers;
+using AutomationTool;
+
+namespace Gauntlet
+{
+	public static class HttpRequest
+	{
+		public static AuthenticationHeaderValue Authentication(string Target, string Token)
+		{
+			return new AuthenticationHeaderValue(Target, Token);
+		}
+		public class Connection
+		{
+			private HttpClient Client;
+			public Connection(string Server, AuthenticationHeaderValue Auth)
+			{
+				Client = new HttpClient();
+				Client.BaseAddress = new Uri(Server);
+				if (Auth == null)
+				{
+					throw new AutomationException(string.Format("Missing Authorization Header for {0}", Server));
+				}
+				Client.DefaultRequestHeaders.Authorization = Auth;
+			}
+			public HttpResponse PostJson(string Path, string JsonString)
+			{
+
+				HttpRequestMessage Request = new HttpRequestMessage(HttpMethod.Post, Path);
+				Request.Content = new StringContent(JsonString, Encoding.UTF8, "application/json");
+				using (HttpResponseMessage Response = Client.SendAsync(Request).Result)
+				{
+					return new HttpResponse(Response.StatusCode, Response.Content.ReadAsStringAsync().Result);
+				}
+			}
+		}
+		public class HttpResponse
+		{
+			public string Content;
+			public HttpStatusCode StatusCode;
+			public HttpResponse(HttpStatusCode InCode, string InContent)
+			{
+				StatusCode = InCode;
+				Content = InContent;
+			}
+			public bool IsSuccessStatusCode { get { return StatusCode == HttpStatusCode.OK; } }
+		}
+	}
+}
\ No newline at end of file
diff --git a/Engine/Source/Programs/AutomationTool/Gauntlet/Unreal/Config/Gauntlet.EpicConfig.cs b/Engine/Source/Programs/AutomationTool/Gauntlet/Unreal/Config/Gauntlet.EpicConfig.cs
--- a/Engine/Source/Programs/AutomationTool/Gauntlet/Unreal/Config/Gauntlet.EpicConfig.cs
+++ b/Engine/Source/Programs/AutomationTool/Gauntlet/Unreal/Config/Gauntlet.EpicConfig.cs
@@ -3,7 +3,10 @@
 using AutomationTool;
 using Gauntlet;
 using System;
+using System.IO;
 using System.Data;
+using System.Linq;
+using System.Text.Json.Serialization;
 using System.Collections.Generic;
 using UnrealBuildTool;
 
@@ -19,20 +22,22 @@ namespace EpicConfig
 		private int BuildID = -1;
 		public override void LoadConfig(string ConfigFilePath)
 		{
+			string RootConfig;
 			if (string.IsNullOrEmpty(ConfigFilePath))
 			{
-				if (HostPlatform.Current.HostEditorPlatform == UnrealTargetPlatform.Win64)
+				RootConfig = @"\\epicgames.net\root\Builds\Automation\UE5\Config";
+				if (HostPlatform.Current.HostEditorPlatform == UnrealTargetPlatform.Linux)
 				{
-					ConfigFilePath = @"\\epicgames.net\root\Builds\Automation\UE5\Config\ue-automation-telemetry-mysql-config.txt";
+					RootConfig = "/mnt/Builds/Automation/UE5/Config";
 				}
-				else if (HostPlatform.Current.HostEditorPlatform == UnrealTargetPlatform.Linux)
+				else if (HostPlatform.Current.HostEditorPlatform == UnrealTargetPlatform.Mac)
 				{
-					ConfigFilePath = "/mnt/Builds/Automation/UE5/Config/ue-automation-telemetry-mysql-config.txt";
-				}
-				else
-				{
-					ConfigFilePath = "/Volumes/Root/Builds/Automation/UE5/Config/ue-automation-telemetry-mysql-config.txt";
+					RootConfig = "/Volumes/Root/Builds/Automation/UE5/Config";
 				}
+				ConfigFilePath = Path.Combine(RootConfig, "ue-automation-telemetry-mysql-config.txt");
+			} else
+			{
+				RootConfig = Path.GetDirectoryName(ConfigFilePath);
 			}
 			base.LoadConfig(ConfigFilePath);
 		}
@@ -120,6 +125,61 @@ namespace EpicConfig
 
 			return true;
 		}
+
+		/// <summary>
+		/// Return the query to fetch data
+		/// </summary>
+		/// <param name="InContext"></param>
+		/// <returns></returns>
+		public override string GetFetchDataQuery(ITelemetryContext InContext)
+        {
+			string Project = InContext.GetProperty("ProjectName")?.ToString();
+			string Configuration = InContext.GetProperty("Configuration")?.ToString();
+			string Platform = InContext.GetProperty("Platform")?.ToString();
+			string Branch = InContext.GetProperty("Branch")?.ToString();
+			string TestName = InContext.GetProperty("TestName")?.ToString();
+			string DataPoint = InContext.GetProperty("DataPoint")?.ToString();
+			string TestContext = InContext.GetProperty("Context")?.ToString();
+
+			string WhereClause = "";
+			if(!string.IsNullOrEmpty(Project))
+            {
+				WhereClause += $"AND T1.project = '{Project}' ";
+            }
+			if (!string.IsNullOrEmpty(Branch))
+			{
+				WhereClause += $"AND T1.stream = '{Branch}' ";
+			}
+			if (!string.IsNullOrEmpty(Configuration))
+			{
+				WhereClause += $"AND T2.config = '{Configuration}' ";
+			}
+			if (!string.IsNullOrEmpty(Platform))
+			{
+				WhereClause += $"AND T2.platform = '{Platform}' ";
+			}
+			if (!string.IsNullOrEmpty(TestName))
+			{
+				WhereClause += $"AND T2.testName IN ('{string.Join("','", TestName.Split(','))}') ";
+			}
+			if (!string.IsNullOrEmpty(DataPoint))
+			{
+				WhereClause += $"AND T2.dataPoint IN ('{string.Join("','", DataPoint.Split(','))}')  ";
+			}
+			if (!string.IsNullOrEmpty(TestContext))
+			{
+				WhereClause += $"AND T2.context IN ('{string.Join("','", TestContext.Split(','))}') ";
+			}
+
+			string ChangelistDateTime = ((DateTime)InContext.GetProperty("ChangelistDateTime")).ToString("yyyy-MM-dd HH:mm:ss");
+
+			string SqlQuery = @$"SELECT T2.testName AS TestName, T2.context AS Context, T2.dataPoint AS DataPoint, T2.recordedValue AS Measurement, T2.unit AS Unit, T2.baseline AS Baseline, T1.changelist AS Changelist, UNIX_TIMESTAMP(T1.changelistDateTime) AS ChangelistDataTime, T1.stream AS Stream, T1.project AS Project, T2.platform AS Platform, T2.config AS Config, T2.id AS ID
+									FROM `{DatabaseName}`.{BuildTableName} AS T1
+										RIGHT JOIN (`{DatabaseName}`.{TableName} AS T2) ON (T1.id = T2.buildId)
+											WHERE T1.changelistDateTime > '{ChangelistDateTime}' {WhereClause} ORDER BY T1.id Desc;";
+
+			return SqlQuery;
+		}
 	}
 	/// <summary>
 	/// Staging UE Telemetry Config
@@ -129,4 +189,5 @@ namespace EpicConfig
 		protected override string TableName { get { return "test_records_staging"; } }
 		protected override string BuildTableName { get { return "test_builds_staging"; } }
 	}
+
 }
diff --git a/Engine/Source/Programs/AutomationTool/Gauntlet/Unreal/Utils/Gauntlet.UnrealTelemetry.cs b/Engine/Source/Programs/AutomationTool/Gauntlet/Unreal/Utils/Gauntlet.UnrealTelemetry.cs
--- a/Engine/Source/Programs/AutomationTool/Gauntlet/Unreal/Utils/Gauntlet.UnrealTelemetry.cs
+++ b/Engine/Source/Programs/AutomationTool/Gauntlet/Unreal/Utils/Gauntlet.UnrealTelemetry.cs
@@ -2,6 +2,9 @@
 
 using System;
 using System.IO;
+using System.Data;
+using System.Text;
+using System.Linq;
 using System.Collections.Generic;
 using System.Text.RegularExpressions;
 using AutomationTool;
@@ -121,6 +124,27 @@ namespace Gauntlet
 			Baseline
 		}
 
+		static public void WriteDataTableToCSV(DataTable Data, string OutputFile)
+        {
+			StringBuilder OutputBuffer = new StringBuilder();
+
+			IEnumerable<string> ColumnNames = Data.Columns.Cast<DataColumn>().Select(C => C.ColumnName);
+			OutputBuffer.AppendLine(string.Join(",", ColumnNames));
+
+			foreach (DataRow Row in Data.Rows)
+			{
+				IEnumerable<string> Fields = Row.ItemArray.Select(F => F.ToString());
+				OutputBuffer.AppendLine(string.Join(",", Fields));
+			}
+
+			string CSVFolder = Path.GetDirectoryName(OutputFile);
+			if (!Directory.Exists(CSVFolder))
+			{
+				Directory.CreateDirectory(CSVFolder);
+			}
+			File.WriteAllText(OutputFile, OutputBuffer.ToString());
+		}
+
 	}
 
 	public class UnrealTelemetryContext : ITelemetryContext
@@ -187,7 +211,7 @@ namespace Gauntlet
 			}
 			if (!string.IsNullOrEmpty(CSVFile) && !File.Exists(CSVFile))
 			{
-				throw new AutomationException(string.Format("CSVFile '{0}' is missing.", CSVDirectory));
+				throw new AutomationException(string.Format("CSVFile '{0}' is missing.", CSVFile));
 			}
 
 			// CSV Mapping
@@ -320,4 +344,144 @@ namespace Gauntlet
 		}
 
 	}
+
+	/// <summary>
+	/// Build command to fetch Unreal Automation telemetry data from a database and save it to csv.
+	/// </summary>
+	class FetchUnrealAutomationTelemetry : BuildCommand
+	{
+		[Help("CSVFile", "Path of the csv file to save.")]
+		[Help("TelemetryConfig", "Telemetry configuration to use to publish to Database. Default: UETelemetryStaging.")]
+		[Help("DatabaseConfigPath", "Path to alternate Database config. Default is TelemetryConfig default.")]
+		[Help("Project", "Target Project name.")]
+		[Help("Platform", "Target platform name. Default: current environment platform.")]
+		[Help("Role", "Target Role name. Default: Editor.")]
+		[Help("Branch", "Target Branch name. Default: Unknown.")]
+		[Help("Configuration", "Target Configuration name. Default: Development.")]
+		[Help("Since", "Filter fetch data from the last 'Since' time. Default: 1month.")]
+		[Help("TestName", "Filter fetch by TestName. Support coma separated list.")]
+		[Help("DataPoint", "Filter fetch by DataPoint. Support coma separated list.")]
+		[Help("Context", "Filter fetch by Context. Support coma separated list.")]
+
+		public override ExitCode Execute()
+		{
+			string CSVFile = ParseParamValue("CSVFile=", "");
+			string Config = ParseParamValue("TelemetryConfig=", "UETelemetryStaging");
+			string DatabaseConfigPath = ParseParamValue("DatabaseConfigPath=", "");
+			string ProjectString = ParseParamValue("Project=", "");
+			string PlatformString = ParseParamValue("Platform=", "");
+			string BranchString = ParseParamValue("Branch=", "Unknown");
+			string RoleString = ParseParamValue("Role=", "Editor");
+			string ConfigurationString = ParseParamValue("Configuration=", "Development");
+
+			string Since = ParseParamValue("Since=", "1month");
+			string TestName = ParseParamValue("TestName=", "");
+			string DataPoint = ParseParamValue("DataPoint=", "");
+			string Context = ParseParamValue("Context=", "");
+
+			if (string.IsNullOrEmpty(CSVFile))
+			{
+				throw new AutomationException("CSVFile argument is missing.");
+			}
+
+			UnrealTelemetryContext Conditions = new UnrealTelemetryContext();
+			if (string.IsNullOrEmpty(ProjectString))
+			{
+				throw new AutomationException("No project specified. Use -Project=ShooterGame etc");
+			}
+			Conditions.SetProperty("ProjectName", ProjectString);
+
+			object Role;
+			if (!Enum.TryParse(typeof(UnrealTargetRole), RoleString, true, out Role))
+			{
+				string AllKeys = string.Join(", ", Enum.GetNames(typeof(UnrealTargetRole)));
+				throw new AutomationException(string.Format("Unknown Role '{0}', it must be one of the values: {1}.", RoleString, AllKeys));
+			}
+			RoleString = ((UnrealTargetRole)Role).ToString();
+
+			object Configuration;
+			if (!Enum.TryParse(typeof(UnrealTargetConfiguration), ConfigurationString, true, out Configuration))
+			{
+				string AllKeys = string.Join(", ", Enum.GetNames(typeof(UnrealTargetConfiguration)));
+				throw new AutomationException(string.Format("Unknown Configuration '{0}', it must be one of the values: {1}.", ConfigurationString, AllKeys));
+			}
+			ConfigurationString = ((UnrealTargetConfiguration)Configuration).ToString();
+			Conditions.SetProperty("Configuration", string.Format("{0} {1}", RoleString, ConfigurationString));
+
+			UnrealTargetPlatform Platform = string.IsNullOrEmpty(PlatformString) ? BuildHostPlatform.Current.Platform : UnrealTargetPlatform.Parse(PlatformString);
+			Conditions.SetProperty("Platform", Platform.ToString());
+
+			Conditions.SetProperty("Branch", BranchString);
+
+			if (!string.IsNullOrEmpty(TestName))
+			{
+				Conditions.SetProperty("TestName", TestName);
+			}
+			if (!string.IsNullOrEmpty(DataPoint))
+			{
+				Conditions.SetProperty("DataPoint", DataPoint);
+			}
+			if (!string.IsNullOrEmpty(Context))
+			{
+				Conditions.SetProperty("Context", Context);
+			}
+
+			Conditions.SetProperty("ChangelistDateTime", SinceStringToDate(Since));
+
+			IDatabaseConfig<TelemetryData> DBConfig = DatabaseConfigManager<TelemetryData>.GetConfigByName(Config);
+			if (DBConfig != null)
+			{
+				DBConfig.LoadConfig(DatabaseConfigPath);
+				IDatabaseDriver<TelemetryData> DB = DBConfig.GetDriver();
+
+				var Results = DB.FetchData(Conditions);
+
+				if(Results is null || Results.Tables.Count == 0)
+                {
+					throw new AutomationException("No result returned from the DB with the requested conditions.");
+                }
+
+				Log.Info("Successfully retrived {0} rows from the last {1}.", Results.Tables[0].Rows.Count, Since);
+
+				UnrealAutomationTelemetry.WriteDataTableToCSV(Results.Tables[0], CSVFile);
+			}
+
+			return 0;
+		}
+
+		private DateTime SinceStringToDate(string SinceString)
+		{
+			Regex Regx = new Regex(@"([0-9]+)\s*([dwmy])", RegexOptions.Compiled | RegexOptions.IgnoreCase);
+			Match DateMatch = Regx.Match(SinceString);
+
+			DateTime SinceDate = DateTime.Now;
+
+			if (DateMatch.Success)
+            {
+				string NumberString = DateMatch.Groups[1].Value;
+				string TimeSpanString = DateMatch.Groups[2].Value.Substring(0, 1).ToLower();
+
+				int Factor = Int32.Parse(NumberString);
+				switch(TimeSpanString)
+                {
+					case "d":
+						SinceDate = SinceDate.AddDays(-Factor);
+						break;
+					case "w":
+						SinceDate = SinceDate.AddDays(-Factor * 7);
+						break;
+					case "m":
+						SinceDate = SinceDate.AddMonths(-Factor);
+						break;
+					case "y":
+						SinceDate = SinceDate.AddYears(-Factor);
+						break;
+				}
+			}
+
+			return SinceDate;
+		}
+
+	}
+
 }
