BuildGraph: Various refactorings prior to adding bytecode executor.

* Moved XML script parsing into UAT. We don't need to evaluate this in Horde.
* Added custom node types for script and expression nodes, rather than wrapping native nodes with a custom script task.

diff --git a/Engine/Source/Programs/AutomationTool/BuildGraph/BgCondition.cs b/Engine/Source/Programs/AutomationTool/BuildGraph/BgCondition.cs
new file mode 100644
--- /dev/null
+++ b/Engine/Source/Programs/AutomationTool/BuildGraph/BgCondition.cs
@@ -0,0 +1,527 @@
+// Copyright Epic Games, Inc. All Rights Reserved.
+
+using System;
+using System.Collections.Generic;
+using System.IO;
+using System.Linq;
+using System.Threading.Tasks;
+
+namespace AutomationTool
+{
+	/// <summary>
+	/// Exception class thrown due to type and syntax errors in condition expressions
+	/// </summary>
+	class BgConditionException : Exception
+	{
+		/// <summary>
+		/// Constructor; formats the exception message with the given String.Format() style parameters.
+		/// </summary>
+		/// <param name="format">Formatting string, in String.Format syntax</param>
+		/// <param name="args">Optional arguments for the string</param>
+		public BgConditionException(string format, params object[] args) : base(String.Format(format, args))
+		{
+		}
+	}
+
+	/// <summary>
+	/// Class to evaluate condition expressions in build scripts, following this grammar:
+	/// 
+	///		or-expression   ::= and-expression
+	///		                  | or-expression "Or" and-expression;
+	///		    
+	///		and-expression  ::= comparison
+	///		                  | and-expression "And" comparison;
+	///		                  
+	///		comparison      ::= scalar
+	///		                  | scalar "==" scalar
+	///		                  | scalar "!=" scalar
+	///		                  | scalar "&lt;" scalar
+	///		                  | scalar "&lt;=" scalar;
+	///		                  | scalar "&gt;" scalar
+	///		                  | scalar "&gt;=" scalar;
+	///		                  
+	///     scalar          ::= "(" or-expression ")"
+	///                       | "!" scalar
+	///                       | "Exists" "(" scalar ")"
+	///                       | "HasTrailingSlash" "(" scalar ")"
+	///                       | string
+	///                       | identifier;
+	///                       
+	///     string          ::= any sequence of characters terminated by single quotes (') or double quotes ("). Not escaped.
+	///     identifier      ::= any sequence of letters, digits, or underscore characters.
+	///     
+	/// The type of each subexpression is always a scalar, which are converted to expression-specific types (eg. booleans, integers) as required.
+	/// Scalar values are case-insensitive strings. The identifier 'true' and the strings "true" and "True" are all identical scalars.
+	/// </summary>
+	public class BgCondition
+	{
+		/// <summary>
+		/// Sentinel added to the end of a sequence of tokens.
+		/// </summary>
+		const string EndToken = "<EOF>";
+
+		/// <summary>
+		/// Tokens for the condition
+		/// </summary>
+		readonly List<string> _tokens = new List<string>();
+
+		/// <summary>
+		/// The current token index
+		/// </summary>
+		int _idx;
+
+		/// <summary>
+		/// Context for evaluating the expression
+		/// </summary>
+		readonly IBgScriptReaderContext _context;
+
+		/// <summary>
+		/// Constructor
+		/// </summary>
+		/// <param name="text">The condition text</param>
+		/// <param name="context">Context for evaluating the expression</param>
+		private BgCondition(string text, IBgScriptReaderContext context)
+		{
+			_context = context;
+
+			Tokenize(text, _tokens);
+		}
+
+		/// <summary>
+		/// Evaluates the given string as a condition. Throws a ConditionException on a type or syntax error.
+		/// </summary>
+		/// <param name="text"></param>
+		/// <param name="context"></param>
+		/// <returns>The result of evaluating the condition</returns>
+		public static ValueTask<bool> EvaluateAsync(string text, IBgScriptReaderContext context)
+		{
+			return new BgCondition(text, context).EvaluateAsync();
+		}
+
+		/// <summary>
+		/// Evaluates the given string as a condition. Throws a ConditionException on a type or syntax error.
+		/// </summary>
+		/// <returns>The result of evaluating the condition</returns>
+		async ValueTask<bool> EvaluateAsync()
+		{
+			bool bResult = true;
+			if (_tokens.Count > 1)
+			{
+				_idx = 0;
+				string result = await EvaluateOrAsync();
+				if (_tokens[_idx] != EndToken)
+				{
+					throw new BgConditionException("Garbage after expression: {0}", String.Join("", _tokens.Skip(_idx)));
+				}
+				bResult = CoerceToBool(result);
+			}
+			return bResult;
+		}
+
+		/// <summary>
+		/// Evaluates an "or-expression" production.
+		/// </summary>
+		/// <returns>A scalar representing the result of evaluating the expression.</returns>
+		async ValueTask<string> EvaluateOrAsync()
+		{
+			// <Condition> Or <Condition> Or...
+			string result = await EvaluateAndAsync();
+			while (String.Compare(_tokens[_idx], "Or", true) == 0)
+			{
+				// Evaluate this condition. We use a binary OR here, because we want to parse everything rather than short-circuit it.
+				_idx++;
+				string lhs = result;
+				string rhs = await EvaluateAndAsync();
+				result = (CoerceToBool(lhs) | CoerceToBool(rhs)) ? "true" : "false";
+			}
+			return result;
+		}
+
+		/// <summary>
+		/// Evaluates an "and-expression" production.
+		/// </summary>
+		/// <returns>A scalar representing the result of evaluating the expression.</returns>
+		async ValueTask<string> EvaluateAndAsync()
+		{
+			// <Condition> And <Condition> And...
+			string result = await EvaluateComparisonAsync();
+			while (String.Compare(_tokens[_idx], "And", true) == 0)
+			{
+				// Evaluate this condition. We use a binary AND here, because we want to parse everything rather than short-circuit it.
+				_idx++;
+				string lhs = result;
+				string rhs = await EvaluateComparisonAsync();
+				result = (CoerceToBool(lhs) & CoerceToBool(rhs)) ? "true" : "false";
+			}
+			return result;
+		}
+
+		/// <summary>
+		/// Evaluates a "comparison" production.
+		/// </summary>
+		/// <returns>The result of evaluating the expression</returns>
+		async ValueTask<string> EvaluateComparisonAsync()
+		{
+			// scalar
+			// scalar == scalar
+			// scalar != scalar
+			// scalar < scalar
+			// scalar <= scalar
+			// scalar > scalar
+			// scalar >= scalar
+
+			string result = await EvaluateScalarAsync();
+			if (_tokens[_idx] == "==")
+			{
+				// Compare two scalars for equality
+				_idx++;
+				string lhs = result;
+				string rhs = await EvaluateScalarAsync();
+				result = (String.Compare(lhs, rhs, true) == 0) ? "true" : "false";
+			}
+			else if (_tokens[_idx] == "!=")
+			{
+				// Compare two scalars for inequality
+				_idx++;
+				string lhs = result;
+				string rhs = await EvaluateScalarAsync();
+				result = (String.Compare(lhs, rhs, true) != 0) ? "true" : "false";
+			}
+			else if (_tokens[_idx] == "<")
+			{
+				// Compares whether the first integer is less than the second
+				_idx++;
+				int lhs = CoerceToInteger(result);
+				int rhs = CoerceToInteger(await EvaluateScalarAsync());
+				result = (lhs < rhs) ? "true" : "false";
+			}
+			else if (_tokens[_idx] == "<=")
+			{
+				// Compares whether the first integer is less than the second
+				_idx++;
+				int lhs = CoerceToInteger(result);
+				int rhs = CoerceToInteger(await EvaluateScalarAsync());
+				result = (lhs <= rhs) ? "true" : "false";
+			}
+			else if (_tokens[_idx] == ">")
+			{
+				// Compares whether the first integer is less than the second
+				_idx++;
+				int lhs = CoerceToInteger(result);
+				int rhs = CoerceToInteger(await EvaluateScalarAsync());
+				result = (lhs > rhs) ? "true" : "false";
+			}
+			else if (_tokens[_idx] == ">=")
+			{
+				// Compares whether the first integer is less than the second
+				_idx++;
+				int lhs = CoerceToInteger(result);
+				int rhs = CoerceToInteger(await EvaluateScalarAsync());
+				result = (lhs >= rhs) ? "true" : "false";
+			}
+			return result;
+		}
+
+		/// <summary>
+		/// Evaluates arguments from a token string. Arguments are all comma-separated tokens until a closing ) is encountered
+		/// </summary>
+		/// <returns>The result of evaluating the expression</returns>
+		IEnumerable<string> EvaluateArguments()
+		{
+			List<string> arguments = new List<string>();
+
+			// skip opening bracket
+			if (_tokens[_idx++] != "(")
+			{
+				throw new BgConditionException("Expected '('");
+			}
+
+			bool didCloseBracket = false;
+
+			while (_idx < _tokens.Count)
+			{
+				string nextToken = _tokens.ElementAt(_idx++);
+
+				if (nextToken == EndToken)
+				{
+					// ran out of items
+					break;
+				}
+				else if (nextToken == ")")
+				{
+					didCloseBracket = true;
+					break;
+				}
+				else if (nextToken != ",")
+				{
+					if (nextToken.First() == '\'' && nextToken.Last() == '\'')
+					{
+						nextToken = nextToken.Substring(1, nextToken.Length - 2);
+					}
+					arguments.Add(nextToken);
+				}
+			}
+
+			if (!didCloseBracket)
+			{
+				throw new BgConditionException("Expected ')'");
+			}
+
+			return arguments;
+		}
+
+		/// <summary>
+		/// Evaluates a "scalar" production.
+		/// </summary>
+		/// <returns>The result of evaluating the expression</returns>
+		async ValueTask<string> EvaluateScalarAsync()
+		{
+			string result;
+			if (_tokens[_idx] == "(")
+			{
+				// Subexpression
+				_idx++;
+				result = await EvaluateOrAsync();
+				if (_tokens[_idx] != ")")
+				{
+					throw new BgConditionException("Expected ')'");
+				}
+				_idx++;
+			}
+			else if (_tokens[_idx] == "!")
+			{
+				// Logical not
+				_idx++;
+				string rhs = await EvaluateScalarAsync();
+				result = CoerceToBool(rhs) ? "false" : "true";
+			}
+			else if (String.Compare(_tokens[_idx], "Exists", true) == 0 && _tokens[_idx + 1] == "(")
+			{
+				// Check whether file or directory exists. Evaluate the argument as a subexpression.
+				_idx++;
+				string argument = await EvaluateScalarAsync();
+				result = await _context.ExistsAsync(argument) ? "true" : "false";
+			}
+			else if (String.Compare(_tokens[_idx], "HasTrailingSlash", true) == 0 && _tokens[_idx + 1] == "(")
+			{
+				// Check whether the given string ends with a slash
+				_idx++;
+				string argument = await EvaluateScalarAsync();
+				result = (argument.Length > 0 && (argument[^1] == Path.DirectorySeparatorChar || argument[^1] == Path.AltDirectorySeparatorChar)) ? "true" : "false";
+			}
+			else if (String.Compare(_tokens[_idx], "Contains", true) == 0 && _tokens[_idx + 1] == "(")
+			{
+				// Check a string contains a substring. If a separator is supplied the string is first split
+				_idx++;
+				IEnumerable<string> arguments = EvaluateArguments();
+
+				if (arguments.Count() != 2)
+				{
+					throw new BgConditionException("Invalid argument count for 'Contains'. Expected (Haystack,Needle)");
+				}
+
+				result = Contains(arguments.ElementAt(0), arguments.ElementAt(1)) ? "true" : "false";
+			}
+			else if (String.Compare(_tokens[_idx], "ContainsItem", true) == 0 && _tokens[_idx + 1] == "(")
+			{
+				// Check a string contains a substring. If a separator is supplied the string is first split
+				_idx++;
+				IEnumerable<string> arguments = EvaluateArguments();
+
+				if (arguments.Count() != 3)
+				{
+					throw new BgConditionException("Invalid argument count for 'ContainsItem'. Expected (Haystack,Needle,HaystackSeparator)");
+				}
+
+				result = ContainsItem(arguments.ElementAt(0), arguments.ElementAt(1), arguments.ElementAt(2)) ? "true" : "false";
+			}
+			else
+			{
+				// Raw scalar. Remove quotes from strings, and allow literals and simple identifiers to pass through directly.
+				string token = _tokens[_idx];
+				if (token.Length >= 2 && (token[0] == '\'' || token[0] == '\"') && token[^1] == token[0])
+				{
+					result = token.Substring(1, token.Length - 2);
+					_idx++;
+				}
+				else if (Char.IsLetterOrDigit(token[0]) || token[0] == '_')
+				{
+					result = token;
+					_idx++;
+				}
+				else
+				{
+					throw new BgConditionException("Token '{0}' is not a valid scalar", token);
+				}
+			}
+			return result;
+		}
+
+		/// <summary>
+		/// Checks whether Haystack contains "Needle". 
+		/// </summary>
+		/// <param name="haystack">The string to search</param>
+		/// <param name="needle">The string to search for</param>
+		/// <returns>True if the path exists, false otherwise.</returns>
+		static bool Contains(string haystack, string needle)
+		{
+			try
+			{
+				return haystack.IndexOf(needle, StringComparison.CurrentCultureIgnoreCase) >= 0;
+			}
+			catch
+			{
+				return false;
+			}
+		}
+
+		/// <summary>
+		/// Checks whether HaystackItems contains "Needle". 
+		/// </summary>
+		/// <param name="haystack">The separated list of items to check</param>
+		/// <param name="needle">The item to check for</param>
+		/// <param name="haystackSeparator">The separator used in Haystack</param>
+		/// <returns>True if the path exists, false otherwise.</returns>
+		static bool ContainsItem(string haystack, string needle, string haystackSeparator)
+		{
+			try
+			{
+				IEnumerable<string> haystackItems = haystack.Split(new string[] { haystackSeparator }, StringSplitOptions.RemoveEmptyEntries);
+				return haystackItems.Any(i => i.ToLower() == needle.ToLower());
+			}
+			catch
+			{
+				return false;
+			}
+		}
+
+		/// <summary>
+		/// Converts a scalar to a boolean value.
+		/// </summary>
+		/// <param name="scalar">The scalar to convert</param>
+		/// <returns>The scalar converted to a boolean value.</returns>
+		static bool CoerceToBool(string scalar)
+		{
+			bool result;
+			if (String.Compare(scalar, "true", true) == 0)
+			{
+				result = true;
+			}
+			else if (String.Compare(scalar, "false", true) == 0)
+			{
+				result = false;
+			}
+			else
+			{
+				throw new BgConditionException("Token '{0}' cannot be coerced to a bool", scalar);
+			}
+			return result;
+		}
+
+		/// <summary>
+		/// Converts a scalar to a boolean value.
+		/// </summary>
+		/// <param name="scalar">The scalar to convert</param>
+		/// <returns>The scalar converted to an integer value.</returns>
+		static int CoerceToInteger(string scalar)
+		{
+			int value;
+			if (!Int32.TryParse(scalar, out value))
+			{
+				throw new BgConditionException("Token '{0}' cannot be coerced to an integer", scalar);
+			}
+			return value;
+		}
+
+		/// <summary>
+		/// Splits an input string up into expression tokens.
+		/// </summary>
+		/// <param name="text">Text to be converted into tokens</param>
+		/// <param name="tokens">List to receive a list of tokens</param>
+		static void Tokenize(string text, List<string> tokens)
+		{
+			int idx = 0;
+			while (idx < text.Length)
+			{
+				int endIdx = idx + 1;
+				if (!Char.IsWhiteSpace(text[idx]))
+				{
+					// Scan to the end of the current token
+					if (Char.IsNumber(text[idx]))
+					{
+						// Number
+						while (endIdx < text.Length && Char.IsNumber(text[endIdx]))
+						{
+							endIdx++;
+						}
+					}
+					else if (Char.IsLetter(text[idx]) || text[idx] == '_')
+					{
+						// Identifier
+						while (endIdx < text.Length && (Char.IsLetterOrDigit(text[endIdx]) || text[endIdx] == '_'))
+						{
+							endIdx++;
+						}
+					}
+					else if (text[idx] == '!' || text[idx] == '<' || text[idx] == '>' || text[idx] == '=')
+					{
+						// Operator that can be followed by an equals character
+						if (endIdx < text.Length && text[endIdx] == '=')
+						{
+							endIdx++;
+						}
+					}
+					else if (text[idx] == '\'' || text[idx] == '\"')
+					{
+						// String
+						if (endIdx < text.Length)
+						{
+							endIdx++;
+							while (endIdx < text.Length && text[endIdx - 1] != text[idx])
+							{
+								endIdx++;
+							}
+						}
+					}
+					tokens.Add(text.Substring(idx, endIdx - idx));
+				}
+				idx = endIdx;
+			}
+			tokens.Add(EndToken);
+		}
+
+		/// <summary>
+		/// Test cases for conditions.
+		/// </summary>
+		public static async Task TestConditions()
+		{
+			await TestConditionAsync("1 == 2", false);
+			await TestConditionAsync("1 == 1", true);
+			await TestConditionAsync("1 != 2", true);
+			await TestConditionAsync("1 != 1", false);
+			await TestConditionAsync("'hello' == 'hello'", true);
+			await TestConditionAsync("'hello' == ('hello')", true);
+			await TestConditionAsync("'hello' == 'world'", false);
+			await TestConditionAsync("'hello' != ('world')", true);
+			await TestConditionAsync("true == ('true')", true);
+			await TestConditionAsync("true == ('True')", true);
+			await TestConditionAsync("true == ('false')", false);
+			await TestConditionAsync("true == !('False')", true);
+			await TestConditionAsync("true == 'true' and 'false' == 'False'", true);
+			await TestConditionAsync("true == 'true' and 'false' == 'true'", false);
+			await TestConditionAsync("true == 'false' or 'false' == 'false'", true);
+			await TestConditionAsync("true == 'false' or 'false' == 'true'", true);
+		}
+
+		/// <summary>
+		/// Helper method to evaluate a condition and check it's the expected result
+		/// </summary>
+		/// <param name="condition">Condition to evaluate</param>
+		/// <param name="expectedResult">The expected result</param>
+		static async Task TestConditionAsync(string condition, bool expectedResult)
+		{
+			bool result = await new BgCondition(condition, null!).EvaluateAsync();
+			Console.WriteLine("{0}: {1} = {2}", (result == expectedResult) ? "PASS" : "FAIL", condition, result);
+		}
+	}
+}
diff --git a/Engine/Source/Programs/AutomationTool/BuildGraph/BgNodeExecutor.cs b/Engine/Source/Programs/AutomationTool/BuildGraph/BgNodeExecutor.cs
new file mode 100644
--- /dev/null
+++ b/Engine/Source/Programs/AutomationTool/BuildGraph/BgNodeExecutor.cs
@@ -0,0 +1,376 @@
+// Copyright Epic Games, Inc. All Rights Reserved.
+
+using System;
+using System.Collections.Generic;
+using System.ComponentModel;
+using System.Diagnostics.CodeAnalysis;
+using System.IO;
+using System.Linq;
+using System.Reflection;
+using System.Runtime.CompilerServices;
+using System.Text;
+using System.Text.RegularExpressions;
+using System.Threading.Tasks;
+using System.Xml;
+using System.Xml.Schema;
+using AutomationTool.Tasks;
+using EpicGames.BuildGraph;
+using EpicGames.BuildGraph.Expressions;
+using EpicGames.Core;
+using Microsoft.Extensions.Logging;
+using OpenTracing;
+using OpenTracing.Util;
+
+#nullable enable
+
+namespace AutomationTool
+{
+	/// <summary>
+	/// Base class for binding and executing nodes
+	/// </summary>
+	abstract class BgNodeExecutor
+	{
+		public abstract Task<bool> ExecuteAsync(JobContext Job, Dictionary<string, HashSet<FileReference>> TagNameToFileSet);
+	}
+
+	class BgExpressionNodeExecutor : BgNodeExecutor
+	{
+		readonly BgExpressionNode Node;
+		MethodInfo? method;
+
+		public BgExpressionNodeExecutor(BgExpressionNode node)
+		{
+			Node = node;
+		}
+
+		public bool Bind(ILogger logger)
+		{
+			method = Node.Method.Bind();
+			return method != null;
+		}
+
+		/// <summary>
+		/// Execute the method given in the 
+		/// </summary>
+		/// <param name="Job"></param>
+		/// <param name="TagNameToFileSet"></param>
+		/// <returns></returns>
+		public override async Task<bool> ExecuteAsync(JobContext Job, Dictionary<string, HashSet<FileReference>> TagNameToFileSet)
+		{
+			if (method == null)
+			{
+				throw new AutomationException("Node has not been bound");
+			}
+
+			HashSet<FileReference> BuildProducts = TagNameToFileSet[Node.DefaultOutput.TagName];
+
+			BgContextImpl Context = new BgContextImpl(Job, TagNameToFileSet);
+
+			ParameterInfo[] parameters = method.GetParameters();
+
+			object?[] arguments = new object[parameters.Length];
+			for (int Idx = 0; Idx < parameters.Length; Idx++)
+			{
+				Type ParameterType = parameters[Idx].ParameterType;
+				if (ParameterType == typeof(BgContext))
+				{
+					arguments[Idx] = Context;
+				}
+				else
+				{
+					arguments[Idx] = Node.Arguments[Idx];
+				}
+			}
+
+			Task Task = (Task)method.Invoke(null, arguments)!;
+			await Task;
+
+			if (Node.Outputs.Count > 0)
+			{
+				object? Result = null;
+				if (method.ReturnType.IsGenericType && method.ReturnType.GetGenericTypeDefinition() == typeof(Task<>))
+				{
+					Type TaskType = Task.GetType();
+					PropertyInfo Property = TaskType.GetProperty(nameof(Task<int>.Result))!;
+					Result = Property!.GetValue(Task);
+				}
+
+				object?[] OutputValues;
+				if (Result is BgFileSet)
+				{
+					OutputValues = new[] { Result };
+				}
+				else if (Result is ITuple Tuple)
+				{
+					OutputValues = Enumerable.Range(0, Tuple.Length).Select(x => Tuple[x]).ToArray();
+				}
+				else
+				{
+					OutputValues = Array.Empty<object>();
+				}
+
+				for (int Idx = 0; Idx < OutputValues.Length; Idx++)
+				{
+					if (OutputValues[Idx] is BgFileSet Token)
+					{
+						string TagName = $"{Node.Name}${Idx}";
+						FileSet FileSet = Token.ComputeValue(Context.Context);
+						TagNameToFileSet[TagName] = new HashSet<FileReference>(FileSet.Flatten().Values);
+						BuildProducts.UnionWith(TagNameToFileSet[TagName]);
+					}
+				}
+			}
+
+			return true;
+		}
+	}
+
+	/// <summary>
+	/// Implementation of <see cref="BgNode"/> for graphs defined through XML syntax
+	/// </summary>
+	class BgScriptNodeExecutor : BgNodeExecutor
+	{
+		/// <summary>
+		/// The script node
+		/// </summary>
+		public BgScriptNode Node { get; }
+
+		/// <summary>
+		/// List of bound task implementations
+		/// </summary>
+		List<BgTaskImpl> _boundTasks = new List<BgTaskImpl>();
+
+		/// <summary>
+		/// Constructor
+		/// </summary>
+		public BgScriptNodeExecutor(BgScriptNode node)
+		{
+			Node = node;
+		}
+
+		public bool Bind(Dictionary<string, ScriptTaskBinding> nameToTask, Dictionary<string, BgNodeOutput> tagNameToNodeOutput, IBgScriptReaderContext context, ILogger logger)
+		{
+			bool bResult = true;
+			foreach (BgTask TaskInfo in Node.Tasks)
+			{
+				BgTaskImpl? boundTask = BindTask(TaskInfo, nameToTask, tagNameToNodeOutput, context, logger);
+				if (boundTask == null)
+				{
+					bResult = false;
+				}
+				else
+				{
+					_boundTasks.Add(boundTask);
+				}
+			}
+			return bResult;
+		}
+
+		BgTaskImpl? BindTask(BgTask TaskInfo, Dictionary<string, ScriptTaskBinding> NameToTask, IReadOnlyDictionary<string, BgNodeOutput> TagNameToNodeOutput, IBgScriptReaderContext Context, ILogger Logger)
+		{
+			// Get the reflection info for this element
+			ScriptTaskBinding? Task;
+			if (!NameToTask.TryGetValue(TaskInfo.Name, out Task))
+			{
+				Logger.LogScriptError(TaskInfo.Location, "Unknown task '{TaskName}'", TaskInfo.Name);
+				return null;
+			}
+
+			// Check all the required parameters are present
+			bool bHasRequiredAttributes = true;
+			foreach (ScriptTaskParameterBinding Parameter in Task.NameToParameter.Values)
+			{
+				if (!Parameter.Optional && !TaskInfo.Arguments.ContainsKey(Parameter.Name))
+				{
+					Logger.LogScriptError(TaskInfo.Location, "Missing required attribute - {AttrName}", Parameter.Name);
+					bHasRequiredAttributes = false;
+				}
+			}
+
+			// Read all the attributes into a parameters object for this task
+			object ParametersObject = Activator.CreateInstance(Task.ParametersClass)!;
+			foreach ((string Name, string Value) in TaskInfo.Arguments)
+			{
+				// Get the field that this attribute should be written to in the parameters object
+				ScriptTaskParameterBinding? Parameter;
+				if (!Task.NameToParameter.TryGetValue(Name, out Parameter))
+				{
+					Logger.LogScriptError(TaskInfo.Location, "Unknown attribute '{AttrName}'", Name);
+					continue;
+				}
+
+				// If it's a collection type, split it into separate values
+				if (Parameter.CollectionType == null)
+				{
+					// Parse it and assign it to the parameters object
+					object? FieldValue = ParseValue(Value, Parameter.ValueType, Context);
+					Parameter.FieldInfo.SetValue(ParametersObject, FieldValue);
+				}
+				else
+				{
+					// Get the collection, or create one if necessary
+					object? CollectionValue = Parameter.FieldInfo.GetValue(ParametersObject);
+					if (CollectionValue == null)
+					{
+						CollectionValue = Activator.CreateInstance(Parameter.FieldInfo.FieldType)!;
+						Parameter.FieldInfo.SetValue(ParametersObject, CollectionValue);
+					}
+
+					// Parse the values and add them to the collection
+					List<string> ValueStrings = BgTaskImpl.SplitDelimitedList(Value);
+					foreach (string ValueString in ValueStrings)
+					{
+						object ElementValue = ParseValue(ValueString, Parameter.ValueType, Context)!;
+						Parameter.CollectionType.InvokeMember("Add", BindingFlags.InvokeMethod | BindingFlags.Instance | BindingFlags.Public, null, CollectionValue, new object[] { ElementValue });
+					}
+				}
+			}
+
+			// Construct the task
+			if (!bHasRequiredAttributes)
+			{
+				return null;
+			}
+
+			// Add it to the list
+			BgTaskImpl NewTask = (BgTaskImpl)Activator.CreateInstance(Task.TaskClass, ParametersObject)!;
+
+			// Set up the source location for diagnostics
+			NewTask.SourceLocation = TaskInfo.Location;
+
+			// Make sure all the read tags are local or listed as a dependency
+			foreach (string ReadTagName in NewTask.FindConsumedTagNames())
+			{
+				BgNodeOutput? Output;
+				if (TagNameToNodeOutput.TryGetValue(ReadTagName, out Output))
+				{
+					if (Output != null && Output.ProducingNode != Node && !Node.Inputs.Contains(Output))
+					{
+						Logger.LogScriptError(TaskInfo.Location, "The tag '{TagName}' is not a dependency of node '{Node}'", ReadTagName, Node.Name);
+					}
+				}
+			}
+
+			// Make sure all the written tags are local or listed as an output
+			foreach (string ModifiedTagName in NewTask.FindProducedTagNames())
+			{
+				BgNodeOutput? Output;
+				if (TagNameToNodeOutput.TryGetValue(ModifiedTagName, out Output))
+				{
+					if (Output != null && !Node.Outputs.Contains(Output))
+					{
+						Logger.LogScriptError(TaskInfo.Location, "The tag '{TagName}' is created by '{Node}', and cannot be modified downstream", Output.TagName, Output.ProducingNode.Name);
+					}
+				}
+			}
+			return NewTask;
+		}
+
+		/// <summary>
+		/// Parse a value of the given type
+		/// </summary>
+		/// <param name="ValueText">The text to parse</param>
+		/// <param name="ValueType">Type of the value to parse</param>
+		/// <param name="Context">Context for the script reader</param>
+		/// <returns>Value that was parsed</returns>
+		static object? ParseValue(string ValueText, Type ValueType, IBgScriptReaderContext Context)
+		{
+			// Parse it and assign it to the parameters object
+			if (ValueType.IsEnum)
+			{
+				return Enum.Parse(ValueType, ValueText);
+			}
+			else if (ValueType == typeof(Boolean))
+			{
+				return BgCondition.EvaluateAsync(ValueText, Context).Result;
+			}
+			else if (ValueType == typeof(FileReference))
+			{
+				return BgTaskImpl.ResolveFile(ValueText);
+			}
+			else if (ValueType == typeof(DirectoryReference))
+			{
+				return BgTaskImpl.ResolveDirectory(ValueText);
+			}
+
+			TypeConverter Converter = TypeDescriptor.GetConverter(ValueType);
+			if (Converter.CanConvertFrom(typeof(string)))
+			{
+				return Converter.ConvertFromString(ValueText);
+			}
+			else
+			{
+				return Convert.ChangeType(ValueText, ValueType);
+			}
+		}
+
+		/// <summary>
+		/// Build all the tasks for this node
+		/// </summary>
+		/// <param name="Job">Information about the current job</param>
+		/// <param name="TagNameToFileSet">Mapping from tag names to the set of files they include. Should be set to contain the node inputs on entry.</param>
+		/// <returns>Whether the task succeeded or not. Exiting with an exception will be caught and treated as a failure.</returns>
+		public override async Task<bool> ExecuteAsync(JobContext Job, Dictionary<string, HashSet<FileReference>> TagNameToFileSet)
+		{
+			// Run each of the tasks in order
+			HashSet<FileReference> BuildProducts = TagNameToFileSet[Node.DefaultOutput.TagName];
+			for (int Idx = 0; Idx < _boundTasks.Count; Idx++)
+			{
+				using (IScope Scope = GlobalTracer.Instance.BuildSpan("Task").WithTag("resource", _boundTasks[Idx].GetTraceName()).StartActive())
+				{
+					ITaskExecutor Executor = _boundTasks[Idx].GetExecutor();
+					if (Executor == null)
+					{
+						// Execute this task directly
+						try
+						{
+							_boundTasks[Idx].GetTraceMetadata(Scope.Span, "");
+							await _boundTasks[Idx].ExecuteAsync(Job, BuildProducts, TagNameToFileSet);
+						}
+						catch (Exception Ex)
+						{
+							ExceptionUtils.AddContext(Ex, "while executing task {0}", _boundTasks[Idx].GetTraceString());
+							if (_boundTasks[Idx].SourceLocation != null)
+							{
+								ExceptionUtils.AddContext(Ex, "at {0}({1})", _boundTasks[Idx].SourceLocation.File, _boundTasks[Idx].SourceLocation.LineNumber);
+							}
+							throw;
+						}
+					}
+					else
+					{
+						_boundTasks[Idx].GetTraceMetadata(Scope.Span, "1.");
+
+						// The task has a custom executor, which may be able to execute several tasks simultaneously. Try to add the following tasks.
+						int FirstIdx = Idx;
+						while (Idx + 1 < Node.Tasks.Count && Executor.Add(_boundTasks[Idx + 1]))
+						{
+							Idx++;
+							_boundTasks[Idx].GetTraceMetadata(Scope.Span, string.Format("{0}.", 1 + Idx - FirstIdx));
+						}
+						try
+						{
+							await Executor.ExecuteAsync(Job, BuildProducts, TagNameToFileSet);
+						}
+						catch (Exception Ex)
+						{
+							for (int TaskIdx = FirstIdx; TaskIdx <= Idx; TaskIdx++)
+							{
+								ExceptionUtils.AddContext(Ex, "while executing {0}", _boundTasks[TaskIdx].GetTraceString());
+							}
+							if (_boundTasks[FirstIdx].SourceLocation != null)
+							{
+								ExceptionUtils.AddContext(Ex, "at {0}({1})", _boundTasks[FirstIdx].SourceLocation.File, _boundTasks[FirstIdx].SourceLocation.LineNumber);
+							}
+							throw;
+						}
+					}
+				}
+			}
+
+			// Remove anything that doesn't exist, since these files weren't explicitly tagged
+			BuildProducts.RemoveWhere(x => !FileReference.Exists(x));
+			return true;
+		}
+	}
+}
\ No newline at end of file
diff --git a/Engine/Source/Programs/AutomationTool/BuildGraph/BgSchema.cs b/Engine/Source/Programs/AutomationTool/BuildGraph/BgSchema.cs
new file mode 100644
--- /dev/null
+++ b/Engine/Source/Programs/AutomationTool/BuildGraph/BgSchema.cs
@@ -0,0 +1,1150 @@
+// Copyright Epic Games, Inc. All Rights Reserved.
+
+using System;
+using System.Collections.Generic;
+using System.Diagnostics;
+using System.Linq;
+using System.Xml;
+using System.Xml.Schema;
+using EpicGames.BuildGraph;
+using EpicGames.Core;
+
+#nullable enable
+
+namespace AutomationTool
+{
+	/// <summary>
+	/// Specifies validation that should be performed on a task parameter.
+	/// </summary>
+	public enum TaskParameterValidationType
+	{
+		/// <summary>
+		/// Allow any valid values for the field type.
+		/// </summary>
+		Default,
+
+		/// <summary>
+		/// A list of tag names separated by semicolons
+		/// </summary>
+		TagList,
+
+		/// <summary>
+		/// A file specification, which may contain tags and wildcards.
+		/// </summary>
+		FileSpec,
+	}
+
+	/// <summary>
+	/// Information about a parameter to a task
+	/// </summary>
+	[DebuggerDisplay("{Name}")]
+	public class BgScriptTaskParameter
+	{
+		/// <summary>
+		/// Name of this parameter
+		/// </summary>
+		public string Name { get; }
+
+		/// <summary>
+		/// The type for values assigned to this field
+		/// </summary>
+		public Type ValueType { get; }
+
+		/// <summary>
+		/// The ICollection interface for this type
+		/// </summary>
+		public Type? CollectionType { get; }
+
+		/// <summary>
+		/// Validation type for this field
+		/// </summary>
+		public TaskParameterValidationType ValidationType { get; }
+
+		/// <summary>
+		/// Whether this parameter is optional
+		/// </summary>
+		public bool Optional { get; }
+
+		/// <summary>
+		/// Constructor
+		/// </summary>
+		public BgScriptTaskParameter(string inName, Type inValueType, TaskParameterValidationType inValidationType, bool bInOptional)
+		{
+			Name = inName;
+			ValueType = inValueType;
+			ValidationType = inValidationType;
+			Optional = bInOptional;
+
+			if (ValueType.IsGenericType && ValueType.GetGenericTypeDefinition() == typeof(Nullable<>))
+			{
+				ValueType = ValueType.GetGenericArguments()[0];
+				Optional = true;
+			}
+
+			if (ValueType.IsClass)
+			{
+				foreach (Type interfaceType in ValueType.GetInterfaces())
+				{
+					if (interfaceType.IsGenericType)
+					{
+						Type genericInterfaceType = interfaceType.GetGenericTypeDefinition();
+						if (genericInterfaceType == typeof(ICollection<>))
+						{
+							CollectionType = interfaceType;
+							ValueType = interfaceType.GetGenericArguments()[0];
+						}
+					}
+				}
+			}
+		}
+	}
+
+	/// <summary>
+	/// Helper class to serialize a task from an xml element
+	/// </summary>
+	[DebuggerDisplay("{Name}")]
+	public class BgScriptTask
+	{
+		/// <summary>
+		/// Name of this task
+		/// </summary>
+		public string Name { get; }
+
+		/// <summary>
+		/// Parameters for this task
+		/// </summary>
+		public List<BgScriptTaskParameter> Parameters { get; }
+
+		/// <summary>
+		/// Constructor
+		/// </summary>
+		/// <param name="name">Name of the task</param>
+		/// <param name="parameters">Parameters for the task</param>
+		public BgScriptTask(string name, List<BgScriptTaskParameter> parameters)
+		{
+			Name = name;
+			Parameters = new List<BgScriptTaskParameter>(parameters);
+		}
+	}
+
+#pragma warning disable CS1591
+	/// <summary>
+	/// Enumeration of standard types used in the schema. Avoids hard-coding names.
+	/// </summary>
+	public enum ScriptSchemaStandardType
+	{
+		Graph,
+		Agent,
+		AgentBody,
+		Node,
+		NodeBody,
+		Aggregate,
+		Report,
+		Badge,
+		Label,
+		Notify,
+		Include,
+		Option,
+		EnvVar,
+		Property,
+		Regex,
+		Macro,
+		MacroBody,
+		Extend,
+		Expand,
+		Annotate,
+		Trace,
+		Warning,
+		Error,
+		Name,
+		NameList,
+		Tag,
+		TagList,
+		NameOrTag,
+		NameOrTagList,
+		QualifiedName,
+		BalancedString,
+		Boolean,
+		Integer,
+		LabelChange
+	}
+#pragma warning restore CS1591
+
+	/// <summary>
+	/// Schema for build graph definitions. Stores information about the supported tasks, and allows validating an XML document.
+	/// </summary>
+	public class BgScriptSchema
+	{
+		/// <summary>
+		/// Name of the root element
+		/// </summary>
+		public const string RootElementName = "BuildGraph";
+
+		/// <summary>
+		/// Namespace for the schema
+		/// </summary>
+		public const string NamespaceUri = "http://www.epicgames.com/BuildGraph";
+
+		/// <summary>
+		/// Qualified name for the string type
+		/// </summary>
+		static readonly XmlQualifiedName s_stringTypeName = new XmlQualifiedName("string", "http://www.w3.org/2001/XMLSchema");
+
+		/// <summary>
+		/// The inner xml schema
+		/// </summary>
+		public readonly XmlSchema CompiledSchema;
+
+		/// <summary>
+		/// Characters which are not permitted in names.
+		/// </summary>
+		public const string IllegalNameCharacters = "^<>:\"/\\|?*";
+
+		/// <summary>
+		/// Pattern which matches any name; alphanumeric characters, with single embedded spaces.
+		/// </summary>
+		const string NamePattern = "[^ " + IllegalNameCharacters + "]+( [^ " + IllegalNameCharacters + "]+)*";
+
+		/// <summary>
+		/// Pattern which matches a list of names, separated by semicolons.
+		/// </summary>
+		const string NameListPattern = NamePattern + "(;" + NamePattern + ")*";
+
+		/// <summary>
+		/// Pattern which matches any tag name; a name with a leading '#' character
+		/// </summary>
+		const string TagPattern = "#" + NamePattern;
+
+		/// <summary>
+		/// Pattern which matches a list of tag names, separated by semicolons;
+		/// </summary>
+		const string TagListPattern = TagPattern + "(;" + TagPattern + ")*";
+
+		/// <summary>
+		/// Pattern which matches any name or tag name; a name with a leading '#' character
+		/// </summary>
+		const string NameOrTagPattern = "#?" + NamePattern;
+
+		/// <summary>
+		/// Pattern which matches a list of names or tag names, separated by semicolons;
+		/// </summary>
+		const string NameOrTagListPattern = NameOrTagPattern + "(;" + NameOrTagPattern + ")*";
+
+		/// <summary>
+		/// Pattern which matches a qualified name.
+		/// </summary>
+		const string QualifiedNamePattern = NamePattern + "(\\." + NamePattern + ")*";
+
+		/// <summary>
+		/// Pattern which matches a property name
+		/// </summary>
+		const string PropertyPattern = "\\$\\(" + NamePattern + "\\)";
+
+		/// <summary>
+		/// Pattern which matches balanced parentheses in a string
+		/// </summary>
+		const string StringWithPropertiesPattern = "[^\\$]*" + "(" + "(" + PropertyPattern + "|" + "\\$[^\\(]" + ")" + "[^\\$]*" + ")+" + "\\$?";
+
+		/// <summary>
+		/// Pattern which matches balanced parentheses in a string
+		/// </summary>
+		const string BalancedStringPattern = "[^\\$]*" + "(" + "(" + PropertyPattern + "|" + "\\$[^\\(]" + ")" + "[^\\$]*" + ")*" + "\\$?";
+
+		private BgScriptSchema(XmlSchema schema)
+		{
+			CompiledSchema = schema;
+		}
+
+		/// <summary>
+		/// Constructor
+		/// </summary>
+		/// <param name="tasks">Set of known tasks</param>
+		/// <param name="primitiveTypes">Mapping of task name to information about how to construct it</param>
+		public BgScriptSchema(IEnumerable<BgScriptTask> tasks, List<(Type, ScriptSchemaStandardType)> primitiveTypes)
+		{
+			// Create a lookup from standard types to their qualified names
+			Dictionary<Type, XmlQualifiedName> typeToSchemaTypeName = new Dictionary<Type, XmlQualifiedName>();
+			typeToSchemaTypeName.Add(typeof(string), GetQualifiedTypeName(ScriptSchemaStandardType.BalancedString));
+			typeToSchemaTypeName.Add(typeof(bool), GetQualifiedTypeName(ScriptSchemaStandardType.Boolean));
+			typeToSchemaTypeName.Add(typeof(int), GetQualifiedTypeName(ScriptSchemaStandardType.Integer));
+			foreach ((Type type, ScriptSchemaStandardType schemaType) in primitiveTypes)
+			{
+				typeToSchemaTypeName.Add(type, GetQualifiedTypeName(schemaType));
+			}
+
+			// Create all the custom user types, and add them to the qualified name lookup
+			List<XmlSchemaType> userTypes = new List<XmlSchemaType>();
+			foreach (Type type in tasks.SelectMany(x => x.Parameters).Select(x => x.ValueType))
+			{
+				if (!typeToSchemaTypeName.ContainsKey(type))
+				{
+					if (type.IsClass && type.GetInterfaces().Any(x => x.GetGenericTypeDefinition() == typeof(ICollection<>)))
+					{
+						typeToSchemaTypeName.Add(type, GetQualifiedTypeName(ScriptSchemaStandardType.BalancedString));
+					}
+					else
+					{
+						string name = type.Name + "UserType";
+						XmlSchemaType schemaType = CreateUserType(name, type);
+						userTypes.Add(schemaType);
+						typeToSchemaTypeName.Add(type, new XmlQualifiedName(name, NamespaceUri));
+					}
+				}
+			}
+
+			// Create all the task types
+			Dictionary<string, XmlSchemaComplexType>? taskNameToType = new Dictionary<string, XmlSchemaComplexType>();
+			foreach (BgScriptTask task in tasks)
+			{
+				XmlSchemaComplexType taskType = new XmlSchemaComplexType();
+				taskType.Name = task.Name + "TaskType";
+				foreach (BgScriptTaskParameter parameter in task.Parameters)
+				{
+					XmlQualifiedName? schemaTypeName = GetQualifiedTypeName(parameter.ValidationType);
+					if (schemaTypeName == null)
+					{
+						schemaTypeName = typeToSchemaTypeName[parameter.ValueType];
+					}
+					taskType.Attributes.Add(CreateSchemaAttribute(parameter.Name, schemaTypeName, parameter.Optional ? XmlSchemaUse.Optional : XmlSchemaUse.Required));
+				}
+				taskType.Attributes.Add(CreateSchemaAttribute("If", ScriptSchemaStandardType.BalancedString, XmlSchemaUse.Optional));
+				taskNameToType.Add(task.Name, taskType);
+			}
+
+			// Create the schema object
+			XmlSchema newSchema = new XmlSchema();
+			newSchema.TargetNamespace = NamespaceUri;
+			newSchema.ElementFormDefault = XmlSchemaForm.Qualified;
+			newSchema.Items.Add(CreateSchemaElement(RootElementName, ScriptSchemaStandardType.Graph));
+			newSchema.Items.Add(CreateGraphType());
+			newSchema.Items.Add(CreateAgentType());
+			newSchema.Items.Add(CreateAgentBodyType());
+			newSchema.Items.Add(CreateNodeType());
+			newSchema.Items.Add(CreateNodeBodyType(taskNameToType));
+			newSchema.Items.Add(CreateAggregateType());
+			newSchema.Items.Add(CreateReportType());
+			newSchema.Items.Add(CreateBadgeType());
+			newSchema.Items.Add(CreateLabelType());
+			newSchema.Items.Add(CreateEnumType(GetTypeName(ScriptSchemaStandardType.LabelChange), typeof(BgLabelChange)));
+			newSchema.Items.Add(CreateNotifyType());
+			newSchema.Items.Add(CreateIncludeType());
+			newSchema.Items.Add(CreateOptionType());
+			newSchema.Items.Add(CreateEnvVarType());
+			newSchema.Items.Add(CreatePropertyType());
+			newSchema.Items.Add(CreateRegexType());
+			newSchema.Items.Add(CreateMacroType());
+			newSchema.Items.Add(CreateMacroBodyType(taskNameToType));
+			newSchema.Items.Add(CreateExtendType());
+			newSchema.Items.Add(CreateExpandType());
+			newSchema.Items.Add(CreateAnnotateType());
+			newSchema.Items.Add(CreateDiagnosticType(ScriptSchemaStandardType.Trace));
+			newSchema.Items.Add(CreateDiagnosticType(ScriptSchemaStandardType.Warning));
+			newSchema.Items.Add(CreateDiagnosticType(ScriptSchemaStandardType.Error));
+			newSchema.Items.Add(CreateSimpleTypeFromRegex(GetTypeName(ScriptSchemaStandardType.Name), "(" + NamePattern + "|" + StringWithPropertiesPattern + ")"));
+			newSchema.Items.Add(CreateSimpleTypeFromRegex(GetTypeName(ScriptSchemaStandardType.NameList), "(" + NameListPattern + "|" + StringWithPropertiesPattern + ")"));
+			newSchema.Items.Add(CreateSimpleTypeFromRegex(GetTypeName(ScriptSchemaStandardType.Tag), "(" + TagPattern + "|" + StringWithPropertiesPattern + ")"));
+			newSchema.Items.Add(CreateSimpleTypeFromRegex(GetTypeName(ScriptSchemaStandardType.TagList), "(" + TagListPattern + "|" + StringWithPropertiesPattern + ")"));
+			newSchema.Items.Add(CreateSimpleTypeFromRegex(GetTypeName(ScriptSchemaStandardType.NameOrTag), "(" + NameOrTagPattern + "|" + StringWithPropertiesPattern + ")"));
+			newSchema.Items.Add(CreateSimpleTypeFromRegex(GetTypeName(ScriptSchemaStandardType.NameOrTagList), "(" + NameOrTagListPattern + "|" + StringWithPropertiesPattern + ")"));
+			newSchema.Items.Add(CreateSimpleTypeFromRegex(GetTypeName(ScriptSchemaStandardType.QualifiedName), "(" + QualifiedNamePattern + "|" + StringWithPropertiesPattern + ")"));
+			newSchema.Items.Add(CreateSimpleTypeFromRegex(GetTypeName(ScriptSchemaStandardType.BalancedString), BalancedStringPattern));
+			newSchema.Items.Add(CreateSimpleTypeFromRegex(GetTypeName(ScriptSchemaStandardType.Boolean), "(true|True|false|False|" + StringWithPropertiesPattern + ")"));
+			newSchema.Items.Add(CreateSimpleTypeFromRegex(GetTypeName(ScriptSchemaStandardType.Integer), "(" + "(-?[1-9][0-9]*|0)" + "|" + StringWithPropertiesPattern + ")"));
+			if (taskNameToType != null)
+			{
+				foreach (XmlSchemaComplexType type in taskNameToType.Values)
+				{
+					newSchema.Items.Add(type);
+				}
+			}
+			foreach (XmlSchemaSimpleType type in userTypes)
+			{
+				newSchema.Items.Add(type);
+			}
+
+			// Now that we've finished, compile it and store it to the class
+			XmlSchemaSet newSchemaSet = new XmlSchemaSet();
+			newSchemaSet.Add(newSchema);
+			newSchemaSet.Compile();
+
+			XmlSchema? schema = null;
+			foreach (XmlSchema? newCompiledSchema in newSchemaSet.Schemas())
+			{
+				schema = newCompiledSchema!;
+			}
+			CompiledSchema = schema!;
+		}
+
+		/// <summary>
+		/// Export the schema to a file
+		/// </summary>
+		/// <param name="file"></param>
+		public void Export(FileReference file)
+		{
+			DirectoryReference.CreateDirectory(file.Directory);
+
+			XmlWriterSettings settings = new XmlWriterSettings();
+			settings.Indent = true;
+			settings.IndentChars = "  ";
+			settings.NewLineChars = "\n";
+
+			DirectoryReference.CreateDirectory(file.Directory);
+			using (XmlWriter writer = XmlWriter.Create(file.FullName, settings))
+			{
+				CompiledSchema.Write(writer);
+			}
+		}
+
+		/// <summary>
+		/// Imports a schema from a file
+		/// </summary>
+		/// <param name="file">The XML file to import</param>
+		/// <returns>A <see cref="BgScriptSchema"/> deserialized from the XML file, or null if file doesn't exist</returns>
+		public static BgScriptSchema? Import(FileReference file)
+		{
+			if (!FileReference.Exists(file))
+			{
+				return null;
+			}
+
+			using (XmlReader schemaFile = XmlReader.Create(file.FullName))
+			{
+				BgScriptSchema importedSchema = new BgScriptSchema(XmlSchema.Read(schemaFile, ValidationCallback)!);
+				return importedSchema;
+			}
+		}
+
+		static void ValidationCallback(object? sender, ValidationEventArgs eventArgs)
+		{
+			if (eventArgs.Severity == XmlSeverityType.Warning)
+			{
+				Log.WriteLine(LogEventType.Warning, "WARNING: {0}", eventArgs.Message);
+			}
+			else if (eventArgs.Severity == XmlSeverityType.Error)
+			{
+				Log.WriteLine(LogEventType.Error, "ERROR: {0}", eventArgs.Message);
+			}
+		}
+
+		/// <summary>
+		/// Gets the bare name for the given script type
+		/// </summary>
+		/// <param name="type">Script type to find the name of</param>
+		/// <returns>Name of the schema type that matches the given script type</returns>
+		static string GetTypeName(ScriptSchemaStandardType type)
+		{
+			return type.ToString() + "Type";
+		}
+
+		/// <summary>
+		/// Gets the qualified name for the given script type
+		/// </summary>
+		/// <param name="type">Script type to find the qualified name for</param>
+		/// <returns>Qualified name of the schema type that matches the given script type</returns>
+		static XmlQualifiedName GetQualifiedTypeName(ScriptSchemaStandardType type)
+		{
+			return new XmlQualifiedName(GetTypeName(type), NamespaceUri);
+		}
+
+		/// <summary>
+		/// Gets the qualified name of the schema type for the given type of validation
+		/// </summary>
+		/// <returns>Qualified name for the corresponding schema type</returns>
+		static XmlQualifiedName? GetQualifiedTypeName(TaskParameterValidationType type)
+		{
+			switch (type)
+			{
+				case TaskParameterValidationType.TagList:
+					return GetQualifiedTypeName(ScriptSchemaStandardType.TagList);
+			}
+			return null;
+		}
+
+		/// <summary>
+		/// Creates the schema type representing the graph type
+		/// </summary>
+		/// <returns>Type definition for a graph</returns>
+		static XmlSchemaType CreateGraphType()
+		{
+			XmlSchemaChoice graphChoice = new XmlSchemaChoice();
+			graphChoice.MinOccurs = 0;
+			graphChoice.MaxOccursString = "unbounded";
+			graphChoice.Items.Add(CreateSchemaElement("Include", ScriptSchemaStandardType.Include));
+			graphChoice.Items.Add(CreateSchemaElement("Option", ScriptSchemaStandardType.Option));
+			graphChoice.Items.Add(CreateSchemaElement("EnvVar", ScriptSchemaStandardType.EnvVar));
+			graphChoice.Items.Add(CreateSchemaElement("Property", ScriptSchemaStandardType.Property));
+			graphChoice.Items.Add(CreateSchemaElement("Regex", ScriptSchemaStandardType.Regex));
+			graphChoice.Items.Add(CreateSchemaElement("Macro", ScriptSchemaStandardType.Macro));
+			graphChoice.Items.Add(CreateSchemaElement("Extend", ScriptSchemaStandardType.Extend));
+			graphChoice.Items.Add(CreateSchemaElement("Agent", ScriptSchemaStandardType.Agent));
+			graphChoice.Items.Add(CreateSchemaElement("Aggregate", ScriptSchemaStandardType.Aggregate));
+			graphChoice.Items.Add(CreateSchemaElement("Report", ScriptSchemaStandardType.Report));
+			graphChoice.Items.Add(CreateSchemaElement("Badge", ScriptSchemaStandardType.Badge));
+			graphChoice.Items.Add(CreateSchemaElement("Label", ScriptSchemaStandardType.Label));
+			graphChoice.Items.Add(CreateSchemaElement("Notify", ScriptSchemaStandardType.Notify));
+			graphChoice.Items.Add(CreateSchemaElement("Annotate", ScriptSchemaStandardType.Annotate));
+			graphChoice.Items.Add(CreateSchemaElement("Trace", ScriptSchemaStandardType.Trace));
+			graphChoice.Items.Add(CreateSchemaElement("Warning", ScriptSchemaStandardType.Warning));
+			graphChoice.Items.Add(CreateSchemaElement("Error", ScriptSchemaStandardType.Error));
+			graphChoice.Items.Add(CreateSchemaElement("Expand", ScriptSchemaStandardType.Expand));
+			graphChoice.Items.Add(CreateDoElement(ScriptSchemaStandardType.Graph));
+			graphChoice.Items.Add(CreateSwitchElement(ScriptSchemaStandardType.Graph));
+			graphChoice.Items.Add(CreateForEachElement(ScriptSchemaStandardType.Graph));
+
+			XmlSchemaComplexType graphType = new XmlSchemaComplexType();
+			graphType.Name = GetTypeName(ScriptSchemaStandardType.Graph);
+			graphType.Particle = graphChoice;
+			return graphType;
+		}
+
+		/// <summary>
+		/// Creates the schema type representing the agent type
+		/// </summary>
+		/// <returns>Type definition for an agent</returns>
+		static XmlSchemaType CreateAgentType()
+		{
+			XmlSchemaComplexContentExtension extension = new XmlSchemaComplexContentExtension();
+			extension.BaseTypeName = GetQualifiedTypeName(ScriptSchemaStandardType.AgentBody);
+			extension.Attributes.Add(CreateSchemaAttribute("Name", s_stringTypeName, XmlSchemaUse.Required));
+			extension.Attributes.Add(CreateSchemaAttribute("Type", ScriptSchemaStandardType.NameList, XmlSchemaUse.Optional));
+			extension.Attributes.Add(CreateSchemaAttribute("If", ScriptSchemaStandardType.BalancedString, XmlSchemaUse.Optional));
+
+			XmlSchemaComplexContent contentModel = new XmlSchemaComplexContent();
+			contentModel.Content = extension;
+
+			XmlSchemaComplexType complexType = new XmlSchemaComplexType();
+			complexType.Name = GetTypeName(ScriptSchemaStandardType.Agent);
+			complexType.ContentModel = contentModel;
+			return complexType;
+		}
+
+		/// <summary>
+		/// Creates the schema type representing the contents of agent type
+		/// </summary>
+		/// <returns>Type definition for an agent</returns>
+		static XmlSchemaType CreateAgentBodyType()
+		{
+			XmlSchemaChoice agentChoice = new XmlSchemaChoice();
+			agentChoice.MinOccurs = 0;
+			agentChoice.MaxOccursString = "unbounded";
+			agentChoice.Items.Add(CreateSchemaElement("Property", ScriptSchemaStandardType.Property));
+			agentChoice.Items.Add(CreateSchemaElement("Regex", ScriptSchemaStandardType.Regex));
+			agentChoice.Items.Add(CreateSchemaElement("EnvVar", ScriptSchemaStandardType.EnvVar));
+			agentChoice.Items.Add(CreateSchemaElement("Node", ScriptSchemaStandardType.Node));
+			agentChoice.Items.Add(CreateSchemaElement("Trace", ScriptSchemaStandardType.Trace));
+			agentChoice.Items.Add(CreateSchemaElement("Label", ScriptSchemaStandardType.Label));
+			agentChoice.Items.Add(CreateSchemaElement("Annotate", ScriptSchemaStandardType.Annotate));
+			agentChoice.Items.Add(CreateSchemaElement("Warning", ScriptSchemaStandardType.Warning));
+			agentChoice.Items.Add(CreateSchemaElement("Error", ScriptSchemaStandardType.Error));
+			agentChoice.Items.Add(CreateSchemaElement("Expand", ScriptSchemaStandardType.Expand));
+			agentChoice.Items.Add(CreateDoElement(ScriptSchemaStandardType.AgentBody));
+			agentChoice.Items.Add(CreateSwitchElement(ScriptSchemaStandardType.AgentBody));
+			agentChoice.Items.Add(CreateForEachElement(ScriptSchemaStandardType.AgentBody));
+
+			XmlSchemaComplexType agentType = new XmlSchemaComplexType();
+			agentType.Name = GetTypeName(ScriptSchemaStandardType.AgentBody);
+			agentType.Particle = agentChoice;
+			return agentType;
+		}
+
+		/// <summary>
+		/// Creates the schema type representing the node type
+		/// </summary>
+		/// <returns>Type definition for a node</returns>
+		static XmlSchemaType CreateNodeType()
+		{
+			XmlSchemaComplexContentExtension extension = new XmlSchemaComplexContentExtension();
+			extension.BaseTypeName = GetQualifiedTypeName(ScriptSchemaStandardType.NodeBody);
+			extension.Attributes.Add(CreateSchemaAttribute("Name", ScriptSchemaStandardType.Name, XmlSchemaUse.Required));
+			extension.Attributes.Add(CreateSchemaAttribute("Requires", ScriptSchemaStandardType.NameOrTagList, XmlSchemaUse.Optional));
+			extension.Attributes.Add(CreateSchemaAttribute("Produces", ScriptSchemaStandardType.TagList, XmlSchemaUse.Optional));
+			extension.Attributes.Add(CreateSchemaAttribute("After", ScriptSchemaStandardType.NameOrTagList, XmlSchemaUse.Optional));
+			extension.Attributes.Add(CreateSchemaAttribute("Token", ScriptSchemaStandardType.BalancedString, XmlSchemaUse.Optional));
+			extension.Attributes.Add(CreateSchemaAttribute("If", ScriptSchemaStandardType.BalancedString, XmlSchemaUse.Optional));
+			extension.Attributes.Add(CreateSchemaAttribute("RunEarly", ScriptSchemaStandardType.Boolean, XmlSchemaUse.Optional));
+			extension.Attributes.Add(CreateSchemaAttribute("NotifyOnWarnings", ScriptSchemaStandardType.Boolean, XmlSchemaUse.Optional));
+			extension.Attributes.Add(CreateSchemaAttribute("Annotations", ScriptSchemaStandardType.BalancedString, XmlSchemaUse.Optional));
+
+			XmlSchemaComplexContent contentModel = new XmlSchemaComplexContent();
+			contentModel.Content = extension;
+
+			XmlSchemaComplexType complexType = new XmlSchemaComplexType();
+			complexType.Name = GetTypeName(ScriptSchemaStandardType.Node);
+			complexType.ContentModel = contentModel;
+			return complexType;
+		}
+
+		/// <summary>
+		/// Creates the schema type representing the body of the node type
+		/// </summary>
+		/// <returns>Type definition for a node</returns>
+		static XmlSchemaType CreateNodeBodyType(Dictionary<string, XmlSchemaComplexType> taskNameToType)
+		{
+			XmlSchemaChoice nodeChoice = new XmlSchemaChoice();
+			nodeChoice.MinOccurs = 0;
+			nodeChoice.MaxOccursString = "unbounded";
+			nodeChoice.Items.Add(CreateSchemaElement("Property", ScriptSchemaStandardType.Property));
+			nodeChoice.Items.Add(CreateSchemaElement("Regex", ScriptSchemaStandardType.Regex));
+			nodeChoice.Items.Add(CreateSchemaElement("EnvVar", ScriptSchemaStandardType.EnvVar));
+			nodeChoice.Items.Add(CreateSchemaElement("Trace", ScriptSchemaStandardType.Trace));
+			nodeChoice.Items.Add(CreateSchemaElement("Warning", ScriptSchemaStandardType.Warning));
+			nodeChoice.Items.Add(CreateSchemaElement("Error", ScriptSchemaStandardType.Error));
+			nodeChoice.Items.Add(CreateSchemaElement("Expand", ScriptSchemaStandardType.Expand));
+			nodeChoice.Items.Add(CreateDoElement(ScriptSchemaStandardType.NodeBody));
+			nodeChoice.Items.Add(CreateSwitchElement(ScriptSchemaStandardType.NodeBody));
+			nodeChoice.Items.Add(CreateForEachElement(ScriptSchemaStandardType.NodeBody));
+
+			if (taskNameToType == null)
+			{
+				nodeChoice.Items.Add(new XmlSchemaAny());
+			}
+			else
+			{
+				foreach (KeyValuePair<string, XmlSchemaComplexType> pair in taskNameToType.OrderBy(x => x.Key))
+				{
+					nodeChoice.Items.Add(CreateSchemaElement(pair.Key, new XmlQualifiedName(pair.Value.Name, NamespaceUri)));
+				}
+			}
+
+			XmlSchemaComplexType nodeType = new XmlSchemaComplexType();
+			nodeType.Name = GetTypeName(ScriptSchemaStandardType.NodeBody);
+			nodeType.Particle = nodeChoice;
+			return nodeType;
+		}
+
+		/// <summary>
+		/// Creates the schema type representing the aggregate type
+		/// </summary>
+		/// <returns>Type definition for an aggregate</returns>
+		static XmlSchemaType CreateAggregateType()
+		{
+			XmlSchemaComplexType aggregateType = new XmlSchemaComplexType();
+			aggregateType.Name = GetTypeName(ScriptSchemaStandardType.Aggregate);
+			aggregateType.Attributes.Add(CreateSchemaAttribute("Name", ScriptSchemaStandardType.Name, XmlSchemaUse.Required));
+			aggregateType.Attributes.Add(CreateSchemaAttribute("Label", ScriptSchemaStandardType.BalancedString, XmlSchemaUse.Optional));
+			aggregateType.Attributes.Add(CreateSchemaAttribute("Requires", ScriptSchemaStandardType.NameOrTagList, XmlSchemaUse.Required));
+			aggregateType.Attributes.Add(CreateSchemaAttribute("Include", ScriptSchemaStandardType.NameOrTagList, XmlSchemaUse.Optional));
+			aggregateType.Attributes.Add(CreateSchemaAttribute("Exclude", ScriptSchemaStandardType.NameOrTagList, XmlSchemaUse.Optional));
+			aggregateType.Attributes.Add(CreateSchemaAttribute("If", ScriptSchemaStandardType.BalancedString, XmlSchemaUse.Optional));
+			return aggregateType;
+		}
+
+		/// <summary>
+		/// Creates the schema type representing the report type
+		/// </summary>
+		/// <returns>Type definition for a report</returns>
+		static XmlSchemaType CreateReportType()
+		{
+			XmlSchemaComplexType reportType = new XmlSchemaComplexType();
+			reportType.Name = GetTypeName(ScriptSchemaStandardType.Report);
+			reportType.Attributes.Add(CreateSchemaAttribute("Name", ScriptSchemaStandardType.Name, XmlSchemaUse.Required));
+			reportType.Attributes.Add(CreateSchemaAttribute("Requires", ScriptSchemaStandardType.NameOrTagList, XmlSchemaUse.Required));
+			reportType.Attributes.Add(CreateSchemaAttribute("If", ScriptSchemaStandardType.BalancedString, XmlSchemaUse.Optional));
+			return reportType;
+		}
+
+		/// <summary>
+		/// Creates the schema type representing the badge type
+		/// </summary>
+		/// <returns>Type definition for a badge</returns>
+		static XmlSchemaType CreateBadgeType()
+		{
+			XmlSchemaComplexType badgeType = new XmlSchemaComplexType();
+			badgeType.Name = GetTypeName(ScriptSchemaStandardType.Badge);
+			badgeType.Attributes.Add(CreateSchemaAttribute("Name", ScriptSchemaStandardType.Name, XmlSchemaUse.Required));
+			badgeType.Attributes.Add(CreateSchemaAttribute("Requires", ScriptSchemaStandardType.NameOrTagList, XmlSchemaUse.Optional));
+			badgeType.Attributes.Add(CreateSchemaAttribute("Targets", ScriptSchemaStandardType.NameOrTagList, XmlSchemaUse.Optional));
+			badgeType.Attributes.Add(CreateSchemaAttribute("Project", ScriptSchemaStandardType.BalancedString, XmlSchemaUse.Required));
+			badgeType.Attributes.Add(CreateSchemaAttribute("Change", ScriptSchemaStandardType.BalancedString, XmlSchemaUse.Optional));
+			badgeType.Attributes.Add(CreateSchemaAttribute("If", ScriptSchemaStandardType.BalancedString, XmlSchemaUse.Optional));
+			return badgeType;
+		}
+
+		/// <summary>
+		/// Creates the schema type representing the label type
+		/// </summary>
+		/// <returns>Type definition for a label</returns>
+		static XmlSchemaType CreateLabelType()
+		{
+			XmlSchemaComplexType labelType = new XmlSchemaComplexType();
+			labelType.Name = GetTypeName(ScriptSchemaStandardType.Label);
+			labelType.Attributes.Add(CreateSchemaAttribute("Name", ScriptSchemaStandardType.Name, XmlSchemaUse.Optional));
+			labelType.Attributes.Add(CreateSchemaAttribute("Category", ScriptSchemaStandardType.Name, XmlSchemaUse.Optional));
+			labelType.Attributes.Add(CreateSchemaAttribute("UgsBadge", ScriptSchemaStandardType.Name, XmlSchemaUse.Optional));
+			labelType.Attributes.Add(CreateSchemaAttribute("UgsProject", ScriptSchemaStandardType.BalancedString, XmlSchemaUse.Optional));
+			labelType.Attributes.Add(CreateSchemaAttribute("Change", ScriptSchemaStandardType.LabelChange, XmlSchemaUse.Optional));
+			labelType.Attributes.Add(CreateSchemaAttribute("Requires", ScriptSchemaStandardType.NameOrTagList, XmlSchemaUse.Required));
+			labelType.Attributes.Add(CreateSchemaAttribute("Include", ScriptSchemaStandardType.NameOrTagList, XmlSchemaUse.Optional));
+			labelType.Attributes.Add(CreateSchemaAttribute("Exclude", ScriptSchemaStandardType.NameOrTagList, XmlSchemaUse.Optional));
+			labelType.Attributes.Add(CreateSchemaAttribute("If", ScriptSchemaStandardType.BalancedString, XmlSchemaUse.Optional));
+			return labelType;
+		}
+
+		/// <summary>
+		/// Creates the schema type representing a notifier
+		/// </summary>
+		/// <returns>Type definition for a notifier</returns>
+		static XmlSchemaType CreateNotifyType()
+		{
+			XmlSchemaComplexType aggregateType = new XmlSchemaComplexType();
+			aggregateType.Name = GetTypeName(ScriptSchemaStandardType.Notify);
+			aggregateType.Attributes.Add(CreateSchemaAttribute("Targets", ScriptSchemaStandardType.NameOrTagList, XmlSchemaUse.Optional));
+			aggregateType.Attributes.Add(CreateSchemaAttribute("Except", ScriptSchemaStandardType.NameOrTagList, XmlSchemaUse.Optional));
+			aggregateType.Attributes.Add(CreateSchemaAttribute("Nodes", ScriptSchemaStandardType.NameOrTagList, XmlSchemaUse.Optional));
+			aggregateType.Attributes.Add(CreateSchemaAttribute("Reports", ScriptSchemaStandardType.NameList, XmlSchemaUse.Optional));
+			aggregateType.Attributes.Add(CreateSchemaAttribute("Users", ScriptSchemaStandardType.BalancedString, XmlSchemaUse.Optional));
+			aggregateType.Attributes.Add(CreateSchemaAttribute("Submitters", ScriptSchemaStandardType.BalancedString, XmlSchemaUse.Optional));
+			aggregateType.Attributes.Add(CreateSchemaAttribute("Warnings", ScriptSchemaStandardType.Boolean, XmlSchemaUse.Optional));
+			aggregateType.Attributes.Add(CreateSchemaAttribute("If", ScriptSchemaStandardType.BalancedString, XmlSchemaUse.Optional));
+			aggregateType.Attributes.Add(CreateSchemaAttribute("Absolute", ScriptSchemaStandardType.BalancedString, XmlSchemaUse.Optional));
+			return aggregateType;
+		}
+
+		/// <summary>
+		/// Creates the schema type representing an annotation
+		/// </summary>
+		/// <returns>Type definition for a notifier</returns>
+		static XmlSchemaType CreateAnnotateType()
+		{
+			XmlSchemaComplexType aggregateType = new XmlSchemaComplexType();
+			aggregateType.Name = GetTypeName(ScriptSchemaStandardType.Annotate);
+			aggregateType.Attributes.Add(CreateSchemaAttribute("Targets", ScriptSchemaStandardType.NameOrTagList, XmlSchemaUse.Optional));
+			aggregateType.Attributes.Add(CreateSchemaAttribute("Except", ScriptSchemaStandardType.NameOrTagList, XmlSchemaUse.Optional));
+			aggregateType.Attributes.Add(CreateSchemaAttribute("Nodes", ScriptSchemaStandardType.NameOrTagList, XmlSchemaUse.Optional));
+			aggregateType.Attributes.Add(CreateSchemaAttribute("Values", ScriptSchemaStandardType.BalancedString, XmlSchemaUse.Optional));
+			aggregateType.Attributes.Add(CreateSchemaAttribute("If", ScriptSchemaStandardType.BalancedString, XmlSchemaUse.Optional));
+			return aggregateType;
+		}
+
+		/// <summary>
+		/// Creates the schema type representing an include type
+		/// </summary>
+		/// <returns>Type definition for an include directive</returns>
+		static XmlSchemaType CreateIncludeType()
+		{
+			XmlSchemaComplexType propertyType = new XmlSchemaComplexType();
+			propertyType.Name = GetTypeName(ScriptSchemaStandardType.Include);
+			propertyType.Attributes.Add(CreateSchemaAttribute("Script", ScriptSchemaStandardType.BalancedString, XmlSchemaUse.Required));
+			propertyType.Attributes.Add(CreateSchemaAttribute("If", ScriptSchemaStandardType.BalancedString, XmlSchemaUse.Optional));
+			return propertyType;
+		}
+
+		/// <summary>
+		/// Creates the schema type representing a parameter type
+		/// </summary>
+		/// <returns>Type definition for a parameter</returns>
+		static XmlSchemaType CreateOptionType()
+		{
+			XmlSchemaComplexType optionType = new XmlSchemaComplexType();
+			optionType.Name = GetTypeName(ScriptSchemaStandardType.Option);
+			optionType.Attributes.Add(CreateSchemaAttribute("Name", ScriptSchemaStandardType.Name, XmlSchemaUse.Required));
+			optionType.Attributes.Add(CreateSchemaAttribute("Restrict", s_stringTypeName, XmlSchemaUse.Optional));
+			optionType.Attributes.Add(CreateSchemaAttribute("DefaultValue", s_stringTypeName, XmlSchemaUse.Required));
+			optionType.Attributes.Add(CreateSchemaAttribute("Description", s_stringTypeName, XmlSchemaUse.Required));
+			optionType.Attributes.Add(CreateSchemaAttribute("If", ScriptSchemaStandardType.BalancedString, XmlSchemaUse.Optional));
+			return optionType;
+		}
+
+		/// <summary>
+		/// Creates the schema type representing a environment variable type
+		/// </summary>
+		/// <returns>Type definition for an environment variable property</returns>
+		static XmlSchemaType CreateEnvVarType()
+		{
+			XmlSchemaComplexType envVarType = new XmlSchemaComplexType();
+			envVarType.Name = GetTypeName(ScriptSchemaStandardType.EnvVar);
+			envVarType.Attributes.Add(CreateSchemaAttribute("Name", ScriptSchemaStandardType.Name, XmlSchemaUse.Required));
+			envVarType.Attributes.Add(CreateSchemaAttribute("If", ScriptSchemaStandardType.BalancedString, XmlSchemaUse.Optional));
+			return envVarType;
+		}
+
+		/// <summary>
+		/// Creates the schema type representing a property type
+		/// </summary>
+		/// <returns>Type definition for a property</returns>
+		static XmlSchemaType CreatePropertyType()
+		{
+			XmlSchemaSimpleContentExtension extension = new XmlSchemaSimpleContentExtension();
+			extension.BaseTypeName = s_stringTypeName;
+			extension.Attributes.Add(CreateSchemaAttribute("Name", ScriptSchemaStandardType.Name, XmlSchemaUse.Required));
+			extension.Attributes.Add(CreateSchemaAttribute("Value", s_stringTypeName, XmlSchemaUse.Optional));
+			extension.Attributes.Add(CreateSchemaAttribute("If", ScriptSchemaStandardType.BalancedString, XmlSchemaUse.Optional));
+
+			XmlSchemaSimpleContent contentModel = new XmlSchemaSimpleContent();
+			contentModel.Content = extension;
+
+			XmlSchemaComplexType propertyType = new XmlSchemaComplexType();
+			propertyType.Name = GetTypeName(ScriptSchemaStandardType.Property);
+			propertyType.ContentModel = contentModel;
+			return propertyType;
+		}
+
+		/// <summary>
+		/// Creates the schema type representing a regex type
+		/// </summary>
+		/// <returns>Type definition for a regex</returns>
+		static XmlSchemaType CreateRegexType()
+		{
+			XmlSchemaSimpleContentExtension extension = new XmlSchemaSimpleContentExtension();
+			extension.BaseTypeName = s_stringTypeName;
+
+			extension.Attributes.Add(CreateSchemaAttribute("Input", s_stringTypeName, XmlSchemaUse.Required));
+			extension.Attributes.Add(CreateSchemaAttribute("Pattern", s_stringTypeName, XmlSchemaUse.Required));
+			extension.Attributes.Add(CreateSchemaAttribute("Capture", ScriptSchemaStandardType.NameList, XmlSchemaUse.Required));
+			extension.Attributes.Add(CreateSchemaAttribute("Optional", ScriptSchemaStandardType.Boolean, XmlSchemaUse.Optional));
+			extension.Attributes.Add(CreateSchemaAttribute("If", ScriptSchemaStandardType.BalancedString, XmlSchemaUse.Optional));
+
+			XmlSchemaSimpleContent contentModel = new XmlSchemaSimpleContent();
+			contentModel.Content = extension;
+
+			XmlSchemaComplexType regexType = new XmlSchemaComplexType();
+			regexType.Name = GetTypeName(ScriptSchemaStandardType.Regex);
+			regexType.ContentModel = contentModel;
+			return regexType;
+		}
+
+		/// <summary>
+		/// Creates the schema type representing the macro type
+		/// </summary>
+		/// <returns>Type definition for a node</returns>
+		static XmlSchemaType CreateMacroType()
+		{
+			XmlSchemaComplexContentExtension extension = new XmlSchemaComplexContentExtension();
+			extension.BaseTypeName = GetQualifiedTypeName(ScriptSchemaStandardType.MacroBody);
+			extension.Attributes.Add(CreateSchemaAttribute("Name", ScriptSchemaStandardType.Name, XmlSchemaUse.Required));
+			extension.Attributes.Add(CreateSchemaAttribute("Arguments", ScriptSchemaStandardType.BalancedString, XmlSchemaUse.Optional));
+			extension.Attributes.Add(CreateSchemaAttribute("OptionalArguments", ScriptSchemaStandardType.BalancedString, XmlSchemaUse.Optional));
+
+			XmlSchemaComplexContent contentModel = new XmlSchemaComplexContent();
+			contentModel.Content = extension;
+
+			XmlSchemaComplexType complexType = new XmlSchemaComplexType();
+			complexType.Name = GetTypeName(ScriptSchemaStandardType.Macro);
+			complexType.ContentModel = contentModel;
+			return complexType;
+		}
+
+		/// <summary>
+		/// Creates the schema type representing the macro type
+		/// </summary>
+		/// <returns>Type definition for a node</returns>
+		static XmlSchemaType CreateMacroBodyType(Dictionary<string, XmlSchemaComplexType> taskNameToType)
+		{
+			XmlSchemaChoice macroChoice = new XmlSchemaChoice();
+			macroChoice.MinOccurs = 0;
+			macroChoice.MaxOccursString = "unbounded";
+
+			// Graph scope
+			macroChoice.Items.Add(CreateSchemaElement("Include", ScriptSchemaStandardType.Include));
+			macroChoice.Items.Add(CreateSchemaElement("Option", ScriptSchemaStandardType.Option));
+			macroChoice.Items.Add(CreateSchemaElement("EnvVar", ScriptSchemaStandardType.EnvVar));
+			macroChoice.Items.Add(CreateSchemaElement("Property", ScriptSchemaStandardType.Property));
+			macroChoice.Items.Add(CreateSchemaElement("Regex", ScriptSchemaStandardType.Regex));
+			macroChoice.Items.Add(CreateSchemaElement("Macro", ScriptSchemaStandardType.Macro));
+			macroChoice.Items.Add(CreateSchemaElement("Agent", ScriptSchemaStandardType.Agent));
+			macroChoice.Items.Add(CreateSchemaElement("Aggregate", ScriptSchemaStandardType.Aggregate));
+			macroChoice.Items.Add(CreateSchemaElement("Report", ScriptSchemaStandardType.Report));
+			macroChoice.Items.Add(CreateSchemaElement("Badge", ScriptSchemaStandardType.Badge));
+			macroChoice.Items.Add(CreateSchemaElement("Notify", ScriptSchemaStandardType.Notify));
+			macroChoice.Items.Add(CreateSchemaElement("Trace", ScriptSchemaStandardType.Trace));
+			macroChoice.Items.Add(CreateSchemaElement("Warning", ScriptSchemaStandardType.Warning));
+			macroChoice.Items.Add(CreateSchemaElement("Error", ScriptSchemaStandardType.Error));
+			macroChoice.Items.Add(CreateSchemaElement("Expand", ScriptSchemaStandardType.Expand));
+			macroChoice.Items.Add(CreateSchemaElement("Label", ScriptSchemaStandardType.Label));
+
+			// Agent scope
+			macroChoice.Items.Add(CreateSchemaElement("Node", ScriptSchemaStandardType.Node));
+
+			// Node scope
+			macroChoice.Items.Add(CreateDoElement(ScriptSchemaStandardType.NodeBody));
+			macroChoice.Items.Add(CreateSwitchElement(ScriptSchemaStandardType.NodeBody));
+			macroChoice.Items.Add(CreateForEachElement(ScriptSchemaStandardType.NodeBody));
+			foreach (KeyValuePair<string, XmlSchemaComplexType> pair in taskNameToType.OrderBy(x => x.Key))
+			{
+				macroChoice.Items.Add(CreateSchemaElement(pair.Key, new XmlQualifiedName(pair.Value.Name, NamespaceUri)));
+			}
+
+			XmlSchemaComplexType nodeType = new XmlSchemaComplexType();
+			nodeType.Name = GetTypeName(ScriptSchemaStandardType.MacroBody);
+			nodeType.Particle = macroChoice;
+			return nodeType;
+		}
+
+		/// <summary>
+		/// Creates the schema type representing the macro type
+		/// </summary>
+		/// <returns>Type definition for a node</returns>
+		static XmlSchemaType CreateExtendType()
+		{
+			XmlSchemaComplexContentExtension extension = new XmlSchemaComplexContentExtension();
+			extension.BaseTypeName = GetQualifiedTypeName(ScriptSchemaStandardType.MacroBody);
+			extension.Attributes.Add(CreateSchemaAttribute("Name", ScriptSchemaStandardType.Name, XmlSchemaUse.Required));
+
+			XmlSchemaComplexContent contentModel = new XmlSchemaComplexContent();
+			contentModel.Content = extension;
+
+			XmlSchemaComplexType complexType = new XmlSchemaComplexType();
+			complexType.Name = GetTypeName(ScriptSchemaStandardType.Extend);
+			complexType.ContentModel = contentModel;
+			return complexType;
+		}
+
+		/// <summary>
+		/// Creates the schema type representing a macro expansion
+		/// </summary>
+		/// <returns>Type definition for expanding a macro</returns>
+		static XmlSchemaType CreateExpandType()
+		{
+			XmlSchemaAnyAttribute anyAttribute = new XmlSchemaAnyAttribute();
+			anyAttribute.ProcessContents = XmlSchemaContentProcessing.Skip;
+
+			XmlSchemaComplexType propertyType = new XmlSchemaComplexType();
+			propertyType.Name = GetTypeName(ScriptSchemaStandardType.Expand);
+			propertyType.Attributes.Add(CreateSchemaAttribute("Name", ScriptSchemaStandardType.BalancedString, XmlSchemaUse.Required));
+			propertyType.Attributes.Add(CreateSchemaAttribute("If", ScriptSchemaStandardType.BalancedString, XmlSchemaUse.Optional));
+			propertyType.AnyAttribute = anyAttribute;
+			return propertyType;
+		}
+
+		/// <summary>
+		/// Creates the schema type representing a warning or error type
+		/// </summary>
+		/// <returns>Type definition for a warning</returns>
+		static XmlSchemaType CreateDiagnosticType(ScriptSchemaStandardType standardType)
+		{
+			XmlSchemaComplexType propertyType = new XmlSchemaComplexType();
+			propertyType.Name = GetTypeName(standardType);
+			propertyType.Attributes.Add(CreateSchemaAttribute("Message", ScriptSchemaStandardType.BalancedString, XmlSchemaUse.Required));
+			propertyType.Attributes.Add(CreateSchemaAttribute("If", ScriptSchemaStandardType.BalancedString, XmlSchemaUse.Optional));
+			return propertyType;
+		}
+
+		/// <summary>
+		/// Creates an element representing a conditional "Do" block, which recursively contains another type
+		/// </summary>
+		/// <param name="innerType">The base type for the do block to contain</param>
+		/// <returns>New schema element for the block</returns>
+		static XmlSchemaElement CreateDoElement(ScriptSchemaStandardType innerType)
+		{
+			XmlSchemaComplexContentExtension extension = new XmlSchemaComplexContentExtension();
+			extension.BaseTypeName = GetQualifiedTypeName(innerType);
+			extension.Attributes.Add(CreateSchemaAttribute("If", ScriptSchemaStandardType.BalancedString, XmlSchemaUse.Optional));
+
+			XmlSchemaComplexContent contentModel = new XmlSchemaComplexContent();
+			contentModel.Content = extension;
+
+			XmlSchemaComplexType schemaType = new XmlSchemaComplexType();
+			schemaType.ContentModel = contentModel;
+
+			XmlSchemaElement element = new XmlSchemaElement();
+			element.Name = "Do";
+			element.SchemaType = schemaType;
+			return element;
+		}
+
+		/// <summary>
+		/// Creates an element representing a conditional "Switch" block, which recursively contains another type
+		/// </summary>
+		/// <param name="innerType">The base type for the do block to contain</param>
+		/// <returns>New schema element for the block</returns>
+		static XmlSchemaElement CreateSwitchElement(ScriptSchemaStandardType innerType)
+		{
+			// Create the "Option" element
+			XmlSchemaComplexContentExtension caseExtension = new XmlSchemaComplexContentExtension();
+			caseExtension.BaseTypeName = GetQualifiedTypeName(innerType);
+			caseExtension.Attributes.Add(CreateSchemaAttribute("If", ScriptSchemaStandardType.BalancedString, XmlSchemaUse.Required));
+
+			XmlSchemaComplexContent caseContentModel = new XmlSchemaComplexContent();
+			caseContentModel.Content = caseExtension;
+
+			XmlSchemaComplexType caseSchemaType = new XmlSchemaComplexType();
+			caseSchemaType.ContentModel = caseContentModel;
+
+			XmlSchemaElement caseElement = new XmlSchemaElement();
+			caseElement.Name = "Case";
+			caseElement.SchemaType = caseSchemaType;
+			caseElement.MinOccurs = 0;
+			caseElement.MaxOccursString = "unbounded";
+
+			// Create the "Otherwise" element
+			XmlSchemaElement otherwiseElement = new XmlSchemaElement();
+			otherwiseElement.Name = "Default";
+			otherwiseElement.SchemaTypeName = GetQualifiedTypeName(innerType);
+			otherwiseElement.MinOccurs = 0;
+			otherwiseElement.MaxOccurs = 1;
+
+			// Create the "Switch" element
+			XmlSchemaSequence switchSequence = new XmlSchemaSequence();
+			switchSequence.Items.Add(caseElement);
+			switchSequence.Items.Add(otherwiseElement);
+
+			XmlSchemaComplexType switchSchemaType = new XmlSchemaComplexType();
+			switchSchemaType.Particle = switchSequence;
+
+			XmlSchemaElement switchElement = new XmlSchemaElement();
+			switchElement.Name = "Switch";
+			switchElement.SchemaType = switchSchemaType;
+			return switchElement;
+		}
+
+		/// <summary>
+		/// Creates an element representing a conditional "ForEach" block, which recursively contains another type
+		/// </summary>
+		/// <param name="innerType">The base type for the foreach block to contain</param>
+		/// <returns>New schema element for the block</returns>
+		static XmlSchemaElement CreateForEachElement(ScriptSchemaStandardType innerType)
+		{
+			XmlSchemaComplexContentExtension extension = new XmlSchemaComplexContentExtension();
+			extension.BaseTypeName = GetQualifiedTypeName(innerType);
+			extension.Attributes.Add(CreateSchemaAttribute("Name", ScriptSchemaStandardType.BalancedString, XmlSchemaUse.Required));
+			extension.Attributes.Add(CreateSchemaAttribute("Values", ScriptSchemaStandardType.BalancedString, XmlSchemaUse.Required));
+			extension.Attributes.Add(CreateSchemaAttribute("Separator", ScriptSchemaStandardType.BalancedString, XmlSchemaUse.Optional));
+			extension.Attributes.Add(CreateSchemaAttribute("If", ScriptSchemaStandardType.BalancedString, XmlSchemaUse.Optional));
+
+			XmlSchemaComplexContent contentModel = new XmlSchemaComplexContent();
+			contentModel.Content = extension;
+
+			XmlSchemaComplexType schemaType = new XmlSchemaComplexType();
+			schemaType.ContentModel = contentModel;
+
+			XmlSchemaElement element = new XmlSchemaElement();
+			element.Name = "ForEach";
+			element.SchemaType = schemaType;
+			return element;
+		}
+
+		/// <summary>
+		/// Constructs an XmlSchemaElement and initializes it with the given parameters
+		/// </summary>
+		/// <param name="name">Element name</param>
+		/// <param name="schemaType">Type enumeration for the attribute</param>
+		/// <returns>A new XmlSchemaElement object</returns>
+		static XmlSchemaElement CreateSchemaElement(string name, ScriptSchemaStandardType schemaType)
+		{
+			return CreateSchemaElement(name, GetQualifiedTypeName(schemaType));
+		}
+
+		/// <summary>
+		/// Constructs an XmlSchemaElement and initializes it with the given parameters
+		/// </summary>
+		/// <param name="name">Element name</param>
+		/// <param name="schemaTypeName">Qualified name of the type for this element</param>
+		/// <returns>A new XmlSchemaElement object</returns>
+		static XmlSchemaElement CreateSchemaElement(string name, XmlQualifiedName schemaTypeName)
+		{
+			XmlSchemaElement element = new XmlSchemaElement();
+			element.Name = name;
+			element.SchemaTypeName = schemaTypeName;
+			return element;
+		}
+
+		/// <summary>
+		/// Constructs an XmlSchemaAttribute and initialize it with the given parameters
+		/// </summary>
+		/// <param name="name">The attribute name</param>
+		/// <param name="schemaType">Type enumeration for the attribute</param>
+		/// <param name="use">Whether the attribute is required or optional</param>
+		/// <returns>A new XmlSchemaAttribute object</returns>
+		static XmlSchemaAttribute CreateSchemaAttribute(string name, ScriptSchemaStandardType schemaType, XmlSchemaUse use)
+		{
+			return CreateSchemaAttribute(name, GetQualifiedTypeName(schemaType), use);
+		}
+
+		/// <summary>
+		/// Constructs an XmlSchemaAttribute and initialize it with the given parameters
+		/// </summary>
+		/// <param name="name">The attribute name</param>
+		/// <param name="schemaTypeName">Qualified name of the type for this attribute</param>
+		/// <param name="use">Whether the attribute is required or optional</param>
+		/// <returns>The new attribute</returns>
+		static XmlSchemaAttribute CreateSchemaAttribute(string name, XmlQualifiedName schemaTypeName, XmlSchemaUse use)
+		{
+			XmlSchemaAttribute attribute = new XmlSchemaAttribute();
+			attribute.Name = name;
+			attribute.SchemaTypeName = schemaTypeName;
+			attribute.Use = use;
+			return attribute;
+		}
+
+		/// <summary>
+		/// Creates a simple type that is the union of two other types
+		/// </summary>
+		/// <param name="name">The name of the type</param>
+		/// <param name="validTypes">List of valid types for the union</param>
+		/// <returns>A simple type which will match the given pattern</returns>
+		static XmlSchemaSimpleType CreateSimpleTypeFromUnion(string name, params XmlSchemaType[] validTypes)
+		{
+			XmlSchemaSimpleTypeUnion union = new XmlSchemaSimpleTypeUnion();
+			foreach (XmlSchemaType validType in validTypes)
+			{
+				union.BaseTypes.Add(validType);
+			}
+
+			XmlSchemaSimpleType unionType = new XmlSchemaSimpleType();
+			unionType.Name = name;
+			unionType.Content = union;
+			return unionType;
+		}
+
+		/// <summary>
+		/// Creates a simple type that matches a regex
+		/// </summary>
+		/// <param name="name">Name of the new type</param>
+		/// <param name="pattern">Regex pattern to match</param>
+		/// <returns>A simple type which will match the given pattern</returns>
+		static XmlSchemaSimpleType CreateSimpleTypeFromRegex(string? name, string pattern)
+		{
+			XmlSchemaPatternFacet patternFacet = new XmlSchemaPatternFacet();
+			patternFacet.Value = pattern;
+
+			XmlSchemaSimpleTypeRestriction restriction = new XmlSchemaSimpleTypeRestriction();
+			restriction.BaseTypeName = s_stringTypeName;
+			restriction.Facets.Add(patternFacet);
+
+			XmlSchemaSimpleType simpleType = new XmlSchemaSimpleType();
+			simpleType.Name = name;
+			simpleType.Content = restriction;
+			return simpleType;
+		}
+
+		/// <summary>
+		/// Create a schema type for the given user type. Currently only handles enumerations.
+		/// </summary>
+		/// <param name="name">Name for the new type</param>
+		/// <param name="type">CLR type information to create a schema type for</param>
+		static XmlSchemaType CreateUserType(string name, Type type)
+		{
+			if (type.IsEnum)
+			{
+				return CreateSimpleTypeFromUnion(name, CreateEnumType(null, type), CreateSimpleTypeFromRegex(null, StringWithPropertiesPattern));
+			}
+			else
+			{
+				throw new Exception($"Cannot create custom type in schema for '{type.Name}'");
+			}
+		}
+
+		/// <summary>
+		/// Create a schema type for the given enum.
+		/// </summary>
+		/// <param name="name">Name for the new type</param>
+		/// <param name="type">CLR type information to create a schema type for</param>
+		static XmlSchemaType CreateEnumType(string? name, Type type)
+		{
+			XmlSchemaSimpleTypeRestriction restriction = new XmlSchemaSimpleTypeRestriction();
+			restriction.BaseTypeName = s_stringTypeName;
+
+			foreach (string enumName in Enum.GetNames(type))
+			{
+				XmlSchemaEnumerationFacet facet = new XmlSchemaEnumerationFacet();
+				facet.Value = enumName;
+				restriction.Facets.Add(facet);
+			}
+
+			XmlSchemaSimpleType schemaType = new XmlSchemaSimpleType();
+			schemaType.Name = name;
+			schemaType.Content = restriction;
+			return schemaType;
+		}
+	}
+}
\ No newline at end of file
diff --git a/Engine/Source/Programs/AutomationTool/BuildGraph/BgScriptReader.cs b/Engine/Source/Programs/AutomationTool/BuildGraph/BgScriptReader.cs
new file mode 100644
--- /dev/null
+++ b/Engine/Source/Programs/AutomationTool/BuildGraph/BgScriptReader.cs
@@ -0,0 +1,2214 @@
+// Copyright Epic Games, Inc. All Rights Reserved.
+
+using System;
+using System.Collections.Generic;
+using System.Diagnostics.CodeAnalysis;
+using System.IO;
+using System.Linq;
+using System.Text;
+using System.Text.RegularExpressions;
+using System.Threading.Tasks;
+using System.Xml;
+using System.Xml.Schema;
+using EpicGames.BuildGraph;
+using EpicGames.Core;
+using Microsoft.Extensions.Logging;
+
+#nullable enable
+
+namespace AutomationTool
+{
+	/// <summary>
+	/// Location of an element within a file
+	/// </summary>
+	public class BgScriptLocation
+	{
+		/// <summary>
+		/// The file containing this element
+		/// </summary>
+		public string File { get; }
+
+		/// <summary>
+		/// Native file object
+		/// </summary>
+		public object NativeFile { get; }
+
+		/// <summary>
+		/// The line number containing this element
+		/// </summary>
+		public int LineNumber { get; }
+
+		/// <summary>
+		/// Constructor
+		/// </summary>
+		/// <param name="file"></param>
+		/// <param name="nativeFile"></param>
+		/// <param name="lineNumber"></param>
+		public BgScriptLocation(string file, object nativeFile, int lineNumber)
+		{
+			File = file;
+			NativeFile = nativeFile;
+			LineNumber = lineNumber;
+		}
+	}
+
+	/// <summary>
+	/// Implementation of XmlDocument which preserves line numbers for its elements
+	/// </summary>
+	public class BgScriptDocument : XmlDocument
+	{
+		/// <summary>
+		/// The file being read
+		/// </summary>
+		string File { get; }
+
+		/// <summary>
+		/// Native file representation
+		/// </summary>
+		object NativeFile { get; }
+
+		/// <summary>
+		/// Interface to the LineInfo on the active XmlReader
+		/// </summary>
+		IXmlLineInfo? _lineInfo;
+
+		/// <summary>
+		/// Set to true if the reader encounters an error
+		/// </summary>
+		bool _bHasErrors;
+
+		/// <summary>
+		/// Logger for validation errors
+		/// </summary>
+		ILogger Logger { get; }
+
+		/// <summary>
+		/// Private constructor. Use ScriptDocument.Load to read an XML document.
+		/// </summary>
+		BgScriptDocument(string inFile, object inNativeFile, ILogger inLogger)
+		{
+			File = inFile;
+			NativeFile = inNativeFile;
+			Logger = inLogger;
+		}
+
+		/// <summary>
+		/// Overrides XmlDocument.CreateElement() to construct ScriptElements rather than XmlElements
+		/// </summary>
+		public override XmlElement CreateElement(string? prefix, string? localName, string? namespaceUri)
+		{
+			BgScriptLocation location = new BgScriptLocation(File, NativeFile, _lineInfo!.LineNumber);
+			return new BgScriptElement(location, prefix!, localName!, namespaceUri!, this);
+		}
+
+		/// <summary>
+		/// Loads a script document from the given file
+		/// </summary>
+		/// <param name="file">The file to load</param>
+		/// <param name="nativeFile"></param>
+		/// <param name="data"></param>
+		/// <param name="schema">The schema to validate against</param>
+		/// <param name="logger">Logger for output messages</param>
+		/// <param name="outDocument">If successful, the document that was read</param>
+		/// <returns>True if the document could be read, false otherwise</returns>
+		public static bool TryRead(string file, object nativeFile, byte[] data, BgScriptSchema schema, ILogger logger, [NotNullWhen(true)] out BgScriptDocument? outDocument)
+		{
+			BgScriptDocument document = new BgScriptDocument(file, nativeFile, logger);
+
+			XmlReaderSettings settings = new XmlReaderSettings();
+			if (schema != null)
+			{
+				settings.Schemas.Add(schema.CompiledSchema);
+				settings.ValidationType = ValidationType.Schema;
+				settings.ValidationEventHandler += document.ValidationEvent;
+			}
+
+			using (MemoryStream stream = new MemoryStream(data))
+			using (XmlReader reader = XmlReader.Create(stream, settings))
+			{
+				// Read the document
+				document._lineInfo = (IXmlLineInfo)reader;
+				try
+				{
+					document.Load(reader);
+				}
+				catch (XmlException ex)
+				{
+					if (!document._bHasErrors)
+					{
+						BgScriptLocation location = new BgScriptLocation(file, nativeFile, ex.LineNumber);
+						logger.LogScriptError(location, "{Message}", ex.Message);
+						document._bHasErrors = true;
+					}
+				}
+
+				// If we hit any errors while parsing
+				if (document._bHasErrors)
+				{
+					outDocument = null;
+					return false;
+				}
+
+				// Check that the root element is valid. If not, we didn't actually validate against the schema.
+				if (document.DocumentElement!.Name != BgScriptSchema.RootElementName)
+				{
+					BgScriptLocation location = new BgScriptLocation(file, nativeFile, 1);
+					logger.LogScriptError(location, "Script does not have a root element called '{ElementName}'", BgScriptSchema.RootElementName);
+					outDocument = null;
+					return false;
+				}
+				if (document.DocumentElement.NamespaceURI != BgScriptSchema.NamespaceUri)
+				{
+					BgScriptLocation location = new BgScriptLocation(file, nativeFile, 1);
+					logger.LogScriptError(location, "Script root element is not in the '{Namespace}' namespace (add the xmlns=\"{NewNamespace}\" attribute)", BgScriptSchema.NamespaceUri, BgScriptSchema.NamespaceUri);
+					outDocument = null;
+					return false;
+				}
+			}
+
+			outDocument = document;
+			return true;
+		}
+
+		/// <summary>
+		/// Callback for validation errors in the document
+		/// </summary>
+		/// <param name="sender">Standard argument for ValidationEventHandler</param>
+		/// <param name="args">Standard argument for ValidationEventHandler</param>
+		void ValidationEvent(object? sender, ValidationEventArgs args)
+		{
+			BgScriptLocation location = new BgScriptLocation(File, NativeFile, args.Exception.LineNumber);
+			if (args.Severity == XmlSeverityType.Warning)
+			{
+				Logger.LogScriptWarning(location, "{Message}", args.Message);
+			}
+			else
+			{
+				Logger.LogScriptError(location, "{Message}", args.Message);
+				_bHasErrors = true;
+			}
+		}
+	}
+
+	/// <summary>
+	/// Implementation of XmlElement which preserves line numbers
+	/// </summary>
+	public class BgScriptElement : XmlElement
+	{
+		/// <summary>
+		/// Location of the element within the file
+		/// </summary>
+		public BgScriptLocation Location { get; }
+
+		/// <summary>
+		/// Constructor
+		/// </summary>
+		public BgScriptElement(BgScriptLocation location, string prefix, string localName, string namespaceUri, BgScriptDocument document)
+			: base(prefix, localName, namespaceUri, document)
+		{
+			Location = location;
+		}
+	}
+
+	/// <summary>
+	/// Stores information about a script function that has been declared
+	/// </summary>
+	class BgScriptMacro
+	{
+		/// <summary>
+		/// Name of the function
+		/// </summary>
+		public string Name { get; }
+
+		/// <summary>
+		/// Element where the function was declared
+		/// </summary>
+		public List<BgScriptElement> Elements { get; } = new List<BgScriptElement>();
+
+		/// <summary>
+		/// The total number of arguments
+		/// </summary>
+		public int NumArguments { get; }
+
+		/// <summary>
+		/// Number of arguments that are required
+		/// </summary>
+		public int NumRequiredArguments { get; }
+
+		/// <summary>
+		/// Maps an argument name to its type
+		/// </summary>
+		public IReadOnlyDictionary<string, int> ArgumentNameToIndex { get; }
+
+		/// <summary>
+		/// Constructor
+		/// </summary>
+		/// <param name="name">Name of the function</param>
+		/// <param name="element">Element containing the function definition</param>
+		/// <param name="argumentNameToIndex">Map of argument name to index</param>
+		/// <param name="numRequiredArguments">Number of arguments that are required. Indices 0 to NumRequiredArguments - 1 are required.</param>
+		public BgScriptMacro(string name, BgScriptElement element, Dictionary<string, int> argumentNameToIndex, int numRequiredArguments)
+		{
+			Name = name;
+			Elements.Add(element);
+			NumArguments = argumentNameToIndex.Count;
+			NumRequiredArguments = numRequiredArguments;
+			ArgumentNameToIndex = argumentNameToIndex;
+		}
+	}
+
+	/// <summary>
+	/// Context used for parsing scripts
+	/// </summary>
+	public interface IBgScriptReaderContext
+	{
+		/// <summary>
+		/// Tests whether the given file or directory exists
+		/// </summary>
+		/// <param name="path">Path to a file</param>
+		/// <returns>True if the file exists</returns>
+		Task<bool> ExistsAsync(string path);
+
+		/// <summary>
+		/// Tries to read a file from the given path
+		/// </summary>
+		/// <param name="path">Path of the file to read</param>
+		/// <returns></returns>
+		Task<byte[]?> ReadAsync(string path);
+
+		/// <summary>
+		/// Finds files matching the given pattern
+		/// </summary>
+		/// <param name="pattern"></param>
+		/// <returns></returns>
+		Task<string[]> FindAsync(string pattern);
+
+		/// <summary>
+		/// Converts a path to its native form, for display to the user
+		/// </summary>
+		/// <param name="path">Path to format</param>
+		/// <returns></returns>
+		object GetNativePath(string path);
+
+		/// <summary>
+		/// Combines one path with another
+		/// </summary>
+		/// <param name="path">The base file path</param>
+		/// <param name="next">Relative path to the next file</param>
+		/// <returns></returns>
+		string CombinePaths(string path, string next);
+	}
+
+	/// <summary>
+	/// Extension methods for writing script error messages
+	/// </summary>
+	public static class BgScriptExtensions
+	{
+		/// <summary>
+		/// Utility method to log a script error at a particular location
+		/// </summary>
+		public static void LogScriptError(this ILogger logger, BgScriptLocation location, string format, params object[] args)
+		{
+			object[] allArgs = new object[args.Length + 2];
+			allArgs[0] = location.NativeFile;
+			allArgs[1] = location.LineNumber;
+			args.CopyTo(allArgs, 2);
+			logger.LogError($"{{Script}}({{Line}}): error: {format}", allArgs);
+		}
+
+		/// <summary>
+		/// Utility method to log a script warning at a particular location
+		/// </summary>
+		public static void LogScriptWarning(this ILogger logger, BgScriptLocation location, string format, params object[] args)
+		{
+			object[] allArgs = new object[args.Length + 2];
+			allArgs[0] = location.NativeFile;
+			allArgs[1] = location.LineNumber;
+			args.CopyTo(allArgs, 2);
+			logger.LogWarning($"{{Script}}({{Line}}): warning: {format}", allArgs);
+		}
+	}
+
+	/// <summary>
+	/// Reader for build graph definitions. Instanced to contain temporary state; public interface is through ScriptReader.TryRead().
+	/// </summary>
+	public abstract class BgScriptReaderBase
+	{
+		/// <summary>
+		/// Interface used for reading files
+		/// </summary>
+		protected IBgScriptReaderContext Context { get; }
+
+		/// <summary>
+		/// List of property name to value lookups. Modifications to properties are scoped to nodes and agents. EnterScope() pushes an empty dictionary onto the end of this list, and LeaveScope() removes one. 
+		/// ExpandProperties() searches from last to first lookup when trying to resolve a property name, and takes the first it finds.
+		/// </summary>
+		protected List<Dictionary<string, string>> ScopedProperties { get; } = new List<Dictionary<string, string>>();
+
+		/// <summary>
+		/// When declaring a property in a nested scope, we enter its name into a set for each parent scope which prevents redeclaration in an OUTER scope later. Subsequent NESTED scopes can redeclare it.
+		/// The former is likely a coding error, since it implies that the scope of the variable was meant to be further out, whereas the latter is common for temporary and loop variables.
+		/// </summary>
+		readonly List<HashSet<string>> _shadowProperties = new List<HashSet<string>>();
+
+		/// <summary>
+		/// Maps from a function name to its definition
+		/// </summary>
+		readonly Dictionary<string, BgScriptMacro> _macroNameToDefinition = new Dictionary<string, BgScriptMacro>();
+
+		/// <summary>
+		/// Schema for the script
+		/// </summary>
+		BgScriptSchema Schema { get; }
+
+		/// <summary>
+		/// Logger for diagnostic messages
+		/// </summary>
+		protected ILogger Logger { get; }
+
+		/// <summary>
+		/// The number of errors encountered during processing so far
+		/// </summary>
+		public int NumErrors { get; private set; }
+
+		/// <summary>
+		/// Private constructor. Use ScriptReader.TryRead() to read a script file.
+		/// </summary>
+		/// <param name="context">Context object</param>
+		/// <param name="schema">Schema for the script</param>
+		/// <param name="logger">Logger for diagnostic messages</param>
+		protected BgScriptReaderBase(IBgScriptReaderContext context, BgScriptSchema schema, ILogger logger)
+		{
+			Context = context;
+			Schema = schema;
+			Logger = logger;
+
+			EnterScope();
+		}
+
+		/// <summary>
+		/// Read the script from the given file
+		/// </summary>
+		/// <param name="file">File to read from</param>
+		protected async Task<bool> TryReadAsync(string file)
+		{
+			// Get the data for this file
+			byte[]? data = await Context.ReadAsync(file);
+			if (data == null)
+			{
+				Logger.LogError("Unable to open file {File}", file);
+				NumErrors++;
+				return false;
+			}
+
+			// Read the document and validate it against the schema
+			BgScriptDocument? document;
+			if (!BgScriptDocument.TryRead(file, Context.GetNativePath(file), data, Schema, Logger, out document))
+			{
+				NumErrors++;
+				return false;
+			}
+
+			// Read the root BuildGraph element
+			await ReadGraphBodyAsync(document.DocumentElement!);
+			return true;
+		}
+
+		/// <summary>
+		/// Reads the contents of a graph
+		/// </summary>
+		/// <param name="element">The parent element to read from</param>
+		async Task ReadGraphBodyAsync(XmlElement element)
+		{
+			foreach (BgScriptElement childElement in element.ChildNodes.OfType<BgScriptElement>())
+			{
+				switch (childElement.Name)
+				{
+					case "Include":
+						await ReadIncludeAsync(childElement);
+						break;
+					case "Option":
+						await ReadOptionAsync(childElement);
+						break;
+					case "Property":
+						await ReadPropertyAsync(childElement);
+						break;
+					case "Regex":
+						await ReadRegexAsync(childElement);
+						break;
+					case "EnvVar":
+						await ReadEnvVarAsync(childElement);
+						break;
+					case "Macro":
+						ReadMacro(childElement);
+						break;
+					case "Extend":
+						await ReadExtendAsync(childElement);
+						break;
+					case "Agent":
+						await ReadAgentAsync(childElement);
+						break;
+					case "Aggregate":
+						await ReadAggregateAsync(childElement);
+						break;
+					case "Report":
+						await ReadReportAsync(childElement);
+						break;
+					case "Badge":
+						await ReadBadgeAsync(childElement);
+						break;
+					case "Label":
+						await ReadLabelAsync(childElement);
+						break;
+					case "Notify":
+						await ReadNotifierAsync(childElement);
+						break;
+					case "Trace":
+						await ReadDiagnosticAsync(childElement, LogLevel.Information);
+						break;
+					case "Warning":
+						await ReadDiagnosticAsync(childElement, LogLevel.Warning);
+						break;
+					case "Error":
+						await ReadDiagnosticAsync(childElement, LogLevel.Error);
+						break;
+					case "Do":
+						await ReadBlockAsync(childElement, ReadGraphBodyAsync);
+						break;
+					case "Switch":
+						await ReadSwitchAsync(childElement, ReadGraphBodyAsync);
+						break;
+					case "ForEach":
+						await ReadForEachAsync(childElement, ReadGraphBodyAsync);
+						break;
+					case "Expand":
+						await ReadExpandAsync(childElement, ReadGraphBodyAsync);
+						break;
+					case "Annotate":
+						await ReadAnnotationAsync(childElement);
+						break;
+					default:
+						LogError(childElement, "Invalid element '{0}'", childElement.Name);
+						break;
+				}
+			}
+		}
+
+		/// <summary>
+		/// Push a new property scope onto the stack
+		/// </summary>
+		protected void EnterScope()
+		{
+			ScopedProperties.Add(new Dictionary<string, string>(StringComparer.OrdinalIgnoreCase));
+			_shadowProperties.Add(new HashSet<string>(StringComparer.OrdinalIgnoreCase));
+		}
+
+		/// <summary>
+		/// Pop a property scope from the stack
+		/// </summary>
+		protected void LeaveScope()
+		{
+			ScopedProperties.RemoveAt(ScopedProperties.Count - 1);
+			_shadowProperties.RemoveAt(_shadowProperties.Count - 1);
+		}
+
+		/// <summary>
+		/// Sets a property value in the current scope
+		/// </summary>
+		/// <param name="element">Element containing the property assignment. Used for error messages if the property is shadowed in another scope.</param>
+		/// <param name="name">Name of the property</param>
+		/// <param name="value">Value for the property</param>
+		protected void SetPropertyValue(BgScriptElement element, string name, string value)
+		{
+			// Find the scope containing this property, defaulting to the current scope
+			int scopeIdx = 0;
+			while (scopeIdx < ScopedProperties.Count - 1 && !ScopedProperties[scopeIdx].ContainsKey(name))
+			{
+				scopeIdx++;
+			}
+
+			// Make sure this property name was not already used in a child scope; it likely indicates an error.
+			if (_shadowProperties[scopeIdx].Contains(name))
+			{
+				LogError(element, "Property '{0}' was already used in a child scope. Move this definition before the previous usage if they are intended to share scope, or use a different name.", name);
+			}
+			else
+			{
+				// Make sure it's added to the shadow property list for every parent scope
+				for (int idx = 0; idx < scopeIdx; idx++)
+				{
+					_shadowProperties[idx].Add(name);
+				}
+				ScopedProperties[scopeIdx][name] = value;
+			}
+		}
+
+		/// <summary>
+		/// Tries to get the value of a property
+		/// </summary>
+		/// <param name="name">Name of the property</param>
+		/// <param name="value">On success, contains the value of the property. Set to null otherwise.</param>
+		/// <returns>True if the property was found, false otherwise</returns>
+		protected bool TryGetPropertyValue(string name, out string? value)
+		{
+			int valueLength = 0;
+			if (name.Contains(":"))
+			{
+				string[] tokens = name.Split(':');
+				name = tokens[0];
+				valueLength = Int32.Parse(tokens[1]);
+			}
+
+			// Check each scope for the property
+			for (int scopeIdx = ScopedProperties.Count - 1; scopeIdx >= 0; scopeIdx--)
+			{
+				string? scopeValue;
+				if (ScopedProperties[scopeIdx].TryGetValue(name, out scopeValue))
+				{
+					value = scopeValue;
+
+					// It's valid for a property to exist but have a null value. It won't be expanded
+					// Handle $(PropName:-6) where PropName might be "Foo"
+					if (value != null && value.Length > Math.Abs(valueLength))
+					{
+						if (valueLength > 0)
+						{
+							value = value.Substring(0, valueLength);
+						}
+						if (valueLength < 0)
+						{
+							value = value.Substring(value.Length + valueLength, -valueLength);
+						}
+					}
+					return true;
+				}
+			}
+
+			// If we didn't find it, return false.
+			value = null;
+			return false;
+		}
+
+		/// <summary>
+		/// Read an include directive, and the contents of the target file
+		/// </summary>
+		/// <param name="element">Xml element to read the definition from</param>
+		async Task ReadIncludeAsync(BgScriptElement element)
+		{
+			if (await EvaluateConditionAsync(element))
+			{
+				HashSet<string> files = new HashSet<string>();
+				foreach (string script in ReadListAttribute(element, "Script"))
+				{
+					string includePath = Context.CombinePaths(element.Location.File, script);
+					if (Regex.IsMatch(includePath, @"\*|\?|\.\.\."))
+					{
+						files.UnionWith(await Context.FindAsync(includePath));
+					}
+					else
+					{
+						files.Add(includePath);
+					}
+				}
+
+				foreach (string file in files.OrderBy(x => x, StringComparer.OrdinalIgnoreCase))
+				{
+					Logger.LogDebug("Including file {File}", file);
+					await TryReadAsync(file);
+				}
+			}
+		}
+
+		/// <summary>
+		/// Reads the definition of a graph option; a parameter which can be set by the user on the command-line or via an environment variable.
+		/// </summary>
+		/// <param name="element">Xml element to read the definition from</param>
+		protected abstract Task ReadOptionAsync(BgScriptElement element);
+
+		/// <summary>
+		/// Reads a property assignment.
+		/// </summary>
+		/// <param name="element">Xml element to read the definition from</param>
+		protected abstract Task ReadPropertyAsync(BgScriptElement element);
+
+		/// <summary>
+		/// Reads a Regex assignment.
+		/// </summary>
+		/// <param name="element">Xml element to read the definition from</param>
+		protected abstract Task ReadRegexAsync(BgScriptElement element);
+
+		/// <summary>
+		/// Reads a property assignment from an environment variable.
+		/// </summary>
+		/// <param name="element">Xml element to read the definition from</param>
+		protected abstract Task ReadEnvVarAsync(BgScriptElement element);
+
+		/// <summary>
+		/// Reads a macro definition
+		/// </summary>
+		/// <param name="element">Xml element to read the definition from</param>
+		void ReadMacro(BgScriptElement element)
+		{
+			string name = element.GetAttribute("Name");
+			if (ValidateName(element, name))
+			{
+				BgScriptMacro? originalDefinition;
+				if (_macroNameToDefinition.TryGetValue(name, out originalDefinition))
+				{
+					BgScriptLocation location = originalDefinition.Elements[0].Location;
+					LogError(element, "Macro '{0}' has already been declared (see {1} line {2})", name, location.File, location.LineNumber);
+				}
+				else
+				{
+					Dictionary<string, int> argumentNameToIndex = new Dictionary<string, int>();
+					ReadMacroArguments(element, "Arguments", argumentNameToIndex);
+
+					int numRequiredArguments = argumentNameToIndex.Count;
+					ReadMacroArguments(element, "OptionalArguments", argumentNameToIndex);
+
+					_macroNameToDefinition.Add(name, new BgScriptMacro(name, element, argumentNameToIndex, numRequiredArguments));
+				}
+			}
+		}
+
+		/// <summary>
+		/// Reads a list of macro arguments from an attribute
+		/// </summary>
+		/// <param name="element">The element containing the attributes</param>
+		/// <param name="attributeName">Name of the attribute containing the arguments</param>
+		/// <param name="argumentNameToIndex">List of arguments to add to</param>
+		void ReadMacroArguments(BgScriptElement element, string attributeName, Dictionary<string, int> argumentNameToIndex)
+		{
+			string attributeValue = ReadAttribute(element, attributeName);
+			if (attributeValue != null)
+			{
+				foreach (string argumentName in attributeValue.Split(new char[] { ';' }, StringSplitOptions.RemoveEmptyEntries))
+				{
+					if (argumentNameToIndex.ContainsKey(argumentName))
+					{
+						LogWarning(element, "Argument '{0}' is listed multiple times", argumentName);
+					}
+					else
+					{
+						argumentNameToIndex.Add(argumentName, argumentNameToIndex.Count);
+					}
+				}
+			}
+		}
+
+		/// <summary>
+		/// Reads a macro definition
+		/// </summary>
+		/// <param name="element">Xml element to read the definition from</param>
+		async Task ReadExtendAsync(BgScriptElement element)
+		{
+			if (await EvaluateConditionAsync(element))
+			{
+				string name = ReadAttribute(element, "Name");
+
+				BgScriptMacro? originalDefinition;
+				if (_macroNameToDefinition.TryGetValue(name, out originalDefinition))
+				{
+					originalDefinition.Elements.Add(element);
+				}
+				else
+				{
+					LogError(element, "Macro '{0}' has not been declared", name);
+				}
+			}
+		}
+
+		/// <summary>
+		/// Reads the definition for an agent.
+		/// </summary>
+		/// <param name="element">Xml element to read the definition from</param>
+		protected abstract Task ReadAgentAsync(BgScriptElement element);
+
+		/// <summary>
+		/// Read the contents of an agent definition
+		/// </summary>
+		/// <param name="element">Xml element to read the definition from</param>
+		protected async Task ReadAgentBodyAsync(BgScriptElement element)
+		{
+			EnterScope();
+			foreach (BgScriptElement childElement in element.ChildNodes.OfType<BgScriptElement>())
+			{
+				switch (childElement.Name)
+				{
+					case "Property":
+						await ReadPropertyAsync(childElement);
+						break;
+					case "Regex":
+						await ReadRegexAsync(childElement);
+						break;
+					case "Node":
+						await ReadNodeAsync(childElement);
+						break;
+					case "Aggregate":
+						await ReadAggregateAsync(childElement);
+						break;
+					case "Trace":
+						await ReadDiagnosticAsync(childElement, LogLevel.Information);
+						break;
+					case "Warning":
+						await ReadDiagnosticAsync(childElement, LogLevel.Warning);
+						break;
+					case "Error":
+						await ReadDiagnosticAsync(childElement, LogLevel.Error);
+						break;
+					case "Label":
+						await ReadLabelAsync(childElement);
+						break;
+					case "Do":
+						await ReadBlockAsync(childElement, ReadAgentBodyAsync);
+						break;
+					case "Switch":
+						await ReadSwitchAsync(childElement, ReadAgentBodyAsync);
+						break;
+					case "ForEach":
+						await ReadForEachAsync(childElement, ReadAgentBodyAsync);
+						break;
+					case "Expand":
+						await ReadExpandAsync(childElement, ReadAgentBodyAsync);
+						break;
+					case "Annotate":
+						await ReadAnnotationAsync(childElement);
+						break;
+					default:
+						LogError(childElement, "Unexpected element type '{0}'", childElement.Name);
+						break;
+				}
+			}
+			LeaveScope();
+		}
+
+		/// <summary>
+		/// Reads the definition for an aggregate
+		/// </summary>
+		/// <param name="element">Xml element to read the definition from</param>
+		protected abstract Task ReadAggregateAsync(BgScriptElement element);
+
+		/// <summary>
+		/// Reads the definition for a report
+		/// </summary>
+		/// <param name="element">Xml element to read the definition from</param>
+		protected abstract Task ReadReportAsync(BgScriptElement element);
+
+		/// <summary>
+		/// Reads the definition for a badge
+		/// </summary>
+		/// <param name="element">Xml element to read the definition from</param>
+		protected abstract Task ReadBadgeAsync(BgScriptElement element);
+
+		/// <summary>
+		/// Reads the definition for a label
+		/// </summary>
+		/// <param name="element">Xml element to read the definition from</param>
+		protected abstract Task ReadLabelAsync(BgScriptElement element);
+
+		/// <summary>
+		/// Reads the definition for a node, and adds it to the given agent
+		/// </summary>
+		/// <param name="element">Xml element to read the definition from</param>
+		protected abstract Task ReadNodeAsync(BgScriptElement element);
+
+		/// <summary>
+		/// Reads the contents of a node element
+		/// </summary>
+		/// <param name="element">Xml element to read the definition from</param>
+		protected async Task ReadNodeBodyAsync(XmlElement element)
+		{
+			EnterScope();
+			foreach (BgScriptElement childElement in element.ChildNodes.OfType<BgScriptElement>())
+			{
+				switch (childElement.Name)
+				{
+					case "Property":
+						await ReadPropertyAsync(childElement);
+						break;
+					case "Regex":
+						await ReadRegexAsync(childElement);
+						break;
+					case "Trace":
+						await ReadDiagnosticAsync(childElement, LogLevel.Information);
+						break;
+					case "Warning":
+						await ReadDiagnosticAsync(childElement, LogLevel.Warning);
+						break;
+					case "Error":
+						await ReadDiagnosticAsync(childElement, LogLevel.Error);
+						break;
+					case "Do":
+						await ReadBlockAsync(childElement, ReadNodeBodyAsync);
+						break;
+					case "Switch":
+						await ReadSwitchAsync(childElement, ReadNodeBodyAsync);
+						break;
+					case "ForEach":
+						await ReadForEachAsync(childElement, ReadNodeBodyAsync);
+						break;
+					case "Expand":
+						await ReadExpandAsync(childElement, ReadNodeBodyAsync);
+						break;
+					default:
+						await ReadTaskAsync(childElement);
+						break;
+				}
+			}
+			LeaveScope();
+		}
+
+		/// <summary>
+		/// Reads a block element
+		/// </summary>
+		/// <param name="element">Xml element to read the definition from</param>
+		/// <param name="readContentsAsync">Delegate to read the contents of the element, if the condition evaluates to true</param>
+		async Task ReadBlockAsync(BgScriptElement element, Func<BgScriptElement, Task> readContentsAsync)
+		{
+			if (await EvaluateConditionAsync(element))
+			{
+				await readContentsAsync(element);
+			}
+		}
+
+		/// <summary>
+		/// Reads a "Switch" element 
+		/// </summary>
+		/// <param name="element">Xml element to read the definition from</param>
+		/// <param name="readContentsAsync">Delegate to read the contents of the element, if the condition evaluates to true</param>
+		protected abstract Task ReadSwitchAsync(BgScriptElement element, Func<BgScriptElement, Task> readContentsAsync);
+
+		/// <summary>
+		/// Reads a "ForEach" element 
+		/// </summary>
+		/// <param name="element">Xml element to read the definition from</param>
+		/// <param name="readContentsAsync">Delegate to read the contents of the element, if the condition evaluates to true</param>
+		async Task ReadForEachAsync(BgScriptElement element, Func<BgScriptElement, Task> readContentsAsync)
+		{
+			EnterScope();
+			if (await EvaluateConditionAsync(element))
+			{
+				string name = ReadAttribute(element, "Name");
+				string separator = ReadAttribute(element, "Separator");
+				if (separator.Length > 1)
+				{
+					LogWarning(element, "Node {0}'s Separator attribute is more than one character ({1}). Defaulting to ;", name, separator);
+					separator = ";";
+				}
+				if (String.IsNullOrEmpty(separator))
+				{
+					separator = ";";
+				}
+				if (ValidateName(element, name))
+				{
+					if (ScopedProperties.Any(x => x.ContainsKey(name)))
+					{
+						LogError(element, "Loop variable '{0}' already exists as a local property in an outer scope", name);
+					}
+					else
+					{
+						// Loop through all the values
+						string[] values = ReadListAttribute(element, "Values", Convert.ToChar(separator));
+						foreach (string value in values)
+						{
+							ScopedProperties[^1][name] = value;
+							await readContentsAsync(element);
+						}
+					}
+				}
+			}
+			LeaveScope();
+		}
+
+		/// <summary>
+		/// Reads an "Expand" element 
+		/// </summary>
+		/// <param name="element">Xml element to read the definition from</param>
+		/// <param name="readContentsAsync">Delegate to read the contents of the element, if the condition evaluates to true</param>
+		async Task ReadExpandAsync(BgScriptElement element, Func<BgScriptElement, Task> readContentsAsync)
+		{
+			if (await EvaluateConditionAsync(element))
+			{
+				string name = ReadAttribute(element, "Name");
+				if (ValidateName(element, name))
+				{
+					BgScriptMacro? macro;
+					if (!_macroNameToDefinition.TryGetValue(name, out macro))
+					{
+						LogError(element, "Macro '{0}' does not exist", name);
+					}
+					else
+					{
+						// Parse the argument list
+						string[] arguments = new string[macro.ArgumentNameToIndex.Count];
+						foreach (XmlAttribute? attribute in element.Attributes)
+						{
+							if (attribute != null && attribute.Name != "Name" && attribute.Name != "If")
+							{
+								int index;
+								if (macro.ArgumentNameToIndex.TryGetValue(attribute.Name, out index))
+								{
+									arguments[index] = ExpandProperties(element, attribute.Value);
+								}
+								else
+								{
+									LogWarning(element, "Macro '{0}' does not take an argument '{1}'", name, attribute.Name);
+								}
+							}
+						}
+
+						// Make sure none of the required arguments are missing
+						bool bHasMissingArguments = false;
+						for (int idx = 0; idx < macro.NumRequiredArguments; idx++)
+						{
+							if (arguments[idx] == null)
+							{
+								LogWarning(element, "Macro '{0}' is missing argument '{1}'", macro.Name, macro.ArgumentNameToIndex.First(x => x.Value == idx).Key);
+								bHasMissingArguments = true;
+							}
+						}
+
+						// Expand the function
+						if (!bHasMissingArguments)
+						{
+							EnterScope();
+							foreach (KeyValuePair<string, int> pair in macro.ArgumentNameToIndex)
+							{
+								ScopedProperties[^1][pair.Key] = arguments[pair.Value] ?? "";
+							}
+							foreach (BgScriptElement macroElement in macro.Elements)
+							{
+								await readContentsAsync(macroElement);
+							}
+							LeaveScope();
+						}
+					}
+				}
+			}
+		}
+
+		/// <summary>
+		/// Reads a task definition from the given element, and add it to the given list
+		/// </summary>
+		/// <param name="element">Xml element to read the definition from</param>
+		protected abstract Task ReadTaskAsync(BgScriptElement element);
+
+		/// <summary>
+		/// Reads the definition for an email notifier
+		/// </summary>
+		/// <param name="element">Xml element to read the definition from</param>
+		protected abstract Task ReadNotifierAsync(BgScriptElement element);
+
+		/// <summary>
+		/// Reads a graph annotation
+		/// </summary>
+		/// <param name="element">Xml element to read the definition from</param>
+		protected abstract Task ReadAnnotationAsync(BgScriptElement element);
+
+		/// <summary>
+		/// Reads a warning from the given element, evaluates the condition on it, and writes it to the log if the condition passes.
+		/// </summary>
+		/// <param name="element">Xml element to read the definition from</param>
+		/// <param name="level">The diagnostic event type</param>
+		protected abstract Task ReadDiagnosticAsync(BgScriptElement element, LogLevel level);
+
+		/// <summary>
+		/// Reads an object name from its defining element. Outputs an error if the name is missing.
+		/// </summary>
+		/// <param name="element">Element to read the name for</param>
+		/// <param name="name">Output variable to receive the name of the object</param>
+		/// <returns>True if the object had a valid name (assigned to the Name variable), false if the name was invalid or missing.</returns>
+		protected bool TryReadObjectName(BgScriptElement element, [NotNullWhen(true)] out string? name)
+		{
+			// Check the name attribute is present
+			if (!element.HasAttribute("Name"))
+			{
+				LogError(element, "Missing 'Name' attribute");
+				name = null;
+				return false;
+			}
+
+			// Get the value of it, strip any leading or trailing whitespace, and make sure it's not empty
+			string value = ReadAttribute(element, "Name");
+			if (!ValidateName(element, value))
+			{
+				name = null;
+				return false;
+			}
+
+			// Return it
+			name = value;
+			return true;
+		}
+
+		/// <summary>
+		/// Checks that the given name is valid syntax
+		/// </summary>
+		/// <param name="element">The element that contains the name</param>
+		/// <param name="name">The name to check</param>
+		/// <returns>True if the name is valid</returns>
+		protected bool ValidateName(BgScriptElement element, string name)
+		{
+			// Check it's not empty
+			if (name.Length == 0)
+			{
+				LogError(element, "Name is empty");
+				return false;
+			}
+
+			// Check there are no invalid characters
+			for (int idx = 0; idx < name.Length; idx++)
+			{
+				if (idx > 0 && name[idx] == ' ' && name[idx - 1] == ' ')
+				{
+					LogError(element, "Consecutive spaces in object name - '{0}'", name);
+					return false;
+				}
+				if (Char.IsControl(name[idx]) || BgScriptSchema.IllegalNameCharacters.IndexOf(name[idx]) != -1)
+				{
+					LogError(element, "Invalid character in object name - '{0}'", name[idx]);
+					return false;
+				}
+			}
+			return true;
+		}
+
+		/// <summary>
+		/// Constructs a regex from a regex string and returns it
+		/// </summary>
+		/// <param name="element">The element that contains the regex</param>
+		/// <param name="regex">The pattern to construct</param>
+		/// <returns>The regex if is valid, otherwise null</returns>
+		protected Regex? ParseRegex(BgScriptElement element, string regex)
+		{
+			if (regex.Length == 0)
+			{
+				LogError(element, "Regex is empty");
+				return null;
+			}
+			try
+			{
+				return new Regex(regex);
+			}
+			catch (ArgumentException invalidRegex)
+			{
+				LogError(element, "Could not construct the Regex, reason: {0}", invalidRegex.Message);
+				return null;
+			}
+		}
+
+		/// <summary>
+		/// Expands any properties and reads an attribute.
+		/// </summary>
+		/// <param name="element">Element to read the attribute from</param>
+		/// <param name="name">Name of the attribute</param>
+		/// <returns>Array of names, with all leading and trailing whitespace removed</returns>
+		protected string ReadAttribute(BgScriptElement element, string name)
+		{
+			return ExpandProperties(element, element.GetAttribute(name));
+		}
+
+		/// <summary>
+		/// Expands any properties and reads a list of strings from an attribute, separated by semi-colon characters
+		/// </summary>
+		/// <param name="element"></param>
+		/// <param name="name"></param>
+		/// <param name="separator"></param>
+		/// <returns>Array of names, with all leading and trailing whitespace removed</returns>
+		protected string[] ReadListAttribute(BgScriptElement element, string name, char separator = ';')
+		{
+			string value = ReadAttribute(element, name);
+			return value.Split(new char[] { separator }).Select(x => x.Trim()).Where(x => x.Length > 0).ToArray();
+		}
+
+		/// <summary>
+		/// Parse a map of annotations
+		/// </summary>
+		/// <param name="element"></param>
+		/// <param name="name"></param>
+		/// <returns></returns>
+		protected Dictionary<string, string> ReadAnnotationsAttribute(BgScriptElement element, string name)
+		{
+			string[] pairs = ReadListAttribute(element, name);
+			return ParseAnnotations(element, pairs);
+		}
+
+		/// <summary>
+		/// Parse a map of annotations
+		/// </summary>
+		/// <param name="element"></param>
+		/// <param name="pairs"></param>
+		/// <returns></returns>
+		private Dictionary<string, string> ParseAnnotations(BgScriptElement element, string[] pairs)
+		{
+			// Find the annotations to apply
+			Dictionary<string, string> pairMap = new Dictionary<string, string>(StringComparer.OrdinalIgnoreCase);
+			foreach (string pair in pairs)
+			{
+				if (!String.IsNullOrWhiteSpace(pair))
+				{
+					int idx = pair.IndexOf('=');
+					if (idx < 0)
+					{
+						LogError(element, "Invalid annotation '{0}'", pair);
+						continue;
+					}
+
+					string key = pair.Substring(0, idx).Trim();
+					if (!Regex.IsMatch(key, @"[a-zA-Z0-9_\.]+"))
+					{
+						LogError(element, "Invalid annotation key '{0}'", pair);
+						continue;
+					}
+					if (pairMap.ContainsKey(key))
+					{
+						LogError(element, "Annotation key '{0}' was specified twice", key);
+						continue;
+					}
+
+					pairMap.Add(key, pair.Substring(idx + 1).Trim());
+				}
+			}
+			return pairMap;
+		}
+
+		/// <summary>
+		/// Reads an attribute from the given XML element, expands any properties in it, and parses it as a boolean.
+		/// </summary>
+		/// <param name="element">Element to read the attribute from</param>
+		/// <param name="name">Name of the attribute</param>
+		/// <param name="bDefaultValue">Default value if the attribute is missing</param>
+		/// <returns>The value of the attribute field</returns>
+		protected bool ReadBooleanAttribute(BgScriptElement element, string name, bool bDefaultValue)
+		{
+			bool bResult = bDefaultValue;
+			if (element.HasAttribute(name))
+			{
+				string value = ReadAttribute(element, name).Trim();
+				if (value.Equals("true", StringComparison.InvariantCultureIgnoreCase))
+				{
+					bResult = true;
+				}
+				else if (value.Equals("false", StringComparison.InvariantCultureIgnoreCase))
+				{
+					bResult = false;
+				}
+				else
+				{
+					LogError(element, "Invalid boolean value '{0}' - expected 'true' or 'false'", value);
+				}
+			}
+			return bResult;
+		}
+
+		/// <summary>
+		/// Reads an attribute from the given XML element, expands any properties in it, and parses it as an integer.
+		/// </summary>
+		/// <param name="element">Element to read the attribute from</param>
+		/// <param name="name">Name of the attribute</param>
+		/// <param name="defaultValue">Default value for the integer, if the attribute is missing</param>
+		/// <returns>The value of the attribute field</returns>
+		protected int ReadIntegerAttribute(BgScriptElement element, string name, int defaultValue)
+		{
+			int result = defaultValue;
+			if (element.HasAttribute(name))
+			{
+				string value = ReadAttribute(element, name).Trim();
+
+				int intValue;
+				if (Int32.TryParse(value, out intValue))
+				{
+					result = intValue;
+				}
+				else
+				{
+					LogError(element, "Invalid integer value '{0}'", value);
+				}
+			}
+			return result;
+		}
+
+		/// <summary>
+		/// Reads an attribute from the given XML element, expands any properties in it, and parses it as an enum of the given type.
+		/// </summary>
+		/// <typeparam name="T">The enum type to parse the attribute as</typeparam>
+		/// <param name="element">Element to read the attribute from</param>
+		/// <param name="name">Name of the attribute</param>
+		/// <param name="defaultValue">Default value for the enum, if the attribute is missing</param>
+		/// <returns>The value of the attribute field</returns>
+		protected T ReadEnumAttribute<T>(BgScriptElement element, string name, T defaultValue) where T : struct
+		{
+			T result = defaultValue;
+			if (element.HasAttribute(name))
+			{
+				string value = ReadAttribute(element, name).Trim();
+
+				T enumValue;
+				if (Enum.TryParse(value, true, out enumValue))
+				{
+					result = enumValue;
+				}
+				else
+				{
+					LogError(element, "Invalid value '{0}' - expected {1}", value, String.Join("/", Enum.GetNames(typeof(T))));
+				}
+			}
+			return result;
+		}
+
+		/// <summary>
+		/// Outputs an error message to the log and increments the number of errors, referencing the file and line number of the element that caused it.
+		/// </summary>
+		/// <param name="element">The script element causing the error</param>
+		/// <param name="format">Standard String.Format()-style format string</param>
+		/// <param name="args">Optional arguments</param>
+		protected void LogError(BgScriptElement element, string format, params object[] args)
+		{
+			Logger.LogScriptError(element.Location, format, args);
+			NumErrors++;
+		}
+
+		/// <summary>
+		/// Outputs a warning message to the log and increments the number of errors, referencing the file and line number of the element that caused it.
+		/// </summary>
+		/// <param name="element">The script element causing the error</param>
+		/// <param name="format">Standard String.Format()-style format string</param>
+		/// <param name="args">Optional arguments</param>
+		protected void LogWarning(BgScriptElement element, string format, params object[] args)
+		{
+			Logger.LogScriptWarning(element.Location, format, args);
+		}
+
+		/// <summary>
+		/// Evaluates the (optional) conditional expression on a given XML element via the If="..." attribute, and returns true if the element is enabled.
+		/// </summary>
+		/// <param name="element">The element to check</param>
+		/// <returns>True if the element's condition evaluates to true (or doesn't have a conditional expression), false otherwise</returns>
+		protected async Task<bool> EvaluateConditionAsync(BgScriptElement element)
+		{
+			// Check if the element has a conditional attribute
+			const string AttributeName = "If";
+			if (!element.HasAttribute(AttributeName))
+			{
+				return true;
+			}
+
+			// If it does, try to evaluate it.
+			try
+			{
+				string text = ExpandProperties(element, element.GetAttribute("If"));
+				return await BgCondition.EvaluateAsync(text, Context);
+			}
+			catch (BgConditionException ex)
+			{
+				LogError(element, "Error in condition: {0}", ex.Message);
+				return false;
+			}
+		}
+
+		/// <summary>
+		/// Expand all the property references (of the form $(PropertyName)) in a string.
+		/// </summary>
+		/// <param name="element">The element containing the string. Used for diagnostic messages.</param>
+		/// <param name="text">The input string to expand properties in</param>
+		/// <returns>The expanded string</returns>
+		protected string ExpandProperties(BgScriptElement element, string text)
+		{
+			string result = text;
+			// Iterate in reverse order to handle cases where there are nested expansions like $(Outer$(Inner))
+			for (int idx = result.LastIndexOf("$("); idx != -1; idx = result.LastIndexOf("$(", idx, idx + 1))
+			{
+				// Find the end of the variable name
+				int endIdx = result.IndexOf(')', idx + 2);
+				if (endIdx == -1)
+				{
+					break;
+				}
+
+				// Extract the variable name from the string
+				string name = result.Substring(idx + 2, endIdx - (idx + 2));
+
+				// Find the value for it, either from the dictionary or the environment block
+				string? value;
+				if (!TryGetPropertyValue(name, out value))
+				{
+					LogWarning(element, "Property '{0}' is not defined", name);
+					value = "";
+				}
+
+				// Check if we've got a value for this variable
+				if (value != null)
+				{
+					// Replace the variable, or skip past it
+					result = result.Substring(0, idx) + value + result.Substring(endIdx + 1);
+				}
+			}
+			return result;
+		}
+	}
+
+	/// <summary>
+	/// Overridden version of <see cref="BgNode"/> which contains a list of tasks
+	/// </summary>
+	class BgScriptNode : BgNode
+	{
+		/// <summary>
+		/// List of tasks to execute
+		/// </summary>
+		public List<BgTask> Tasks { get; } = new List<BgTask>();
+
+		/// <summary>
+		/// Constructor
+		/// </summary>
+		public BgScriptNode(string name, IReadOnlyList<BgNodeOutput> inputs, IReadOnlyList<string> outputNames, IReadOnlyList<BgNode> inputDependencies, IReadOnlyList<BgNode> orderDependencies, IReadOnlyList<FileReference> requiredTokens)
+			: base(name, inputs, outputNames, inputDependencies, orderDependencies, requiredTokens)
+		{
+		}
+	}
+
+	/// <summary>
+	/// Reader for build graph definitions. Instanced to contain temporary state; public interface is through ScriptReader.TryRead().
+	/// </summary>
+	public class BgScriptReader : BgScriptReaderBase
+	{
+		/// <summary>
+		/// The current graph
+		/// </summary>
+		readonly BgGraph _graph = new BgGraph();
+
+		/// <summary>
+		/// Arguments for evaluating the graph
+		/// </summary>
+		readonly Dictionary<string, string> _arguments;
+
+		/// <summary>
+		/// The name of the node if only a single node is going to be built, otherwise null.
+		/// </summary>
+		readonly string? _singleNodeName;
+
+		BgAgent? _enclosingAgent;
+		BgScriptNode? _enclosingNode;
+
+		/// <summary>
+		/// Private constructor. Use ScriptReader.TryRead() to read a script file.
+		/// </summary>
+		/// <param name="context">Context object</param>
+		/// <param name="defaultProperties">Default properties available to the script</param>
+		/// <param name="arguments">Arguments passed in to the graph on the command line</param>
+		/// <param name="singleNodeName">If a single node will be processed, the name of that node.</param>
+		/// <param name="schema">Schema for the script</param>
+		/// <param name="logger">Logger for diagnostic messages</param>
+		private BgScriptReader(IBgScriptReaderContext context, IDictionary<string, string> defaultProperties, IReadOnlyDictionary<string, string> arguments, string? singleNodeName, BgScriptSchema schema, ILogger logger)
+			: base(context, schema, logger)
+		{
+			_arguments = new Dictionary<string, string>(arguments, StringComparer.OrdinalIgnoreCase);
+			_singleNodeName = singleNodeName;
+
+			foreach (KeyValuePair<string, string> pair in defaultProperties)
+			{
+				SetPropertyValue(null!, pair.Key, pair.Value);
+			}
+		}
+
+		/// <summary>
+		/// Try to read a script file from the given file.
+		/// </summary>
+		/// <param name="context">Supplies context about the parse</param>
+		/// <param name="file">File to read from</param>
+		/// <param name="arguments">Arguments passed in to the graph on the command line</param>
+		/// <param name="defaultProperties">Default properties available to the script</param>
+		/// <param name="schema">Schema for the script</param>
+		/// <param name="logger">Logger for output messages</param>
+		/// <param name="singleNodeName">If a single node will be processed, the name of that node.</param>
+		/// <returns>True if the graph was read, false if there were errors</returns>
+		public static async Task<BgGraph?> ReadAsync(IBgScriptReaderContext context, string file, Dictionary<string, string> arguments, Dictionary<string, string> defaultProperties, BgScriptSchema schema, ILogger logger, string? singleNodeName = null)
+		{
+			// Read the file and build the graph
+			BgScriptReader reader = new BgScriptReader(context, defaultProperties, arguments, singleNodeName, schema, logger);
+			if (!await reader.TryReadAsync(file) || reader.NumErrors > 0)
+			{
+				return null;
+			}
+
+			// Make sure all the arguments were valid
+			HashSet<string> validArgumentNames = new HashSet<string>(reader._graph.Options.Select(x => x.Name), StringComparer.OrdinalIgnoreCase);
+			validArgumentNames.Add("PreflightChange");
+
+			foreach (string argumentName in arguments.Keys)
+			{
+				if (!validArgumentNames.Contains(argumentName))
+				{
+					logger.LogWarning("Unknown argument '{ArgumentName}' for '{Script}'", argumentName, context.GetNativePath(file));
+				}
+			}
+
+			// Return the constructed graph
+			return reader._graph;
+		}
+
+		/// <summary>
+		/// Reads the definition of a graph option; a parameter which can be set by the user on the command-line or via an environment variable.
+		/// </summary>
+		/// <param name="element">Xml element to read the definition from</param>
+		protected override async Task ReadOptionAsync(BgScriptElement element)
+		{
+			if (await EvaluateConditionAsync(element))
+			{
+				string name = ReadAttribute(element, "Name");
+				if (ValidateName(element, name))
+				{
+					// Make sure we're at global scope
+					if (ScopedProperties.Count > 1)
+					{
+						throw new Exception("Incorrect scope depth for reading option settings");
+					}
+
+					// Check if the property already exists. If it does, we don't need to register it as an option.
+					string? existingValue;
+					if (TryGetPropertyValue(name, out existingValue) && existingValue != null)
+					{
+						// If there's a restriction on this definition, check it matches
+						string restrict = ReadAttribute(element, "Restrict");
+						if (!String.IsNullOrEmpty(restrict) && !Regex.IsMatch(existingValue, "^" + restrict + "$", RegexOptions.IgnoreCase))
+						{
+							LogError(element, "'{0} is already set to '{1}', which does not match the given restriction ('{2}')", name, existingValue, restrict);
+						}
+					}
+					else
+					{
+						// Create a new option object to store the settings
+						string description = ReadAttribute(element, "Description");
+						string defaultValue = ReadAttribute(element, "DefaultValue");
+						BgOption option = new BgOption(name, description, defaultValue);
+						_graph.Options.Add(option);
+
+						// Get the value of this property
+						string? value;
+						if (!_arguments.TryGetValue(name, out value))
+						{
+							value = option.DefaultValue;
+						}
+						SetPropertyValue(element, name, value);
+
+						// If there's a restriction on it, check it's valid
+						string restrict = ReadAttribute(element, "Restrict");
+						if (!String.IsNullOrEmpty(restrict))
+						{
+							string pattern = "^(" + restrict + ")$";
+							if (!Regex.IsMatch(value, pattern, RegexOptions.IgnoreCase))
+							{
+								LogError(element, "'{0}' is not a valid value for '{1}' (required: '{2}')", value, name, restrict);
+							}
+							if (option.DefaultValue != value && !Regex.IsMatch(option.DefaultValue, pattern, RegexOptions.IgnoreCase))
+							{
+								LogError(element, "Default value '{0}' is not valid for '{1}' (required: '{2}')", option.DefaultValue, name, restrict);
+							}
+						}
+					}
+				}
+			}
+		}
+
+		/// <summary>
+		/// Reads a property assignment.
+		/// </summary>
+		/// <param name="element">Xml element to read the definition from</param>
+		protected override async Task ReadPropertyAsync(BgScriptElement element)
+		{
+			if (await EvaluateConditionAsync(element))
+			{
+				string name = ReadAttribute(element, "Name");
+				if (ValidateName(element, name))
+				{
+					string value = ReadAttribute(element, "Value");
+					if (element.HasChildNodes)
+					{
+						// Read the element content, and append each line to the value as a semicolon delimited list
+						StringBuilder builder = new StringBuilder(value);
+						foreach (string line in element.InnerText.Split('\n'))
+						{
+							string trimLine = ExpandProperties(element, line.Trim());
+							if (trimLine.Length > 0)
+							{
+								if (builder.Length > 0)
+								{
+									builder.Append(";");
+								}
+								builder.Append(trimLine);
+							}
+						}
+						value = builder.ToString();
+					}
+					SetPropertyValue(element, name, value);
+				}
+			}
+		}
+
+		/// <summary>
+		/// Reads a Regex assignment.
+		/// </summary>
+		/// <param name="element">Xml element to read the definition from</param>
+		protected override async Task ReadRegexAsync(BgScriptElement element)
+		{
+			if (await EvaluateConditionAsync(element))
+			{
+				// Get the pattern
+				string regexString = ReadAttribute(element, "Pattern");
+
+				// Make sure its a valid regex.
+				Regex? regexValue = ParseRegex(element, regexString);
+				if (regexValue != null)
+				{
+					// read the names in 
+					string[] captureNames = ReadListAttribute(element, "Capture");
+
+					// get number of groups we passed in
+					int[] groupNumbers = regexValue.GetGroupNumbers();
+
+					// make sure the number of property names is the same as the number of match groups
+					// this includes the entire string match group as [0], so don't count that one.
+					if (captureNames.Length != groupNumbers.Count() - 1)
+					{
+						LogError(element, "MatchGroup count: {0} does not match the number of names specified: {1}", groupNumbers.Count() - 1, captureNames.Length);
+					}
+					else
+					{
+						// apply the regex to the value
+						string input = ReadAttribute(element, "Input");
+						Match match = regexValue.Match(input);
+
+						bool optional = await BgCondition.EvaluateAsync(ReadAttribute(element, "Optional"), Context);
+						if (!match.Success)
+						{
+							if (!optional)
+							{
+								LogError(element, "Regex {0} did not find a match against input string {1}", regexString, input);
+							}
+						}
+						else
+						{
+							// assign each property to the group it matches, skip over [0]
+							for (int matchIdx = 1; matchIdx < groupNumbers.Count(); matchIdx++)
+							{
+								SetPropertyValue(element, captureNames[matchIdx - 1], match.Groups[matchIdx].Value);
+							}
+						}
+					}
+				}
+			}
+		}
+
+		/// <summary>
+		/// Reads a property assignment from an environment variable.
+		/// </summary>
+		/// <param name="element">Xml element to read the definition from</param>
+		protected override async Task ReadEnvVarAsync(BgScriptElement element)
+		{
+			if (await EvaluateConditionAsync(element))
+			{
+				string name = ReadAttribute(element, "Name");
+				if (ValidateName(element, name))
+				{
+					string value = Environment.GetEnvironmentVariable(name) ?? "";
+					SetPropertyValue(element, name, value);
+				}
+			}
+		}
+
+		/// <summary>
+		/// Reads the definition for an agent.
+		/// </summary>
+		/// <param name="element">Xml element to read the definition from</param>
+		protected override async Task ReadAgentAsync(BgScriptElement element)
+		{
+			string? name;
+			if (await EvaluateConditionAsync(element) && TryReadObjectName(element, out name))
+			{
+				// Read the valid agent types. This may be omitted if we're continuing an existing agent.
+				string[] types = ReadListAttribute(element, "Type");
+
+				// Create the agent object, or continue an existing one
+				BgAgent? agent;
+				if (_graph.NameToAgent.TryGetValue(name, out agent))
+				{
+					if (types.Length > 0 && agent.PossibleTypes.Length > 0)
+					{
+						if (types.Length != agent.PossibleTypes.Length || !types.SequenceEqual(agent.PossibleTypes, StringComparer.InvariantCultureIgnoreCase))
+						{
+							LogError(element, "Agent types ({0}) were different than previous agent definition with types ({1}). Must either be empty or match exactly.", String.Join(",", types), String.Join(",", agent.PossibleTypes));
+						}
+					}
+				}
+				else
+				{
+					if (types.Length == 0)
+					{
+						LogError(element, "Missing type for agent '{0}'", name);
+					}
+					agent = new BgAgent(name, types);
+					_graph.NameToAgent.Add(name, agent);
+					_graph.Agents.Add(agent);
+				}
+
+				// Process all the child elements.
+				_enclosingAgent = agent;
+				await ReadAgentBodyAsync(element);
+				_enclosingAgent = null;
+			}
+		}
+
+		/// <summary>
+		/// Reads the definition for an aggregate
+		/// </summary>
+		/// <param name="element">Xml element to read the definition from</param>
+		protected override async Task ReadAggregateAsync(BgScriptElement element)
+		{
+			string? name;
+			if (await EvaluateConditionAsync(element) && TryReadObjectName(element, out name) && CheckNameIsUnique(element, name))
+			{
+				string[] requiredNames = ReadListAttribute(element, "Requires");
+
+				BgAggregate newAggregate = new BgAggregate(name);
+				foreach (BgNode referencedNode in ResolveReferences(element, requiredNames))
+				{
+					newAggregate.RequiredNodes.Add(referencedNode);
+				}
+				_graph.NameToAggregate[name] = newAggregate;
+
+				string labelCategoryName = ReadAttribute(element, "Label");
+				if (!String.IsNullOrEmpty(labelCategoryName))
+				{
+					BgLabel label;
+
+					// Create the label
+					int slashIdx = labelCategoryName.IndexOf('/');
+					if (slashIdx != -1)
+					{
+						label = new BgLabel(labelCategoryName.Substring(slashIdx + 1), labelCategoryName.Substring(0, slashIdx), null, null, BgLabelChange.Current);
+					}
+					else
+					{
+						label = new BgLabel(labelCategoryName, "Other", null, null, BgLabelChange.Current);
+					}
+
+					// Find all the included nodes
+					foreach (BgNode requiredNode in newAggregate.RequiredNodes)
+					{
+						label.RequiredNodes.Add(requiredNode);
+						label.IncludedNodes.Add(requiredNode);
+						label.IncludedNodes.UnionWith(requiredNode.OrderDependencies);
+					}
+
+					string[] includedNames = ReadListAttribute(element, "Include");
+					foreach (BgNode includedNode in ResolveReferences(element, includedNames))
+					{
+						label.IncludedNodes.Add(includedNode);
+						label.IncludedNodes.UnionWith(includedNode.OrderDependencies);
+					}
+
+					string[] excludedNames = ReadListAttribute(element, "Exclude");
+					foreach (BgNode excludedNode in ResolveReferences(element, excludedNames))
+					{
+						label.IncludedNodes.Remove(excludedNode);
+						label.IncludedNodes.ExceptWith(excludedNode.OrderDependencies);
+					}
+
+					_graph.Labels.Add(label);
+				}
+			}
+		}
+
+		/// <summary>
+		/// Reads the definition for a report
+		/// </summary>
+		/// <param name="element">Xml element to read the definition from</param>
+		protected override async Task ReadReportAsync(BgScriptElement element)
+		{
+			string? name;
+			if (await EvaluateConditionAsync(element) && TryReadObjectName(element, out name) && CheckNameIsUnique(element, name))
+			{
+				string[] requiredNames = ReadListAttribute(element, "Requires");
+
+				BgReport newReport = new BgReport(name);
+				foreach (BgNode referencedNode in ResolveReferences(element, requiredNames))
+				{
+					newReport.Nodes.Add(referencedNode);
+					newReport.Nodes.UnionWith(referencedNode.OrderDependencies);
+				}
+				_graph.NameToReport.Add(name, newReport);
+			}
+		}
+
+		/// <summary>
+		/// Reads the definition for a badge
+		/// </summary>
+		/// <param name="element">Xml element to read the definition from</param>
+		protected override async Task ReadBadgeAsync(BgScriptElement element)
+		{
+			string? name;
+			if (await EvaluateConditionAsync(element) && TryReadObjectName(element, out name))
+			{
+				string[] requiredNames = ReadListAttribute(element, "Requires");
+				string[] targetNames = ReadListAttribute(element, "Targets");
+				string project = ReadAttribute(element, "Project");
+				int change = ReadIntegerAttribute(element, "Change", 0);
+
+				BgBadge newBadge = new BgBadge(name, project, change);
+				foreach (BgNode referencedNode in ResolveReferences(element, requiredNames))
+				{
+					newBadge.Nodes.Add(referencedNode);
+				}
+				foreach (BgNode referencedNode in ResolveReferences(element, targetNames))
+				{
+					newBadge.Nodes.Add(referencedNode);
+					newBadge.Nodes.UnionWith(referencedNode.OrderDependencies);
+				}
+				_graph.Badges.Add(newBadge);
+			}
+		}
+
+		/// <summary>
+		/// Reads the definition for a label
+		/// </summary>
+		/// <param name="element">Xml element to read the definition from</param>
+		protected override async Task ReadLabelAsync(BgScriptElement element)
+		{
+			if (await EvaluateConditionAsync(element))
+			{
+				string name = ReadAttribute(element, "Name");
+				if (!String.IsNullOrEmpty(name))
+				{
+					ValidateName(element, name);
+				}
+
+				string category = ReadAttribute(element, "Category");
+
+				string[] requiredNames = ReadListAttribute(element, "Requires");
+				string[] includedNames = ReadListAttribute(element, "Include");
+				string[] excludedNames = ReadListAttribute(element, "Exclude");
+
+				string ugsBadge = ReadAttribute(element, "UgsBadge");
+				string ugsProject = ReadAttribute(element, "UgsProject");
+
+				BgLabelChange change = ReadEnumAttribute<BgLabelChange>(element, "Change", BgLabelChange.Current);
+
+				BgLabel newLabel = new BgLabel(name, category, ugsBadge, ugsProject, change);
+				foreach (BgNode referencedNode in ResolveReferences(element, requiredNames))
+				{
+					newLabel.RequiredNodes.Add(referencedNode);
+					newLabel.IncludedNodes.Add(referencedNode);
+					newLabel.IncludedNodes.UnionWith(referencedNode.OrderDependencies);
+				}
+				foreach (BgNode includedNode in ResolveReferences(element, includedNames))
+				{
+					newLabel.IncludedNodes.Add(includedNode);
+					newLabel.IncludedNodes.UnionWith(includedNode.OrderDependencies);
+				}
+				foreach (BgNode excludedNode in ResolveReferences(element, excludedNames))
+				{
+					newLabel.IncludedNodes.Remove(excludedNode);
+					newLabel.IncludedNodes.ExceptWith(excludedNode.OrderDependencies);
+				}
+				_graph.Labels.Add(newLabel);
+			}
+		}
+
+		/// <summary>
+		/// Reads the definition for a node, and adds it to the given agent
+		/// </summary>
+		/// <param name="element">Xml element to read the definition from</param>
+		protected override async Task ReadNodeAsync(BgScriptElement element)
+		{
+			string? name;
+			if (await EvaluateConditionAsync(element) && TryReadObjectName(element, out name))
+			{
+				string[] requiresNames = ReadListAttribute(element, "Requires");
+				string[] producesNames = ReadListAttribute(element, "Produces");
+				string[] afterNames = ReadListAttribute(element, "After");
+				string[] tokenFileNames = ReadListAttribute(element, "Token");
+				bool bRunEarly = ReadBooleanAttribute(element, "RunEarly", false);
+				bool bNotifyOnWarnings = ReadBooleanAttribute(element, "NotifyOnWarnings", true);
+				Dictionary<string, string> annotations = ReadAnnotationsAttribute(element, "Annotations");
+
+				// Resolve all the inputs we depend on
+				HashSet<BgNodeOutput> inputs = ResolveInputReferences(element, requiresNames);
+
+				// Gather up all the input dependencies, and check they're all upstream of the current node
+				HashSet<BgNode> inputDependencies = new HashSet<BgNode>();
+				foreach (BgNode inputDependency in inputs.Select(x => x.ProducingNode).Distinct())
+				{
+					inputDependencies.Add(inputDependency);
+				}
+
+				// Remove all the lock names from the list of required names
+				HashSet<FileReference> requiredTokens = new HashSet<FileReference>(tokenFileNames.Select(x => new FileReference(x)));
+
+				// Recursively include all their dependencies too
+				foreach (BgNode inputDependency in inputDependencies.ToArray())
+				{
+					requiredTokens.UnionWith(inputDependency.RequiredTokens);
+					inputDependencies.UnionWith(inputDependency.InputDependencies);
+				}
+
+				// Validate all the outputs
+				List<string> validOutputNames = new List<string>();
+				foreach (string producesName in producesNames)
+				{
+					BgNodeOutput? existingOutput;
+					if (_graph.TagNameToNodeOutput.TryGetValue(producesName, out existingOutput))
+					{
+						LogError(element, "Output tag '{0}' is already generated by node '{1}'", producesName, existingOutput.ProducingNode.Name);
+					}
+					else if (!producesName.StartsWith("#"))
+					{
+						LogError(element, "Output tag names must begin with a '#' character ('{0}')", producesName);
+					}
+					else
+					{
+						validOutputNames.Add(producesName);
+					}
+				}
+
+				// Gather up all the order dependencies
+				HashSet<BgNode> orderDependencies = new HashSet<BgNode>(inputDependencies);
+				orderDependencies.UnionWith(ResolveReferences(element, afterNames));
+
+				// Recursively include all their order dependencies too
+				foreach (BgNode orderDependency in orderDependencies.ToArray())
+				{
+					orderDependencies.UnionWith(orderDependency.OrderDependencies);
+				}
+
+				// Check that we're not dependent on anything completing that is declared after the initial declaration of this agent.
+				int agentIdx = _graph.Agents.IndexOf(_enclosingAgent!);
+				for (int idx = agentIdx + 1; idx < _graph.Agents.Count; idx++)
+				{
+					foreach (BgNode node in _graph.Agents[idx].Nodes.Where(x => orderDependencies.Contains(x)))
+					{
+						LogError(element, "Node '{0}' has a dependency on '{1}', which was declared after the initial definition of '{2}'.", name, node.Name, _enclosingAgent!.Name);
+					}
+				}
+
+				// Construct and register the node
+				if (CheckNameIsUnique(element, name))
+				{
+					// Add it to the node lookup
+					BgScriptNode newNode = new BgScriptNode(name, inputs.ToArray(), validOutputNames.ToArray(), inputDependencies.ToArray(), orderDependencies.ToArray(), requiredTokens.ToArray());
+					newNode.RunEarly = bRunEarly;
+					newNode.NotifyOnWarnings = bNotifyOnWarnings;
+					foreach ((string key, string value) in annotations)
+					{
+						newNode.Annotations[key] = value;
+					}
+					_graph.NameToNode.Add(name, newNode);
+
+					// Register all the output tags in the global name table.
+					foreach (BgNodeOutput output in newNode.Outputs)
+					{
+						BgNodeOutput? existingOutput;
+						if (_graph.TagNameToNodeOutput.TryGetValue(output.TagName, out existingOutput))
+						{
+							LogError(element, "Node '{0}' already has an output called '{1}'", existingOutput.ProducingNode.Name, output.TagName);
+						}
+						else
+						{
+							_graph.TagNameToNodeOutput.Add(output.TagName, output);
+						}
+					}
+
+					// Add all the tasks
+					_enclosingNode = newNode;
+					await ReadNodeBodyAsync(element);
+					_enclosingNode = null;
+
+					// Add it to the current agent
+					_enclosingAgent!.Nodes.Add(newNode);
+				}
+			}
+		}
+
+		/// <summary>
+		/// Reads a "Switch" element 
+		/// </summary>
+		/// <param name="element">Xml element to read the definition from</param>
+		/// <param name="readContentsAsync">Delegate to read the contents of the element, if the condition evaluates to true</param>
+		protected override async Task ReadSwitchAsync(BgScriptElement element, Func<BgScriptElement, Task> readContentsAsync)
+		{
+			foreach (BgScriptElement childElement in element.ChildNodes.OfType<BgScriptElement>())
+			{
+				if (childElement.Name == "Default" || await EvaluateConditionAsync(childElement))
+				{
+					await readContentsAsync(childElement);
+					break;
+				}
+			}
+		}
+
+		/// <summary>
+		/// Reads a task definition from the given element, and add it to the given list
+		/// </summary>
+		/// <param name="element">Xml element to read the definition from</param>
+		protected override async Task ReadTaskAsync(BgScriptElement element)
+		{
+			// If we're running a single node and this element's parent isn't the single node to run, ignore the error and return.
+			if (!String.IsNullOrWhiteSpace(_singleNodeName) && _enclosingNode!.Name != _singleNodeName)
+			{
+				return;
+			}
+
+			if (await EvaluateConditionAsync(element))
+			{
+				BgTask info = new BgTask(element.Location, element.Name);
+				foreach (XmlAttribute? attribute in element.Attributes)
+				{
+					if (String.Compare(attribute!.Name, "If", StringComparison.InvariantCultureIgnoreCase) != 0)
+					{
+						string expandedValue = ExpandProperties(element, attribute.Value);
+						info.Arguments.Add(attribute.Name, expandedValue);
+					}
+				}
+				_enclosingNode!.Tasks.Add(info);
+			}
+		}
+
+		/// <summary>
+		/// Reads the definition for an email notifier
+		/// </summary>
+		/// <param name="element">Xml element to read the definition from</param>
+		protected override async Task ReadNotifierAsync(BgScriptElement element)
+		{
+			if (await EvaluateConditionAsync(element))
+			{
+				string[] targetNames = ReadListAttribute(element, "Targets");
+				string[] exceptNames = ReadListAttribute(element, "Except");
+				string[] individualNodeNames = ReadListAttribute(element, "Nodes");
+				string[] reportNames = ReadListAttribute(element, "Reports");
+				string[] users = ReadListAttribute(element, "Users");
+				string[] submitters = ReadListAttribute(element, "Submitters");
+				bool? bWarnings = element.HasAttribute("Warnings") ? (bool?)ReadBooleanAttribute(element, "Warnings", true) : null;
+				bool bAbsolute = element.HasAttribute("Absolute") && ReadBooleanAttribute(element, "Absolute", true);
+
+				// Find the list of targets which are included, and recurse through all their dependencies
+				HashSet<BgNode> nodes = new HashSet<BgNode>();
+				if (targetNames != null)
+				{
+					HashSet<BgNode> targetNodes = ResolveReferences(element, targetNames);
+					foreach (BgNode node in targetNodes)
+					{
+						nodes.Add(node);
+						nodes.UnionWith(node.InputDependencies);
+					}
+				}
+
+				// Add all the individually referenced nodes
+				if (individualNodeNames != null)
+				{
+					HashSet<BgNode> individualNodes = ResolveReferences(element, individualNodeNames);
+					nodes.UnionWith(individualNodes);
+				}
+
+				// Exclude all the exceptions
+				if (exceptNames != null)
+				{
+					HashSet<BgNode> exceptNodes = ResolveReferences(element, exceptNames);
+					nodes.ExceptWith(exceptNodes);
+				}
+
+				// Update all the referenced nodes with the settings
+				foreach (BgNode node in nodes)
+				{
+					if (users != null)
+					{
+						if (bAbsolute)
+						{
+							node.NotifyUsers = new HashSet<string>(users);
+						}
+						else
+						{
+							node.NotifyUsers.UnionWith(users);
+						}
+					}
+					if (submitters != null)
+					{
+						if (bAbsolute)
+						{
+							node.NotifySubmitters = new HashSet<string>(submitters);
+						}
+						else
+						{
+							node.NotifySubmitters.UnionWith(submitters);
+						}
+					}
+					if (bWarnings.HasValue)
+					{
+						node.NotifyOnWarnings = bWarnings.Value;
+					}
+				}
+
+				// Add the users to the list of reports
+				if (reportNames != null)
+				{
+					foreach (string reportName in reportNames)
+					{
+						BgReport? report;
+						if (_graph.NameToReport.TryGetValue(reportName, out report))
+						{
+							report.NotifyUsers.UnionWith(users);
+						}
+						else
+						{
+							LogError(element, "Report '{0}' has not been defined", reportName);
+						}
+					}
+				}
+			}
+		}
+
+		/// <inheritdoc/>
+		protected override async Task ReadAnnotationAsync(BgScriptElement element)
+		{
+			if (await EvaluateConditionAsync(element))
+			{
+				string[] targetNames = ReadListAttribute(element, "Targets");
+				string[] exceptNames = ReadListAttribute(element, "Except");
+				string[] individualNodeNames = ReadListAttribute(element, "Nodes");
+				Dictionary<string, string> annotations = ReadAnnotationsAttribute(element, "Values");
+
+				// Find the list of targets which are included, and recurse through all their dependencies
+				HashSet<BgNode> nodes = new HashSet<BgNode>();
+				if (targetNames != null)
+				{
+					HashSet<BgNode> targetNodes = ResolveReferences(element, targetNames);
+					foreach (BgNode node in targetNodes)
+					{
+						nodes.Add(node);
+						nodes.UnionWith(node.InputDependencies);
+					}
+				}
+
+				// Add all the individually referenced nodes
+				if (individualNodeNames != null)
+				{
+					HashSet<BgNode> individualNodes = ResolveReferences(element, individualNodeNames);
+					nodes.UnionWith(individualNodes);
+				}
+
+				// Exclude all the exceptions
+				if (exceptNames != null)
+				{
+					HashSet<BgNode> exceptNodes = ResolveReferences(element, exceptNames);
+					nodes.ExceptWith(exceptNodes);
+				}
+
+				// Update all the referenced nodes with the settings
+				foreach (BgNode node in nodes)
+				{
+					foreach ((string key, string value) in annotations)
+					{
+						node.Annotations[key] = value;
+					}
+				}
+			}
+		}
+
+		/// <summary>
+		/// Reads a warning from the given element, evaluates the condition on it, and writes it to the log if the condition passes.
+		/// </summary>
+		/// <param name="element">Xml element to read the definition from</param>
+		/// <param name="level">The diagnostic event type</param>
+		protected override async Task ReadDiagnosticAsync(BgScriptElement element, LogLevel level)
+		{
+			if (await EvaluateConditionAsync(element))
+			{
+				string message = ReadAttribute(element, "Message");
+
+				BgDiagnostic diagnostic = new BgDiagnostic(element.Location.File, element.Location.LineNumber, level, message);
+				if (_enclosingNode != null)
+				{
+					_enclosingNode.Diagnostics.Add(diagnostic);
+				}
+				else if (_enclosingAgent != null)
+				{
+					_enclosingAgent.Diagnostics.Add(diagnostic);
+				}
+				else
+				{
+					_graph.Diagnostics.Add(diagnostic);
+				}
+			}
+		}
+
+		/// <summary>
+		/// Checks that the given name does not already used to refer to a node, and print an error if it is.
+		/// </summary>
+		/// <param name="element">Xml element to read from</param>
+		/// <param name="name">Name of the alias</param>
+		/// <returns>True if the name was registered correctly, false otherwise.</returns>
+		bool CheckNameIsUnique(BgScriptElement element, string name)
+		{
+			// Get the nodes that it maps to
+			if (_graph.ContainsName(name))
+			{
+				LogError(element, "'{0}' is already defined; cannot add a second time", name);
+				return false;
+			}
+			return true;
+		}
+
+		/// <summary>
+		/// Resolve a list of references to a set of nodes
+		/// </summary>
+		/// <param name="element">Element used to locate any errors</param>
+		/// <param name="referenceNames">Sequence of names to look up</param>
+		/// <returns>Hashset of all the nodes included by the given names</returns>
+		HashSet<BgNode> ResolveReferences(BgScriptElement element, IEnumerable<string> referenceNames)
+		{
+			HashSet<BgNode> nodes = new HashSet<BgNode>();
+			foreach (string referenceName in referenceNames)
+			{
+				BgNode[]? otherNodes;
+				if (_graph.TryResolveReference(referenceName, out otherNodes))
+				{
+					nodes.UnionWith(otherNodes);
+				}
+				else if (!referenceName.StartsWith("#") && _graph.TagNameToNodeOutput.ContainsKey("#" + referenceName))
+				{
+					LogError(element, "Reference to '{0}' cannot be resolved; did you mean '#{0}'?", referenceName);
+				}
+				else
+				{
+					LogError(element, "Reference to '{0}' cannot be resolved; check it has been defined.", referenceName);
+				}
+			}
+			return nodes;
+		}
+
+		/// <summary>
+		/// Resolve a list of references to a set of nodes
+		/// </summary>
+		/// <param name="element">Element used to locate any errors</param>
+		/// <param name="referenceNames">Sequence of names to look up</param>
+		/// <returns>Set of all the nodes included by the given names</returns>
+		HashSet<BgNodeOutput> ResolveInputReferences(BgScriptElement element, IEnumerable<string> referenceNames)
+		{
+			HashSet<BgNodeOutput> inputs = new HashSet<BgNodeOutput>();
+			foreach (string referenceName in referenceNames)
+			{
+				BgNodeOutput[]? referenceInputs;
+				if (_graph.TryResolveInputReference(referenceName, out referenceInputs))
+				{
+					inputs.UnionWith(referenceInputs);
+				}
+				else if (!referenceName.StartsWith("#") && _graph.TagNameToNodeOutput.ContainsKey("#" + referenceName))
+				{
+					LogError(element, "Reference to '{0}' cannot be resolved; did you mean '#{0}'?", referenceName);
+				}
+				else
+				{
+					LogError(element, "Reference to '{0}' cannot be resolved; check it has been defined.", referenceName);
+				}
+			}
+			return inputs;
+		}
+	}
+}
\ No newline at end of file
diff --git a/Engine/Source/Programs/AutomationTool/BuildGraph/BgScriptWriter.cs b/Engine/Source/Programs/AutomationTool/BuildGraph/BgScriptWriter.cs
new file mode 100644
--- /dev/null
+++ b/Engine/Source/Programs/AutomationTool/BuildGraph/BgScriptWriter.cs
@@ -0,0 +1,182 @@
+// Copyright Epic Games, Inc. All Rights Reserved.
+
+using System;
+using System.Collections.Generic;
+using System.Diagnostics.CodeAnalysis;
+using System.IO;
+using System.Linq;
+using System.Text;
+using System.Text.RegularExpressions;
+using System.Threading.Tasks;
+using System.Xml;
+using System.Xml.Schema;
+using EpicGames.BuildGraph;
+using EpicGames.Core;
+using Microsoft.Extensions.Logging;
+
+namespace AutomationTool
+{
+	/// <summary>
+	/// Implementation of XmlDocument which preserves line numbers for its elements
+	/// </summary>
+	public static class BgScriptWriter
+	{
+		/// <summary>
+		/// Writes a preprocessed build graph to a script file
+		/// </summary>
+		/// <param name="file">The file to load</param>
+		/// <param name="schemaFile">Schema file for validation</param>
+		public static void Write(this BgGraph graph, FileReference file, FileReference schemaFile)
+		{
+			XmlWriterSettings settings = new XmlWriterSettings();
+			settings.Indent = true;
+			settings.IndentChars = "\t";
+
+			using (XmlWriter writer = XmlWriter.Create(file.FullName, settings))
+			{
+				writer.WriteStartElement("BuildGraph", "http://www.epicgames.com/BuildGraph");
+
+				if (schemaFile != null)
+				{
+					writer.WriteAttributeString("schemaLocation", "http://www.w3.org/2001/XMLSchema-instance", "http://www.epicgames.com/BuildGraph " + schemaFile.MakeRelativeTo(file.Directory));
+				}
+
+				foreach (BgAgent agent in graph.Agents)
+				{
+					agent.Write(writer);
+				}
+
+				foreach (BgAggregate aggregate in graph.NameToAggregate.Values)
+				{
+					// If the aggregate has no required elements, skip it.
+					if (aggregate.RequiredNodes.Count == 0)
+					{
+						continue;
+					}
+
+					writer.WriteStartElement("Aggregate");
+					writer.WriteAttributeString("Name", aggregate.Name);
+					writer.WriteAttributeString("Requires", String.Join(";", aggregate.RequiredNodes.Select(x => x.Name)));
+					writer.WriteEndElement();
+				}
+
+				foreach (BgLabel label in graph.Labels)
+				{
+					writer.WriteStartElement("Label");
+					if (label.DashboardCategory != null)
+					{
+						writer.WriteAttributeString("Category", label.DashboardCategory);
+					}
+					writer.WriteAttributeString("Name", label.DashboardName);
+					writer.WriteAttributeString("Requires", String.Join(";", label.RequiredNodes.Select(x => x.Name)));
+
+					HashSet<BgNode> includedNodes = new HashSet<BgNode>(label.IncludedNodes);
+					includedNodes.ExceptWith(label.IncludedNodes.SelectMany(x => x.InputDependencies));
+					includedNodes.ExceptWith(label.RequiredNodes);
+					if (includedNodes.Count > 0)
+					{
+						writer.WriteAttributeString("Include", String.Join(";", includedNodes.Select(x => x.Name)));
+					}
+
+					HashSet<BgNode> excludedNodes = new HashSet<BgNode>(label.IncludedNodes);
+					excludedNodes.UnionWith(label.IncludedNodes.SelectMany(x => x.InputDependencies));
+					excludedNodes.ExceptWith(label.IncludedNodes);
+					excludedNodes.ExceptWith(excludedNodes.ToArray().SelectMany(x => x.InputDependencies));
+					if (excludedNodes.Count > 0)
+					{
+						writer.WriteAttributeString("Exclude", String.Join(";", excludedNodes.Select(x => x.Name)));
+					}
+					writer.WriteEndElement();
+				}
+
+				foreach (BgReport report in graph.NameToReport.Values)
+				{
+					writer.WriteStartElement("Report");
+					writer.WriteAttributeString("Name", report.Name);
+					writer.WriteAttributeString("Requires", String.Join(";", report.Nodes.Select(x => x.Name)));
+					writer.WriteEndElement();
+				}
+
+				foreach (BgBadge badge in graph.Badges)
+				{
+					writer.WriteStartElement("Badge");
+					writer.WriteAttributeString("Name", badge.Name);
+					if (badge.Project != null)
+					{
+						writer.WriteAttributeString("Project", badge.Project);
+					}
+					if (badge.Change != 0)
+					{
+						writer.WriteAttributeString("Change", badge.Change.ToString());
+					}
+					writer.WriteAttributeString("Requires", String.Join(";", badge.Nodes.Select(x => x.Name)));
+					writer.WriteEndElement();
+				}
+
+				writer.WriteEndElement();
+			}
+		}
+
+		/// <summary>
+		/// Writes this agent group out to a file, filtering nodes by a controlling trigger
+		/// </summary>
+		/// <param name="writer">The XML writer to output to</param>
+		public static void Write(this BgAgent agent, XmlWriter writer)
+		{
+			writer.WriteStartElement("Agent");
+			writer.WriteAttributeString("Name", agent.Name);
+			writer.WriteAttributeString("Type", String.Join(";", agent.PossibleTypes));
+			foreach (BgNode node in agent.Nodes)
+			{
+				node.Write(writer);
+			}
+			writer.WriteEndElement();
+		}
+
+		/// <summary>
+		/// Write this node to an XML writer
+		/// </summary>
+		/// <param name="writer">The writer to output the node to</param>
+		public static void Write(this BgNode node, XmlWriter writer)
+		{
+			writer.WriteStartElement("Node");
+			writer.WriteAttributeString("Name", node.Name);
+
+			string[] requireNames = node.Inputs.Select(x => x.TagName).ToArray();
+			if (requireNames.Length > 0)
+			{
+				writer.WriteAttributeString("Requires", String.Join(";", requireNames));
+			}
+
+			string[] producesNames = node.Outputs.Where(x => x != node.DefaultOutput).Select(x => x.TagName).ToArray();
+			if (producesNames.Length > 0)
+			{
+				writer.WriteAttributeString("Produces", String.Join(";", producesNames));
+			}
+
+			string[] afterNames = node.GetDirectOrderDependencies().Except(node.InputDependencies).Select(x => x.Name).ToArray();
+			if (afterNames.Length > 0)
+			{
+				writer.WriteAttributeString("After", String.Join(";", afterNames));
+			}
+
+			if (!node.NotifyOnWarnings)
+			{
+				writer.WriteAttributeString("NotifyOnWarnings", node.NotifyOnWarnings.ToString());
+			}
+
+			if (node.RunEarly)
+			{
+				writer.WriteAttributeString("RunEarly", node.RunEarly.ToString());
+			}
+
+			BgScriptNode scriptNode = (BgScriptNode)node;
+			foreach (BgTask task in scriptNode.Tasks)
+			{
+				task.Write(writer);
+			}
+
+			writer.WriteEndElement();
+		}
+	}
+}
\ No newline at end of file
diff --git a/Engine/Source/Programs/AutomationTool/BuildGraph/BgTask.cs b/Engine/Source/Programs/AutomationTool/BuildGraph/BgTask.cs
new file mode 100644
--- /dev/null
+++ b/Engine/Source/Programs/AutomationTool/BuildGraph/BgTask.cs
@@ -0,0 +1,51 @@
+// Copyright Epic Games, Inc. All Rights Reserved.
+
+using System.Collections.Generic;
+using System.Xml;
+
+namespace AutomationTool
+{
+	/// <summary>
+	/// A task invocation
+	/// </summary>
+	public class BgTask
+	{
+		/// <summary>
+		/// Line number in a source file that this task was declared. Optional; used for log messages.
+		/// </summary>
+		public BgScriptLocation Location { get; }
+
+		/// <summary>
+		/// Name of the task
+		/// </summary>
+		public string Name { get; set; }
+
+		/// <summary>
+		/// Arguments for the task
+		/// </summary>
+		public Dictionary<string, string> Arguments { get; } = new Dictionary<string, string>();
+
+		/// <summary>
+		/// Constructor
+		/// </summary>
+		public BgTask(BgScriptLocation location, string name)
+		{
+			Location = location;
+			Name = name;
+		}
+
+		/// <summary>
+		/// Write to an xml file
+		/// </summary>
+		/// <param name="writer"></param>
+		public void Write(XmlWriter writer)
+		{
+			writer.WriteStartElement(Name);
+			foreach (KeyValuePair<string, string> argument in Arguments)
+			{
+				writer.WriteAttributeString(argument.Key, argument.Value);
+			}
+			writer.WriteEndElement();
+		}
+	}
+}
diff --git a/Engine/Source/Programs/AutomationTool/BuildGraph/BuildGraph.cs b/Engine/Source/Programs/AutomationTool/BuildGraph/BuildGraph.cs
--- a/Engine/Source/Programs/AutomationTool/BuildGraph/BuildGraph.cs
+++ b/Engine/Source/Programs/AutomationTool/BuildGraph/BuildGraph.cs
@@ -19,6 +19,7 @@ using UnrealBuildBase;
 using UnrealBuildTool;
 using System.Threading.Tasks;
 using EpicGames.BuildGraph.Expressions;
+using AutomationTool.Tasks;
 
 namespace AutomationTool
 {
@@ -727,17 +728,17 @@ namespace AutomationTool
 			}
 
 			// Create tasks for the entire graph
-			Dictionary<BgTask, BgTaskImpl> TaskInfoToTask = new Dictionary<BgTask, BgTaskImpl>();
+			Dictionary<BgNode, BgNodeExecutor> NodeToExecutor = new Dictionary<BgNode, BgNodeExecutor>();
 			if (bSkipValidation && SingleNode != null)
 			{
-				if (!CreateTaskInstances(Graph, SingleNode, NameToTask, TaskInfoToTask))
+				if (!BindNodes(SingleNode, NameToTask, Graph.TagNameToNodeOutput, NodeToExecutor))
 				{
 					return ExitCode.Error_Unknown;
 				}
 			}
 			else
 			{
-				if (!CreateTaskInstances(Graph, NameToTask, TaskInfoToTask))
+				if (!BindNodes(Graph, NameToTask, NodeToExecutor))
 				{
 					return ExitCode.Error_Unknown;
 				}
@@ -748,14 +749,14 @@ namespace AutomationTool
 			{
 				if(SingleNode != null)
 				{
-					if(!await BuildNodeAsync(new JobContext(this), Graph, SingleNode, TaskInfoToTask, Storage, bWithBanner: true))
+					if(!await BuildNodeAsync(new JobContext(this), Graph, SingleNode, NodeToExecutor, Storage, bWithBanner: true))
 					{
 						return ExitCode.Error_Unknown;
 					}
 				}
 				else
 				{
-					if(!await BuildAllNodesAsync(new JobContext(this), Graph, TaskInfoToTask, Storage))
+					if(!await BuildAllNodesAsync(new JobContext(this), Graph, NodeToExecutor, Storage))
 					{
 						return ExitCode.Error_Unknown;
 					}
@@ -764,179 +765,39 @@ namespace AutomationTool
 			return ExitCode.Success;
 		}
 
-		bool CreateTaskInstances(BgGraph Graph, Dictionary<string, ScriptTaskBinding> NameToTask, Dictionary<BgTask, BgTaskImpl> TaskInfoToTask)
+		bool BindNodes(BgGraph Graph, Dictionary<string, ScriptTaskBinding> NameToTask, Dictionary<BgNode, BgNodeExecutor> NodeToExecutor)
 		{
 			bool bResult = true;
 			foreach (BgAgent Agent in Graph.Agents)
 			{
 				foreach (BgNode Node in Agent.Nodes)
 				{
-					bResult &= CreateTaskInstances(Graph, Node, NameToTask, TaskInfoToTask);
+					bResult &= BindNodes(Node, NameToTask, Graph.TagNameToNodeOutput, NodeToExecutor);
 				}
 			}
 			return bResult;
 		}
 
-		bool CreateTaskInstances(BgGraph Graph, BgNode Node, Dictionary<string, ScriptTaskBinding> NameToTask, Dictionary<BgTask, BgTaskImpl> TaskInfoToTask)
+		bool BindNodes(BgNode Node, Dictionary<string, ScriptTaskBinding> NameToTask, Dictionary<string, BgNodeOutput> TagNameToNodeOutput, Dictionary<BgNode, BgNodeExecutor> NodeToExecutor)
 		{
-			bool bResult = true;
-			foreach (BgTask TaskInfo in Node.Tasks)
+			if (Node is BgScriptNode ScriptNode)
 			{
-				BgTaskImpl Task = BindTask(Node, TaskInfo, NameToTask, Graph.TagNameToNodeOutput);
-				if (Task == null)
-				{
-					bResult = false;
-				}
-				else
-				{
-					TaskInfoToTask.Add(TaskInfo, Task);
-				}
+				BgScriptNodeExecutor executor = new BgScriptNodeExecutor(ScriptNode);
+				NodeToExecutor[Node] = executor;
+				return executor.Bind(NameToTask, TagNameToNodeOutput, Context, Logger);
 			}
-			return bResult;
-		}
-
-		BgTaskImpl BindTask(BgNode Node, BgTask TaskInfo, Dictionary<string, ScriptTaskBinding> NameToTask, IReadOnlyDictionary<string, BgNodeOutput> TagNameToNodeOutput)
-		{
-			// Get the reflection info for this element
-			ScriptTaskBinding Task;
-			if (!NameToTask.TryGetValue(TaskInfo.Name, out Task))
+			else if (Node is BgExpressionNode ExpressionNode)
 			{
-				OutputBindingError(TaskInfo, "Unknown task '{TaskName}'", TaskInfo.Name);
-				return null;
-			}
-
-			// Check all the required parameters are present
-			bool bHasRequiredAttributes = true;
-			foreach (ScriptTaskParameterBinding Parameter in Task.NameToParameter.Values)
-			{
-				if (!Parameter.Optional && !TaskInfo.Arguments.ContainsKey(Parameter.Name))
-				{
-					OutputBindingError(TaskInfo, "Missing required attribute - {AttrName}", Parameter.Name);
-					bHasRequiredAttributes = false;
-				}
-			}
-
-			// Read all the attributes into a parameters object for this task
-			object ParametersObject = Activator.CreateInstance(Task.ParametersClass);
-			foreach ((string Name, string Value) in TaskInfo.Arguments)
-			{
-				// Get the field that this attribute should be written to in the parameters object
-				ScriptTaskParameterBinding Parameter;
-				if (!Task.NameToParameter.TryGetValue(Name, out Parameter))
-				{
-					OutputBindingError(TaskInfo, "Unknown attribute '{AttrName}'", Name);
-					continue;
-				}
-
-				// If it's a collection type, split it into separate values
-				if (Parameter.CollectionType == null)
-				{
-					// Parse it and assign it to the parameters object
-					object FieldValue = ParseValue(Value, Parameter.ValueType);
-					Parameter.FieldInfo.SetValue(ParametersObject, FieldValue);
-				}
-				else
-				{
-					// Get the collection, or create one if necessary
-					object CollectionValue = Parameter.FieldInfo.GetValue(ParametersObject);
-					if (CollectionValue == null)
-					{
-						CollectionValue = Activator.CreateInstance(Parameter.FieldInfo.FieldType);
-						Parameter.FieldInfo.SetValue(ParametersObject, CollectionValue);
-					}
-
-					// Parse the values and add them to the collection
-					List<string> ValueStrings = BgTaskImpl.SplitDelimitedList(Value);
-					foreach (string ValueString in ValueStrings)
-					{
-						object ElementValue = ParseValue(ValueString, Parameter.ValueType);
-						Parameter.CollectionType.InvokeMember("Add", BindingFlags.InvokeMethod | BindingFlags.Instance | BindingFlags.Public, null, CollectionValue, new object[] { ElementValue });
-					}
-				}
-			}
-
-			// Construct the task
-			if (!bHasRequiredAttributes)
-			{
-				return null;
-			}
-
-			// Add it to the list
-			BgTaskImpl NewTask = (BgTaskImpl)Activator.CreateInstance(Task.TaskClass, ParametersObject);
-
-			// Set up the source location for diagnostics
-			NewTask.SourceLocation = TaskInfo.Location;
-
-			// Make sure all the read tags are local or listed as a dependency
-			foreach (string ReadTagName in NewTask.FindConsumedTagNames())
-			{
-				BgNodeOutput Output;
-				if (TagNameToNodeOutput.TryGetValue(ReadTagName, out Output))
-				{
-					if (Output != null && Output.ProducingNode != Node && !Node.Inputs.Contains(Output))
-					{
-						OutputBindingError(TaskInfo, "The tag '{TagName}' is not a dependency of node '{Node}'", ReadTagName, Node.Name);
-					}
-				}
-			}
-
-			// Make sure all the written tags are local or listed as an output
-			foreach (string ModifiedTagName in NewTask.FindProducedTagNames())
-			{
-				BgNodeOutput Output;
-				if (TagNameToNodeOutput.TryGetValue(ModifiedTagName, out Output))
-				{
-					if (Output != null && !Node.Outputs.Contains(Output))
-					{
-						OutputBindingError(TaskInfo, "The tag '{0}' is created by '{1}', and cannot be modified downstream", Output.TagName, Output.ProducingNode.Name);
-					}
-				}
-			}
-			return NewTask;
-		}
-
-		/// <summary>
-		/// Parse a value of the given type
-		/// </summary>
-		/// <param name="ValueText">The text to parse</param>
-		/// <param name="ValueType">Type of the value to parse</param>
-		/// <returns>Value that was parsed</returns>
-		object ParseValue(string ValueText, Type ValueType)
-		{
-			// Parse it and assign it to the parameters object
-			if (ValueType.IsEnum)
-			{
-				return Enum.Parse(ValueType, ValueText);
-			}
-			else if (ValueType == typeof(Boolean))
-			{
-				return BgCondition.EvaluateAsync(ValueText, Context).Result;
-			}
-			else if (ValueType == typeof(FileReference))
-			{
-				return BgTaskImpl.ResolveFile(ValueText);
-			}
-			else if (ValueType == typeof(DirectoryReference))
-			{
-				return BgTaskImpl.ResolveDirectory(ValueText);
-			}
-
-			TypeConverter Converter = TypeDescriptor.GetConverter(ValueType);
-			if (Converter.CanConvertFrom(typeof(string)))
-			{
-				return Converter.ConvertFromString(ValueText);
+				BgExpressionNodeExecutor executor = new BgExpressionNodeExecutor(ExpressionNode);
+				NodeToExecutor[Node] = executor;
+				return executor.Bind(Logger);
 			}
 			else
 			{
-				return Convert.ChangeType(ValueText, ValueType);
+				throw new NotImplementedException();
 			}
 		}
 
-		void OutputBindingError(BgTask Task, string Format, params object[] Args)
-		{
-			Logger.LogScriptError(Task.Location, Format, Args);
-		}
-
 		static void FindAvailableGraphs(Dictionary<string, Type> NameToType)
 		{
 			foreach (Assembly LoadedAssembly in ScriptManager.AllScriptAssemblies)
@@ -1121,9 +982,10 @@ namespace AutomationTool
 		/// </summary>
 		/// <param name="Job">Information about the current job</param>
 		/// <param name="Graph">The graph instance</param>
+		/// <param name="NodeToExecutor">Map from node to executor</param>
 		/// <param name="Storage">The temp storage backend which stores the shared state</param>
 		/// <returns>True if everything built successfully</returns>
-		async Task<bool> BuildAllNodesAsync(JobContext Job, BgGraph Graph, Dictionary<BgTask, BgTaskImpl> TaskInfoToTask, TempStorage Storage)
+		async Task<bool> BuildAllNodesAsync(JobContext Job, BgGraph Graph, Dictionary<BgNode, BgNodeExecutor> NodeToExecutor, TempStorage Storage)
 		{
 			// Build a flat list of nodes to execute, in order
 			BgNode[] NodesToExecute = Graph.Agents.SelectMany(x => x.Nodes).ToArray();
@@ -1148,7 +1010,7 @@ namespace AutomationTool
 				if(!Storage.IsComplete(NodeToExecute.Name))
 				{
 					LogInformation("");
-					if(!await BuildNodeAsync(Job, Graph, NodeToExecute, TaskInfoToTask, Storage, bWithBanner: false))
+					if(!await BuildNodeAsync(Job, Graph, NodeToExecute, NodeToExecutor, Storage, bWithBanner: false))
 					{
 						return false;
 					} 
@@ -1164,11 +1026,11 @@ namespace AutomationTool
 		/// <param name="Job">Information about the current job</param>
 		/// <param name="Graph">The graph to which the node belongs. Used to determine which outputs need to be transferred to temp storage.</param>
 		/// <param name="Node">The node to build</param>
-		/// <param name="TaskInfoToTask">Map from task info instance to task instance</param>
+		/// <param name="NodeToExecutor">Map from node to executor</param>
 		/// <param name="Storage">The temp storage backend which stores the shared state</param>
 		/// <param name="bWithBanner">Whether to write a banner before and after this node's log output</param>
 		/// <returns>True if the node built successfully, false otherwise.</returns>
-		async Task<bool> BuildNodeAsync(JobContext Job, BgGraph Graph, BgNode Node, Dictionary<BgTask, BgTaskImpl> TaskInfoToTask, TempStorage Storage, bool bWithBanner)
+		async Task<bool> BuildNodeAsync(JobContext Job, BgGraph Graph, BgNode Node, Dictionary<BgNode, BgNodeExecutor> NodeToExecutor, TempStorage Storage, bool bWithBanner)
 		{
 			DirectoryReference RootDir = new DirectoryReference(CommandUtils.CmdEnv.LocalRoot);
 
@@ -1225,7 +1087,7 @@ namespace AutomationTool
 				Console.WriteLine();
 				CommandUtils.LogInformation("========== Starting: {0} ==========", Node.Name);
 			}
-			if(!await ExecuteTasksAsync(Node, Job, TaskInfoToTask, TagNameToFileSet))
+			if(!await NodeToExecutor[Node].ExecuteAsync(Job, TagNameToFileSet))
 			{
 				return false;
 			}
@@ -1359,78 +1221,6 @@ namespace AutomationTool
 			return true;
 		}
 
-		/// <summary>
-		/// Build all the tasks for this node
-		/// </summary>
-		/// <param name="Job">Information about the current job</param>
-		/// <param name="TaskInfoToTask">Map from TaskInfo to Task object</param>
-		/// <param name="TagNameToFileSet">Mapping from tag names to the set of files they include. Should be set to contain the node inputs on entry.</param>
-		/// <returns>Whether the task succeeded or not. Exiting with an exception will be caught and treated as a failure.</returns>
-		async Task<bool> ExecuteTasksAsync(BgNode Node, JobContext Job, Dictionary<BgTask, BgTaskImpl> TaskInfoToTask, Dictionary<string, HashSet<FileReference>> TagNameToFileSet)
-		{
-			List<BgTaskImpl> Tasks = Node.Tasks.ConvertAll(x => TaskInfoToTask[x]);
-
-			// Run each of the tasks in order
-			HashSet<FileReference> BuildProducts = TagNameToFileSet[Node.DefaultOutput.TagName];
-			for (int Idx = 0; Idx < Tasks.Count; Idx++)
-			{
-				using (IScope Scope = GlobalTracer.Instance.BuildSpan("Task").WithTag("resource", Tasks[Idx].GetTraceName()).StartActive())
-				{
-					ITaskExecutor Executor = Tasks[Idx].GetExecutor();
-					if (Executor == null)
-					{
-						// Execute this task directly
-						try
-						{
-							Tasks[Idx].GetTraceMetadata(Scope.Span, "");
-							await Tasks[Idx].ExecuteAsync(Job, BuildProducts, TagNameToFileSet);
-						}
-						catch (Exception Ex)
-						{
-							ExceptionUtils.AddContext(Ex, "while executing task {0}", Tasks[Idx].GetTraceString());
-							if (Tasks[Idx].SourceLocation != null)
-							{
-								ExceptionUtils.AddContext(Ex, "at {0}({1})", Tasks[Idx].SourceLocation.File, Tasks[Idx].SourceLocation.LineNumber);
-							}
-							throw;
-						}
-					}
-					else
-					{
-						Tasks[Idx].GetTraceMetadata(Scope.Span, "1.");
-
-						// The task has a custom executor, which may be able to execute several tasks simultaneously. Try to add the following tasks.
-						int FirstIdx = Idx;
-						while (Idx + 1 < Tasks.Count && Executor.Add(Tasks[Idx + 1]))
-						{
-							Idx++;
-							Tasks[Idx].GetTraceMetadata(Scope.Span, string.Format("{0}.", 1 + Idx - FirstIdx));
-						}
-						try
-						{
-							await Executor.ExecuteAsync(Job, BuildProducts, TagNameToFileSet);
-						}
-						catch (Exception Ex)
-						{
-							for (int TaskIdx = FirstIdx; TaskIdx <= Idx; TaskIdx++)
-							{
-								ExceptionUtils.AddContext(Ex, "while executing {0}", Tasks[TaskIdx].GetTraceString());
-							}
-							if (Tasks[FirstIdx].SourceLocation != null)
-							{
-								ExceptionUtils.AddContext(Ex, "at {0}({1})", Tasks[FirstIdx].SourceLocation.File, Tasks[FirstIdx].SourceLocation.LineNumber);
-							}
-							throw;
-						}
-					}
-				}
-			}
-
-			// Remove anything that doesn't exist, since these files weren't explicitly tagged
-			BuildProducts.RemoveWhere(x => !FileReference.Exists(x));
-			return true;
-		}
-
 		/// <summary>
 		/// Generate HTML documentation for all the tasks
 		/// </summary>
diff --git a/Engine/Source/Programs/AutomationTool/BuildGraph/Tasks/CallMethodTask.cs b/Engine/Source/Programs/AutomationTool/BuildGraph/Tasks/CallMethodTask.cs
--- a/Engine/Source/Programs/AutomationTool/BuildGraph/Tasks/CallMethodTask.cs
+++ b/Engine/Source/Programs/AutomationTool/BuildGraph/Tasks/CallMethodTask.cs
@@ -99,49 +99,50 @@ namespace AutomationTool.Tasks
 		public string? Tags;
 	}
 
-	/// <summary>
-	/// Spawns the editor to run a commandlet.
-	/// </summary>
-	[TaskElement("CallMethod", typeof(CallMethodTaskParameters))]
-	public class CallMethodTask : BgTaskImpl
+	class BgContextImpl : BgContext
 	{
-		class BgStateImpl : BgContext
-		{
-			JobContext JobContext;
+		JobContext JobContext;
 
-			public BgStateImpl(JobContext JobContext, Dictionary<string, HashSet<FileReference>> TagNameToFileSet)
-				: base(CreateExprContext(TagNameToFileSet))
-			{
-				this.JobContext = JobContext;
-			}
+		public BgContextImpl(JobContext JobContext, Dictionary<string, HashSet<FileReference>> TagNameToFileSet)
+			: base(CreateExprContext(TagNameToFileSet))
+		{
+			this.JobContext = JobContext;
+		}
 
-			static BgExprContext CreateExprContext(Dictionary<string, HashSet<FileReference>> TagNameToFileSet)
+		static BgExprContext CreateExprContext(Dictionary<string, HashSet<FileReference>> TagNameToFileSet)
+		{
+			BgExprContext Context = new BgExprContext();
+			foreach (KeyValuePair<string, HashSet<FileReference>> Pair in TagNameToFileSet)
 			{
-				BgExprContext Context = new BgExprContext();
-				foreach (KeyValuePair<string, HashSet<FileReference>> Pair in TagNameToFileSet)
-				{
-					Context.TagNameToFileSet[Pair.Key] = FileSet.FromFiles(Unreal.RootDirectory, Pair.Value);
-				}
-				return Context;
+				Context.TagNameToFileSet[Pair.Key] = FileSet.FromFiles(Unreal.RootDirectory, Pair.Value);
 			}
+			return Context;
+		}
 
-			public override string Stream => CommandUtils.P4Enabled ? CommandUtils.P4Env.Branch : "";
+		public override string Stream => CommandUtils.P4Enabled ? CommandUtils.P4Env.Branch : "";
 
-			public override int Change => CommandUtils.P4Enabled ? CommandUtils.P4Env.Changelist : 0;
+		public override int Change => CommandUtils.P4Enabled ? CommandUtils.P4Env.Changelist : 0;
 
-			public override int CodeChange => CommandUtils.P4Enabled ? CommandUtils.P4Env.CodeChangelist : 0;
+		public override int CodeChange => CommandUtils.P4Enabled ? CommandUtils.P4Env.CodeChangelist : 0;
 
-			public override (int Major, int Minor, int Patch) EngineVersion
+		public override (int Major, int Minor, int Patch) EngineVersion
+		{
+			get
 			{
-				get
-				{
-					ReadOnlyBuildVersion Current = ReadOnlyBuildVersion.Current;
-					return (Current.MajorVersion, Current.MinorVersion, Current.PatchVersion);
-				}
+				ReadOnlyBuildVersion Current = ReadOnlyBuildVersion.Current;
+				return (Current.MajorVersion, Current.MinorVersion, Current.PatchVersion);
 			}
-
-			public override bool IsBuildMachine => CommandUtils.IsBuildMachine;
 		}
+
+		public override bool IsBuildMachine => CommandUtils.IsBuildMachine;
+	}
+
+	/// <summary>
+	/// Spawns the editor to run a commandlet.
+	/// </summary>
+	[TaskElement("CallMethod", typeof(CallMethodTaskParameters))]
+	public class CallMethodTask : BgTaskImpl
+	{
 				
 		/// <summary>
 		/// Parameters for this task
@@ -190,7 +191,7 @@ namespace AutomationTool.Tasks
 
 			ParameterInfo[] MethodParameters = Method.GetParameters();
 
-			BgStateImpl State = new BgStateImpl(Job, TagNameToFileSet);
+			BgContextImpl State = new BgContextImpl(Job, TagNameToFileSet);
 
 			object[] Arguments = new object[MethodParameters.Length];
 			for (int Idx = 0; Idx < MethodParameters.Length; Idx++)
diff --git a/Engine/Source/Programs/Shared/EpicGames.BuildGraph/BgAgent.cs b/Engine/Source/Programs/Shared/EpicGames.BuildGraph/BgAgent.cs
--- a/Engine/Source/Programs/Shared/EpicGames.BuildGraph/BgAgent.cs
+++ b/Engine/Source/Programs/Shared/EpicGames.BuildGraph/BgAgent.cs
@@ -44,21 +44,5 @@ namespace EpicGames.BuildGraph
 			Name = inName;
 			PossibleTypes = inPossibleTypes;
 		}
-
-		/// <summary>
-		/// Writes this agent group out to a file, filtering nodes by a controlling trigger
-		/// </summary>
-		/// <param name="writer">The XML writer to output to</param>
-		public void Write(XmlWriter writer)
-		{
-			writer.WriteStartElement("Agent");
-			writer.WriteAttributeString("Name", Name);
-			writer.WriteAttributeString("Type", String.Join(";", PossibleTypes));
-			foreach (BgNode node in Nodes)
-			{
-				node.Write(writer);
-			}
-			writer.WriteEndElement();
-		}
 	}
 }
diff --git a/Engine/Source/Programs/Shared/EpicGames.BuildGraph/BgCondition.cs b/Engine/Source/Programs/Shared/EpicGames.BuildGraph/BgCondition.cs
deleted file mode 100644
diff --git a/Engine/Source/Programs/Shared/EpicGames.BuildGraph/BgDiagnosticSpec.cs b/Engine/Source/Programs/Shared/EpicGames.BuildGraph/BgDiagnosticSpec.cs
--- a/Engine/Source/Programs/Shared/EpicGames.BuildGraph/BgDiagnosticSpec.cs
+++ b/Engine/Source/Programs/Shared/EpicGames.BuildGraph/BgDiagnosticSpec.cs
@@ -38,10 +38,9 @@ namespace EpicGames.BuildGraph
 		/// </summary>
 		internal void AddToGraph(BgExprContext context, BgGraph graph, BgAgent? enclosingAgent, BgNode? enclosingNode)
 		{
-			BgScriptLocation location = new BgScriptLocation("(unknown)", "(unknown)", 1);
 			string messageValue = Message.Compute(context);
 
-			BgDiagnostic diagnostic = new BgDiagnostic(location.File, location.LineNumber, Level, messageValue);
+			BgDiagnostic diagnostic = new BgDiagnostic("unknown", 1, Level, messageValue);
 			if (enclosingNode != null)
 			{
 				enclosingNode.Diagnostics.Add(diagnostic);
diff --git a/Engine/Source/Programs/Shared/EpicGames.BuildGraph/BgGraph.cs b/Engine/Source/Programs/Shared/EpicGames.BuildGraph/BgGraph.cs
--- a/Engine/Source/Programs/Shared/EpicGames.BuildGraph/BgGraph.cs
+++ b/Engine/Source/Programs/Shared/EpicGames.BuildGraph/BgGraph.cs
@@ -287,102 +287,6 @@ namespace EpicGames.BuildGraph
 			Badges.RemoveAll(x => x.Nodes.Any(y => !retainNodes.Contains(y)));
 		}
 
-		/// <summary>
-		/// Writes a preprocessed build graph to a script file
-		/// </summary>
-		/// <param name="file">The file to load</param>
-		/// <param name="schemaFile">Schema file for validation</param>
-		public void Write(FileReference file, FileReference schemaFile)
-		{
-			XmlWriterSettings settings = new XmlWriterSettings();
-			settings.Indent = true;
-			settings.IndentChars = "\t";
-
-			using (XmlWriter writer = XmlWriter.Create(file.FullName, settings))
-			{
-				writer.WriteStartElement("BuildGraph", "http://www.epicgames.com/BuildGraph");
-
-				if (schemaFile != null)
-				{
-					writer.WriteAttributeString("schemaLocation", "http://www.w3.org/2001/XMLSchema-instance", "http://www.epicgames.com/BuildGraph " + schemaFile.MakeRelativeTo(file.Directory));
-				}
-
-				foreach (BgAgent agent in Agents)
-				{
-					agent.Write(writer);
-				}
-
-				foreach (BgAggregate aggregate in NameToAggregate.Values)
-				{
-					// If the aggregate has no required elements, skip it.
-					if (aggregate.RequiredNodes.Count == 0)
-					{
-						continue;
-					}
-
-					writer.WriteStartElement("Aggregate");
-					writer.WriteAttributeString("Name", aggregate.Name);
-					writer.WriteAttributeString("Requires", String.Join(";", aggregate.RequiredNodes.Select(x => x.Name)));
-					writer.WriteEndElement();
-				}
-
-				foreach (BgLabel label in Labels)
-				{
-					writer.WriteStartElement("Label");
-					if (label.DashboardCategory != null)
-					{
-						writer.WriteAttributeString("Category", label.DashboardCategory);
-					}
-					writer.WriteAttributeString("Name", label.DashboardName);
-					writer.WriteAttributeString("Requires", String.Join(";", label.RequiredNodes.Select(x => x.Name)));
-
-					HashSet<BgNode> includedNodes = new HashSet<BgNode>(label.IncludedNodes);
-					includedNodes.ExceptWith(label.IncludedNodes.SelectMany(x => x.InputDependencies));
-					includedNodes.ExceptWith(label.RequiredNodes);
-					if (includedNodes.Count > 0)
-					{
-						writer.WriteAttributeString("Include", String.Join(";", includedNodes.Select(x => x.Name)));
-					}
-
-					HashSet<BgNode> excludedNodes = new HashSet<BgNode>(label.IncludedNodes);
-					excludedNodes.UnionWith(label.IncludedNodes.SelectMany(x => x.InputDependencies));
-					excludedNodes.ExceptWith(label.IncludedNodes);
-					excludedNodes.ExceptWith(excludedNodes.ToArray().SelectMany(x => x.InputDependencies));
-					if (excludedNodes.Count > 0)
-					{
-						writer.WriteAttributeString("Exclude", String.Join(";", excludedNodes.Select(x => x.Name)));
-					}
-					writer.WriteEndElement();
-				}
-
-				foreach (BgReport report in NameToReport.Values)
-				{
-					writer.WriteStartElement("Report");
-					writer.WriteAttributeString("Name", report.Name);
-					writer.WriteAttributeString("Requires", String.Join(";", report.Nodes.Select(x => x.Name)));
-					writer.WriteEndElement();
-				}
-
-				foreach (BgBadge badge in Badges)
-				{
-					writer.WriteStartElement("Badge");
-					writer.WriteAttributeString("Name", badge.Name);
-					if (badge.Project != null)
-					{
-						writer.WriteAttributeString("Project", badge.Project);
-					}
-					if (badge.Change != 0)
-					{
-						writer.WriteAttributeString("Change", badge.Change.ToString());
-					}
-					writer.WriteAttributeString("Requires", String.Join(";", badge.Nodes.Select(x => x.Name)));
-					writer.WriteEndElement();
-				}
-
-				writer.WriteEndElement();
-			}
-		}
-
 		/// <summary>
 		/// Export the build graph to a Json file, for parallel execution by the build system
 		/// </summary>
diff --git a/Engine/Source/Programs/Shared/EpicGames.BuildGraph/BgNode.cs b/Engine/Source/Programs/Shared/EpicGames.BuildGraph/BgNode.cs
--- a/Engine/Source/Programs/Shared/EpicGames.BuildGraph/BgNode.cs
+++ b/Engine/Source/Programs/Shared/EpicGames.BuildGraph/BgNode.cs
@@ -3,6 +3,9 @@
 using System;
 using System.Collections.Generic;
 using System.Linq;
+using System.Linq.Expressions;
+using System.Reflection;
+using System.Threading.Tasks;
 using System.Xml;
 using EpicGames.Core;
 
@@ -44,6 +47,63 @@ namespace EpicGames.BuildGraph
 		}
 	}
 
+	/// <summary>
+	/// Method to be executed for a node
+	/// </summary>
+	public class BgMethod
+	{
+		/// <summary>
+		/// Full name of the class containing the method to execute
+		/// </summary>
+		public string ClassName { get; }
+
+		/// <summary>
+		/// Name of the method to execute
+		/// </summary>
+		public string MethodName { get; }
+
+		/// <summary>
+		/// Constructor
+		/// </summary>
+		public BgMethod(string className, string methodName)
+		{
+			ClassName = className;
+			MethodName = methodName;
+		}
+
+		/// <summary>
+		/// Constructor
+		/// </summary>
+		public BgMethod(MethodInfo method)
+			: this(method.DeclaringType!.AssemblyQualifiedName!, method.Name)
+		{
+		}
+
+		/// <summary>
+		/// Bind this method name to a method instance
+		/// </summary>
+		/// <returns>The resolved method instance</returns>
+		public MethodInfo Bind()
+		{
+			Type? Type = Type.GetType(ClassName);
+			if (Type == null)
+			{
+				throw new BgNodeException($"Unable to find class '{ClassName}'");
+			}
+
+			MethodInfo? Method = Type.GetMethod(MethodName, BindingFlags.Public | BindingFlags.NonPublic | BindingFlags.Static);
+			if (Method == null)
+			{
+				throw new BgNodeException($"Unable to find method '{Type.FullName}.{MethodName}'");
+			}
+
+			return Method;
+		}
+
+		/// <inheritdoc/>
+		public override string ToString() => $"{MethodName}";
+	}
+
 	/// <summary>
 	/// Defines a node, a container for tasks and the smallest unit of execution that can be run as part of a build graph.
 	/// </summary>
@@ -79,11 +139,6 @@ namespace EpicGames.BuildGraph
 		/// </summary>
 		public IReadOnlyList<FileReference> RequiredTokens { get; }
 
-		/// <summary>
-		/// List of tasks to execute
-		/// </summary>
-		public List<BgTask> Tasks { get; } = new List<BgTask>();
-
 		/// <summary>
 		/// List of email addresses to notify if this node fails.
 		/// </summary>
@@ -172,56 +227,38 @@ namespace EpicGames.BuildGraph
 		}
 
 		/// <summary>
-		/// Write this node to an XML writer
+		/// Returns the name of this node
 		/// </summary>
-		/// <param name="writer">The writer to output the node to</param>
-		public void Write(XmlWriter writer)
+		/// <returns>The name of this node</returns>
+		public override string ToString()
 		{
-			writer.WriteStartElement("Node");
-			writer.WriteAttributeString("Name", Name);
-
-			string[] requireNames = Inputs.Select(x => x.TagName).ToArray();
-			if (requireNames.Length > 0)
-			{
-				writer.WriteAttributeString("Requires", String.Join(";", requireNames));
-			}
-
-			string[] producesNames = Outputs.Where(x => x != DefaultOutput).Select(x => x.TagName).ToArray();
-			if (producesNames.Length > 0)
-			{
-				writer.WriteAttributeString("Produces", String.Join(";", producesNames));
-			}
-
-			string[] afterNames = GetDirectOrderDependencies().Except(InputDependencies).Select(x => x.Name).ToArray();
-			if (afterNames.Length > 0)
-			{
-				writer.WriteAttributeString("After", String.Join(";", afterNames));
-			}
-
-			if (!NotifyOnWarnings)
-			{
-				writer.WriteAttributeString("NotifyOnWarnings", NotifyOnWarnings.ToString());
-			}
+			return Name;
+		}
+	}
 
-			if (RunEarly)
-			{
-				writer.WriteAttributeString("RunEarly", RunEarly.ToString());
-			}
+	/// <summary>
+	/// Node constructed from a bytecode expression
+	/// </summary>
+	public class BgExpressionNode : BgNode
+	{
+		/// <summary>
+		/// The method to execute for this node
+		/// </summary>
+		public BgMethod Method { get; }
 
-			foreach (BgTask task in Tasks)
-			{
-				task.Write(writer);
-			}
-			writer.WriteEndElement();
-		}
+		/// <summary>
+		/// Arguments for invoking the method
+		/// </summary>
+		public IReadOnlyList<object?> Arguments { get; }
 
 		/// <summary>
-		/// Returns the name of this node
+		/// Constructor
 		/// </summary>
-		/// <returns>The name of this node</returns>
-		public override string ToString()
+		public BgExpressionNode(string name, BgMethod method, IReadOnlyList<object?> arguments, IReadOnlyList<BgNodeOutput> inputs, IReadOnlyList<string> outputNames, IReadOnlyList<BgNode> inputDependencies, IReadOnlyList<BgNode> orderDependencies, IReadOnlyList<FileReference> requiredTokens)
+			: base(name, inputs, outputNames, inputDependencies, orderDependencies, requiredTokens)
 		{
-			return Name;
+			Method = method;
+			Arguments = arguments;
 		}
 	}
 }
diff --git a/Engine/Source/Programs/Shared/EpicGames.BuildGraph/BgNodeSpec.cs b/Engine/Source/Programs/Shared/EpicGames.BuildGraph/BgNodeSpec.cs
--- a/Engine/Source/Programs/Shared/EpicGames.BuildGraph/BgNodeSpec.cs
+++ b/Engine/Source/Programs/Shared/EpicGames.BuildGraph/BgNodeSpec.cs
@@ -53,8 +53,7 @@ namespace EpicGames.BuildGraph
 	/// </summary>
 	public class BgNodeSpec
 	{
-		internal string _className;
-		internal string _methodName;
+		internal BgMethod _method;
 		internal IBgExpr?[] _argumentExprs;
 		internal BgFileSet[] ResultExprs { get; }
 
@@ -98,14 +97,7 @@ namespace EpicGames.BuildGraph
 		/// </summary>
 		internal BgNodeSpec(MethodCallExpression call)
 		{
-			string methodSignature = $"{call.Method.DeclaringType?.Name}.{call.Method.Name}";
-			if (!call.Method.IsStatic)
-			{
-				throw new BgNodeException($"Node {methodSignature} must be static");
-			}
-
-			_className = call.Method.DeclaringType!.AssemblyQualifiedName!;
-			_methodName = call.Method.Name;
+			_method = new BgMethod(call.Method);
 
 			try
 			{
@@ -115,7 +107,7 @@ namespace EpicGames.BuildGraph
 			}
 			catch (Exception ex)
 			{
-				ExceptionUtils.AddContext(ex, $"while calling method {methodSignature}");
+				ExceptionUtils.AddContext(ex, $"while calling method {_method}");
 				throw;
 			}
 
@@ -379,28 +371,10 @@ namespace EpicGames.BuildGraph
 			BgNode[] afterNodes = afterDependencies.Select(x => x.ProducingNode).Distinct().ToArray();
 			bool runEarly = CanRunEarly.Compute(context);
 
-			BgNode node = new BgNode(name, inputDependencies, outputNames, inputNodes, afterNodes, Array.Empty<FileReference>());
+			List<object?> arguments = _argumentExprs.ConvertAll(x => x?.Compute(context));
+			BgExpressionNode node = new BgExpressionNode(name, _method, arguments, inputDependencies, outputNames, inputNodes, afterNodes, Array.Empty<FileReference>());
 			node.RunEarly = runEarly;
 
-			BgScriptLocation location = new BgScriptLocation("(unknown)", "(unknown)", 1);
-
-			BgTask task = new BgTask(location, "CallMethod");
-			task.Arguments["Class"] = _className;
-			task.Arguments["Method"] = _methodName;
-			for (int idx = 0; idx < _argumentExprs.Length; idx++)
-			{
-				IBgExpr? argumentExpr = _argumentExprs[idx];
-				if (argumentExpr != null)
-				{
-					task.Arguments[$"Arg{idx + 1}"] = BgType.Get(argumentExpr.GetType()).SerializeArgument(argumentExpr, context);
-				}
-			}
-			if (outputNames != null)
-			{
-				task.Arguments["Tags"] = String.Join(";", outputNames);
-			}
-			node.Tasks.Add(task);
-
 			agent.Nodes.Add(node);
 			graph.NameToNode.Add(name, node);
 
diff --git a/Engine/Source/Programs/Shared/EpicGames.BuildGraph/BgSchema.cs b/Engine/Source/Programs/Shared/EpicGames.BuildGraph/BgSchema.cs
deleted file mode 100644
diff --git a/Engine/Source/Programs/Shared/EpicGames.BuildGraph/BgScript.cs b/Engine/Source/Programs/Shared/EpicGames.BuildGraph/BgScript.cs
deleted file mode 100644
diff --git a/Engine/Source/Programs/Shared/EpicGames.BuildGraph/BgScriptReader.cs b/Engine/Source/Programs/Shared/EpicGames.BuildGraph/BgScriptReader.cs
deleted file mode 100644
diff --git a/Engine/Source/Programs/Shared/EpicGames.BuildGraph/BgTask.cs b/Engine/Source/Programs/Shared/EpicGames.BuildGraph/BgTask.cs
deleted file mode 100644
