Fix a call to the base method: use the overloaded version instead.

diff --git a/Engine/Source/Programs/AutomationTool/Gauntlet/Editor/RunEditorTests.cs b/Engine/Source/Programs/AutomationTool/Gauntlet/Editor/RunEditorTests.cs
--- a/Engine/Source/Programs/AutomationTool/Gauntlet/Editor/RunEditorTests.cs
+++ b/Engine/Source/Programs/AutomationTool/Gauntlet/Editor/RunEditorTests.cs
@@ -48,7 +48,7 @@ namespace EditorTest
 		/// <param name="AppConfig"></param>
 		public override void ApplyToConfig(UnrealAppConfig AppConfig, UnrealSessionRole ConfigRole, IEnumerable<UnrealSessionRole> OtherRoles)
 		{
-			base.ApplyToConfig(AppConfig);
+			base.ApplyToConfig(AppConfig, ConfigRole, OtherRoles);
 
 			if (AppConfig.ProcessType.IsEditor())
 			{
