BuildGraph: Fix UDN documentation output.

diff --git a/Engine/Source/Programs/AutomationTool/AutomationUtils/BuildCommand.cs b/Engine/Source/Programs/AutomationTool/AutomationUtils/BuildCommand.cs
--- a/Engine/Source/Programs/AutomationTool/AutomationUtils/BuildCommand.cs
+++ b/Engine/Source/Programs/AutomationTool/AutomationUtils/BuildCommand.cs
@@ -160,7 +160,7 @@ namespace AutomationTool
 		/// Parses an argument as an enum.
 		/// </summary>
 		/// <param name="Param">Name of the parameter to read.</param>
-		/// <returns>Returns the value that was parsed.
+		/// <returns>Returns the value that was parsed.</returns>
 		public Nullable<T> ParseOptionalEnumParam<T>(string Param) where T : struct
 		{
 			string ValueString = ParseParamValue(Param);
@@ -183,7 +183,7 @@ namespace AutomationTool
 		/// Parses an argument as an enum. Throws an exception if the parameter is not specified.
 		/// </summary>
 		/// <param name="Param">Name of the parameter to read.</param>
-		/// <returns>Returns the value that was parsed.
+		/// <returns>Returns the value that was parsed.</returns>
 		public T ParseRequiredEnumParamEnum<T>(string Param) where T : struct
 		{
 			Nullable<T> Value = ParseOptionalEnumParam<T>(Param);
diff --git a/Engine/Source/Programs/AutomationTool/AutomationUtils/CommandUtils.cs b/Engine/Source/Programs/AutomationTool/AutomationUtils/CommandUtils.cs
--- a/Engine/Source/Programs/AutomationTool/AutomationUtils/CommandUtils.cs
+++ b/Engine/Source/Programs/AutomationTool/AutomationUtils/CommandUtils.cs
@@ -28,7 +28,7 @@ namespace AutomationTool
 {
 	/// <summary>
 	/// Wrapper around List with support for multi parameter constructor, i.e:
-	///   var Maps = new ParamList<string>("Map1", "Map2");
+	///   var Maps = new ParamList&lt;string&gt;("Map1", "Map2");
 	/// </summary>
 	/// <typeparam name="T"></typeparam>
 	public class ParamList<T> : List<T>
@@ -1825,7 +1825,6 @@ namespace AutomationTool
 		/// <summary>
 		/// Moves files in parallel
         /// </summary>
-		/// <param
 		/// <param name="SourceAndTargetPairs">Pairs of source and target files</param>
 		public static void ParallelMoveFiles(IEnumerable<KeyValuePair<FileReference, FileReference>> SourceAndTargetPairs)
 		{
@@ -1835,7 +1834,6 @@ namespace AutomationTool
 		/// <summary>
 		/// Moves files in parallel
 		/// </summary>
-		/// <param
 		/// <param name="SourceAndTargetPairs">Pairs of source and target files</param>
 		/// <param name="Overwrite">Whether or not to overwrite target files if they already exist</param>
 		public static void ParallelMoveFiles(IEnumerable<KeyValuePair<FileReference, FileReference>> SourceAndTargetPairs, bool Overwrite)
@@ -2082,7 +2080,7 @@ namespace AutomationTool
 		}
 
 		/// <summary>
-		/// Telemetry data for the current run. Add -WriteTelemetry=<Path> to the command line to export to disk.
+		/// Telemetry data for the current run. Add -WriteTelemetry=[Path] to the command line to export to disk.
 		/// </summary>
 		public static TelemetryData Telemetry = new TelemetryData();
 
@@ -2739,7 +2737,7 @@ namespace AutomationTool
         }
 
         /// <summary>
-        /// Flushes the time to <see cref="CmdEnv.CSVFile"/> if we are the build machine and that environment variable is specified.
+        /// Flushes the time to <see cref="CommandEnvironment.CSVFile"/> if we are the build machine and that environment variable is specified.
         /// Call Finish manually with an alternate name to use that one instead. Useful for dynamically generated names that you can't specify at construction.
         /// </summary>
         /// <param name="AlternateName">Used in place of the Name specified during construction.</param>
diff --git a/Engine/Source/Programs/AutomationTool/AutomationUtils/MCPPublic.cs b/Engine/Source/Programs/AutomationTool/AutomationUtils/MCPPublic.cs
--- a/Engine/Source/Programs/AutomationTool/AutomationUtils/MCPPublic.cs
+++ b/Engine/Source/Programs/AutomationTool/AutomationUtils/MCPPublic.cs
@@ -1338,7 +1338,7 @@ namespace EpicGames.MCP.Automation
 			public string ManifestA;
 			/// <summary>
 			/// The install tags to use for ManifestA.
-			/// </summary
+			/// </summary>
 			public HashSet<string> InstallTagsA;
 			/// <summary>
 			/// The file path to the update manifest.
@@ -1479,7 +1479,7 @@ namespace EpicGames.MCP.Automation
 			public string ArtifactId;
 			/// <summary>
 			/// The install tags to use for ManifestA.
-			/// </summary
+			/// </summary>
 			public HashSet<string> InstallTagsA;
 			/// <summary>
 			/// The build version to generate optimized deltas from 
@@ -1566,7 +1566,7 @@ namespace EpicGames.MCP.Automation
 			/// from files that are not expected to be available. Most of the time this should be the union of all tags in TagSetSplit, unless tagging is changed between the two manifests.
 			/// An empty string must be included to include untagged files.
 			/// Leaving this variable null will include all files.
-			/// </summary
+			/// </summary>
 			public HashSet<string> PrevManifestTags;
 			/// <summary>
 			/// Optionally provide a list of tagsets to split chunkdb files on. First all data from the tagset at index 0 will be saved, then any extra data needed
diff --git a/Engine/Source/Programs/AutomationTool/AutomationUtils/P4Utils.cs b/Engine/Source/Programs/AutomationTool/AutomationUtils/P4Utils.cs
--- a/Engine/Source/Programs/AutomationTool/AutomationUtils/P4Utils.cs
+++ b/Engine/Source/Programs/AutomationTool/AutomationUtils/P4Utils.cs
@@ -3639,7 +3639,7 @@ namespace AutomationTool
 		}
 
 		/// <summary>
-		/// Run 'p4 files <cmdline>'and return a list of the files in the changelist (files being deleted are excluded)
+		/// Run 'p4 files [cmdline]'and return a list of the files in the changelist (files being deleted are excluded)
 		/// </summary>
 		/// <param name="CommandLine"></param>
 		/// <returns>List of files in the specified directory.</returns>
@@ -3666,7 +3666,7 @@ namespace AutomationTool
 		}
 
 		/// <summary>
-		/// Run 'p4 opened <cmdline>'and return a list of the files in the changelist (files being deleted are excluded)
+		/// Run 'p4 opened [cmdline]'and return a list of the files in the changelist (files being deleted are excluded)
 		/// </summary>
 		/// <param name="CommandLine"></param>
 		/// <returns>List of files in the specified directory.</returns>
diff --git a/Engine/Source/Programs/AutomationTool/AutomationUtils/Platform.cs b/Engine/Source/Programs/AutomationTool/AutomationUtils/Platform.cs
--- a/Engine/Source/Programs/AutomationTool/AutomationUtils/Platform.cs
+++ b/Engine/Source/Programs/AutomationTool/AutomationUtils/Platform.cs
@@ -522,7 +522,7 @@ namespace AutomationTool
 		/// <summary>
 		/// return true if we need to change the case of filenames outside of pak files
 		/// </summary>
-		/// <param name="FileType" The staged file type to check (UFS vs SsytemNonUFS, etc)
+		/// <param name="FileType">The staged file type to check (UFS vs SsytemNonUFS, etc)</param>
 		/// <returns>true if files should be lower-cased during staging, for the given filetype</returns>
 		public virtual bool DeployLowerCaseFilenames(StagedFileType FileType)
 		{
@@ -532,7 +532,7 @@ namespace AutomationTool
 		/// <summary>
 		/// return true if we need to change the case of a particular file
 		/// </summary>
-		/// <param name="FileType" The staged file type to check (UFS vs SsytemNonUFS, etc)
+		/// <param name="FileType">The staged file type to check (UFS vs SsytemNonUFS, etc)</param>
 		/// <returns>true if files should be lower-cased during staging, for the given filetype</returns>
 		public virtual bool DeployLowerCaseFile(FileReference File, StagedFileType FileType)
 		{
@@ -637,7 +637,7 @@ namespace AutomationTool
 		/// <summary>
 		/// return true if we need to call Remap of a specific file type
 		/// </summary>
-		/// <param name="FileType" The staged file type to check (UFS vs SsytemNonUFS, etc)
+		/// <param name="FileType">The staged file type to check (UFS vs SsytemNonUFS, etc)</param>
 		/// <returns>true if files should be remaped, for the given filetype</returns>
 		public virtual bool RemapFileType(StagedFileType FileType)
 		{
@@ -655,7 +655,6 @@ namespace AutomationTool
 		/// <summary>
 		/// Tri-state - The intent is to override command line parameters for pak if needed per platform.
 		/// </summary>
-		///
 		public enum PakType { Always, Never, DontCare };
 
 		public virtual PakType RequiresPak(ProjectParams Params)
diff --git a/Engine/Source/Programs/AutomationTool/AutomationUtils/ProjectParams.cs b/Engine/Source/Programs/AutomationTool/AutomationUtils/ProjectParams.cs
--- a/Engine/Source/Programs/AutomationTool/AutomationUtils/ProjectParams.cs
+++ b/Engine/Source/Programs/AutomationTool/AutomationUtils/ProjectParams.cs
@@ -1762,13 +1762,16 @@ namespace AutomationTool
 		/// <summary>
         /// </summary>
         public bool AddPatchLevel;
+
         /// <summary>
         /// Are we staging the unmodified pak files from the base release
+		/// </summary>
         public bool StageBaseReleasePaks;
 
-        /// Name of dlc to cook and package (if this paramter is supplied cooks the dlc and packages it into the dlc directory)
-        /// </summary>
-        public FileReference DLCFile;
+		/// <summary>
+		/// Name of dlc to cook and package (if this paramter is supplied cooks the dlc and packages it into the dlc directory)
+		/// </summary>
+		public FileReference DLCFile;
 
         /// <summary>
         /// Enable cooking of engine content when cooking dlc 
@@ -1939,7 +1942,7 @@ namespace AutomationTool
 		[Help("bundlename", "string to use as the bundle name when deploying to mobile device")]
         public string BundleName;
 
-		//<summary>
+		/// <summary>
 		/// Stage: Specifies a list of extra targets that should be staged along with a client
 		/// </summary>
 		public ParamList<string> ExtraTargetsToStageWithClient = new ParamList<string>();
diff --git a/Engine/Source/Programs/AutomationTool/BuildGraph/BgScriptNode.cs b/Engine/Source/Programs/AutomationTool/BuildGraph/BgScriptNode.cs
new file mode 100644
--- /dev/null
+++ b/Engine/Source/Programs/AutomationTool/BuildGraph/BgScriptNode.cs
@@ -0,0 +1,296 @@
+// Copyright Epic Games, Inc. All Rights Reserved.
+
+using System;
+using System.Collections.Generic;
+using System.ComponentModel;
+using System.Diagnostics.CodeAnalysis;
+using System.IO;
+using System.Linq;
+using System.Reflection;
+using System.Runtime.CompilerServices;
+using System.Text;
+using System.Text.RegularExpressions;
+using System.Threading.Tasks;
+using System.Xml;
+using System.Xml.Schema;
+using AutomationTool.Tasks;
+using EpicGames.BuildGraph;
+using EpicGames.BuildGraph.Expressions;
+using EpicGames.Core;
+using Microsoft.Extensions.Logging;
+using OpenTracing;
+using OpenTracing.Util;
+
+#nullable enable
+
+namespace AutomationTool
+{
+#if false
+	/// <summary>
+	/// Base class for binding and executing nodes
+	/// </summary>
+	abstract class BgNodeExecutor
+	{
+		public abstract Task<bool> ExecuteAsync(JobContext Job, Dictionary<string, HashSet<FileReference>> TagNameToFileSet);
+	}
+
+	class BgExpressionNodeExecutor : BgNodeExecutor
+	{
+		readonly BgExpressionNode Node;
+		MethodInfo? method;
+
+		public BgExpressionNodeExecutor(BgExpressionNode node)
+		{
+			Node = node;
+		}
+
+	/// <summary>
+	/// Implementation of <see cref="BgNode"/> for graphs defined through XML syntax
+	/// </summary>
+	class BgScriptNode : BgNode
+	{
+		/// <summary>
+		/// List of tasks to execute
+		/// </summary>
+		public List<BgTask> Tasks { get; } = new List<BgTask>();
+
+		/// <summary>
+		/// List of bound task implementations
+		/// </summary>
+		List<BgTaskImpl> _boundTasks = new List<BgTaskImpl>();
+
+		/// <summary>
+		/// Constructor
+		/// </summary>
+		public BgScriptNodeExecutor(string name, IReadOnlyList<BgNodeOutput> inputs, IReadOnlyList<string> outputNames, IReadOnlyList<BgNode> inputDependencies, IReadOnlyList<BgNode> orderDependencies, IReadOnlyList<FileReference> requiredTokens)
+			: base(name, inputs, outputNames, inputDependencies, orderDependencies, requiredTokens)
+		{
+		}
+
+		public bool Bind(Dictionary<string, ScriptTaskBinding> nameToTask, Dictionary<string, BgNodeOutput> tagNameToNodeOutput, IBgScriptReaderContext context, ILogger logger)
+		{
+			bool bResult = true;
+			foreach (BgTask TaskInfo in Node.Tasks)
+			{
+				BgTaskImpl? boundTask = BindTask(TaskInfo, nameToTask, tagNameToNodeOutput, context, logger);
+				if (boundTask == null)
+				{
+					bResult = false;
+				}
+				else
+				{
+					_boundTasks.Add(boundTask);
+				}
+			}
+			return bResult;
+		}
+
+		BgTaskImpl? BindTask(BgTask TaskInfo, Dictionary<string, ScriptTaskBinding> NameToTask, IReadOnlyDictionary<string, BgNodeOutput> TagNameToNodeOutput, IBgScriptReaderContext Context, ILogger Logger)
+		{
+			// Get the reflection info for this element
+			ScriptTaskBinding? Task;
+			if (!NameToTask.TryGetValue(TaskInfo.Name, out Task))
+			{
+				Logger.LogScriptError(TaskInfo.Location, "Unknown task '{TaskName}'", TaskInfo.Name);
+				return null;
+			}
+
+			// Check all the required parameters are present
+			bool bHasRequiredAttributes = true;
+			foreach (ScriptTaskParameterBinding Parameter in Task.NameToParameter.Values)
+			{
+				if (!Parameter.Optional && !TaskInfo.Arguments.ContainsKey(Parameter.Name))
+				{
+					Logger.LogScriptError(TaskInfo.Location, "Missing required attribute - {AttrName}", Parameter.Name);
+					bHasRequiredAttributes = false;
+				}
+			}
+
+			// Read all the attributes into a parameters object for this task
+			object ParametersObject = Activator.CreateInstance(Task.ParametersClass)!;
+			foreach ((string Name, string Value) in TaskInfo.Arguments)
+			{
+				// Get the field that this attribute should be written to in the parameters object
+				ScriptTaskParameterBinding? Parameter;
+				if (!Task.NameToParameter.TryGetValue(Name, out Parameter))
+				{
+					Logger.LogScriptError(TaskInfo.Location, "Unknown attribute '{AttrName}'", Name);
+					continue;
+				}
+
+				// If it's a collection type, split it into separate values
+				if (Parameter.CollectionType == null)
+				{
+					// Parse it and assign it to the parameters object
+					object? FieldValue = ParseValue(Value, Parameter.ValueType, Context);
+					Parameter.FieldInfo.SetValue(ParametersObject, FieldValue);
+				}
+				else
+				{
+					// Get the collection, or create one if necessary
+					object? CollectionValue = Parameter.FieldInfo.GetValue(ParametersObject);
+					if (CollectionValue == null)
+					{
+						CollectionValue = Activator.CreateInstance(Parameter.FieldInfo.FieldType)!;
+						Parameter.FieldInfo.SetValue(ParametersObject, CollectionValue);
+					}
+
+					// Parse the values and add them to the collection
+					List<string> ValueStrings = BgTaskImpl.SplitDelimitedList(Value);
+					foreach (string ValueString in ValueStrings)
+					{
+						object ElementValue = ParseValue(ValueString, Parameter.ValueType, Context)!;
+						Parameter.CollectionType.InvokeMember("Add", BindingFlags.InvokeMethod | BindingFlags.Instance | BindingFlags.Public, null, CollectionValue, new object[] { ElementValue });
+					}
+				}
+			}
+
+			// Construct the task
+			if (!bHasRequiredAttributes)
+			{
+				return null;
+			}
+
+			// Add it to the list
+			BgTaskImpl NewTask = (BgTaskImpl)Activator.CreateInstance(Task.TaskClass, ParametersObject)!;
+
+			// Set up the source location for diagnostics
+			NewTask.SourceLocation = TaskInfo.Location;
+
+			// Make sure all the read tags are local or listed as a dependency
+			foreach (string ReadTagName in NewTask.FindConsumedTagNames())
+			{
+				BgNodeOutput? Output;
+				if (TagNameToNodeOutput.TryGetValue(ReadTagName, out Output))
+				{
+					if (Output != null && Output.ProducingNode != Node && !Node.Inputs.Contains(Output))
+					{
+						Logger.LogScriptError(TaskInfo.Location, "The tag '{TagName}' is not a dependency of node '{Node}'", ReadTagName, Node.Name);
+					}
+				}
+			}
+
+			// Make sure all the written tags are local or listed as an output
+			foreach (string ModifiedTagName in NewTask.FindProducedTagNames())
+			{
+				BgNodeOutput? Output;
+				if (TagNameToNodeOutput.TryGetValue(ModifiedTagName, out Output))
+				{
+					if (Output != null && !Node.Outputs.Contains(Output))
+					{
+						Logger.LogScriptError(TaskInfo.Location, "The tag '{TagName}' is created by '{Node}', and cannot be modified downstream", Output.TagName, Output.ProducingNode.Name);
+					}
+				}
+			}
+			return NewTask;
+		}
+
+		/// <summary>
+		/// Parse a value of the given type
+		/// </summary>
+		/// <param name="ValueText">The text to parse</param>
+		/// <param name="ValueType">Type of the value to parse</param>
+		/// <param name="Context">Context for the script reader</param>
+		/// <returns>Value that was parsed</returns>
+		static object? ParseValue(string ValueText, Type ValueType, IBgScriptReaderContext Context)
+		{
+			// Parse it and assign it to the parameters object
+			if (ValueType.IsEnum)
+			{
+				return Enum.Parse(ValueType, ValueText);
+			}
+			else if (ValueType == typeof(Boolean))
+			{
+				return BgCondition.EvaluateAsync(ValueText, Context).Result;
+			}
+			else if (ValueType == typeof(FileReference))
+			{
+				return BgTaskImpl.ResolveFile(ValueText);
+			}
+			else if (ValueType == typeof(DirectoryReference))
+			{
+				return BgTaskImpl.ResolveDirectory(ValueText);
+			}
+
+			TypeConverter Converter = TypeDescriptor.GetConverter(ValueType);
+			if (Converter.CanConvertFrom(typeof(string)))
+			{
+				return Converter.ConvertFromString(ValueText);
+			}
+			else
+			{
+				return Convert.ChangeType(ValueText, ValueType);
+			}
+		}
+
+		/// <summary>
+		/// Build all the tasks for this node
+		/// </summary>
+		/// <param name="Job">Information about the current job</param>
+		/// <param name="TagNameToFileSet">Mapping from tag names to the set of files they include. Should be set to contain the node inputs on entry.</param>
+		/// <returns>Whether the task succeeded or not. Exiting with an exception will be caught and treated as a failure.</returns>
+		public override async Task<bool> ExecuteAsync(JobContext Job, Dictionary<string, HashSet<FileReference>> TagNameToFileSet)
+		{
+			// Run each of the tasks in order
+			HashSet<FileReference> BuildProducts = TagNameToFileSet[Node.DefaultOutput.TagName];
+			for (int Idx = 0; Idx < _boundTasks.Count; Idx++)
+			{
+				using (IScope Scope = GlobalTracer.Instance.BuildSpan("Task").WithTag("resource", _boundTasks[Idx].GetTraceName()).StartActive())
+				{
+					ITaskExecutor Executor = _boundTasks[Idx].GetExecutor();
+					if (Executor == null)
+					{
+						// Execute this task directly
+						try
+						{
+							_boundTasks[Idx].GetTraceMetadata(Scope.Span, "");
+							await _boundTasks[Idx].ExecuteAsync(Job, BuildProducts, TagNameToFileSet);
+						}
+						catch (Exception Ex)
+						{
+							ExceptionUtils.AddContext(Ex, "while executing task {0}", _boundTasks[Idx].GetTraceString());
+							if (_boundTasks[Idx].SourceLocation != null)
+							{
+								ExceptionUtils.AddContext(Ex, "at {0}({1})", _boundTasks[Idx].SourceLocation.File, _boundTasks[Idx].SourceLocation.LineNumber);
+							}
+							throw;
+						}
+					}
+					else
+					{
+						_boundTasks[Idx].GetTraceMetadata(Scope.Span, "1.");
+
+						// The task has a custom executor, which may be able to execute several tasks simultaneously. Try to add the following tasks.
+						int FirstIdx = Idx;
+						while (Idx + 1 < Node.Tasks.Count && Executor.Add(_boundTasks[Idx + 1]))
+						{
+							Idx++;
+							_boundTasks[Idx].GetTraceMetadata(Scope.Span, string.Format("{0}.", 1 + Idx - FirstIdx));
+						}
+						try
+						{
+							await Executor.ExecuteAsync(Job, BuildProducts, TagNameToFileSet);
+						}
+						catch (Exception Ex)
+						{
+							for (int TaskIdx = FirstIdx; TaskIdx <= Idx; TaskIdx++)
+							{
+								ExceptionUtils.AddContext(Ex, "while executing {0}", _boundTasks[TaskIdx].GetTraceString());
+							}
+							if (_boundTasks[FirstIdx].SourceLocation != null)
+							{
+								ExceptionUtils.AddContext(Ex, "at {0}({1})", _boundTasks[FirstIdx].SourceLocation.File, _boundTasks[FirstIdx].SourceLocation.LineNumber);
+							}
+							throw;
+						}
+					}
+				}
+			}
+
+			// Remove anything that doesn't exist, since these files weren't explicitly tagged
+			BuildProducts.RemoveWhere(x => !FileReference.Exists(x));
+			return true;
+		}
+	}
+#endif
+}
diff --git a/Engine/Source/Programs/AutomationTool/BuildGraph/BgScriptWriter.cs b/Engine/Source/Programs/AutomationTool/BuildGraph/BgScriptWriter.cs
--- a/Engine/Source/Programs/AutomationTool/BuildGraph/BgScriptWriter.cs
+++ b/Engine/Source/Programs/AutomationTool/BuildGraph/BgScriptWriter.cs
@@ -24,6 +24,7 @@ namespace AutomationTool
 		/// <summary>
 		/// Writes a preprocessed build graph to a script file
 		/// </summary>
+		/// <param name="graph">Graph to output</param>
 		/// <param name="file">The file to load</param>
 		/// <param name="schemaFile">Schema file for validation</param>
 		public static void Write(this BgGraph graph, FileReference file, FileReference schemaFile)
@@ -120,6 +121,7 @@ namespace AutomationTool
 		/// <summary>
 		/// Writes this agent group out to a file, filtering nodes by a controlling trigger
 		/// </summary>
+		/// <param name="agent">Agent to output</param>
 		/// <param name="writer">The XML writer to output to</param>
 		public static void Write(this BgAgent agent, XmlWriter writer)
 		{
@@ -136,6 +138,7 @@ namespace AutomationTool
 		/// <summary>
 		/// Write this node to an XML writer
 		/// </summary>
+		/// <param name="node">Node to output</param>
 		/// <param name="writer">The writer to output the node to</param>
 		public static void Write(this BgNode node, XmlWriter writer)
 		{
diff --git a/Engine/Source/Programs/AutomationTool/BuildGraph/BuildGraph.Automation.csproj b/Engine/Source/Programs/AutomationTool/BuildGraph/BuildGraph.Automation.csproj
--- a/Engine/Source/Programs/AutomationTool/BuildGraph/BuildGraph.Automation.csproj
+++ b/Engine/Source/Programs/AutomationTool/BuildGraph/BuildGraph.Automation.csproj
@@ -15,6 +15,7 @@
     <OutputPath>..\..\..\..\Binaries\DotNET\AutomationTool\AutomationScripts\BuildGraph</OutputPath>
     <TreatWarningsAsErrors>true</TreatWarningsAsErrors>
     <DebugType>pdbonly</DebugType>
+    <GenerateDocumentationFile>True</GenerateDocumentationFile>
   </PropertyGroup>
   <PropertyGroup Condition=" '$(Configuration)|$(Platform)' == 'Development|AnyCPU' ">
     <DefineConstants>$(DefineConstants);TRACE</DefineConstants>
diff --git a/Engine/Source/Programs/AutomationTool/BuildGraph/BuildGraph.cs b/Engine/Source/Programs/AutomationTool/BuildGraph/BuildGraph.cs
--- a/Engine/Source/Programs/AutomationTool/BuildGraph/BuildGraph.cs
+++ b/Engine/Source/Programs/AutomationTool/BuildGraph/BuildGraph.cs
@@ -1319,18 +1319,28 @@ namespace AutomationTool
 							XmlElement ParameterElement;
 							if (MemberNameToElement.TryGetValue("F:" + Parameter.FieldInfo.DeclaringType.FullName + "." + Parameter.Name, out ParameterElement))
 							{
-								string TypeName = Parameter.FieldInfo.FieldType.Name;
+								Type FieldType = Parameter.FieldInfo.FieldType;
+								if (FieldType.IsGenericType && FieldType.GetGenericTypeDefinition() == typeof(Nullable<>))
+								{
+									FieldType = FieldType.GetGenericArguments()[0];
+								}
+
+								string TypeName;
 								if (Parameter.ValidationType != TaskParameterValidationType.Default)
 								{
-									StringBuilder NewTypeName = new StringBuilder(Parameter.ValidationType.ToString());
-									for (int Idx = 1; Idx < NewTypeName.Length; Idx++)
-									{
-										if (Char.IsLower(NewTypeName[Idx - 1]) && Char.IsUpper(NewTypeName[Idx]))
-										{
-											NewTypeName.Insert(Idx, ' ');
-										}
-									}
-									TypeName = NewTypeName.ToString();
+									TypeName = Parameter.ValidationType.ToString();
+								}
+								else if (FieldType == typeof(int))
+								{
+									TypeName = "Integer";
+								}
+								else if (FieldType == typeof(HashSet<DirectoryReference>))
+								{
+									TypeName = "DirectoryList";
+								}
+								else
+								{
+									TypeName = FieldType.Name;
 								}
 
 								string[] Columns = new string[4];
diff --git a/Engine/Source/Programs/AutomationTool/BuildGraph/Tasks/AwsAssumeRoleTask.cs b/Engine/Source/Programs/AutomationTool/BuildGraph/Tasks/AwsAssumeRoleTask.cs
--- a/Engine/Source/Programs/AutomationTool/BuildGraph/Tasks/AwsAssumeRoleTask.cs
+++ b/Engine/Source/Programs/AutomationTool/BuildGraph/Tasks/AwsAssumeRoleTask.cs
@@ -56,7 +56,7 @@ namespace AutomationTool.Tasks
 	}
 
 	/// <summary>
-	/// Spawns AWS CLI and waits for it to complete.
+	/// Assumes an AWS role.
 	/// </summary>
 	[TaskElement("Aws-AssumeRole", typeof(AwsAssumeRoleTaskParameters))]
 	public class AwsAssumeRoleTask : SpawnTaskBase
diff --git a/Engine/Source/Programs/AutomationTool/BuildGraph/Tasks/CallMethodTask.cs b/Engine/Source/Programs/AutomationTool/BuildGraph/Tasks/CallMethodTask.cs
--- a/Engine/Source/Programs/AutomationTool/BuildGraph/Tasks/CallMethodTask.cs
+++ b/Engine/Source/Programs/AutomationTool/BuildGraph/Tasks/CallMethodTask.cs
@@ -138,12 +138,11 @@ namespace AutomationTool.Tasks
 	}
 
 	/// <summary>
-	/// Spawns the editor to run a commandlet.
+	/// Invokes a C# method.
 	/// </summary>
 	[TaskElement("CallMethod", typeof(CallMethodTaskParameters))]
 	public class CallMethodTask : BgTaskImpl
 	{
-				
 		/// <summary>
 		/// Parameters for this task
 		/// </summary>
diff --git a/Engine/Source/Programs/AutomationTool/BuildGraph/Tasks/CommandletTask.cs b/Engine/Source/Programs/AutomationTool/BuildGraph/Tasks/CommandletTask.cs
--- a/Engine/Source/Programs/AutomationTool/BuildGraph/Tasks/CommandletTask.cs
+++ b/Engine/Source/Programs/AutomationTool/BuildGraph/Tasks/CommandletTask.cs
@@ -144,6 +144,9 @@
 		}
 	}
 
+	/// <summary>
+	/// Task wrapper methods
+	/// </summary>
 	public static partial class StandardTasks
 	{
 		/// <summary>
diff --git a/Engine/Source/Programs/AutomationTool/BuildGraph/Tasks/CopyTask.cs b/Engine/Source/Programs/AutomationTool/BuildGraph/Tasks/CopyTask.cs
--- a/Engine/Source/Programs/AutomationTool/BuildGraph/Tasks/CopyTask.cs
+++ b/Engine/Source/Programs/AutomationTool/BuildGraph/Tasks/CopyTask.cs
@@ -131,6 +131,12 @@ namespace AutomationTool.Tasks
 			}
 		}
 
+		/// <summary>
+		/// 
+		/// </summary>
+		/// <param name="TargetFileToSourceFile"></param>
+		/// <param name="Overwrite"></param>
+		/// <returns></returns>
 		public static Task ExecuteAsync(Dictionary<FileReference, FileReference> TargetFileToSourceFile, bool Overwrite)
 		{
 			//  If we're not overwriting, remove any files where the destination file already exists.
@@ -233,6 +239,9 @@ namespace AutomationTool.Tasks
 		}
 	}
 
+	/// <summary>
+	/// Extension methods
+	/// </summary>
 	public static partial class BgStateExtensions
 	{
 		/// <summary>
diff --git a/Engine/Source/Programs/AutomationTool/BuildGraph/Tasks/CsCompileTask.cs b/Engine/Source/Programs/AutomationTool/BuildGraph/Tasks/CsCompileTask.cs
--- a/Engine/Source/Programs/AutomationTool/BuildGraph/Tasks/CsCompileTask.cs
+++ b/Engine/Source/Programs/AutomationTool/BuildGraph/Tasks/CsCompileTask.cs
@@ -329,32 +329,63 @@ namespace AutomationTool.Tasks
 		}
 	}
 
+	/// <summary>
+	/// Output from compiling a csproj file
+	/// </summary>
 	public class CsCompileOutput
 	{
+		/// <summary>
+		/// Empty instance of CsCompileOutput
+		/// </summary>
 		public static CsCompileOutput Empty { get; } = new CsCompileOutput(FileSet.Empty, FileSet.Empty);
 
+		/// <summary>
+		/// Output binaries
+		/// </summary>
 		public FileSet Binaries { get; }
+
+		/// <summary>
+		/// Referenced output
+		/// </summary>
 		public FileSet References { get; }
 
+		/// <summary>
+		/// Constructor
+		/// </summary>
 		public CsCompileOutput(FileSet Binaries, FileSet References)
 		{
 			this.Binaries = Binaries;
 			this.References = References;
 		}
 
+		/// <summary>
+		/// Merge all outputs from this project
+		/// </summary>
+		/// <returns></returns>
 		public FileSet Merge()
 		{
 			return Binaries + References;
 		}
 
+		/// <summary>
+		/// Merges two outputs together
+		/// </summary>
 		public static CsCompileOutput operator +(CsCompileOutput Lhs, CsCompileOutput Rhs)
 		{
 			return new CsCompileOutput(Lhs.Binaries + Rhs.Binaries, Lhs.References + Rhs.References);
 		}
 	}
 
+	/// <summary>
+	/// Extension methods for csproj compilation
+	/// </summary>
 	public static class CsCompileOutputExtensions
 	{
+		/// <summary>
+		/// 
+		/// </summary>
+		/// <param name="Task"></param>
+		/// <returns></returns>
 		public static async Task<FileSet> MergeAsync(this Task<CsCompileOutput> Task)
 		{
 			return (await Task).Merge();
diff --git a/Engine/Source/Programs/AutomationTool/BuildGraph/Tasks/GatherBuildProductsFromFileTask.cs b/Engine/Source/Programs/AutomationTool/BuildGraph/Tasks/GatherBuildProductsFromFileTask.cs
--- a/Engine/Source/Programs/AutomationTool/BuildGraph/Tasks/GatherBuildProductsFromFileTask.cs
+++ b/Engine/Source/Programs/AutomationTool/BuildGraph/Tasks/GatherBuildProductsFromFileTask.cs
@@ -12,8 +12,14 @@ using UnrealBuildTool;
 
 namespace AutomationTool.Tasks
 {	
+	/// <summary>
+	/// Parameters for <see cref="GatherBuildProductsFromFileTask"/>
+	/// </summary>
 	public class GatherBuildProductsFromFileTaskParameters
 	{
+		/// <summary>
+		/// 
+		/// </summary>
 		[TaskParameter]
 		public string BuildProductsFile;
 	}
diff --git a/Engine/Source/Programs/AutomationTool/BuildGraph/Tasks/GitCloneTask.cs b/Engine/Source/Programs/AutomationTool/BuildGraph/Tasks/GitCloneTask.cs
--- a/Engine/Source/Programs/AutomationTool/BuildGraph/Tasks/GitCloneTask.cs
+++ b/Engine/Source/Programs/AutomationTool/BuildGraph/Tasks/GitCloneTask.cs
@@ -41,7 +41,7 @@ namespace AutomationTool.Tasks
 	}
 
 	/// <summary>
-	/// Spawns Git and waits for it to complete.
+	/// Clones a Git repository into a local path.
 	/// </summary>
 	[TaskElement("Git-Clone", typeof(GitCloneTaskParameters))]
 	public class GitCloneTask : BgTaskImpl
diff --git a/Engine/Source/Programs/AutomationTool/BuildGraph/Tasks/SetVersionTask.cs b/Engine/Source/Programs/AutomationTool/BuildGraph/Tasks/SetVersionTask.cs
--- a/Engine/Source/Programs/AutomationTool/BuildGraph/Tasks/SetVersionTask.cs
+++ b/Engine/Source/Programs/AutomationTool/BuildGraph/Tasks/SetVersionTask.cs
@@ -136,8 +136,16 @@ namespace AutomationTool.Tasks
 		}
 	}
 
+	/// <summary>
+	/// Task wrapper methods
+	/// </summary>
 	public static partial class StandardTasks
 	{
+		/// <summary>
+		/// Execute a task instance
+		/// </summary>
+		/// <param name="Task"></param>
+		/// <returns></returns>
 		public static async Task<FileSet> ExecuteAsync(BgTaskImpl Task)
 		{
 			HashSet<FileReference> BuildProducts = new HashSet<FileReference>();
@@ -145,6 +153,9 @@ namespace AutomationTool.Tasks
 			return FileSet.FromFiles(Unreal.RootDirectory, BuildProducts);
 		}
 
+		/// <summary>
+		/// Updates the current engine version
+		/// </summary>
 		public static async Task<FileSet> SetVersionAsync(int Change, string Branch, int? CompatibleChange = null, string Build = null, bool? Licensee = null, bool? Promoted = null, bool? SkipWrite = null)
 		{
 			SetVersionTaskParameters Parameters = new SetVersionTaskParameters();
diff --git a/Engine/Source/Programs/AutomationTool/BuildGraph/Tasks/SpawnTask.cs b/Engine/Source/Programs/AutomationTool/BuildGraph/Tasks/SpawnTask.cs
--- a/Engine/Source/Programs/AutomationTool/BuildGraph/Tasks/SpawnTask.cs
+++ b/Engine/Source/Programs/AutomationTool/BuildGraph/Tasks/SpawnTask.cs
@@ -207,7 +207,6 @@ namespace AutomationTool.Tasks
 		/// <summary>
 		/// Execute an external program
 		/// </summary>
-		/// <param name="State">Execution state</param>
 		/// <param name="Exe">Executable to spawn.</param>
 		/// <param name="Arguments">Arguments for the newly created process.</param>
 		/// <param name="WorkingDir">Working directory for spawning the new task.</param>
diff --git a/Engine/Source/Programs/AutomationTool/BuildGraph/Tasks/TagReceiptTask.cs b/Engine/Source/Programs/AutomationTool/BuildGraph/Tasks/TagReceiptTask.cs
--- a/Engine/Source/Programs/AutomationTool/BuildGraph/Tasks/TagReceiptTask.cs
+++ b/Engine/Source/Programs/AutomationTool/BuildGraph/Tasks/TagReceiptTask.cs
@@ -228,8 +228,14 @@ namespace AutomationTool.Tasks
 		}
 	}
 
+	/// <summary>
+	/// Extension methods
+	/// </summary>
 	public static class TaskExtensions
 	{
+		/// <summary>
+		/// Task that tags build products and/or runtime dependencies by reading from *.target files.
+		/// </summary>
 		public static async Task<FileSet> TagReceiptsAsync(this FileSet Files, DirectoryReference EngineDir = null, bool BuildProducts = false, BuildProductType? BuildProductType = null, bool RuntimeDependencies = false, StagedFileType? StagedFileType = null)
 		{
 			HashSet<FileReference> Result = await TagReceiptTask.ExecuteAsync(EngineDir ?? Unreal.EngineDirectory, Files, BuildProducts, BuildProductType, RuntimeDependencies, StagedFileType);
