Add build steps for static analysis of UBT and UAT, and to test project file generation.

diff --git a/Engine/Source/Programs/AutomationTool/AutomationTool.csproj b/Engine/Source/Programs/AutomationTool/AutomationTool.csproj
--- a/Engine/Source/Programs/AutomationTool/AutomationTool.csproj
+++ b/Engine/Source/Programs/AutomationTool/AutomationTool.csproj
@@ -27,7 +27,7 @@
     <AppendTargetFrameworkToOutputPath>false</AppendTargetFrameworkToOutputPath>
     <GenerateAssemblyInfo>false</GenerateAssemblyInfo>
     <GenerateTargetFrameworkAttribute>false</GenerateTargetFrameworkAttribute>
-    <Configurations>Debug;Release;Development</Configurations>
+    <Configurations>Debug;Release;Development;Analyze</Configurations>
     <ApplicationManifest>app.manifest</ApplicationManifest>
     <RootNamespace>AutomationTool</RootNamespace>
     <AssemblyName>AutomationTool</AssemblyName>
@@ -53,6 +53,12 @@
     <DebugSymbols>true</DebugSymbols>
     <DebugType>full</DebugType>
   </PropertyGroup>
+  <PropertyGroup Condition=" '$(Configuration)|$(Platform)' == 'Analyze|AnyCPU' ">
+    <RunAnalyzersDuringBuild>True</RunAnalyzersDuringBuild>
+    <TreatWarningsAsErrors>false</TreatWarningsAsErrors>
+    <GenerateDocumentationFile>True</GenerateDocumentationFile>
+    <DocumentationFile></DocumentationFile>
+  </PropertyGroup>
   
   <ItemGroup>
     <Reference Include="Ionic.Zip.Reduced">
diff --git a/Engine/Source/Programs/AutomationTool/BuildGraph/BuildGraph.Automation.csproj b/Engine/Source/Programs/AutomationTool/BuildGraph/BuildGraph.Automation.csproj
--- a/Engine/Source/Programs/AutomationTool/BuildGraph/BuildGraph.Automation.csproj
+++ b/Engine/Source/Programs/AutomationTool/BuildGraph/BuildGraph.Automation.csproj
@@ -8,7 +8,7 @@
     <AppendTargetFrameworkToOutputPath>false</AppendTargetFrameworkToOutputPath>
     <GenerateAssemblyInfo>false</GenerateAssemblyInfo>
     <GenerateTargetFrameworkAttribute>false</GenerateTargetFrameworkAttribute>
-    <Configurations>Debug;Release;Development</Configurations>
+    <Configurations>Debug;Release;Development;Analyze</Configurations>
     <RootNamespace>AutomationTool</RootNamespace>
     <AssemblyName>BuildGraph.Automation</AssemblyName>
     <WarningsNotAsErrors>612,618</WarningsNotAsErrors>
@@ -30,6 +30,12 @@
     <DebugSymbols>true</DebugSymbols>
     <DebugType>full</DebugType>
   </PropertyGroup>
+  <PropertyGroup Condition="'$(Configuration)|$(Platform)' == 'Analyze|AnyCPU'">
+    <DocumentationFile></DocumentationFile>
+    <TreatWarningsAsErrors>false</TreatWarningsAsErrors>
+    <RunAnalyzersDuringBuild>True</RunAnalyzersDuringBuild>
+    <GenerateDocumentationFile>True</GenerateDocumentationFile>
+  </PropertyGroup>
   <ItemGroup>
     <Reference Include="fastJSON">
       <HintPath>..\..\..\..\Binaries\ThirdParty\fastJSON\netstandard2.0\fastJSON.dll</HintPath>
diff --git a/Engine/Source/Programs/AutomationTool/Program.cs b/Engine/Source/Programs/AutomationTool/Program.cs
--- a/Engine/Source/Programs/AutomationTool/Program.cs
+++ b/Engine/Source/Programs/AutomationTool/Program.cs
@@ -17,14 +17,17 @@ using Microsoft.Extensions.Logging;
 
 namespace AutomationToolDriver
 {
+	/// <summary>
+	/// Main entry point
+	/// </summary>
 	public partial class Program
 	{
 		/// <summary>
 		/// Parses command line parameter.
 		/// </summary>
-		/// <param name="ParamIndex">Parameter index</param>
-		/// <param name="CommandLine">Command line</param>
+		/// <param name="CurrentParam">Parameter</param>
 		/// <param name="CurrentCommand">Recently parsed command</param>
+		/// <param name="Logger">Logger for output</param>
 		/// <returns>True if the parameter has been successfully parsed.</returns>
 		private static void ParseParam(string CurrentParam, CommandInfo CurrentCommand, ILogger Logger)
 		{
@@ -230,6 +233,7 @@ namespace AutomationToolDriver
 		/// Parse the command line and create a list of commands to execute.
 		/// </summary>
 		/// <param name="Arguments">Command line</param>
+		/// <param name="Logger">Logger for output</param>
 		public static void ParseCommandLine(string[] Arguments, ILogger Logger)
 		{
 			AutomationToolCommandLine = new ParsedCommandLine(
@@ -311,7 +315,11 @@ namespace AutomationToolDriver
 		static ParsedCommandLine AutomationToolCommandLine;
 		static StartupTraceListener StartupListener = new StartupTraceListener();
 
-		// Do not add [STAThread] here. It will cause deadlocks in platform automation code.
+
+		/// <summary>
+		/// Main entry point
+		/// </summary>
+		/// <remarks>Do not add [STAThread] here. It will cause deadlocks in platform automation code.</remarks>
 		public static async Task<int> Main(string[] Arguments)
 		{
 			ILogger Logger = Log.Logger;
diff --git a/Engine/Source/Programs/Shared/EpicGames.Build/EpicGames.Build.csproj b/Engine/Source/Programs/Shared/EpicGames.Build/EpicGames.Build.csproj
--- a/Engine/Source/Programs/Shared/EpicGames.Build/EpicGames.Build.csproj
+++ b/Engine/Source/Programs/Shared/EpicGames.Build/EpicGames.Build.csproj
@@ -8,7 +8,7 @@
     <RootNamespace>UnrealBuildBase</RootNamespace>
     <GenerateAssemblyInfo>false</GenerateAssemblyInfo>
     <GenerateTargetFrameworkAttribute>false</GenerateTargetFrameworkAttribute>
-    <Configurations>Debug;Release;Development</Configurations>
+    <Configurations>Debug;Release;Development;Analyze</Configurations>
     <Nullable>enable</Nullable>
     <DebugType>pdbonly</DebugType>
     <ProduceReferenceAssemblyInOutDir>true</ProduceReferenceAssemblyInOutDir>
@@ -26,7 +26,12 @@
     <DocumentationFile></DocumentationFile>
     <Optimize>true</Optimize>
   </PropertyGroup>
-  
+  <PropertyGroup Condition=" '$(Configuration)|$(Platform)' == 'Analyze|AnyCPU' ">
+    <DocumentationFile></DocumentationFile>
+    <RunAnalyzersDuringBuild>True</RunAnalyzersDuringBuild>
+    <GenerateDocumentationFile>True</GenerateDocumentationFile>
+    <NoWarn>1701;1702;1591</NoWarn>
+  </PropertyGroup>
   <ItemGroup>
     <Compile Include="..\MetaData.cs">
       <Link>Properties\MetaData.cs</Link>
diff --git a/Engine/Source/Programs/Shared/EpicGames.BuildGraph/EpicGames.BuildGraph.csproj b/Engine/Source/Programs/Shared/EpicGames.BuildGraph/EpicGames.BuildGraph.csproj
--- a/Engine/Source/Programs/Shared/EpicGames.BuildGraph/EpicGames.BuildGraph.csproj
+++ b/Engine/Source/Programs/Shared/EpicGames.BuildGraph/EpicGames.BuildGraph.csproj
@@ -6,12 +6,12 @@
     <RootNamespace>EpicGames.BuildGraph</RootNamespace>
     <GenerateAssemblyInfo>false</GenerateAssemblyInfo>
     <GenerateTargetFrameworkAttribute>false</GenerateTargetFrameworkAttribute>
-    <Configurations>Debug;Release;Development</Configurations>
+    <Configurations>Debug;Release;Development;Analyze</Configurations>
     <Nullable>enable</Nullable>
     <DebugType>pdbonly</DebugType>
     <GenerateDocumentationFile>True</GenerateDocumentationFile>
   </PropertyGroup>
-  <PropertyGroup Condition=" '$(Configuration)|$(Platform)' == 'Debug|AnyCPU' ">
+  <PropertyGroup Condition="'$(Configuration)|$(Platform)' == 'Debug|AnyCPU'">
     <DocumentationFile></DocumentationFile>
   </PropertyGroup>
   <PropertyGroup Condition="'$(Configuration)|$(Platform)' == 'Development|AnyCPU'">
@@ -22,7 +22,13 @@
     <DocumentationFile></DocumentationFile>
     <Optimize>true</Optimize>
   </PropertyGroup>
-  
+  <PropertyGroup Condition=" '$(Configuration)|$(Platform)' == 'Analyze|AnyCPU' ">
+    <RunAnalyzersDuringBuild>True</RunAnalyzersDuringBuild>
+    <TreatWarningsAsErrors>false</TreatWarningsAsErrors>
+    <GenerateDocumentationFile>True</GenerateDocumentationFile>
+    <DocumentationFile></DocumentationFile>
+  </PropertyGroup>
+
   <ItemGroup>
     <Compile Include="..\MetaData.cs">
       <Link>Properties\MetaData.cs</Link>
diff --git a/Engine/Source/Programs/Shared/EpicGames.Core/EpicGames.Core.csproj b/Engine/Source/Programs/Shared/EpicGames.Core/EpicGames.Core.csproj
--- a/Engine/Source/Programs/Shared/EpicGames.Core/EpicGames.Core.csproj
+++ b/Engine/Source/Programs/Shared/EpicGames.Core/EpicGames.Core.csproj
@@ -7,7 +7,7 @@
     <RootNamespace>EpicGames.Core</RootNamespace>
     <GenerateAssemblyInfo>false</GenerateAssemblyInfo>
     <GenerateTargetFrameworkAttribute>false</GenerateTargetFrameworkAttribute>
-    <Configurations>Debug;Release;Development</Configurations>
+    <Configurations>Debug;Release;Development;Analyze</Configurations>
     <Nullable>enable</Nullable>
     <AllowUnsafeBlocks>true</AllowUnsafeBlocks>
     <DebugType>pdbonly</DebugType>
@@ -24,6 +24,12 @@
     <DocumentationFile></DocumentationFile>
     <Optimize>true</Optimize>
   </PropertyGroup>
+  <PropertyGroup Condition=" '$(Configuration)|$(Platform)' == 'Analyze|AnyCPU' ">
+    <RunAnalyzersDuringBuild>True</RunAnalyzersDuringBuild>
+    <TreatWarningsAsErrors>false</TreatWarningsAsErrors>
+    <GenerateDocumentationFile>True</GenerateDocumentationFile>
+    <DocumentationFile></DocumentationFile>
+  </PropertyGroup>
   <ItemGroup>
     <Compile Include="..\..\Shared\MetaData.cs">
       <Link>Properties\MetaData.cs</Link>
diff --git a/Engine/Source/Programs/Shared/EpicGames.Horde/EpicGames.Horde.csproj b/Engine/Source/Programs/Shared/EpicGames.Horde/EpicGames.Horde.csproj
--- a/Engine/Source/Programs/Shared/EpicGames.Horde/EpicGames.Horde.csproj
+++ b/Engine/Source/Programs/Shared/EpicGames.Horde/EpicGames.Horde.csproj
@@ -7,11 +7,12 @@
     <RootNamespace>EpicGames.Horde</RootNamespace>
     <GenerateAssemblyInfo>false</GenerateAssemblyInfo>
     <GenerateTargetFrameworkAttribute>false</GenerateTargetFrameworkAttribute>
-    <Configurations>Debug;Release;Development</Configurations>
+    <Configurations>Debug;Release;Development;Analyze</Configurations>
     <Nullable>enable</Nullable>
     <DebugType>pdbonly</DebugType>
     <GenerateDocumentationFile>True</GenerateDocumentationFile>
     <AllowUnsafeBlocks>True</AllowUnsafeBlocks>
+    <RunAnalyzersDuringBuild>false</RunAnalyzersDuringBuild>
   </PropertyGroup>
   <PropertyGroup Condition=" '$(Configuration)|$(Platform)' == 'Debug|AnyCPU' ">
     <DocumentationFile></DocumentationFile>
@@ -24,6 +25,11 @@
     <DocumentationFile></DocumentationFile>
     <Optimize>true</Optimize>
   </PropertyGroup>
+  <PropertyGroup Condition=" '$(Configuration)|$(Platform)' == 'Analyze|AnyCPU' ">
+    <DocumentationFile></DocumentationFile>
+    <GenerateDocumentationFile>True</GenerateDocumentationFile>
+    <RunAnalyzersDuringBuild>True</RunAnalyzersDuringBuild>
+  </PropertyGroup>
   <ItemGroup>
     <Compile Include="..\..\Shared\MetaData.cs">
       <Link>Properties\MetaData.cs</Link>
diff --git a/Engine/Source/Programs/Shared/EpicGames.Serialization/EpicGames.Serialization.csproj b/Engine/Source/Programs/Shared/EpicGames.Serialization/EpicGames.Serialization.csproj
--- a/Engine/Source/Programs/Shared/EpicGames.Serialization/EpicGames.Serialization.csproj
+++ b/Engine/Source/Programs/Shared/EpicGames.Serialization/EpicGames.Serialization.csproj
@@ -10,6 +10,7 @@
     <Configurations>Debug;Release;Development</Configurations>
     <Nullable>enable</Nullable>
     <DebugType>pdbonly</DebugType>
+    <RunAnalyzersDuringBuild>false</RunAnalyzersDuringBuild>
   </PropertyGroup>
   <PropertyGroup Condition=" '$(Configuration)|$(Platform)' == 'Debug|AnyCPU' ">
     <DocumentationFile></DocumentationFile>
@@ -22,6 +23,11 @@
     <DocumentationFile></DocumentationFile>
     <Optimize>true</Optimize>
   </PropertyGroup>
+  <PropertyGroup Condition=" '$(Configuration)|$(Platform)' == 'Analyze|AnyCPU' ">
+    <DocumentationFile></DocumentationFile>
+    <RunAnalyzersDuringBuild>True</RunAnalyzersDuringBuild>
+    <GenerateDocumentationFile>True</GenerateDocumentationFile>
+  </PropertyGroup>
   <ItemGroup>
     <Compile Include="..\..\Shared\MetaData.cs">
       <Link>Properties\MetaData.cs</Link>
diff --git a/Engine/Source/Programs/UnrealBuildTool/UnrealBuildTool.csproj b/Engine/Source/Programs/UnrealBuildTool/UnrealBuildTool.csproj
--- a/Engine/Source/Programs/UnrealBuildTool/UnrealBuildTool.csproj
+++ b/Engine/Source/Programs/UnrealBuildTool/UnrealBuildTool.csproj
@@ -8,7 +8,7 @@
     <AppendTargetFrameworkToOutputPath>false</AppendTargetFrameworkToOutputPath>
     <GenerateAssemblyInfo>false</GenerateAssemblyInfo>
     <GenerateTargetFrameworkAttribute>false</GenerateTargetFrameworkAttribute>
-    <Configurations>Debug;Release;Development</Configurations>
+    <Configurations>Debug;Release;Development;Analyze</Configurations>
     <ApplicationManifest>app.manifest</ApplicationManifest>
     <AssemblyName>UnrealBuildTool</AssemblyName>
     <RootNamespace>UnrealBuildTool</RootNamespace>
@@ -42,6 +42,14 @@
     <DebugSymbols>true</DebugSymbols>
     <DebugType>full</DebugType>
   </PropertyGroup>
+  <PropertyGroup Condition=" '$(Configuration)|$(Platform)' == 'Analyze|AnyCPU' ">
+    <OutputPath>..\..\..\Binaries\DotNET\UnrealBuildTool</OutputPath>
+    <DocumentationFile>..\..\..\Binaries\DotNET\UnrealBuildTool\UnrealBuildTool.xml</DocumentationFile>
+    <RunAnalyzersDuringBuild>True</RunAnalyzersDuringBuild>
+    <TreatWarningsAsErrors>false</TreatWarningsAsErrors>
+    <GenerateDocumentationFile>True</GenerateDocumentationFile>
+    <DocumentationFile></DocumentationFile>
+  </PropertyGroup>
   <ItemGroup>
     <Reference Include="Ionic.Zip.Reduced">
       <HintPath>..\..\..\Binaries\DotNET\Ionic.Zip.Reduced.dll</HintPath>
