BuildGraph: Remove unused abstractions added to support reading BuildGraph scripts directly from P4.

diff --git a/Engine/Source/Programs/AutomationTool/BuildGraph/BgCondition.cs b/Engine/Source/Programs/AutomationTool/BuildGraph/BgCondition.cs
--- a/Engine/Source/Programs/AutomationTool/BuildGraph/BgCondition.cs
+++ b/Engine/Source/Programs/AutomationTool/BuildGraph/BgCondition.cs
@@ -1,10 +1,12 @@
 // Copyright Epic Games, Inc. All Rights Reserved.
 
+using EpicGames.Core;
 using System;
 using System.Collections.Generic;
 using System.IO;
 using System.Linq;
 using System.Threading.Tasks;
+using UnrealBuildBase;
 
 namespace AutomationTool
 {
@@ -70,20 +72,12 @@ namespace AutomationTool
 		/// </summary>
 		int _idx;
 
-		/// <summary>
-		/// Context for evaluating the expression
-		/// </summary>
-		readonly IBgScriptReaderContext _context;
-
 		/// <summary>
 		/// Constructor
 		/// </summary>
 		/// <param name="text">The condition text</param>
-		/// <param name="context">Context for evaluating the expression</param>
-		private BgCondition(string text, IBgScriptReaderContext context)
+		private BgCondition(string text)
 		{
-			_context = context;
-
 			Tokenize(text, _tokens);
 		}
 
@@ -91,11 +85,10 @@ namespace AutomationTool
 		/// Evaluates the given string as a condition. Throws a ConditionException on a type or syntax error.
 		/// </summary>
 		/// <param name="text"></param>
-		/// <param name="context"></param>
 		/// <returns>The result of evaluating the condition</returns>
-		public static ValueTask<bool> EvaluateAsync(string text, IBgScriptReaderContext context)
+		public static ValueTask<bool> EvaluateAsync(string text)
 		{
-			return new BgCondition(text, context).EvaluateAsync();
+			return new BgCondition(text).EvaluateAsync();
 		}
 
 		/// <summary>
@@ -300,7 +293,7 @@ namespace AutomationTool
 				// Check whether file or directory exists. Evaluate the argument as a subexpression.
 				_idx++;
 				string argument = await EvaluateScalarAsync();
-				result = await _context.ExistsAsync(argument) ? "true" : "false";
+				result = Exists(argument) ? "true" : "false";
 			}
 			else if (String.Compare(_tokens[_idx], "HasTrailingSlash", true) == 0 && _tokens[_idx + 1] == "(")
 			{
@@ -357,6 +350,23 @@ namespace AutomationTool
 			return result;
 		}
 
+		/// <summary>
+		/// Determine if a path exists
+		/// </summary>
+		/// <param name="Path"></param>
+		/// <returns></returns>
+		static bool Exists(string Path)
+		{
+			try
+			{
+				return FileReference.Exists(FileReference.Combine(Unreal.RootDirectory, Path)) || DirectoryReference.Exists(DirectoryReference.Combine(Unreal.RootDirectory, Path));
+			}
+			catch
+			{
+				return false;
+			}
+		}
+
 		/// <summary>
 		/// Checks whether Haystack contains "Needle". 
 		/// </summary>
@@ -520,7 +530,7 @@ namespace AutomationTool
 		/// <param name="expectedResult">The expected result</param>
 		static async Task TestConditionAsync(string condition, bool expectedResult)
 		{
-			bool result = await new BgCondition(condition, null!).EvaluateAsync();
+			bool result = await new BgCondition(condition).EvaluateAsync();
 			Console.WriteLine("{0}: {1} = {2}", (result == expectedResult) ? "PASS" : "FAIL", condition, result);
 		}
 	}
diff --git a/Engine/Source/Programs/AutomationTool/BuildGraph/BgNodeExecutor.cs b/Engine/Source/Programs/AutomationTool/BuildGraph/BgNodeExecutor.cs
--- a/Engine/Source/Programs/AutomationTool/BuildGraph/BgNodeExecutor.cs
+++ b/Engine/Source/Programs/AutomationTool/BuildGraph/BgNodeExecutor.cs
@@ -148,12 +148,12 @@ namespace AutomationTool
 			Node = node;
 		}
 
-		public bool Bind(Dictionary<string, ScriptTaskBinding> nameToTask, Dictionary<string, BgNodeOutput> tagNameToNodeOutput, IBgScriptReaderContext context, ILogger logger)
+		public bool Bind(Dictionary<string, ScriptTaskBinding> nameToTask, Dictionary<string, BgNodeOutput> tagNameToNodeOutput, ILogger logger)
 		{
 			bool bResult = true;
 			foreach (BgTask TaskInfo in Node.Tasks)
 			{
-				BgTaskImpl? boundTask = BindTask(TaskInfo, nameToTask, tagNameToNodeOutput, context, logger);
+				BgTaskImpl? boundTask = BindTask(TaskInfo, nameToTask, tagNameToNodeOutput, logger);
 				if (boundTask == null)
 				{
 					bResult = false;
@@ -166,7 +166,7 @@ namespace AutomationTool
 			return bResult;
 		}
 
-		BgTaskImpl? BindTask(BgTask TaskInfo, Dictionary<string, ScriptTaskBinding> NameToTask, IReadOnlyDictionary<string, BgNodeOutput> TagNameToNodeOutput, IBgScriptReaderContext Context, ILogger Logger)
+		BgTaskImpl? BindTask(BgTask TaskInfo, Dictionary<string, ScriptTaskBinding> NameToTask, IReadOnlyDictionary<string, BgNodeOutput> TagNameToNodeOutput, ILogger Logger)
 		{
 			// Get the reflection info for this element
 			ScriptTaskBinding? Task;
@@ -203,7 +203,7 @@ namespace AutomationTool
 				if (Parameter.CollectionType == null)
 				{
 					// Parse it and assign it to the parameters object
-					object? FieldValue = ParseValue(Value, Parameter.ValueType, Context);
+					object? FieldValue = ParseValue(Value, Parameter.ValueType);
 					Parameter.FieldInfo.SetValue(ParametersObject, FieldValue);
 				}
 				else
@@ -220,7 +220,7 @@ namespace AutomationTool
 					List<string> ValueStrings = BgTaskImpl.SplitDelimitedList(Value);
 					foreach (string ValueString in ValueStrings)
 					{
-						object ElementValue = ParseValue(ValueString, Parameter.ValueType, Context)!;
+						object ElementValue = ParseValue(ValueString, Parameter.ValueType)!;
 						Parameter.CollectionType.InvokeMember("Add", BindingFlags.InvokeMethod | BindingFlags.Instance | BindingFlags.Public, null, CollectionValue, new object[] { ElementValue });
 					}
 				}
@@ -271,9 +271,8 @@ namespace AutomationTool
 		/// </summary>
 		/// <param name="ValueText">The text to parse</param>
 		/// <param name="ValueType">Type of the value to parse</param>
-		/// <param name="Context">Context for the script reader</param>
 		/// <returns>Value that was parsed</returns>
-		static object? ParseValue(string ValueText, Type ValueType, IBgScriptReaderContext Context)
+		static object? ParseValue(string ValueText, Type ValueType)
 		{
 			// Parse it and assign it to the parameters object
 			if (ValueType.IsEnum)
@@ -282,7 +281,7 @@ namespace AutomationTool
 			}
 			else if (ValueType == typeof(Boolean))
 			{
-				return BgCondition.EvaluateAsync(ValueText, Context).Result;
+				return BgCondition.EvaluateAsync(ValueText).Result;
 			}
 			else if (ValueType == typeof(FileReference))
 			{
diff --git a/Engine/Source/Programs/AutomationTool/BuildGraph/BgScriptReader.cs b/Engine/Source/Programs/AutomationTool/BuildGraph/BgScriptReader.cs
--- a/Engine/Source/Programs/AutomationTool/BuildGraph/BgScriptReader.cs
+++ b/Engine/Source/Programs/AutomationTool/BuildGraph/BgScriptReader.cs
@@ -13,6 +13,7 @@
 using EpicGames.BuildGraph;
 using EpicGames.Core;
 using Microsoft.Extensions.Logging;
+using UnrealBuildBase;
 
 #nullable enable
 
@@ -26,12 +27,7 @@
 		/// <summary>
 		/// The file containing this element
 		/// </summary>
-		public string File { get; }
+		public FileReference File { get; }
 
 		/// <summary>
-		/// Native file object
-		/// </summary>
-		public object NativeFile { get; }
-
-		/// <summary>
 		/// The line number containing this element
@@ -42,12 +38,10 @@
 		/// Constructor
 		/// </summary>
 		/// <param name="file"></param>
-		/// <param name="nativeFile"></param>
 		/// <param name="lineNumber"></param>
-		public BgScriptLocation(string file, object nativeFile, int lineNumber)
+		public BgScriptLocation(FileReference file, int lineNumber)
 		{
 			File = file;
-			NativeFile = nativeFile;
 			LineNumber = lineNumber;
 		}
 	}
@@ -60,12 +54,7 @@
 		/// <summary>
 		/// The file being read
 		/// </summary>
-		string File { get; }
+		FileReference File { get; }
 
 		/// <summary>
-		/// Native file representation
-		/// </summary>
-		object NativeFile { get; }
-
-		/// <summary>
 		/// Interface to the LineInfo on the active XmlReader
@@ -85,10 +74,9 @@
 		/// <summary>
 		/// Private constructor. Use ScriptDocument.Load to read an XML document.
 		/// </summary>
-		BgScriptDocument(string inFile, object inNativeFile, ILogger inLogger)
+		BgScriptDocument(FileReference inFile, ILogger inLogger)
 		{
 			File = inFile;
-			NativeFile = inNativeFile;
 			Logger = inLogger;
 		}
 
@@ -97,7 +85,7 @@
 		/// </summary>
 		public override XmlElement CreateElement(string? prefix, string? localName, string? namespaceUri)
 		{
-			BgScriptLocation location = new BgScriptLocation(File, NativeFile, _lineInfo!.LineNumber);
+			BgScriptLocation location = new BgScriptLocation(File, _lineInfo!.LineNumber);
 			return new BgScriptElement(location, prefix!, localName!, namespaceUri!, this);
 		}
 
@@ -105,15 +93,14 @@
 		/// Loads a script document from the given file
 		/// </summary>
 		/// <param name="file">The file to load</param>
-		/// <param name="nativeFile"></param>
 		/// <param name="data"></param>
 		/// <param name="schema">The schema to validate against</param>
 		/// <param name="logger">Logger for output messages</param>
 		/// <param name="outDocument">If successful, the document that was read</param>
 		/// <returns>True if the document could be read, false otherwise</returns>
-		public static bool TryRead(string file, object nativeFile, byte[] data, BgScriptSchema schema, ILogger logger, [NotNullWhen(true)] out BgScriptDocument? outDocument)
+		public static bool TryRead(FileReference file, byte[] data, BgScriptSchema schema, ILogger logger, [NotNullWhen(true)] out BgScriptDocument? outDocument)
 		{
-			BgScriptDocument document = new BgScriptDocument(file, nativeFile, logger);
+			BgScriptDocument document = new BgScriptDocument(file, logger);
 
 			XmlReaderSettings settings = new XmlReaderSettings();
 			if (schema != null)
@@ -136,7 +123,7 @@
 				{
 					if (!document._bHasErrors)
 					{
-						BgScriptLocation location = new BgScriptLocation(file, nativeFile, ex.LineNumber);
+						BgScriptLocation location = new BgScriptLocation(file, ex.LineNumber);
 						logger.LogScriptError(location, "{Message}", ex.Message);
 						document._bHasErrors = true;
 					}
@@ -152,14 +139,14 @@
 				// Check that the root element is valid. If not, we didn't actually validate against the schema.
 				if (document.DocumentElement!.Name != BgScriptSchema.RootElementName)
 				{
-					BgScriptLocation location = new BgScriptLocation(file, nativeFile, 1);
+					BgScriptLocation location = new BgScriptLocation(file, 1);
 					logger.LogScriptError(location, "Script does not have a root element called '{ElementName}'", BgScriptSchema.RootElementName);
 					outDocument = null;
 					return false;
 				}
 				if (document.DocumentElement.NamespaceURI != BgScriptSchema.NamespaceUri)
 				{
-					BgScriptLocation location = new BgScriptLocation(file, nativeFile, 1);
+					BgScriptLocation location = new BgScriptLocation(file, 1);
 					logger.LogScriptError(location, "Script root element is not in the '{Namespace}' namespace (add the xmlns=\"{NewNamespace}\" attribute)", BgScriptSchema.NamespaceUri, BgScriptSchema.NamespaceUri);
 					outDocument = null;
 					return false;
@@ -177,7 +164,7 @@
 		/// <param name="args">Standard argument for ValidationEventHandler</param>
 		void ValidationEvent(object? sender, ValidationEventArgs args)
 		{
-			BgScriptLocation location = new BgScriptLocation(File, NativeFile, args.Exception.LineNumber);
+			BgScriptLocation location = new BgScriptLocation(File, args.Exception.LineNumber);
 			if (args.Severity == XmlSeverityType.Warning)
 			{
 				Logger.LogScriptWarning(location, "{Message}", args.Message);
@@ -258,47 +245,5 @@
 	}
 
 	/// <summary>
-	/// Context used for parsing scripts
-	/// </summary>
-	public interface IBgScriptReaderContext
-	{
-		/// <summary>
-		/// Tests whether the given file or directory exists
-		/// </summary>
-		/// <param name="path">Path to a file</param>
-		/// <returns>True if the file exists</returns>
-		Task<bool> ExistsAsync(string path);
-
-		/// <summary>
-		/// Tries to read a file from the given path
-		/// </summary>
-		/// <param name="path">Path of the file to read</param>
-		/// <returns></returns>
-		Task<byte[]?> ReadAsync(string path);
-
-		/// <summary>
-		/// Finds files matching the given pattern
-		/// </summary>
-		/// <param name="pattern"></param>
-		/// <returns></returns>
-		Task<string[]> FindAsync(string pattern);
-
-		/// <summary>
-		/// Converts a path to its native form, for display to the user
-		/// </summary>
-		/// <param name="path">Path to format</param>
-		/// <returns></returns>
-		object GetNativePath(string path);
-
-		/// <summary>
-		/// Combines one path with another
-		/// </summary>
-		/// <param name="path">The base file path</param>
-		/// <param name="next">Relative path to the next file</param>
-		/// <returns></returns>
-		string CombinePaths(string path, string next);
-	}
-
-	/// <summary>
 	/// Extension methods for writing script error messages
 	/// </summary>
@@ -310,7 +255,7 @@
 		public static void LogScriptError(this ILogger logger, BgScriptLocation location, string format, params object[] args)
 		{
 			object[] allArgs = new object[args.Length + 2];
-			allArgs[0] = location.NativeFile;
+			allArgs[0] = location.File;
 			allArgs[1] = location.LineNumber;
 			args.CopyTo(allArgs, 2);
 			logger.LogError($"{{Script}}({{Line}}): error: {format}", allArgs);
@@ -322,7 +267,7 @@
 		public static void LogScriptWarning(this ILogger logger, BgScriptLocation location, string format, params object[] args)
 		{
 			object[] allArgs = new object[args.Length + 2];
-			allArgs[0] = location.NativeFile;
+			allArgs[0] = location.File;
 			allArgs[1] = location.LineNumber;
 			args.CopyTo(allArgs, 2);
 			logger.LogWarning($"{{Script}}({{Line}}): warning: {format}", allArgs);
@@ -330,15 +275,29 @@
 	}
 
 	/// <summary>
+	/// Overridden version of <see cref="BgNode"/> which contains a list of tasks
+	/// </summary>
+	class BgScriptNode : BgNode
+	{
+		/// <summary>
+		/// List of tasks to execute
+		/// </summary>
+		public List<BgTask> Tasks { get; } = new List<BgTask>();
+
+		/// <summary>
+		/// Constructor
+		/// </summary>
+		public BgScriptNode(string name, IReadOnlyList<BgNodeOutput> inputs, IReadOnlyList<string> outputNames, IReadOnlyList<BgNode> inputDependencies, IReadOnlyList<BgNode> orderDependencies, IReadOnlyList<FileReference> requiredTokens)
+			: base(name, inputs, outputNames, inputDependencies, orderDependencies, requiredTokens)
+		{
+		}
+	}
+
+	/// <summary>
 	/// Reader for build graph definitions. Instanced to contain temporary state; public interface is through ScriptReader.TryRead().
 	/// </summary>
-	public abstract class BgScriptReaderBase
+	public class BgScriptReader
 	{
 		/// <summary>
-		/// Interface used for reading files
-		/// </summary>
-		protected IBgScriptReaderContext Context { get; }
-
-		/// <summary>
 		/// List of property name to value lookups. Modifications to properties are scoped to nodes and agents. EnterScope() pushes an empty dictionary onto the end of this list, and LeaveScope() removes one. 
 		/// ExpandProperties() searches from last to first lookup when trying to resolve a property name, and takes the first it finds.
@@ -357,5 +316,20 @@
 		readonly Dictionary<string, BgScriptMacro> _macroNameToDefinition = new Dictionary<string, BgScriptMacro>();
 
 		/// <summary>
+		/// The current graph
+		/// </summary>
+		readonly BgGraph _graph = new BgGraph();
+
+		/// <summary>
+		/// Arguments for evaluating the graph
+		/// </summary>
+		readonly Dictionary<string, string> _arguments;
+
+		/// <summary>
+		/// The name of the node if only a single node is going to be built, otherwise null.
+		/// </summary>
+		readonly string? _singleNodeName;
+
+		/// <summary>
 		/// Schema for the script
 		/// </summary>
@@ -371,29 +345,76 @@
 		/// </summary>
 		public int NumErrors { get; private set; }
 
+		BgAgent? _enclosingAgent;
+		BgScriptNode? _enclosingNode;
+
 		/// <summary>
 		/// Private constructor. Use ScriptReader.TryRead() to read a script file.
 		/// </summary>
-		/// <param name="context">Context object</param>
+		/// <param name="defaultProperties">Default properties available to the script</param>
+		/// <param name="arguments">Arguments passed in to the graph on the command line</param>
+		/// <param name="singleNodeName">If a single node will be processed, the name of that node.</param>
 		/// <param name="schema">Schema for the script</param>
 		/// <param name="logger">Logger for diagnostic messages</param>
-		protected BgScriptReaderBase(IBgScriptReaderContext context, BgScriptSchema schema, ILogger logger)
+		protected BgScriptReader(IDictionary<string, string> defaultProperties, IReadOnlyDictionary<string, string> arguments, string? singleNodeName, BgScriptSchema schema, ILogger logger)
 		{
-			Context = context;
 			Schema = schema;
 			Logger = logger;
 
 			EnterScope();
+
+			_arguments = new Dictionary<string, string>(arguments, StringComparer.OrdinalIgnoreCase);
+			_singleNodeName = singleNodeName;
+
+			foreach (KeyValuePair<string, string> pair in defaultProperties)
+			{
+				SetPropertyValue(null!, pair.Key, pair.Value);
+			}
 		}
 
 		/// <summary>
+		/// Try to read a script file from the given file.
+		/// </summary>
+		/// <param name="file">File to read from</param>
+		/// <param name="arguments">Arguments passed in to the graph on the command line</param>
+		/// <param name="defaultProperties">Default properties available to the script</param>
+		/// <param name="schema">Schema for the script</param>
+		/// <param name="logger">Logger for output messages</param>
+		/// <param name="singleNodeName">If a single node will be processed, the name of that node.</param>
+		/// <returns>True if the graph was read, false if there were errors</returns>
+		public static async Task<BgGraph?> ReadAsync(FileReference file, Dictionary<string, string> arguments, Dictionary<string, string> defaultProperties, BgScriptSchema schema, ILogger logger, string? singleNodeName = null)
+		{
+			// Read the file and build the graph
+			BgScriptReader reader = new BgScriptReader(defaultProperties, arguments, singleNodeName, schema, logger);
+			if (!await reader.TryReadAsync(file) || reader.NumErrors > 0)
+			{
+				return null;
+			}
+
+			// Make sure all the arguments were valid
+			HashSet<string> validArgumentNames = new HashSet<string>(reader._graph.Options.Select(x => x.Name), StringComparer.OrdinalIgnoreCase);
+			validArgumentNames.Add("PreflightChange");
+
+			foreach (string argumentName in arguments.Keys)
+			{
+				if (!validArgumentNames.Contains(argumentName))
+				{
+					logger.LogWarning("Unknown argument '{ArgumentName}' for '{Script}'", argumentName, file);
+				}
+			}
+
+			// Return the constructed graph
+			return reader._graph;
+		}
+
+		/// <summary>
 		/// Read the script from the given file
 		/// </summary>
 		/// <param name="file">File to read from</param>
-		protected async Task<bool> TryReadAsync(string file)
+		protected async Task<bool> TryReadAsync(FileReference file)
 		{
 			// Get the data for this file
-			byte[]? data = await Context.ReadAsync(file);
+			byte[]? data = await FileReference.ReadAllBytesAsync(file);
 			if (data == null)
 			{
 				Logger.LogError("Unable to open file {File}", file);
@@ -403,7 +424,7 @@
 
 			// Read the document and validate it against the schema
 			BgScriptDocument? document;
-			if (!BgScriptDocument.TryRead(file, Context.GetNativePath(file), data, Schema, Logger, out document))
+			if (!BgScriptDocument.TryRead(file, data, Schema, Logger, out document))
 			{
 				NumErrors++;
 				return false;
@@ -597,21 +618,23 @@
 		{
 			if (await EvaluateConditionAsync(element))
 			{
-				HashSet<string> files = new HashSet<string>();
+				string basePath = element.Location.File.MakeRelativeTo(Unreal.RootDirectory).Replace(Path.DirectorySeparatorChar, '/');
+
+				HashSet<FileReference> files = new HashSet<FileReference>();
 				foreach (string script in ReadListAttribute(element, "Script"))
 				{
-					string includePath = Context.CombinePaths(element.Location.File, script);
+					string includePath = CombinePaths(basePath, script);
 					if (Regex.IsMatch(includePath, @"\*|\?|\.\.\."))
 					{
-						files.UnionWith(await Context.FindAsync(includePath));
+						files.UnionWith(FindMatchingFiles(includePath));
 					}
 					else
 					{
-						files.Add(includePath);
+						files.Add(new FileReference(includePath));
 					}
 				}
 
-				foreach (string file in files.OrderBy(x => x, StringComparer.OrdinalIgnoreCase))
+				foreach (FileReference file in files.OrderBy(x => x.FullName, StringComparer.OrdinalIgnoreCase))
 				{
 					Logger.LogDebug("Including file {File}", file);
 					await TryReadAsync(file);
@@ -620,28 +643,224 @@
 		}
 
 		/// <summary>
+		/// Combine two paths without validating the result
+		/// </summary>
+		static string CombinePaths(string basePath, string nextPath)
+		{
+			if (Path.IsPathRooted(nextPath))
+			{
+				return nextPath;
+			}
+
+			List<string> fragments = new List<string>(basePath.Split('/'));
+			fragments.RemoveAt(fragments.Count - 1);
+
+			foreach (string appendFragment in nextPath.Split('/'))
+			{
+				if (appendFragment.Equals(".", StringComparison.Ordinal))
+				{
+					continue;
+				}
+				else if (appendFragment.Equals("..", StringComparison.Ordinal))
+				{
+					if (fragments.Count > 0)
+					{
+						fragments.RemoveAt(fragments.Count - 1);
+					}
+					else
+					{
+						throw new Exception($"Path '{nextPath}' cannot be combined with '{basePath}'");
+					}
+				}
+				else
+				{
+					fragments.Add(appendFragment);
+				}
+			}
+			return String.Join('/', fragments);
+		}
+
+		/// <summary>
+		/// Find files matching a pattern
+		/// </summary>
+		/// <param name="Pattern"></param>
+		/// <returns></returns>
+		IEnumerable<FileReference> FindMatchingFiles(string Pattern)
+		{
+			FileFilter Filter = new FileFilter();
+			Filter.AddRule(Pattern, FileFilterType.Include);
+
+			return Filter.ApplyToDirectory(Unreal.RootDirectory, true);
+		}
+
+		/// <summary>
 		/// Reads the definition of a graph option; a parameter which can be set by the user on the command-line or via an environment variable.
 		/// </summary>
 		/// <param name="element">Xml element to read the definition from</param>
-		protected abstract Task ReadOptionAsync(BgScriptElement element);
+		async Task ReadOptionAsync(BgScriptElement element)
+		{
+			if (await EvaluateConditionAsync(element))
+			{
+				string name = ReadAttribute(element, "Name");
+				if (ValidateName(element, name))
+				{
+					// Make sure we're at global scope
+					if (ScopedProperties.Count > 1)
+					{
+						throw new Exception("Incorrect scope depth for reading option settings");
+					}
+
+					// Check if the property already exists. If it does, we don't need to register it as an option.
+					string? existingValue;
+					if (TryGetPropertyValue(name, out existingValue) && existingValue != null)
+					{
+						// If there's a restriction on this definition, check it matches
+						string restrict = ReadAttribute(element, "Restrict");
+						if (!String.IsNullOrEmpty(restrict) && !Regex.IsMatch(existingValue, "^" + restrict + "$", RegexOptions.IgnoreCase))
+						{
+							LogError(element, "'{0} is already set to '{1}', which does not match the given restriction ('{2}')", name, existingValue, restrict);
+						}
+					}
+					else
+					{
+						// Create a new option object to store the settings
+						string description = ReadAttribute(element, "Description");
+						string defaultValue = ReadAttribute(element, "DefaultValue");
+						BgOption option = new BgOption(name, description, defaultValue);
+						_graph.Options.Add(option);
+
+						// Get the value of this property
+						string? value;
+						if (!_arguments.TryGetValue(name, out value))
+						{
+							value = option.DefaultValue;
+						}
+						SetPropertyValue(element, name, value);
+
+						// If there's a restriction on it, check it's valid
+						string restrict = ReadAttribute(element, "Restrict");
+						if (!String.IsNullOrEmpty(restrict))
+						{
+							string pattern = "^(" + restrict + ")$";
+							if (!Regex.IsMatch(value, pattern, RegexOptions.IgnoreCase))
+							{
+								LogError(element, "'{0}' is not a valid value for '{1}' (required: '{2}')", value, name, restrict);
+							}
+							if (option.DefaultValue != value && !Regex.IsMatch(option.DefaultValue, pattern, RegexOptions.IgnoreCase))
+							{
+								LogError(element, "Default value '{0}' is not valid for '{1}' (required: '{2}')", option.DefaultValue, name, restrict);
+							}
+						}
+					}
+				}
+			}
+		}
 
 		/// <summary>
 		/// Reads a property assignment.
 		/// </summary>
 		/// <param name="element">Xml element to read the definition from</param>
-		protected abstract Task ReadPropertyAsync(BgScriptElement element);
+		async Task ReadPropertyAsync(BgScriptElement element)
+		{
+			if (await EvaluateConditionAsync(element))
+			{
+				string name = ReadAttribute(element, "Name");
+				if (ValidateName(element, name))
+				{
+					string value = ReadAttribute(element, "Value");
+					if (element.HasChildNodes)
+					{
+						// Read the element content, and append each line to the value as a semicolon delimited list
+						StringBuilder builder = new StringBuilder(value);
+						foreach (string line in element.InnerText.Split('\n'))
+						{
+							string trimLine = ExpandProperties(element, line.Trim());
+							if (trimLine.Length > 0)
+							{
+								if (builder.Length > 0)
+								{
+									builder.Append(";");
+								}
+								builder.Append(trimLine);
+							}
+						}
+						value = builder.ToString();
+					}
+					SetPropertyValue(element, name, value);
+				}
+			}
+		}
 
 		/// <summary>
 		/// Reads a Regex assignment.
 		/// </summary>
 		/// <param name="element">Xml element to read the definition from</param>
-		protected abstract Task ReadRegexAsync(BgScriptElement element);
+		async Task ReadRegexAsync(BgScriptElement element)
+		{
+			if (await EvaluateConditionAsync(element))
+			{
+				// Get the pattern
+				string regexString = ReadAttribute(element, "Pattern");
+
+				// Make sure its a valid regex.
+				Regex? regexValue = ParseRegex(element, regexString);
+				if (regexValue != null)
+				{
+					// read the names in 
+					string[] captureNames = ReadListAttribute(element, "Capture");
+
+					// get number of groups we passed in
+					int[] groupNumbers = regexValue.GetGroupNumbers();
+
+					// make sure the number of property names is the same as the number of match groups
+					// this includes the entire string match group as [0], so don't count that one.
+					if (captureNames.Length != groupNumbers.Count() - 1)
+					{
+						LogError(element, "MatchGroup count: {0} does not match the number of names specified: {1}", groupNumbers.Count() - 1, captureNames.Length);
+					}
+					else
+					{
+						// apply the regex to the value
+						string input = ReadAttribute(element, "Input");
+						Match match = regexValue.Match(input);
+
+						bool optional = await BgCondition.EvaluateAsync(ReadAttribute(element, "Optional"));
+						if (!match.Success)
+						{
+							if (!optional)
+							{
+								LogError(element, "Regex {0} did not find a match against input string {1}", regexString, input);
+							}
+						}
+						else
+						{
+							// assign each property to the group it matches, skip over [0]
+							for (int matchIdx = 1; matchIdx < groupNumbers.Count(); matchIdx++)
+							{
+								SetPropertyValue(element, captureNames[matchIdx - 1], match.Groups[matchIdx].Value);
+							}
+						}
+					}
+				}
+			}
+		}
 
 		/// <summary>
 		/// Reads a property assignment from an environment variable.
 		/// </summary>
 		/// <param name="element">Xml element to read the definition from</param>
-		protected abstract Task ReadEnvVarAsync(BgScriptElement element);
+		async Task ReadEnvVarAsync(BgScriptElement element)
+		{
+			if (await EvaluateConditionAsync(element))
+			{
+				string name = ReadAttribute(element, "Name");
+				if (ValidateName(element, name))
+				{
+					string value = Environment.GetEnvironmentVariable(name) ?? "";
+					SetPropertyValue(element, name, value);
+				}
+			}
+		}
 
 		/// <summary>
 		/// Reads a macro definition
@@ -722,7 +941,43 @@
 		/// Reads the definition for an agent.
 		/// </summary>
 		/// <param name="element">Xml element to read the definition from</param>
-		protected abstract Task ReadAgentAsync(BgScriptElement element);
+		async Task ReadAgentAsync(BgScriptElement element)
+		{
+			string? name;
+			if (await EvaluateConditionAsync(element) && TryReadObjectName(element, out name))
+			{
+				// Read the valid agent types. This may be omitted if we're continuing an existing agent.
+				string[] types = ReadListAttribute(element, "Type");
+
+				// Create the agent object, or continue an existing one
+				BgAgent? agent;
+				if (_graph.NameToAgent.TryGetValue(name, out agent))
+				{
+					if (types.Length > 0 && agent.PossibleTypes.Length > 0)
+					{
+						if (types.Length != agent.PossibleTypes.Length || !types.SequenceEqual(agent.PossibleTypes, StringComparer.InvariantCultureIgnoreCase))
+						{
+							LogError(element, "Agent types ({0}) were different than previous agent definition with types ({1}). Must either be empty or match exactly.", String.Join(",", types), String.Join(",", agent.PossibleTypes));
+						}
+					}
+				}
+				else
+				{
+					if (types.Length == 0)
+					{
+						LogError(element, "Missing type for agent '{0}'", name);
+					}
+					agent = new BgAgent(name, types);
+					_graph.NameToAgent.Add(name, agent);
+					_graph.Agents.Add(agent);
+				}
+
+				// Process all the child elements.
+				_enclosingAgent = agent;
+				await ReadAgentBodyAsync(element);
+				_enclosingAgent = null;
+			}
+		}
 
 		/// <summary>
 		/// Read the contents of an agent definition
@@ -732,91 +987,3 @@
 		{
 			EnterScope();
 			foreach (BgScriptElement childElement in element.ChildNodes.OfType<BgScriptElement>())
-			{
-				switch (childElement.Name)
-				{
-					case "Property":
-						await ReadPropertyAsync(childElement);
-						break;
-					case "Regex":
-						await ReadRegexAsync(childElement);
-						break;
-					case "Node":
-						await ReadNodeAsync(childElement);
-						break;
-					case "Aggregate":
-						await ReadAggregateAsync(childElement);
-						break;
-					case "Trace":
-						await ReadDiagnosticAsync(childElement, LogLevel.Information);
-						break;
-					case "Warning":
-						await ReadDiagnosticAsync(childElement, LogLevel.Warning);
-						break;
-					case "Error":
-						await ReadDiagnosticAsync(childElement, LogLevel.Error);
-						break;
-					case "Label":
-						await ReadLabelAsync(childElement);
-						break;
-					case "Do":
-						await ReadBlockAsync(childElement, ReadAgentBodyAsync);
-						break;
-					case "Switch":
-						await ReadSwitchAsync(childElement, ReadAgentBodyAsync);
-						break;
-					case "ForEach":
-						await ReadForEachAsync(childElement, ReadAgentBodyAsync);
-						break;
-					case "Expand":
-						await ReadExpandAsync(childElement, ReadAgentBodyAsync);
-						break;
-					case "Annotate":
-						await ReadAnnotationAsync(childElement);
-						break;
-					default:
-						LogError(childElement, "Unexpected element type '{0}'", childElement.Name);
-						break;
-				}
-			}
-			LeaveScope();
-		}
-
-		/// <summary>
-		/// Reads the definition for an aggregate
-		/// </summary>
-		/// <param name="element">Xml element to read the definition from</param>
-		protected abstract Task ReadAggregateAsync(BgScriptElement element);
-
-		/// <summary>
-		/// Reads the definition for a report
-		/// </summary>
-		/// <param name="element">Xml element to read the definition from</param>
-		protected abstract Task ReadReportAsync(BgScriptElement element);
-
-		/// <summary>
-		/// Reads the definition for a badge
-		/// </summary>
-		/// <param name="element">Xml element to read the definition from</param>
-		protected abstract Task ReadBadgeAsync(BgScriptElement element);
-
-		/// <summary>
-		/// Reads the definition for a label
-		/// </summary>
-		/// <param name="element">Xml element to read the definition from</param>
-		protected abstract Task ReadLabelAsync(BgScriptElement element);
-
-		/// <summary>
-		/// Reads the definition for a node, and adds it to the given agent
-		/// </summary>
-		/// <param name="element">Xml element to read the definition from</param>
-		protected abstract Task ReadNodeAsync(BgScriptElement element);
-
-		/// <summary>
-		/// Reads the contents of a node element
-		/// </summary>
-		/// <param name="element">Xml element to read the definition from</param>
-		protected async Task ReadNodeBodyAsync(XmlElement element)
-		{
-			EnterScope();
-			foreach (BgScriptElement childElement in element.ChildNodes.OfType<BgScriptElement>())
@@ -829,6 +996,12 @@
 					case "Regex":
 						await ReadRegexAsync(childElement);
 						break;
+					case "Node":
+						await ReadNodeAsync(childElement);
+						break;
+					case "Aggregate":
+						await ReadAggregateAsync(childElement);
+						break;
 					case "Trace":
 						await ReadDiagnosticAsync(childElement, LogLevel.Information);
 						break;
@@ -838,20 +1011,26 @@
 					case "Error":
 						await ReadDiagnosticAsync(childElement, LogLevel.Error);
 						break;
+					case "Label":
+						await ReadLabelAsync(childElement);
+						break;
 					case "Do":
-						await ReadBlockAsync(childElement, ReadNodeBodyAsync);
+						await ReadBlockAsync(childElement, ReadAgentBodyAsync);
 						break;
 					case "Switch":
-						await ReadSwitchAsync(childElement, ReadNodeBodyAsync);
+						await ReadSwitchAsync(childElement, ReadAgentBodyAsync);
 						break;
 					case "ForEach":
-						await ReadForEachAsync(childElement, ReadNodeBodyAsync);
+						await ReadForEachAsync(childElement, ReadAgentBodyAsync);
 						break;
 					case "Expand":
-						await ReadExpandAsync(childElement, ReadNodeBodyAsync);
+						await ReadExpandAsync(childElement, ReadAgentBodyAsync);
 						break;
+					case "Annotate":
+						await ReadAnnotationAsync(childElement);
+						break;
 					default:
-						await ReadTaskAsync(childElement);
+						LogError(childElement, "Unexpected element type '{0}'", childElement.Name);
 						break;
 				}
 			}
@@ -859,809 +1038,10 @@
 		}
 
 		/// <summary>
-		/// Reads a block element
-		/// </summary>
-		/// <param name="element">Xml element to read the definition from</param>
-		/// <param name="readContentsAsync">Delegate to read the contents of the element, if the condition evaluates to true</param>
-		async Task ReadBlockAsync(BgScriptElement element, Func<BgScriptElement, Task> readContentsAsync)
-		{
-			if (await EvaluateConditionAsync(element))
-			{
-				await readContentsAsync(element);
-			}
-		}
-
-		/// <summary>
-		/// Reads a "Switch" element 
-		/// </summary>
-		/// <param name="element">Xml element to read the definition from</param>
-		/// <param name="readContentsAsync">Delegate to read the contents of the element, if the condition evaluates to true</param>
-		protected abstract Task ReadSwitchAsync(BgScriptElement element, Func<BgScriptElement, Task> readContentsAsync);
-
-		/// <summary>
-		/// Reads a "ForEach" element 
-		/// </summary>
-		/// <param name="element">Xml element to read the definition from</param>
-		/// <param name="readContentsAsync">Delegate to read the contents of the element, if the condition evaluates to true</param>
-		async Task ReadForEachAsync(BgScriptElement element, Func<BgScriptElement, Task> readContentsAsync)
-		{
-			EnterScope();
-			if (await EvaluateConditionAsync(element))
-			{
-				string name = ReadAttribute(element, "Name");
-				string separator = ReadAttribute(element, "Separator");
-				if (separator.Length > 1)
-				{
-					LogWarning(element, "Node {0}'s Separator attribute is more than one character ({1}). Defaulting to ;", name, separator);
-					separator = ";";
-				}
-				if (String.IsNullOrEmpty(separator))
-				{
-					separator = ";";
-				}
-				if (ValidateName(element, name))
-				{
-					if (ScopedProperties.Any(x => x.ContainsKey(name)))
-					{
-						LogError(element, "Loop variable '{0}' already exists as a local property in an outer scope", name);
-					}
-					else
-					{
-						// Loop through all the values
-						string[] values = ReadListAttribute(element, "Values", Convert.ToChar(separator));
-						foreach (string value in values)
-						{
-							ScopedProperties[^1][name] = value;
-							await readContentsAsync(element);
-						}
-					}
-				}
-			}
-			LeaveScope();
-		}
-
-		/// <summary>
-		/// Reads an "Expand" element 
-		/// </summary>
-		/// <param name="element">Xml element to read the definition from</param>
-		/// <param name="readContentsAsync">Delegate to read the contents of the element, if the condition evaluates to true</param>
-		async Task ReadExpandAsync(BgScriptElement element, Func<BgScriptElement, Task> readContentsAsync)
-		{
-			if (await EvaluateConditionAsync(element))
-			{
-				string name = ReadAttribute(element, "Name");
-				if (ValidateName(element, name))
-				{
-					BgScriptMacro? macro;
-					if (!_macroNameToDefinition.TryGetValue(name, out macro))
-					{
-						LogError(element, "Macro '{0}' does not exist", name);
-					}
-					else
-					{
-						// Parse the argument list
-						string[] arguments = new string[macro.ArgumentNameToIndex.Count];
-						foreach (XmlAttribute? attribute in element.Attributes)
-						{
-							if (attribute != null && attribute.Name != "Name" && attribute.Name != "If")
-							{
-								int index;
-								if (macro.ArgumentNameToIndex.TryGetValue(attribute.Name, out index))
-								{
-									arguments[index] = ExpandProperties(element, attribute.Value);
-								}
-								else
-								{
-									LogWarning(element, "Macro '{0}' does not take an argument '{1}'", name, attribute.Name);
-								}
-							}
-						}
-
-						// Make sure none of the required arguments are missing
-						bool bHasMissingArguments = false;
-						for (int idx = 0; idx < macro.NumRequiredArguments; idx++)
-						{
-							if (arguments[idx] == null)
-							{
-								LogWarning(element, "Macro '{0}' is missing argument '{1}'", macro.Name, macro.ArgumentNameToIndex.First(x => x.Value == idx).Key);
-								bHasMissingArguments = true;
-							}
-						}
-
-						// Expand the function
-						if (!bHasMissingArguments)
-						{
-							EnterScope();
-							foreach (KeyValuePair<string, int> pair in macro.ArgumentNameToIndex)
-							{
-								ScopedProperties[^1][pair.Key] = arguments[pair.Value] ?? "";
-							}
-							foreach (BgScriptElement macroElement in macro.Elements)
-							{
-								await readContentsAsync(macroElement);
-							}
-							LeaveScope();
-						}
-					}
-				}
-			}
-		}
-
-		/// <summary>
-		/// Reads a task definition from the given element, and add it to the given list
-		/// </summary>
-		/// <param name="element">Xml element to read the definition from</param>
-		protected abstract Task ReadTaskAsync(BgScriptElement element);
-
-		/// <summary>
-		/// Reads the definition for an email notifier
-		/// </summary>
-		/// <param name="element">Xml element to read the definition from</param>
-		protected abstract Task ReadNotifierAsync(BgScriptElement element);
-
-		/// <summary>
-		/// Reads a graph annotation
-		/// </summary>
-		/// <param name="element">Xml element to read the definition from</param>
-		protected abstract Task ReadAnnotationAsync(BgScriptElement element);
-
-		/// <summary>
-		/// Reads a warning from the given element, evaluates the condition on it, and writes it to the log if the condition passes.
-		/// </summary>
-		/// <param name="element">Xml element to read the definition from</param>
-		/// <param name="level">The diagnostic event type</param>
-		protected abstract Task ReadDiagnosticAsync(BgScriptElement element, LogLevel level);
-
-		/// <summary>
-		/// Reads an object name from its defining element. Outputs an error if the name is missing.
-		/// </summary>
-		/// <param name="element">Element to read the name for</param>
-		/// <param name="name">Output variable to receive the name of the object</param>
-		/// <returns>True if the object had a valid name (assigned to the Name variable), false if the name was invalid or missing.</returns>
-		protected bool TryReadObjectName(BgScriptElement element, [NotNullWhen(true)] out string? name)
-		{
-			// Check the name attribute is present
-			if (!element.HasAttribute("Name"))
-			{
-				LogError(element, "Missing 'Name' attribute");
-				name = null;
-				return false;
-			}
-
-			// Get the value of it, strip any leading or trailing whitespace, and make sure it's not empty
-			string value = ReadAttribute(element, "Name");
-			if (!ValidateName(element, value))
-			{
-				name = null;
-				return false;
-			}
-
-			// Return it
-			name = value;
-			return true;
-		}
-
-		/// <summary>
-		/// Checks that the given name is valid syntax
-		/// </summary>
-		/// <param name="element">The element that contains the name</param>
-		/// <param name="name">The name to check</param>
-		/// <returns>True if the name is valid</returns>
-		protected bool ValidateName(BgScriptElement element, string name)
-		{
-			// Check it's not empty
-			if (name.Length == 0)
-			{
-				LogError(element, "Name is empty");
-				return false;
-			}
-
-			// Check there are no invalid characters
-			for (int idx = 0; idx < name.Length; idx++)
-			{
-				if (idx > 0 && name[idx] == ' ' && name[idx - 1] == ' ')
-				{
-					LogError(element, "Consecutive spaces in object name - '{0}'", name);
-					return false;
-				}
-				if (Char.IsControl(name[idx]) || BgScriptSchema.IllegalNameCharacters.IndexOf(name[idx]) != -1)
-				{
-					LogError(element, "Invalid character in object name - '{0}'", name[idx]);
-					return false;
-				}
-			}
-			return true;
-		}
-
-		/// <summary>
-		/// Constructs a regex from a regex string and returns it
-		/// </summary>
-		/// <param name="element">The element that contains the regex</param>
-		/// <param name="regex">The pattern to construct</param>
-		/// <returns>The regex if is valid, otherwise null</returns>
-		protected Regex? ParseRegex(BgScriptElement element, string regex)
-		{
-			if (regex.Length == 0)
-			{
-				LogError(element, "Regex is empty");
-				return null;
-			}
-			try
-			{
-				return new Regex(regex);
-			}
-			catch (ArgumentException invalidRegex)
-			{
-				LogError(element, "Could not construct the Regex, reason: {0}", invalidRegex.Message);
-				return null;
-			}
-		}
-
-		/// <summary>
-		/// Expands any properties and reads an attribute.
-		/// </summary>
-		/// <param name="element">Element to read the attribute from</param>
-		/// <param name="name">Name of the attribute</param>
-		/// <returns>Array of names, with all leading and trailing whitespace removed</returns>
-		protected string ReadAttribute(BgScriptElement element, string name)
-		{
-			return ExpandProperties(element, element.GetAttribute(name));
-		}
-
-		/// <summary>
-		/// Expands any properties and reads a list of strings from an attribute, separated by semi-colon characters
-		/// </summary>
-		/// <param name="element"></param>
-		/// <param name="name"></param>
-		/// <param name="separator"></param>
-		/// <returns>Array of names, with all leading and trailing whitespace removed</returns>
-		protected string[] ReadListAttribute(BgScriptElement element, string name, char separator = ';')
-		{
-			string value = ReadAttribute(element, name);
-			return value.Split(new char[] { separator }).Select(x => x.Trim()).Where(x => x.Length > 0).ToArray();
-		}
-
-		/// <summary>
-		/// Parse a map of annotations
-		/// </summary>
-		/// <param name="element"></param>
-		/// <param name="name"></param>
-		/// <returns></returns>
-		protected Dictionary<string, string> ReadAnnotationsAttribute(BgScriptElement element, string name)
-		{
-			string[] pairs = ReadListAttribute(element, name);
-			return ParseAnnotations(element, pairs);
-		}
-
-		/// <summary>
-		/// Parse a map of annotations
-		/// </summary>
-		/// <param name="element"></param>
-		/// <param name="pairs"></param>
-		/// <returns></returns>
-		private Dictionary<string, string> ParseAnnotations(BgScriptElement element, string[] pairs)
-		{
-			// Find the annotations to apply
-			Dictionary<string, string> pairMap = new Dictionary<string, string>(StringComparer.OrdinalIgnoreCase);
-			foreach (string pair in pairs)
-			{
-				if (!String.IsNullOrWhiteSpace(pair))
-				{
-					int idx = pair.IndexOf('=');
-					if (idx < 0)
-					{
-						LogError(element, "Invalid annotation '{0}'", pair);
-						continue;
-					}
-
-					string key = pair.Substring(0, idx).Trim();
-					if (!Regex.IsMatch(key, @"[a-zA-Z0-9_\.]+"))
-					{
-						LogError(element, "Invalid annotation key '{0}'", pair);
-						continue;
-					}
-					if (pairMap.ContainsKey(key))
-					{
-						LogError(element, "Annotation key '{0}' was specified twice", key);
-						continue;
-					}
-
-					pairMap.Add(key, pair.Substring(idx + 1).Trim());
-				}
-			}
-			return pairMap;
-		}
-
-		/// <summary>
-		/// Reads an attribute from the given XML element, expands any properties in it, and parses it as a boolean.
-		/// </summary>
-		/// <param name="element">Element to read the attribute from</param>
-		/// <param name="name">Name of the attribute</param>
-		/// <param name="bDefaultValue">Default value if the attribute is missing</param>
-		/// <returns>The value of the attribute field</returns>
-		protected bool ReadBooleanAttribute(BgScriptElement element, string name, bool bDefaultValue)
-		{
-			bool bResult = bDefaultValue;
-			if (element.HasAttribute(name))
-			{
-				string value = ReadAttribute(element, name).Trim();
-				if (value.Equals("true", StringComparison.InvariantCultureIgnoreCase))
-				{
-					bResult = true;
-				}
-				else if (value.Equals("false", StringComparison.InvariantCultureIgnoreCase))
-				{
-					bResult = false;
-				}
-				else
-				{
-					LogError(element, "Invalid boolean value '{0}' - expected 'true' or 'false'", value);
-				}
-			}
-			return bResult;
-		}
-
-		/// <summary>
-		/// Reads an attribute from the given XML element, expands any properties in it, and parses it as an integer.
-		/// </summary>
-		/// <param name="element">Element to read the attribute from</param>
-		/// <param name="name">Name of the attribute</param>
-		/// <param name="defaultValue">Default value for the integer, if the attribute is missing</param>
-		/// <returns>The value of the attribute field</returns>
-		protected int ReadIntegerAttribute(BgScriptElement element, string name, int defaultValue)
-		{
-			int result = defaultValue;
-			if (element.HasAttribute(name))
-			{
-				string value = ReadAttribute(element, name).Trim();
-
-				int intValue;
-				if (Int32.TryParse(value, out intValue))
-				{
-					result = intValue;
-				}
-				else
-				{
-					LogError(element, "Invalid integer value '{0}'", value);
-				}
-			}
-			return result;
-		}
-
-		/// <summary>
-		/// Reads an attribute from the given XML element, expands any properties in it, and parses it as an enum of the given type.
-		/// </summary>
-		/// <typeparam name="T">The enum type to parse the attribute as</typeparam>
-		/// <param name="element">Element to read the attribute from</param>
-		/// <param name="name">Name of the attribute</param>
-		/// <param name="defaultValue">Default value for the enum, if the attribute is missing</param>
-		/// <returns>The value of the attribute field</returns>
-		protected T ReadEnumAttribute<T>(BgScriptElement element, string name, T defaultValue) where T : struct
-		{
-			T result = defaultValue;
-			if (element.HasAttribute(name))
-			{
-				string value = ReadAttribute(element, name).Trim();
-
-				T enumValue;
-				if (Enum.TryParse(value, true, out enumValue))
-				{
-					result = enumValue;
-				}
-				else
-				{
-					LogError(element, "Invalid value '{0}' - expected {1}", value, String.Join("/", Enum.GetNames(typeof(T))));
-				}
-			}
-			return result;
-		}
-
-		/// <summary>
-		/// Outputs an error message to the log and increments the number of errors, referencing the file and line number of the element that caused it.
-		/// </summary>
-		/// <param name="element">The script element causing the error</param>
-		/// <param name="format">Standard String.Format()-style format string</param>
-		/// <param name="args">Optional arguments</param>
-		protected void LogError(BgScriptElement element, string format, params object[] args)
-		{
-			Logger.LogScriptError(element.Location, format, args);
-			NumErrors++;
-		}
-
-		/// <summary>
-		/// Outputs a warning message to the log and increments the number of errors, referencing the file and line number of the element that caused it.
-		/// </summary>
-		/// <param name="element">The script element causing the error</param>
-		/// <param name="format">Standard String.Format()-style format string</param>
-		/// <param name="args">Optional arguments</param>
-		protected void LogWarning(BgScriptElement element, string format, params object[] args)
-		{
-			Logger.LogScriptWarning(element.Location, format, args);
-		}
-
-		/// <summary>
-		/// Evaluates the (optional) conditional expression on a given XML element via the If="..." attribute, and returns true if the element is enabled.
-		/// </summary>
-		/// <param name="element">The element to check</param>
-		/// <returns>True if the element's condition evaluates to true (or doesn't have a conditional expression), false otherwise</returns>
-		protected async Task<bool> EvaluateConditionAsync(BgScriptElement element)
-		{
-			// Check if the element has a conditional attribute
-			const string AttributeName = "If";
-			if (!element.HasAttribute(AttributeName))
-			{
-				return true;
-			}
-
-			// If it does, try to evaluate it.
-			try
-			{
-				string text = ExpandProperties(element, element.GetAttribute("If"));
-				return await BgCondition.EvaluateAsync(text, Context);
-			}
-			catch (BgConditionException ex)
-			{
-				LogError(element, "Error in condition: {0}", ex.Message);
-				return false;
-			}
-		}
-
-		/// <summary>
-		/// Expand all the property references (of the form $(PropertyName)) in a string.
-		/// </summary>
-		/// <param name="element">The element containing the string. Used for diagnostic messages.</param>
-		/// <param name="text">The input string to expand properties in</param>
-		/// <returns>The expanded string</returns>
-		protected string ExpandProperties(BgScriptElement element, string text)
-		{
-			string result = text;
-			// Iterate in reverse order to handle cases where there are nested expansions like $(Outer$(Inner))
-			for (int idx = result.LastIndexOf("$("); idx != -1; idx = result.LastIndexOf("$(", idx, idx + 1))
-			{
-				// Find the end of the variable name
-				int endIdx = result.IndexOf(')', idx + 2);
-				if (endIdx == -1)
-				{
-					break;
-				}
-
-				// Extract the variable name from the string
-				string name = result.Substring(idx + 2, endIdx - (idx + 2));
-
-				// Find the value for it, either from the dictionary or the environment block
-				string? value;
-				if (!TryGetPropertyValue(name, out value))
-				{
-					LogWarning(element, "Property '{0}' is not defined", name);
-					value = "";
-				}
-
-				// Check if we've got a value for this variable
-				if (value != null)
-				{
-					// Replace the variable, or skip past it
-					result = result.Substring(0, idx) + value + result.Substring(endIdx + 1);
-				}
-			}
-			return result;
-		}
-	}
-
-	/// <summary>
-	/// Overridden version of <see cref="BgNode"/> which contains a list of tasks
-	/// </summary>
-	class BgScriptNode : BgNode
-	{
-		/// <summary>
-		/// List of tasks to execute
-		/// </summary>
-		public List<BgTask> Tasks { get; } = new List<BgTask>();
-
-		/// <summary>
-		/// Constructor
-		/// </summary>
-		public BgScriptNode(string name, IReadOnlyList<BgNodeOutput> inputs, IReadOnlyList<string> outputNames, IReadOnlyList<BgNode> inputDependencies, IReadOnlyList<BgNode> orderDependencies, IReadOnlyList<FileReference> requiredTokens)
-			: base(name, inputs, outputNames, inputDependencies, orderDependencies, requiredTokens)
-		{
-		}
-	}
-
-	/// <summary>
-	/// Reader for build graph definitions. Instanced to contain temporary state; public interface is through ScriptReader.TryRead().
-	/// </summary>
-	public class BgScriptReader : BgScriptReaderBase
-	{
-		/// <summary>
-		/// The current graph
-		/// </summary>
-		readonly BgGraph _graph = new BgGraph();
-
-		/// <summary>
-		/// Arguments for evaluating the graph
-		/// </summary>
-		readonly Dictionary<string, string> _arguments;
-
-		/// <summary>
-		/// The name of the node if only a single node is going to be built, otherwise null.
-		/// </summary>
-		readonly string? _singleNodeName;
-
-		BgAgent? _enclosingAgent;
-		BgScriptNode? _enclosingNode;
-
-		/// <summary>
-		/// Private constructor. Use ScriptReader.TryRead() to read a script file.
-		/// </summary>
-		/// <param name="context">Context object</param>
-		/// <param name="defaultProperties">Default properties available to the script</param>
-		/// <param name="arguments">Arguments passed in to the graph on the command line</param>
-		/// <param name="singleNodeName">If a single node will be processed, the name of that node.</param>
-		/// <param name="schema">Schema for the script</param>
-		/// <param name="logger">Logger for diagnostic messages</param>
-		private BgScriptReader(IBgScriptReaderContext context, IDictionary<string, string> defaultProperties, IReadOnlyDictionary<string, string> arguments, string? singleNodeName, BgScriptSchema schema, ILogger logger)
-			: base(context, schema, logger)
-		{
-			_arguments = new Dictionary<string, string>(arguments, StringComparer.OrdinalIgnoreCase);
-			_singleNodeName = singleNodeName;
-
-			foreach (KeyValuePair<string, string> pair in defaultProperties)
-			{
-				SetPropertyValue(null!, pair.Key, pair.Value);
-			}
-		}
-
-		/// <summary>
-		/// Try to read a script file from the given file.
-		/// </summary>
-		/// <param name="context">Supplies context about the parse</param>
-		/// <param name="file">File to read from</param>
-		/// <param name="arguments">Arguments passed in to the graph on the command line</param>
-		/// <param name="defaultProperties">Default properties available to the script</param>
-		/// <param name="schema">Schema for the script</param>
-		/// <param name="logger">Logger for output messages</param>
-		/// <param name="singleNodeName">If a single node will be processed, the name of that node.</param>
-		/// <returns>True if the graph was read, false if there were errors</returns>
-		public static async Task<BgGraph?> ReadAsync(IBgScriptReaderContext context, string file, Dictionary<string, string> arguments, Dictionary<string, string> defaultProperties, BgScriptSchema schema, ILogger logger, string? singleNodeName = null)
-		{
-			// Read the file and build the graph
-			BgScriptReader reader = new BgScriptReader(context, defaultProperties, arguments, singleNodeName, schema, logger);
-			if (!await reader.TryReadAsync(file) || reader.NumErrors > 0)
-			{
-				return null;
-			}
-
-			// Make sure all the arguments were valid
-			HashSet<string> validArgumentNames = new HashSet<string>(reader._graph.Options.Select(x => x.Name), StringComparer.OrdinalIgnoreCase);
-			validArgumentNames.Add("PreflightChange");
-
-			foreach (string argumentName in arguments.Keys)
-			{
-				if (!validArgumentNames.Contains(argumentName))
-				{
-					logger.LogWarning("Unknown argument '{ArgumentName}' for '{Script}'", argumentName, context.GetNativePath(file));
-				}
-			}
-
-			// Return the constructed graph
-			return reader._graph;
-		}
-
-		/// <summary>
-		/// Reads the definition of a graph option; a parameter which can be set by the user on the command-line or via an environment variable.
-		/// </summary>
-		/// <param name="element">Xml element to read the definition from</param>
-		protected override async Task ReadOptionAsync(BgScriptElement element)
-		{
-			if (await EvaluateConditionAsync(element))
-			{
-				string name = ReadAttribute(element, "Name");
-				if (ValidateName(element, name))
-				{
-					// Make sure we're at global scope
-					if (ScopedProperties.Count > 1)
-					{
-						throw new Exception("Incorrect scope depth for reading option settings");
-					}
-
-					// Check if the property already exists. If it does, we don't need to register it as an option.
-					string? existingValue;
-					if (TryGetPropertyValue(name, out existingValue) && existingValue != null)
-					{
-						// If there's a restriction on this definition, check it matches
-						string restrict = ReadAttribute(element, "Restrict");
-						if (!String.IsNullOrEmpty(restrict) && !Regex.IsMatch(existingValue, "^" + restrict + "$", RegexOptions.IgnoreCase))
-						{
-							LogError(element, "'{0} is already set to '{1}', which does not match the given restriction ('{2}')", name, existingValue, restrict);
-						}
-					}
-					else
-					{
-						// Create a new option object to store the settings
-						string description = ReadAttribute(element, "Description");
-						string defaultValue = ReadAttribute(element, "DefaultValue");
-						BgOption option = new BgOption(name, description, defaultValue);
-						_graph.Options.Add(option);
-
-						// Get the value of this property
-						string? value;
-						if (!_arguments.TryGetValue(name, out value))
-						{
-							value = option.DefaultValue;
-						}
-						SetPropertyValue(element, name, value);
-
-						// If there's a restriction on it, check it's valid
-						string restrict = ReadAttribute(element, "Restrict");
-						if (!String.IsNullOrEmpty(restrict))
-						{
-							string pattern = "^(" + restrict + ")$";
-							if (!Regex.IsMatch(value, pattern, RegexOptions.IgnoreCase))
-							{
-								LogError(element, "'{0}' is not a valid value for '{1}' (required: '{2}')", value, name, restrict);
-							}
-							if (option.DefaultValue != value && !Regex.IsMatch(option.DefaultValue, pattern, RegexOptions.IgnoreCase))
-							{
-								LogError(element, "Default value '{0}' is not valid for '{1}' (required: '{2}')", option.DefaultValue, name, restrict);
-							}
-						}
-					}
-				}
-			}
-		}
-
-		/// <summary>
-		/// Reads a property assignment.
-		/// </summary>
-		/// <param name="element">Xml element to read the definition from</param>
-		protected override async Task ReadPropertyAsync(BgScriptElement element)
-		{
-			if (await EvaluateConditionAsync(element))
-			{
-				string name = ReadAttribute(element, "Name");
-				if (ValidateName(element, name))
-				{
-					string value = ReadAttribute(element, "Value");
-					if (element.HasChildNodes)
-					{
-						// Read the element content, and append each line to the value as a semicolon delimited list
-						StringBuilder builder = new StringBuilder(value);
-						foreach (string line in element.InnerText.Split('\n'))
-						{
-							string trimLine = ExpandProperties(element, line.Trim());
-							if (trimLine.Length > 0)
-							{
-								if (builder.Length > 0)
-								{
-									builder.Append(";");
-								}
-								builder.Append(trimLine);
-							}
-						}
-						value = builder.ToString();
-					}
-					SetPropertyValue(element, name, value);
-				}
-			}
-		}
-
-		/// <summary>
-		/// Reads a Regex assignment.
-		/// </summary>
-		/// <param name="element">Xml element to read the definition from</param>
-		protected override async Task ReadRegexAsync(BgScriptElement element)
-		{
-			if (await EvaluateConditionAsync(element))
-			{
-				// Get the pattern
-				string regexString = ReadAttribute(element, "Pattern");
-
-				// Make sure its a valid regex.
-				Regex? regexValue = ParseRegex(element, regexString);
-				if (regexValue != null)
-				{
-					// read the names in 
-					string[] captureNames = ReadListAttribute(element, "Capture");
-
-					// get number of groups we passed in
-					int[] groupNumbers = regexValue.GetGroupNumbers();
-
-					// make sure the number of property names is the same as the number of match groups
-					// this includes the entire string match group as [0], so don't count that one.
-					if (captureNames.Length != groupNumbers.Count() - 1)
-					{
-						LogError(element, "MatchGroup count: {0} does not match the number of names specified: {1}", groupNumbers.Count() - 1, captureNames.Length);
-					}
-					else
-					{
-						// apply the regex to the value
-						string input = ReadAttribute(element, "Input");
-						Match match = regexValue.Match(input);
-
-						bool optional = await BgCondition.EvaluateAsync(ReadAttribute(element, "Optional"), Context);
-						if (!match.Success)
-						{
-							if (!optional)
-							{
-								LogError(element, "Regex {0} did not find a match against input string {1}", regexString, input);
-							}
-						}
-						else
-						{
-							// assign each property to the group it matches, skip over [0]
-							for (int matchIdx = 1; matchIdx < groupNumbers.Count(); matchIdx++)
-							{
-								SetPropertyValue(element, captureNames[matchIdx - 1], match.Groups[matchIdx].Value);
-							}
-						}
-					}
-				}
-			}
-		}
-
-		/// <summary>
-		/// Reads a property assignment from an environment variable.
-		/// </summary>
-		/// <param name="element">Xml element to read the definition from</param>
-		protected override async Task ReadEnvVarAsync(BgScriptElement element)
-		{
-			if (await EvaluateConditionAsync(element))
-			{
-				string name = ReadAttribute(element, "Name");
-				if (ValidateName(element, name))
-				{
-					string value = Environment.GetEnvironmentVariable(name) ?? "";
-					SetPropertyValue(element, name, value);
-				}
-			}
-		}
-
-		/// <summary>
-		/// Reads the definition for an agent.
-		/// </summary>
-		/// <param name="element">Xml element to read the definition from</param>
-		protected override async Task ReadAgentAsync(BgScriptElement element)
-		{
-			string? name;
-			if (await EvaluateConditionAsync(element) && TryReadObjectName(element, out name))
-			{
-				// Read the valid agent types. This may be omitted if we're continuing an existing agent.
-				string[] types = ReadListAttribute(element, "Type");
-
-				// Create the agent object, or continue an existing one
-				BgAgent? agent;
-				if (_graph.NameToAgent.TryGetValue(name, out agent))
-				{
-					if (types.Length > 0 && agent.PossibleTypes.Length > 0)
-					{
-						if (types.Length != agent.PossibleTypes.Length || !types.SequenceEqual(agent.PossibleTypes, StringComparer.InvariantCultureIgnoreCase))
-						{
-							LogError(element, "Agent types ({0}) were different than previous agent definition with types ({1}). Must either be empty or match exactly.", String.Join(",", types), String.Join(",", agent.PossibleTypes));
-						}
-					}
-				}
-				else
-				{
-					if (types.Length == 0)
-					{
-						LogError(element, "Missing type for agent '{0}'", name);
-					}
-					agent = new BgAgent(name, types);
-					_graph.NameToAgent.Add(name, agent);
-					_graph.Agents.Add(agent);
-				}
-
-				// Process all the child elements.
-				_enclosingAgent = agent;
-				await ReadAgentBodyAsync(element);
-				_enclosingAgent = null;
-			}
-		}
-
-		/// <summary>
 		/// Reads the definition for an aggregate
 		/// </summary>
 		/// <param name="element">Xml element to read the definition from</param>
-		protected override async Task ReadAggregateAsync(BgScriptElement element)
+		async Task ReadAggregateAsync(BgScriptElement element)
 		{
 			string? name;
 			if (await EvaluateConditionAsync(element) && TryReadObjectName(element, out name) && CheckNameIsUnique(element, name))
@@ -1722,7 +1102,7 @@
 		/// Reads the definition for a report
 		/// </summary>
 		/// <param name="element">Xml element to read the definition from</param>
-		protected override async Task ReadReportAsync(BgScriptElement element)
+		async Task ReadReportAsync(BgScriptElement element)
 		{
 			string? name;
 			if (await EvaluateConditionAsync(element) && TryReadObjectName(element, out name) && CheckNameIsUnique(element, name))
@@ -1743,7 +1123,7 @@
 		/// Reads the definition for a badge
 		/// </summary>
 		/// <param name="element">Xml element to read the definition from</param>
-		protected override async Task ReadBadgeAsync(BgScriptElement element)
+		async Task ReadBadgeAsync(BgScriptElement element)
 		{
 			string? name;
 			if (await EvaluateConditionAsync(element) && TryReadObjectName(element, out name))
@@ -1771,7 +1151,7 @@
 		/// Reads the definition for a label
 		/// </summary>
 		/// <param name="element">Xml element to read the definition from</param>
-		protected override async Task ReadLabelAsync(BgScriptElement element)
+		async Task ReadLabelAsync(BgScriptElement element)
 		{
 			if (await EvaluateConditionAsync(element))
 			{
@@ -1817,7 +1197,7 @@
 		/// Reads the definition for a node, and adds it to the given agent
 		/// </summary>
 		/// <param name="element">Xml element to read the definition from</param>
-		protected override async Task ReadNodeAsync(BgScriptElement element)
+		async Task ReadNodeAsync(BgScriptElement element)
 		{
 			string? name;
 			if (await EvaluateConditionAsync(element) && TryReadObjectName(element, out name))
@@ -1928,11 +1308,70 @@
 		}
 
 		/// <summary>
+		/// Reads the contents of a node element
+		/// </summary>
+		/// <param name="element">Xml element to read the definition from</param>
+		protected async Task ReadNodeBodyAsync(XmlElement element)
+		{
+			EnterScope();
+			foreach (BgScriptElement childElement in element.ChildNodes.OfType<BgScriptElement>())
+			{
+				switch (childElement.Name)
+				{
+					case "Property":
+						await ReadPropertyAsync(childElement);
+						break;
+					case "Regex":
+						await ReadRegexAsync(childElement);
+						break;
+					case "Trace":
+						await ReadDiagnosticAsync(childElement, LogLevel.Information);
+						break;
+					case "Warning":
+						await ReadDiagnosticAsync(childElement, LogLevel.Warning);
+						break;
+					case "Error":
+						await ReadDiagnosticAsync(childElement, LogLevel.Error);
+						break;
+					case "Do":
+						await ReadBlockAsync(childElement, ReadNodeBodyAsync);
+						break;
+					case "Switch":
+						await ReadSwitchAsync(childElement, ReadNodeBodyAsync);
+						break;
+					case "ForEach":
+						await ReadForEachAsync(childElement, ReadNodeBodyAsync);
+						break;
+					case "Expand":
+						await ReadExpandAsync(childElement, ReadNodeBodyAsync);
+						break;
+					default:
+						await ReadTaskAsync(childElement);
+						break;
+				}
+			}
+			LeaveScope();
+		}
+
+		/// <summary>
+		/// Reads a block element
+		/// </summary>
+		/// <param name="element">Xml element to read the definition from</param>
+		/// <param name="readContentsAsync">Delegate to read the contents of the element, if the condition evaluates to true</param>
+		async Task ReadBlockAsync(BgScriptElement element, Func<BgScriptElement, Task> readContentsAsync)
+		{
+			if (await EvaluateConditionAsync(element))
+			{
+				await readContentsAsync(element);
+			}
+		}
+
+		/// <summary>
 		/// Reads a "Switch" element 
 		/// </summary>
 		/// <param name="element">Xml element to read the definition from</param>
 		/// <param name="readContentsAsync">Delegate to read the contents of the element, if the condition evaluates to true</param>
-		protected override async Task ReadSwitchAsync(BgScriptElement element, Func<BgScriptElement, Task> readContentsAsync)
+		async Task ReadSwitchAsync(BgScriptElement element, Func<BgScriptElement, Task> readContentsAsync)
 		{
 			foreach (BgScriptElement childElement in element.ChildNodes.OfType<BgScriptElement>())
 			{
@@ -1943,3 +1382,112 @@
 				}
 			}
 		}
+
+		/// <summary>
+		/// Reads a "ForEach" element 
+		/// </summary>
+		/// <param name="element">Xml element to read the definition from</param>
+		/// <param name="readContentsAsync">Delegate to read the contents of the element, if the condition evaluates to true</param>
+		async Task ReadForEachAsync(BgScriptElement element, Func<BgScriptElement, Task> readContentsAsync)
+		{
+			EnterScope();
+			if (await EvaluateConditionAsync(element))
+			{
+				string name = ReadAttribute(element, "Name");
+				string separator = ReadAttribute(element, "Separator");
+				if (separator.Length > 1)
+				{
+					LogWarning(element, "Node {0}'s Separator attribute is more than one character ({1}). Defaulting to ;", name, separator);
+					separator = ";";
+				}
+				if (String.IsNullOrEmpty(separator))
+				{
+					separator = ";";
+				}
+				if (ValidateName(element, name))
+				{
+					if (ScopedProperties.Any(x => x.ContainsKey(name)))
+					{
+						LogError(element, "Loop variable '{0}' already exists as a local property in an outer scope", name);
+					}
+					else
+					{
+						// Loop through all the values
+						string[] values = ReadListAttribute(element, "Values", Convert.ToChar(separator));
+						foreach (string value in values)
+						{
+							ScopedProperties[^1][name] = value;
+							await readContentsAsync(element);
+						}
+					}
+				}
+			}
+			LeaveScope();
+		}
+
+		/// <summary>
+		/// Reads an "Expand" element 
+		/// </summary>
+		/// <param name="element">Xml element to read the definition from</param>
+		/// <param name="readContentsAsync">Delegate to read the contents of the element, if the condition evaluates to true</param>
+		async Task ReadExpandAsync(BgScriptElement element, Func<BgScriptElement, Task> readContentsAsync)
+		{
+			if (await EvaluateConditionAsync(element))
+			{
+				string name = ReadAttribute(element, "Name");
+				if (ValidateName(element, name))
+				{
+					BgScriptMacro? macro;
+					if (!_macroNameToDefinition.TryGetValue(name, out macro))
+					{
+						LogError(element, "Macro '{0}' does not exist", name);
+					}
+					else
+					{
+						// Parse the argument list
+						string[] arguments = new string[macro.ArgumentNameToIndex.Count];
+						foreach (XmlAttribute? attribute in element.Attributes)
+						{
+							if (attribute != null && attribute.Name != "Name" && attribute.Name != "If")
+							{
+								int index;
+								if (macro.ArgumentNameToIndex.TryGetValue(attribute.Name, out index))
+								{
+									arguments[index] = ExpandProperties(element, attribute.Value);
+								}
+								else
+								{
+									LogWarning(element, "Macro '{0}' does not take an argument '{1}'", name, attribute.Name);
+								}
+							}
+						}
+
+						// Make sure none of the required arguments are missing
+						bool bHasMissingArguments = false;
+						for (int idx = 0; idx < macro.NumRequiredArguments; idx++)
+						{
+							if (arguments[idx] == null)
+							{
+								LogWarning(element, "Macro '{0}' is missing argument '{1}'", macro.Name, macro.ArgumentNameToIndex.First(x => x.Value == idx).Key);
+								bHasMissingArguments = true;
+							}
+						}
+
+						// Expand the function
+						if (!bHasMissingArguments)
+						{
+							EnterScope();
+							foreach (KeyValuePair<string, int> pair in macro.ArgumentNameToIndex)
+							{
+								ScopedProperties[^1][pair.Key] = arguments[pair.Value] ?? "";
+							}
+							foreach (BgScriptElement macroElement in macro.Elements)
+							{
+								await readContentsAsync(macroElement);
+							}
+							LeaveScope();
+						}
+					}
+				}
+			}
+		}
@@ -1948,7 +1496,7 @@
 		/// Reads a task definition from the given element, and add it to the given list
 		/// </summary>
 		/// <param name="element">Xml element to read the definition from</param>
-		protected override async Task ReadTaskAsync(BgScriptElement element)
+		async Task ReadTaskAsync(BgScriptElement element)
 		{
 			// If we're running a single node and this element's parent isn't the single node to run, ignore the error and return.
 			if (!String.IsNullOrWhiteSpace(_singleNodeName) && _enclosingNode!.Name != _singleNodeName)
@@ -1975,7 +1523,7 @@
 		/// Reads the definition for an email notifier
 		/// </summary>
 		/// <param name="element">Xml element to read the definition from</param>
-		protected override async Task ReadNotifierAsync(BgScriptElement element)
+		async Task ReadNotifierAsync(BgScriptElement element)
 		{
 			if (await EvaluateConditionAsync(element))
 			{
@@ -1990,3 +1538,92 @@
 
 				// Find the list of targets which are included, and recurse through all their dependencies
 				HashSet<BgNode> nodes = new HashSet<BgNode>();
+				if (targetNames != null)
+				{
+					HashSet<BgNode> targetNodes = ResolveReferences(element, targetNames);
+					foreach (BgNode node in targetNodes)
+					{
+						nodes.Add(node);
+						nodes.UnionWith(node.InputDependencies);
+					}
+				}
+
+				// Add all the individually referenced nodes
+				if (individualNodeNames != null)
+				{
+					HashSet<BgNode> individualNodes = ResolveReferences(element, individualNodeNames);
+					nodes.UnionWith(individualNodes);
+				}
+
+				// Exclude all the exceptions
+				if (exceptNames != null)
+				{
+					HashSet<BgNode> exceptNodes = ResolveReferences(element, exceptNames);
+					nodes.ExceptWith(exceptNodes);
+				}
+
+				// Update all the referenced nodes with the settings
+				foreach (BgNode node in nodes)
+				{
+					if (users != null)
+					{
+						if (bAbsolute)
+						{
+							node.NotifyUsers = new HashSet<string>(users);
+						}
+						else
+						{
+							node.NotifyUsers.UnionWith(users);
+						}
+					}
+					if (submitters != null)
+					{
+						if (bAbsolute)
+						{
+							node.NotifySubmitters = new HashSet<string>(submitters);
+						}
+						else
+						{
+							node.NotifySubmitters.UnionWith(submitters);
+						}
+					}
+					if (bWarnings.HasValue)
+					{
+						node.NotifyOnWarnings = bWarnings.Value;
+					}
+				}
+
+				// Add the users to the list of reports
+				if (reportNames != null)
+				{
+					foreach (string reportName in reportNames)
+					{
+						BgReport? report;
+						if (_graph.NameToReport.TryGetValue(reportName, out report))
+						{
+							report.NotifyUsers.UnionWith(users);
+						}
+						else
+						{
+							LogError(element, "Report '{0}' has not been defined", reportName);
+						}
+					}
+				}
+			}
+		}
+
+		/// <summary>
+		/// Reads a graph annotation
+		/// </summary>
+		/// <param name="element">Xml element to read the definition from</param>
+		async Task ReadAnnotationAsync(BgScriptElement element)
+		{
+			if (await EvaluateConditionAsync(element))
+			{
+				string[] targetNames = ReadListAttribute(element, "Targets");
+				string[] exceptNames = ReadListAttribute(element, "Except");
+				string[] individualNodeNames = ReadListAttribute(element, "Nodes");
+				Dictionary<string, string> annotations = ReadAnnotationsAttribute(element, "Values");
+
+				// Find the list of targets which are included, and recurse through all their dependencies
+				HashSet<BgNode> nodes = new HashSet<BgNode>();
@@ -2017,92 +1654,6 @@
 				// Update all the referenced nodes with the settings
 				foreach (BgNode node in nodes)
 				{
-					if (users != null)
-					{
-						if (bAbsolute)
-						{
-							node.NotifyUsers = new HashSet<string>(users);
-						}
-						else
-						{
-							node.NotifyUsers.UnionWith(users);
-						}
-					}
-					if (submitters != null)
-					{
-						if (bAbsolute)
-						{
-							node.NotifySubmitters = new HashSet<string>(submitters);
-						}
-						else
-						{
-							node.NotifySubmitters.UnionWith(submitters);
-						}
-					}
-					if (bWarnings.HasValue)
-					{
-						node.NotifyOnWarnings = bWarnings.Value;
-					}
-				}
-
-				// Add the users to the list of reports
-				if (reportNames != null)
-				{
-					foreach (string reportName in reportNames)
-					{
-						BgReport? report;
-						if (_graph.NameToReport.TryGetValue(reportName, out report))
-						{
-							report.NotifyUsers.UnionWith(users);
-						}
-						else
-						{
-							LogError(element, "Report '{0}' has not been defined", reportName);
-						}
-					}
-				}
-			}
-		}
-
-		/// <inheritdoc/>
-		protected override async Task ReadAnnotationAsync(BgScriptElement element)
-		{
-			if (await EvaluateConditionAsync(element))
-			{
-				string[] targetNames = ReadListAttribute(element, "Targets");
-				string[] exceptNames = ReadListAttribute(element, "Except");
-				string[] individualNodeNames = ReadListAttribute(element, "Nodes");
-				Dictionary<string, string> annotations = ReadAnnotationsAttribute(element, "Values");
-
-				// Find the list of targets which are included, and recurse through all their dependencies
-				HashSet<BgNode> nodes = new HashSet<BgNode>();
-				if (targetNames != null)
-				{
-					HashSet<BgNode> targetNodes = ResolveReferences(element, targetNames);
-					foreach (BgNode node in targetNodes)
-					{
-						nodes.Add(node);
-						nodes.UnionWith(node.InputDependencies);
-					}
-				}
-
-				// Add all the individually referenced nodes
-				if (individualNodeNames != null)
-				{
-					HashSet<BgNode> individualNodes = ResolveReferences(element, individualNodeNames);
-					nodes.UnionWith(individualNodes);
-				}
-
-				// Exclude all the exceptions
-				if (exceptNames != null)
-				{
-					HashSet<BgNode> exceptNodes = ResolveReferences(element, exceptNames);
-					nodes.ExceptWith(exceptNodes);
-				}
-
-				// Update all the referenced nodes with the settings
-				foreach (BgNode node in nodes)
-				{
 					foreach ((string key, string value) in annotations)
 					{
 						node.Annotations[key] = value;
@@ -2116,13 +1667,13 @@
 		/// </summary>
 		/// <param name="element">Xml element to read the definition from</param>
 		/// <param name="level">The diagnostic event type</param>
-		protected override async Task ReadDiagnosticAsync(BgScriptElement element, LogLevel level)
+		async Task ReadDiagnosticAsync(BgScriptElement element, LogLevel level)
 		{
 			if (await EvaluateConditionAsync(element))
 			{
 				string message = ReadAttribute(element, "Message");
 
-				BgDiagnostic diagnostic = new BgDiagnostic(element.Location.File, element.Location.LineNumber, level, message);
+				BgDiagnostic diagnostic = new BgDiagnostic(element.Location.File.FullName, element.Location.LineNumber, level, message);
 				if (_enclosingNode != null)
 				{
 					_enclosingNode.Diagnostics.Add(diagnostic);
@@ -2139,6 +1690,35 @@
 		}
 
 		/// <summary>
+		/// Reads an object name from its defining element. Outputs an error if the name is missing.
+		/// </summary>
+		/// <param name="element">Element to read the name for</param>
+		/// <param name="name">Output variable to receive the name of the object</param>
+		/// <returns>True if the object had a valid name (assigned to the Name variable), false if the name was invalid or missing.</returns>
+		protected bool TryReadObjectName(BgScriptElement element, [NotNullWhen(true)] out string? name)
+		{
+			// Check the name attribute is present
+			if (!element.HasAttribute("Name"))
+			{
+				LogError(element, "Missing 'Name' attribute");
+				name = null;
+				return false;
+			}
+
+			// Get the value of it, strip any leading or trailing whitespace, and make sure it's not empty
+			string value = ReadAttribute(element, "Name");
+			if (!ValidateName(element, value))
+			{
+				name = null;
+				return false;
+			}
+
+			// Return it
+			name = value;
+			return true;
+		}
+
+		/// <summary>
 		/// Checks that the given name does not already used to refer to a node, and print an error if it is.
 		/// </summary>
 		/// <param name="element">Xml element to read from</param>
@@ -2210,5 +1790,316 @@
 			}
 			return inputs;
 		}
+
+		/// <summary>
+		/// Checks that the given name is valid syntax
+		/// </summary>
+		/// <param name="element">The element that contains the name</param>
+		/// <param name="name">The name to check</param>
+		/// <returns>True if the name is valid</returns>
+		protected bool ValidateName(BgScriptElement element, string name)
+		{
+			// Check it's not empty
+			if (name.Length == 0)
+			{
+				LogError(element, "Name is empty");
+				return false;
+			}
+
+			// Check there are no invalid characters
+			for (int idx = 0; idx < name.Length; idx++)
+			{
+				if (idx > 0 && name[idx] == ' ' && name[idx - 1] == ' ')
+				{
+					LogError(element, "Consecutive spaces in object name - '{0}'", name);
+					return false;
+				}
+				if (Char.IsControl(name[idx]) || BgScriptSchema.IllegalNameCharacters.IndexOf(name[idx]) != -1)
+				{
+					LogError(element, "Invalid character in object name - '{0}'", name[idx]);
+					return false;
+				}
+			}
+			return true;
+		}
+
+		/// <summary>
+		/// Constructs a regex from a regex string and returns it
+		/// </summary>
+		/// <param name="element">The element that contains the regex</param>
+		/// <param name="regex">The pattern to construct</param>
+		/// <returns>The regex if is valid, otherwise null</returns>
+		protected Regex? ParseRegex(BgScriptElement element, string regex)
+		{
+			if (regex.Length == 0)
+			{
+				LogError(element, "Regex is empty");
+				return null;
+			}
+			try
+			{
+				return new Regex(regex);
+			}
+			catch (ArgumentException invalidRegex)
+			{
+				LogError(element, "Could not construct the Regex, reason: {0}", invalidRegex.Message);
+				return null;
+			}
+		}
+
+		/// <summary>
+		/// Expands any properties and reads an attribute.
+		/// </summary>
+		/// <param name="element">Element to read the attribute from</param>
+		/// <param name="name">Name of the attribute</param>
+		/// <returns>Array of names, with all leading and trailing whitespace removed</returns>
+		protected string ReadAttribute(BgScriptElement element, string name)
+		{
+			return ExpandProperties(element, element.GetAttribute(name));
+		}
+
+		/// <summary>
+		/// Expands any properties and reads a list of strings from an attribute, separated by semi-colon characters
+		/// </summary>
+		/// <param name="element"></param>
+		/// <param name="name"></param>
+		/// <param name="separator"></param>
+		/// <returns>Array of names, with all leading and trailing whitespace removed</returns>
+		protected string[] ReadListAttribute(BgScriptElement element, string name, char separator = ';')
+		{
+			string value = ReadAttribute(element, name);
+			return value.Split(new char[] { separator }).Select(x => x.Trim()).Where(x => x.Length > 0).ToArray();
+		}
+
+		/// <summary>
+		/// Parse a map of annotations
+		/// </summary>
+		/// <param name="element"></param>
+		/// <param name="name"></param>
+		/// <returns></returns>
+		protected Dictionary<string, string> ReadAnnotationsAttribute(BgScriptElement element, string name)
+		{
+			string[] pairs = ReadListAttribute(element, name);
+			return ParseAnnotations(element, pairs);
+		}
+
+		/// <summary>
+		/// Parse a map of annotations
+		/// </summary>
+		/// <param name="element"></param>
+		/// <param name="pairs"></param>
+		/// <returns></returns>
+		private Dictionary<string, string> ParseAnnotations(BgScriptElement element, string[] pairs)
+		{
+			// Find the annotations to apply
+			Dictionary<string, string> pairMap = new Dictionary<string, string>(StringComparer.OrdinalIgnoreCase);
+			foreach (string pair in pairs)
+			{
+				if (!String.IsNullOrWhiteSpace(pair))
+				{
+					int idx = pair.IndexOf('=');
+					if (idx < 0)
+					{
+						LogError(element, "Invalid annotation '{0}'", pair);
+						continue;
+					}
+
+					string key = pair.Substring(0, idx).Trim();
+					if (!Regex.IsMatch(key, @"[a-zA-Z0-9_\.]+"))
+					{
+						LogError(element, "Invalid annotation key '{0}'", pair);
+						continue;
+					}
+					if (pairMap.ContainsKey(key))
+					{
+						LogError(element, "Annotation key '{0}' was specified twice", key);
+						continue;
+					}
+
+					pairMap.Add(key, pair.Substring(idx + 1).Trim());
+				}
+			}
+			return pairMap;
+		}
+
+		/// <summary>
+		/// Reads an attribute from the given XML element, expands any properties in it, and parses it as a boolean.
+		/// </summary>
+		/// <param name="element">Element to read the attribute from</param>
+		/// <param name="name">Name of the attribute</param>
+		/// <param name="bDefaultValue">Default value if the attribute is missing</param>
+		/// <returns>The value of the attribute field</returns>
+		protected bool ReadBooleanAttribute(BgScriptElement element, string name, bool bDefaultValue)
+		{
+			bool bResult = bDefaultValue;
+			if (element.HasAttribute(name))
+			{
+				string value = ReadAttribute(element, name).Trim();
+				if (value.Equals("true", StringComparison.InvariantCultureIgnoreCase))
+				{
+					bResult = true;
+				}
+				else if (value.Equals("false", StringComparison.InvariantCultureIgnoreCase))
+				{
+					bResult = false;
+				}
+				else
+				{
+					LogError(element, "Invalid boolean value '{0}' - expected 'true' or 'false'", value);
+				}
+			}
+			return bResult;
+		}
+
+		/// <summary>
+		/// Reads an attribute from the given XML element, expands any properties in it, and parses it as an integer.
+		/// </summary>
+		/// <param name="element">Element to read the attribute from</param>
+		/// <param name="name">Name of the attribute</param>
+		/// <param name="defaultValue">Default value for the integer, if the attribute is missing</param>
+		/// <returns>The value of the attribute field</returns>
+		protected int ReadIntegerAttribute(BgScriptElement element, string name, int defaultValue)
+		{
+			int result = defaultValue;
+			if (element.HasAttribute(name))
+			{
+				string value = ReadAttribute(element, name).Trim();
+
+				int intValue;
+				if (Int32.TryParse(value, out intValue))
+				{
+					result = intValue;
+				}
+				else
+				{
+					LogError(element, "Invalid integer value '{0}'", value);
+				}
+			}
+			return result;
+		}
+
+		/// <summary>
+		/// Reads an attribute from the given XML element, expands any properties in it, and parses it as an enum of the given type.
+		/// </summary>
+		/// <typeparam name="T">The enum type to parse the attribute as</typeparam>
+		/// <param name="element">Element to read the attribute from</param>
+		/// <param name="name">Name of the attribute</param>
+		/// <param name="defaultValue">Default value for the enum, if the attribute is missing</param>
+		/// <returns>The value of the attribute field</returns>
+		protected T ReadEnumAttribute<T>(BgScriptElement element, string name, T defaultValue) where T : struct
+		{
+			T result = defaultValue;
+			if (element.HasAttribute(name))
+			{
+				string value = ReadAttribute(element, name).Trim();
+
+				T enumValue;
+				if (Enum.TryParse(value, true, out enumValue))
+				{
+					result = enumValue;
+				}
+				else
+				{
+					LogError(element, "Invalid value '{0}' - expected {1}", value, String.Join("/", Enum.GetNames(typeof(T))));
+				}
+			}
+			return result;
+		}
+
+		/// <summary>
+		/// Outputs an error message to the log and increments the number of errors, referencing the file and line number of the element that caused it.
+		/// </summary>
+		/// <param name="element">The script element causing the error</param>
+		/// <param name="format">Standard String.Format()-style format string</param>
+		/// <param name="args">Optional arguments</param>
+		protected void LogError(BgScriptElement element, string format, params object[] args)
+		{
+			Logger.LogScriptError(element.Location, format, args);
+			NumErrors++;
+		}
+
+		/// <summary>
+		/// Outputs a warning message to the log and increments the number of errors, referencing the file and line number of the element that caused it.
+		/// </summary>
+		/// <param name="element">The script element causing the error</param>
+		/// <param name="format">Standard String.Format()-style format string</param>
+		/// <param name="args">Optional arguments</param>
+		protected void LogWarning(BgScriptElement element, string format, params object[] args)
+		{
+			Logger.LogScriptWarning(element.Location, format, args);
+		}
+
+		/// <summary>
+		/// Evaluates the (optional) conditional expression on a given XML element via the If="..." attribute, and returns true if the element is enabled.
+		/// </summary>
+		/// <param name="element">The element to check</param>
+		/// <returns>True if the element's condition evaluates to true (or doesn't have a conditional expression), false otherwise</returns>
+		protected async Task<bool> EvaluateConditionAsync(BgScriptElement element)
+		{
+			// Check if the element has a conditional attribute
+			const string AttributeName = "If";
+			if (!element.HasAttribute(AttributeName))
+			{
+				return true;
+			}
+
+			// If it does, try to evaluate it.
+			try
+			{
+				string text = ExpandProperties(element, element.GetAttribute("If"));
+				return await BgCondition.EvaluateAsync(text);
+			}
+			catch (BgConditionException ex)
+			{
+				LogError(element, "Error in condition: {0}", ex.Message);
+				return false;
+			}
+		}
+
+		/// <summary>
+		/// Expand all the property references (of the form $(PropertyName)) in a string.
+		/// </summary>
+		/// <param name="element">The element containing the string. Used for diagnostic messages.</param>
+		/// <param name="text">The input string to expand properties in</param>
+		/// <returns>The expanded string</returns>
+		protected string ExpandProperties(BgScriptElement element, string text)
+		{
+			string result = text;
+			// Iterate in reverse order to handle cases where there are nested expansions like $(Outer$(Inner))
+			for (int idx = result.LastIndexOf("$("); idx != -1; idx = result.LastIndexOf("$(", idx, idx + 1))
+			{
+				// Find the end of the variable name
+				int endIdx = result.IndexOf(')', idx + 2);
+				if (endIdx == -1)
+				{
+					break;
+				}
+
+				// Extract the variable name from the string
+				string name = result.Substring(idx + 2, endIdx - (idx + 2));
+
+				// Find the value for it, either from the dictionary or the environment block
+				string? value;
+				if (!TryGetPropertyValue(name, out value))
+				{
+					LogWarning(element, "Property '{0}' is not defined", name);
+					value = "";
+				}
+
+				// Check if we've got a value for this variable
+				if (value != null)
+				{
+					// Replace the variable, or skip past it
+					result = result.Substring(0, idx) + value + result.Substring(endIdx + 1);
+				}
+			}
+			return result;
+		}
+
+		/// <inheritdoc/>
+		public object GetNativePath(string Path)
+		{
+			return FileReference.Combine(Unreal.RootDirectory, Path).FullName;
+		}
 	}
 }
diff --git a/Engine/Source/Programs/AutomationTool/BuildGraph/BuildGraph.cs b/Engine/Source/Programs/AutomationTool/BuildGraph/BuildGraph.cs
--- a/Engine/Source/Programs/AutomationTool/BuildGraph/BuildGraph.cs
+++ b/Engine/Source/Programs/AutomationTool/BuildGraph/BuildGraph.cs
@@ -103,100 +103,6 @@ namespace AutomationTool
 		}
 	}
 
-	/// <summary>
-	/// Implementation of <see cref="IBgScriptReaderContext"/> which reads files from disk
-	/// </summary>
-	class ScriptReaderFileContext : IBgScriptReaderContext
-	{
-		/// <summary>
-		/// Directory to resolve relative paths to
-		/// </summary>
-		public DirectoryReference RootDirectory { get; set; } = Unreal.RootDirectory;
-
-		/// <inheritdoc/>
-		public object GetNativePath(string Path)
-		{
-			return FileReference.Combine(RootDirectory, Path).FullName;
-		}
-
-		/// <inheritdoc/>
-		public Task<bool> ExistsAsync(string Path)
-		{
-			try
-			{
-				return Task.FromResult(FileReference.Exists(FileReference.Combine(RootDirectory, Path)) || DirectoryReference.Exists(DirectoryReference.Combine(RootDirectory, Path)));
-			}
-			catch
-			{
-				return Task.FromResult(false);
-			}
-		}
-
-		/// <inheritdoc/>
-		public async Task<byte[]> ReadAsync(string Path)
-		{
-			try
-			{
-				FileReference File = FileReference.Combine(RootDirectory, Path);
-				if (FileReference.Exists(File))
-				{
-					return await FileReference.ReadAllBytesAsync(File);
-				}
-			}
-			catch
-			{
-			}
-			return null;
-		}
-
-		/// <inheritdoc/>
-		public Task<string[]> FindAsync(string Pattern)
-		{
-			FileFilter Filter = new FileFilter();
-			Filter.AddRule(Pattern, FileFilterType.Include);
-
-			List<string> Files = Filter.ApplyToDirectory(RootDirectory, true).ConvertAll(x => x.MakeRelativeTo(RootDirectory).Replace('\\', '/'));
-			Files.Sort(StringComparer.OrdinalIgnoreCase);
-			return Task.FromResult(Files.ToArray());
-		}
-
-		/// <inheritdoc/>
-		public string CombinePaths(string basePath, string nextPath)
-		{
-			if (Path.IsPathRooted(nextPath))
-			{
-				return nextPath;
-			}
-
-			List<string> fragments = new List<string>(basePath.Split('/'));
-			fragments.RemoveAt(fragments.Count - 1);
-
-			foreach (string appendFragment in nextPath.Split('/'))
-			{
-				if (appendFragment.Equals(".", StringComparison.Ordinal))
-				{
-					continue;
-				}
-				else if (appendFragment.Equals("..", StringComparison.Ordinal))
-				{
-					if (fragments.Count > 0)
-					{
-						fragments.RemoveAt(fragments.Count - 1);
-					}
-					else
-					{
-						throw new Exception($"Path '{nextPath}' cannot be combined with '{basePath}'");
-					}
-				}
-				else
-				{
-					fragments.Add(appendFragment);
-				}
-			}
-			return String.Join('/', fragments);
-		}
-	}
-
 	/// <summary>
 	/// Tool to execute build automation scripts for UE projects, which can be run locally or in parallel across a build farm (assuming synchronization and resource allocation implemented by a separate system).
 	///
@@ -276,11 +182,6 @@ namespace AutomationTool
 			}
 		}
 
-		/// <summary>
-		/// Context object for reading scripts and evaluating conditions
-		/// </summary>
-		ScriptReaderFileContext Context { get; } = new ScriptReaderFileContext();
-
 		/// <summary>
 		/// Main entry point for the BuildGraph command
 		/// </summary>
@@ -501,14 +402,9 @@ namespace AutomationTool
 
 				// Normalize the script filename
 				FileReference FullScriptFile = FileReference.Combine(Unreal.RootDirectory, ScriptFileName);
-				if (!FullScriptFile.IsUnderDirectory(Context.RootDirectory))
-				{
-					Context.RootDirectory = FullScriptFile.Directory;
-				}
-				ScriptFileName = FullScriptFile.MakeRelativeTo(Context.RootDirectory).Replace('\\', '/');
 
 				// Read the script from disk
-				Graph = BgScriptReader.ReadAsync(Context, ScriptFileName, Arguments, DefaultProperties, Schema, Logger, SingleNodeName).Result;
+				Graph = BgScriptReader.ReadAsync(FullScriptFile, Arguments, DefaultProperties, Schema, Logger, SingleNodeName).Result;
 				if (Graph == null)
 				{
 					return ExitCode.Error_Unknown;
@@ -784,7 +680,7 @@ namespace AutomationTool
 			{
 				BgScriptNodeExecutor executor = new BgScriptNodeExecutor(ScriptNode);
 				NodeToExecutor[Node] = executor;
-				return executor.Bind(NameToTask, TagNameToNodeOutput, Context, Logger);
+				return executor.Bind(NameToTask, TagNameToNodeOutput, Logger);
 			}
 			else if (Node is BgExpressionNode ExpressionNode)
 			{
diff --git a/Engine/Source/Programs/AutomationTool/BuildGraph/Task.cs b/Engine/Source/Programs/AutomationTool/BuildGraph/Task.cs
--- a/Engine/Source/Programs/AutomationTool/BuildGraph/Task.cs
+++ b/Engine/Source/Programs/AutomationTool/BuildGraph/Task.cs
@@ -184,7 +184,7 @@ namespace AutomationTool
 		/// <param name="Prefix">Prefix for metadata entries</param>
 		public virtual void GetTraceMetadata(ITraceSpan Span, string Prefix)
 		{
-			Span.AddMetadata(Prefix + "source.file", SourceLocation.File);
+			Span.AddMetadata(Prefix + "source.file", SourceLocation.File.FullName);
 			Span.AddMetadata(Prefix + "source.line", SourceLocation.LineNumber.ToString());
 		}
 		
@@ -195,7 +195,7 @@ namespace AutomationTool
 		/// <param name="Prefix">Prefix for metadata entries</param>
 		public virtual void GetTraceMetadata(ISpan Span, string Prefix)
 		{
-			Span.SetTag(Prefix + "source.file", SourceLocation.File);
+			Span.SetTag(Prefix + "source.file", SourceLocation.File.FullName);
 			Span.SetTag(Prefix + "source.line", SourceLocation.LineNumber);
 		}
 
