Changing the order of sha1 timestamp servers

comodoca isn't preforming very well resulting in over an hour to sign editor files, changing the order results in ~10 minute or less sign time

#rnx

diff --git a/Engine/Source/Programs/AutomationTool/AutomationUtils/CommandUtils.cs b/Engine/Source/Programs/AutomationTool/AutomationUtils/CommandUtils.cs
--- a/Engine/Source/Programs/AutomationTool/AutomationUtils/CommandUtils.cs
+++ b/Engine/Source/Programs/AutomationTool/AutomationUtils/CommandUtils.cs
@@ -2884,8 +2884,8 @@ namespace AutomationTool
 
 		static readonly string[] TimestampServersSHA1 =
 		{
-			"http://timestamp.comodoca.com/authenticode",
 			"http://timestamp.digicert.com",
+			"http://timestamp.comodoca.com/authenticode",
 			"http://timestamp.globalsign.com/scripts/timstamp.dll"
 		};
 
@@ -2934,8 +2934,10 @@ namespace AutomationTool
 					}
 
 					// Append the files for this batch
+					// per: https://docs.microsoft.com/en-us/dotnet/api/system.diagnostics.process.start?view=net-6.0#system-diagnostics-process-start(system-diagnostics-processstartinfo)
+					// The length of the application + command line arguments cannot exceed 2080 characters, otherwise a Win32Exception will be thrown.
 					int NextFileIdx = FileIdx;
-					while(NextFileIdx < Files.Count && CommandLine.Length + Files[NextFileIdx].FullName.Length < 2000)
+					while(NextFileIdx < Files.Count && SignToolPath.Length + CommandLine.Length + Files[NextFileIdx].FullName.Length < 2080)
 					{
 						CommandLine.AppendFormat(" \"{0}\"", Files[NextFileIdx]);
 						NextFileIdx++;
