Rename PakBlacklist to PakDenylist

diff --git a/Engine/Source/Programs/AutomationTool/Scripts/CopyBuildToStagingDirectory.Automation.cs b/Engine/Source/Programs/AutomationTool/Scripts/CopyBuildToStagingDirectory.Automation.cs
--- a/Engine/Source/Programs/AutomationTool/Scripts/CopyBuildToStagingDirectory.Automation.cs
+++ b/Engine/Source/Programs/AutomationTool/Scripts/CopyBuildToStagingDirectory.Automation.cs
@@ -2261,10 +2261,10 @@ namespace AutomationScripts
 			List<string> DenyList = null;
 			if (SC.StageTargetConfigurations.Count == 1)
 			{
-				FileReference PakDenyListFilename = FileReference.Combine(SC.ProjectRoot, "Platforms", SC.PlatformDir, "Build", string.Format("PakBlacklist-{0}.txt", SC.StageTargetConfigurations[0].ToString()));
+				FileReference PakDenyListFilename = FileReference.Combine(SC.ProjectRoot, "Platforms", SC.PlatformDir, "Build", string.Format("PakDenylist-{0}.txt", SC.StageTargetConfigurations[0].ToString()));
 				if (!FileReference.Exists(PakDenyListFilename))
 				{
-					PakDenyListFilename = FileReference.Combine(SC.ProjectRoot, "Build", SC.PlatformDir, string.Format("PakBlacklist-{0}.txt", SC.StageTargetConfigurations[0].ToString()));
+					PakDenyListFilename = FileReference.Combine(SC.ProjectRoot, "Build", SC.PlatformDir, string.Format("PakDenylist-{0}.txt", SC.StageTargetConfigurations[0].ToString()));
 				}
 				if (FileReference.Exists(PakDenyListFilename))
 				{
