Turnkey fast SDK switching is only attempted after the best version has been selected from the file source. This avoids re-downloading & re-installing an SDK unnecessarily.

diff --git a/Engine/Source/Programs/AutomationTool/Turnkey/Commands/VerifySdk.cs b/Engine/Source/Programs/AutomationTool/Turnkey/Commands/VerifySdk.cs
--- a/Engine/Source/Programs/AutomationTool/Turnkey/Commands/VerifySdk.cs
+++ b/Engine/Source/Programs/AutomationTool/Turnkey/Commands/VerifySdk.cs
@@ -122,26 +122,6 @@ namespace Turnkey.Commands
 				// install if out of date, or if forcing it
 				if (bForceSdkInstall || (bUpdateIfNeeded && bHasOutOfDateSDK))
 				{
-					// if the platform has a valid sdk but isn't the "installed one", then try to switch to it
-					if ((LocalState & SdkUtils.LocalAvailability.InstalledSdk_ValidInactiveVersionExists) != 0)
-					{
-
-						// find the highest number that is valid (because a valid version exists, we know there will be at least one valid version)
-						// @todo turnkey another place where dealing with switching if we have multiple sdks won't work (see the "Sdk" param to IsVersionValid - we aren't allowing for althernate sdk types there)
-						string BestAlternateVersion = PlatformSDK.GetAllInstalledSDKVersions().ToList().OrderByDescending(x => x).Where(x => PlatformSDK.IsVersionValid(x, "Sdk")).First();
-
-						bool bWasSwitched = PlatformSDK.SwitchToAlternateSDK(BestAlternateVersion, false);
-
-						if (bWasSwitched == true)
-						{
-							TurnkeyUtils.Log("Fast-switched to already-installed version {0}", BestAlternateVersion);
-
-							// if SwitchToAlternateSDK returns true, then we are good to go!
-							continue;
-						}
-					}
-
-
 					// gather the SDKs that we need to install
 					List<FileSource> SdksToInstall = new List<FileSource>();
 					if (!bPreferFullSdk)
@@ -219,6 +199,19 @@ namespace Turnkey.Commands
 
 						foreach (FileSource Sdk in SdksToInstall)
 						{
+							// attempt to fast switch to the best one if it's already fully installed, unless we are forcing a reinstall
+							if (!bForceSdkInstall && Sdk.Type == FileSource.SourceType.Full && PlatformSDK.GetAllInstalledSDKVersions().Contains(Sdk.Version))
+							{
+								bool bWasSwitched = PlatformSDK.SwitchToAlternateSDK(Sdk.Version, false);
+								if (bWasSwitched == true)
+								{
+									TurnkeyUtils.Log("Fast-switched to already-installed version {0}", Sdk.Version);
+
+									// if SwitchToAlternateSDK returns true, then we are good to go!
+									continue;
+								}
+							}
+
 							if (Sdk.DownloadOrInstall(Platform, TurnkeyContext, null, bUnattended) == false)
 							{
 								TurnkeyUtils.Log("Failed to install {0}", Sdk.Name);
