Change Tcp Messaging log level to Verbose on Horde

diff --git a/Engine/Source/Programs/AutomationTool/Gauntlet/Unreal/Automation/UE.Automation.cs b/Engine/Source/Programs/AutomationTool/Gauntlet/Unreal/Automation/UE.Automation.cs
--- a/Engine/Source/Programs/AutomationTool/Gauntlet/Unreal/Automation/UE.Automation.cs
+++ b/Engine/Source/Programs/AutomationTool/Gauntlet/Unreal/Automation/UE.Automation.cs
@@ -291,6 +291,15 @@ namespace UE
 				}
 			}
 
+			if (!string.IsNullOrEmpty(VerboseLogCategories))
+			{
+				string LogCmdsString = string.Join(',',
+					VerboseLogCategories
+					.Split(',')
+					.Select(LogName => LogName + " Verbose"));
+				AppConfig.CommandLine += string.Format(" -LogCmds=\"{0}\"", LogCmdsString);
+			}
+
 		}
 	}
 
diff --git a/Engine/Source/Programs/AutomationTool/Gauntlet/Unreal/Base/Gauntlet.UnrealTestConfiguration.cs b/Engine/Source/Programs/AutomationTool/Gauntlet/Unreal/Base/Gauntlet.UnrealTestConfiguration.cs
--- a/Engine/Source/Programs/AutomationTool/Gauntlet/Unreal/Base/Gauntlet.UnrealTestConfiguration.cs
+++ b/Engine/Source/Programs/AutomationTool/Gauntlet/Unreal/Base/Gauntlet.UnrealTestConfiguration.cs
@@ -658,6 +658,12 @@ namespace Gauntlet
 		[AutoParam(false)]
 		public bool CookedEditor { get; set; }
 
+		/// <summary>
+		/// Enforce Verbose logging for a list of loggers
+		/// </summary>
+		[AutoParam]
+		public string VerboseLogCategories { get; set; }
+
 		// Member variables 
 
 		/// <summary>
