Add retrying to potential network share failure points in temp storage.

#rnx

diff --git a/Engine/Source/Programs/AutomationTool/BuildGraph/TempStorage.cs b/Engine/Source/Programs/AutomationTool/BuildGraph/TempStorage.cs
--- a/Engine/Source/Programs/AutomationTool/BuildGraph/TempStorage.cs
+++ b/Engine/Source/Programs/AutomationTool/BuildGraph/TempStorage.cs
@@ -14,6 +14,7 @@ using System.Threading;
 using System.Threading.Tasks;
 using System.Xml;
 using System.Xml.Serialization;
+using Microsoft.Extensions.Logging;
 
 namespace AutomationTool
 {
