Stop cooked cooker staging code from including the -Cmd executables PDB from being staged

diff --git a/Engine/Source/Programs/AutomationTool/CookedEditor/CookedEditor.Automation.cs b/Engine/Source/Programs/AutomationTool/CookedEditor/CookedEditor.Automation.cs
--- a/Engine/Source/Programs/AutomationTool/CookedEditor/CookedEditor.Automation.cs
+++ b/Engine/Source/Programs/AutomationTool/CookedEditor/CookedEditor.Automation.cs
@@ -648,6 +648,9 @@ public class MakeCookedEditor : BuildCommand
 			SC.FilesToStage.NonUFSFiles.Remove(new StagedFileReference(Path.Combine(Context.ProjectName, "Binaries", "Linux", MainCookedTarget)));
 			SC.FilesToStage.NonUFSDebugFiles.Remove(new StagedFileReference(Path.Combine(Context.ProjectName, "Binaries", "Win64", MainCookedTarget + ".pdb")));
 			SC.FilesToStage.NonUFSDebugFiles.Remove(new StagedFileReference(Path.Combine(Context.ProjectName, "Binaries", "Linux", MainCookedTarget + ".sym")));
+			// todo: find out why pdbs are in the non-debug file list and get rid of them
+			SC.FilesToStage.NonUFSFiles.Remove(new StagedFileReference(Path.Combine(Context.ProjectName, "Binaries", "Win64", MainCookedTarget + ".pdb")));
+			SC.FilesToStage.NonUFSFiles.Remove(new StagedFileReference(Path.Combine(Context.ProjectName, "Binaries", "Linux", MainCookedTarget + ".sym")));
 		}
 
 	}
