Update MacPlatform.ProcessArchivedProject to use safer file IO operations

Added a SafeRenameDirectory similar to SafeReaname file that uses Directory.Move

#rnx

diff --git a/Engine/Source/Programs/AutomationTool/AutomationUtils/CommandUtils.cs b/Engine/Source/Programs/AutomationTool/AutomationUtils/CommandUtils.cs
--- a/Engine/Source/Programs/AutomationTool/AutomationUtils/CommandUtils.cs
+++ b/Engine/Source/Programs/AutomationTool/AutomationUtils/CommandUtils.cs
@@ -657,11 +657,11 @@ namespace AutomationTool
 		/// </summary>
 		/// <param name="OldName">Old name</param>
 		/// <param name="NewName">new name</param>
-		public static void RenameDirectory(string OldName, string NewName)
+		public static void RenameDirectory(string OldName, string NewName, bool bQuiet = false)
 		{
 			var OldNormalized = ConvertSeparators(PathSeparator.Default, OldName);
 			var NewNormalized = ConvertSeparators(PathSeparator.Default, NewName);
-			Directory.Move(OldNormalized, NewNormalized);
+			InternalUtils.SafeRenameDirectory(OldNormalized, NewNormalized, bQuiet);
 		}
 
 		/// <summary>
diff --git a/Engine/Source/Programs/AutomationTool/AutomationUtils/Utils.cs b/Engine/Source/Programs/AutomationTool/AutomationUtils/Utils.cs
--- a/Engine/Source/Programs/AutomationTool/AutomationUtils/Utils.cs
+++ b/Engine/Source/Programs/AutomationTool/AutomationUtils/Utils.cs
@@ -256,6 +256,46 @@ namespace AutomationTool
 			}
 		}
 
+
+		/// <summary>
+		/// Renames/moves a directory and all its contents
+		/// </summary>
+		/// <param name="Oldname"></param>
+		/// <param name="NewName"></param>
+		/// <param name="bQuiet"></param>
+		/// <returns>True if the directory was moved, false otehrwise</returns>
+		public static bool SafeRenameDirectory(string OldName, string NewName, bool bQuiet = false)
+		{
+			if (!bQuiet)
+			{
+				Log.TraceLog("SafeRenameDirectory {0} {1}", OldName, NewName);
+			}
+			const int MaxAttempts = 10;
+			int Attempts = 0;
+
+			bool Result = true;
+			do
+			{
+				Result = true;
+				try
+				{
+					Directory.Move(OldName, NewName);
+				}
+				catch (Exception Ex)
+				{
+					if (Directory.Exists(OldName) == true || Directory.Exists(NewName) == false)
+					{
+						Log.TraceWarning("Failed to rename {0} to {1}", OldName, NewName);
+						Log.TraceWarning(LogUtils.FormatException(Ex));
+						Result = false;
+					}
+				}
+			}
+			while (Result == false && ++Attempts < MaxAttempts);
+
+			return Result;
+		}
+
 		/// <summary>
 		/// Renames/moves a file.
 		/// </summary>
diff --git a/Engine/Source/Programs/AutomationTool/Mac/MacPlatform.Automation.cs b/Engine/Source/Programs/AutomationTool/Mac/MacPlatform.Automation.cs
--- a/Engine/Source/Programs/AutomationTool/Mac/MacPlatform.Automation.cs
+++ b/Engine/Source/Programs/AutomationTool/Mac/MacPlatform.Automation.cs
@@ -1,4 +1,4 @@
-﻿// Copyright Epic Games, Inc. All Rights Reserved.
+// Copyright Epic Games, Inc. All Rights Reserved.
 using System;
 using System.Collections.Generic;
 using System.Linq;
@@ -472,7 +472,7 @@ public class MacPlatform : Platform
 			{
 				// when combining multiple platforms, don't merge the content into the .app, use the one in the Binaries directory
 				BundlePath = CombinePaths(SC.ArchiveDirectory.FullName, SC.ShortProjectName, "Binaries", "Mac", ExeName + ".app");
-				if (!Directory.Exists(BundlePath))
+				if (!DirectoryExists(BundlePath))
 				{
 					// if the .app wasn't there, just skip out (we don't require executables when combining)
 					return;
@@ -482,34 +482,27 @@ public class MacPlatform : Platform
 			string TargetPath = CombinePaths(BundlePath, "Contents", "UE");
 			if (!SC.bIsCombiningMultiplePlatforms)
 			{
-				if (Directory.Exists(BundlePath))
-				{
-					Directory.Delete(BundlePath, true);
-				}
+				DeleteDirectory(true, BundlePath);
 
 				string SourceBundlePath = CombinePaths(SC.ArchiveDirectory.FullName, SC.ShortProjectName, "Binaries", "Mac", ExeName + ".app");
-				if (!Directory.Exists(SourceBundlePath))
+				if (!DirectoryExists(SourceBundlePath))
 				{
 					SourceBundlePath = CombinePaths(SC.ArchiveDirectory.FullName, "Engine", "Binaries", "Mac", ExeName + ".app");
-					if (!Directory.Exists(SourceBundlePath))
+					if (!DirectoryExists(SourceBundlePath))
 					{
 						SourceBundlePath = CombinePaths(SC.ArchiveDirectory.FullName, "Engine", "Binaries", "Mac", "UE4.app");
 					}
 				}
-				Directory.Move(SourceBundlePath, BundlePath);
+				RenameDirectory(SourceBundlePath, BundlePath, true);
 
-				if (DirectoryExists(TargetPath))
-				{
-					Directory.Delete(TargetPath, true);
-				}
+				DeleteDirectory(true, TargetPath);
 
-				// First, move all files and folders inside the app bundle
 				string[] StagedFiles = Directory.GetFiles(SC.ArchiveDirectory.FullName, "*", SearchOption.TopDirectoryOnly);
 				foreach (string FilePath in StagedFiles)
 				{
 					string TargetFilePath = CombinePaths(TargetPath, Path.GetFileName(FilePath));
-					Directory.CreateDirectory(Path.GetDirectoryName(TargetFilePath));
-					File.Move(FilePath, TargetFilePath);
+					CreateDirectory(Path.GetDirectoryName(TargetFilePath));
+					RenameFile(FilePath, TargetFilePath, true);
 				}
 
 				string[] StagedDirectories = Directory.GetDirectories(SC.ArchiveDirectory.FullName, "*", SearchOption.TopDirectoryOnly);
@@ -519,59 +512,53 @@ public class MacPlatform : Platform
 					if (!DirName.EndsWith(".app"))
 					{
 						string TargetDirPath = CombinePaths(TargetPath, DirName);
-						Directory.CreateDirectory(Path.GetDirectoryName(TargetDirPath));
-						Directory.Move(DirPath, TargetDirPath);
+						CreateDirectory(Path.GetDirectoryName(TargetDirPath));
+						RenameDirectory(DirPath, TargetDirPath, true);
 					}
 				}
 			}
 
 			// Update executable name, icon and entry in Info.plist
 			string UE4GamePath = CombinePaths(BundlePath, "Contents", "MacOS", ExeName);
-			if (!SC.IsCodeBasedProject && ExeName != SC.ShortProjectName && File.Exists(UE4GamePath))
+			if (!SC.IsCodeBasedProject && ExeName != SC.ShortProjectName && FileExists(UE4GamePath))
 			{
 				string GameExePath = CombinePaths(BundlePath, "Contents", "MacOS", SC.ShortProjectName);
-				File.Delete(GameExePath);
-				File.Move(UE4GamePath, GameExePath);
+				DeleteFile(GameExePath);
+				RenameFile(UE4GamePath, GameExePath);
 
 				string DefaultIconPath = CombinePaths(BundlePath, "Contents", "Resources", "UnrealGame.icns");
 				string CustomIconSrcPath = CombinePaths(BundlePath, "Contents", "Resources", "Application.icns");
 				string CustomIconDestPath = CombinePaths(BundlePath, "Contents", "Resources", SC.ShortProjectName + ".icns");
-				if (File.Exists(CustomIconSrcPath))
+				if (FileExists(CustomIconSrcPath))
 				{
-					File.Delete(DefaultIconPath);
-					if (File.Exists(CustomIconDestPath))
-					{
-						File.Delete(CustomIconDestPath);
-					}
-					File.Move(CustomIconSrcPath, CustomIconDestPath);
+					DeleteFile(DefaultIconPath);
+					DeleteFile(CustomIconDestPath);
+					RenameFile(CustomIconSrcPath, CustomIconDestPath, true);
 				}
-				else if (File.Exists(DefaultIconPath))
+				else if (FileExists(DefaultIconPath))
 				{
-					if (File.Exists(CustomIconDestPath))
-					{
-						File.Delete(CustomIconDestPath);
-					}
-					File.Move(DefaultIconPath, CustomIconDestPath);
+					DeleteFile(CustomIconDestPath);
+					RenameFile(DefaultIconPath, CustomIconDestPath, true);
 				}
 
 				string InfoPlistPath = CombinePaths(BundlePath, "Contents", "Info.plist");
 				string InfoPlistContents = File.ReadAllText(InfoPlistPath);
 				InfoPlistContents = InfoPlistContents.Replace(ExeName, SC.ShortProjectName);
 				InfoPlistContents = InfoPlistContents.Replace("<string>UnrealGame</string>", "<string>" + SC.ShortProjectName + "</string>");
-				File.Delete(InfoPlistPath);
-				File.WriteAllText(InfoPlistPath, InfoPlistContents);
+				DeleteFile(InfoPlistPath);
+				WriteAllText(InfoPlistPath, InfoPlistContents);
 			}
 
 			if (!SC.bIsCombiningMultiplePlatforms)
 			{
 				// creating these directories when the content isn't moved into the application causes it 
 				// to fail to load, and isn't needed
-				Directory.CreateDirectory(CombinePaths(TargetPath, "Engine", "Binaries", "Mac"));
-				Directory.CreateDirectory(CombinePaths(TargetPath, SC.ShortProjectName, "Binaries", "Mac"));
+				CreateDirectory(CombinePaths(TargetPath, "Engine", "Binaries", "Mac"));
+				CreateDirectory(CombinePaths(TargetPath, SC.ShortProjectName, "Binaries", "Mac"));
 			}
 
 			// Find any dSYM files in the Manifest_DebugFiles_Mac file, move them to the archive directory, and remove them from the manifest.
-			string[] DebugManifests = CommandUtils.FindFiles("Manifest_DebugFiles_Mac.txt", true, SC.ArchiveDirectory.FullName);
+			string[] DebugManifests = FindFiles("Manifest_DebugFiles_Mac.txt", true, SC.ArchiveDirectory.FullName);
 			if ( DebugManifests.Count() > 0 )
 			{
 				string DebugManifest=DebugManifests[0];
@@ -587,10 +574,10 @@ public class MacPlatform : Platform
 						string FoundDebugFile = ManifestLine.Substring(0, TabIndex);
 						if (FoundDebugFile.Contains(".dSYM"))
 						{
-							FoundDebugFile = CommandUtils.CombinePaths(TargetPath, FoundDebugFile);
+							FoundDebugFile = CombinePaths(TargetPath, FoundDebugFile);
 							string MovedDebugFile = CombinePaths(SC.ArchiveDirectory.FullName, Path.GetFileName(FoundDebugFile));
 
-							File.Move(FoundDebugFile, MovedDebugFile);
+							RenameFile(FoundDebugFile, MovedDebugFile);
 
 							Log.TraceInformation("Moving debug file: '{0}')", FoundDebugFile);
 							ManifestLines.RemoveAt(ManifestLineIndex);
@@ -610,9 +597,9 @@ public class MacPlatform : Platform
 
 			if (ExeDSYMName != ProjectDSYMName)
 			{
-				if (File.Exists(ExeDSYMName))
+				if (FileExists(ExeDSYMName))
 				{
-					File.Move(ExeDSYMName, ProjectDSYMName);
+					RenameFile(ExeDSYMName, ProjectDSYMName);
 				}
 			}
 
