BuildCookRun: Fixed the creation of multiple packages\platforms in a single call. The upluginmanifest file needs to be platform specific. It will prevent some of the packages from running otherwise.

Also added the name of the platform when dumping the Manifest into FinalCopy_XXX.txt

diff --git a/Engine/Source/Programs/AutomationTool/Scripts/CopyBuildToStagingDirectory.Automation.cs b/Engine/Source/Programs/AutomationTool/Scripts/CopyBuildToStagingDirectory.Automation.cs
--- a/Engine/Source/Programs/AutomationTool/Scripts/CopyBuildToStagingDirectory.Automation.cs
+++ b/Engine/Source/Programs/AutomationTool/Scripts/CopyBuildToStagingDirectory.Automation.cs
@@ -1748,7 +1748,7 @@ namespace AutomationScripts
 			{
 				string PluginManifestName = String.Format("{0}.upluginmanifest", ManifestName);
 
-				FileReference PluginManifestFile = FileReference.Combine(SC.ProjectRoot, "Intermediate", "Staging", PluginManifestName);
+			FileReference PluginManifestFile = FileReference.Combine(SC.ProjectRoot, "Intermediate", "Staging", SC.StageTargetPlatform.PlatformType.ToString() , PluginManifestName);
 				DirectoryReference.CreateDirectory(PluginManifestFile.Directory);
 
 				using (JsonWriter Writer = new JsonWriter(PluginManifestFile.FullName))
@@ -4078,7 +4078,7 @@ namespace AutomationScripts
 				}
 			}
 
-			string BaseManifestFileName = CombinePaths(CmdEnv.LogFolder, "FinalCopy" + (SC.DedicatedServer ? "_Server" : ""));
+		string BaseManifestFileName = CombinePaths(CmdEnv.LogFolder, "FinalCopy" + (SC.DedicatedServer ? "_Server" : "") + SC.StageTargetPlatform.PlatformType.ToString());
 			DumpManifest(SC.FilesToStage.UFSFiles, BaseManifestFileName + "_UFSFiles.txt");
 
 			if (!SC.Stage || Params.SkipStage)
