Add retry for notary-app calls
#rnx

diff --git a/Engine/Source/Programs/AutomationTool/BuildGraph/Tasks/NotarizeTask.cs b/Engine/Source/Programs/AutomationTool/BuildGraph/Tasks/NotarizeTask.cs
--- a/Engine/Source/Programs/AutomationTool/BuildGraph/Tasks/NotarizeTask.cs
+++ b/Engine/Source/Programs/AutomationTool/BuildGraph/Tasks/NotarizeTask.cs
@@ -95,9 +95,26 @@ namespace AutomationTool.Tasks
 			int ExitCode = 0;
 			CommandUtils.LogInformation("Uploading {0} to the notarization server...", Dmg.FullName);
 			string CommandLine = string.Format("altool --notarize-app --primary-bundle-id \"{0}\" --username \"{1}\" --password \"@keychain:{2}\" --file \"{3}\"", Parameters.BundleID, Parameters.UserName, Parameters.KeyChainID, Dmg.FullName);
-			string Output = CommandUtils.RunAndLog("xcrun", CommandLine, out ExitCode);
-			if(ExitCode != 0)
+			string Output = "";
+			const int MaxNumRetries = 5;
+			
+			for(int NumRetries = 0;;NumRetries++)
 			{
+				Output = CommandUtils.RunAndLog("xcrun", CommandLine, out ExitCode);
+				
+				if(ExitCode == 0)
+				{
+					break;
+				}
+				
+				if (NumRetries < MaxNumRetries)
+				{
+					CommandUtils.LogInformation("--notarize-app failed with exit {0} attempting retry {1} of {2}", ExitCode, NumRetries, MaxNumRetries);
+					Thread.Sleep(2000);
+					continue;
+				}
+				
+				CommandUtils.LogInformation("Retries have been exhausted");
 				throw new AutomationException("--notarize-app failed with exit {0}", ExitCode);
 			}
 
