Gauntlet - Re-generate sessionID when restarting test

diff --git a/Engine/Source/Programs/AutomationTool/Gauntlet/Unreal/Automation/UE.Automation.cs b/Engine/Source/Programs/AutomationTool/Gauntlet/Unreal/Automation/UE.Automation.cs
--- a/Engine/Source/Programs/AutomationTool/Gauntlet/Unreal/Automation/UE.Automation.cs
+++ b/Engine/Source/Programs/AutomationTool/Gauntlet/Unreal/Automation/UE.Automation.cs
@@ -38,6 +38,10 @@ namespace UE
 			{
 				vulkan
 			}
+			public enum Mac
+			{
+				metal
+			}
 		}
 
 		/// <summary>
@@ -122,6 +126,7 @@ namespace UE
 		/// Used for having the editor and any client communicate
 		/// </summary>
 		public string SessionID = Guid.NewGuid().ToString();
+		public void GenerateSessionID() { SessionID = Guid.NewGuid().ToString(); }
 
 
 		/// <summary>
@@ -501,6 +506,11 @@ namespace UE
 			LastAutomationEntryCount = 0;
 			TestPassResults = null;
 
+			if (GetConfiguration() is AutomationTestConfig Config)
+			{
+				Config.GenerateSessionID();
+			}
+
 			return base.RestartTest();
 		}
 
