Extending windows PGO skeleton to include windows-based platform extensions.

#rnx

diff --git a/Engine/Source/Programs/AutomationTool/Gauntlet/Platform/Windows/Gauntlet.WindowsPGO.cs b/Engine/Source/Programs/AutomationTool/Gauntlet/Platform/Windows/Gauntlet.WindowsPGO.cs
--- a/Engine/Source/Programs/AutomationTool/Gauntlet/Platform/Windows/Gauntlet.WindowsPGO.cs
+++ b/Engine/Source/Programs/AutomationTool/Gauntlet/Platform/Windows/Gauntlet.WindowsPGO.cs
@@ -7,12 +7,18 @@ namespace Gauntlet
 	/// <summary>
 	/// Skeleton Windows PGO platform implementation (primarily used to test PGO locally with editor builds)
 	/// </summary>
-	internal class WindowsPGOPlatform : IPGOPlatform
+	internal abstract class WinBasePGOPlatform : IPGOPlatform
 	{
+		UnrealTargetPlatform Platform;
+
+		protected WinBasePGOPlatform(UnrealTargetPlatform InPlatform)
+		{
+			Platform = InPlatform;
+		}
 
 		public UnrealTargetPlatform GetPlatform()
 		{
-			return UnrealTargetPlatform.Win64;
+			return Platform;
 		}
 
 		public void GatherResults(string ArtifactPath)
@@ -31,6 +37,13 @@ namespace Gauntlet
 			return false;
 		}
 	}
+
+	internal class WindowsPGOPlatform : WinBasePGOPlatform
+	{
+		public WindowsPGOPlatform() : base(UnrealTargetPlatform.Win64)
+		{
+		}
+	}
 }
 
 
