LowLevelTests does not allow the command line to be replaced because it is passed to catch2 which fails on unknown arguments.

#rnx

diff --git a/Engine/Source/Programs/AutomationTool/LowLevelTests/RunLowLevelTests.cs b/Engine/Source/Programs/AutomationTool/LowLevelTests/RunLowLevelTests.cs
--- a/Engine/Source/Programs/AutomationTool/LowLevelTests/RunLowLevelTests.cs
+++ b/Engine/Source/Programs/AutomationTool/LowLevelTests/RunLowLevelTests.cs
@@ -595,7 +595,7 @@ namespace LowLevelTests
 			Platform = InPlatform;
 			BuildPath = InBuildPath;
 			ExecutablePath = InExecutablePath;
-			Flags = BuildFlags.CanReplaceCommandLine | BuildFlags.CanReplaceExecutable | BuildFlags.Loose;
+			Flags = BuildFlags.CanReplaceExecutable | BuildFlags.Loose;
 		}
 	}
 }
\ No newline at end of file
