Remove unacceptable string from Gauntlet

diff --git a/Engine/Source/Programs/AutomationTool/Gauntlet/Framework/Gauntlet.TestExecutor.cs b/Engine/Source/Programs/AutomationTool/Gauntlet/Framework/Gauntlet.TestExecutor.cs
--- a/Engine/Source/Programs/AutomationTool/Gauntlet/Framework/Gauntlet.TestExecutor.cs
+++ b/Engine/Source/Programs/AutomationTool/Gauntlet/Framework/Gauntlet.TestExecutor.cs
@@ -481,7 +481,7 @@ namespace Gauntlet
 					}
 
 					// report all tests
-					ReportMasterSummary(CurrentTestPass + 1, Options.TestIterations, PassDuration, CompletedTests);
+					ReportMainSummary(CurrentTestPass + 1, Options.TestIterations, PassDuration, CompletedTests);
 
 					if (FailedCount > 0 && Options.StopOnError)
 					{
@@ -622,7 +622,7 @@ namespace Gauntlet
 		/// <param name="Duration"></param>
 		/// <param name="AllInfo"></param>
 		/// <returns></returns>
-		void ReportMasterSummary(int CurrentPass, int NumPasses, TimeSpan Duration, IEnumerable<TestExecutionInfo> AllInfo)
+		void ReportMainSummary(int CurrentPass, int NumPasses, TimeSpan Duration, IEnumerable<TestExecutionInfo> AllInfo)
 		{
 
 			MarkdownBuilder MB = new MarkdownBuilder();
diff --git a/Engine/Source/Programs/AutomationTool/Gauntlet/Unreal/Utils/Gauntlet.UnrealHelpers.cs b/Engine/Source/Programs/AutomationTool/Gauntlet/Unreal/Utils/Gauntlet.UnrealHelpers.cs
--- a/Engine/Source/Programs/AutomationTool/Gauntlet/Unreal/Utils/Gauntlet.UnrealHelpers.cs
+++ b/Engine/Source/Programs/AutomationTool/Gauntlet/Unreal/Utils/Gauntlet.UnrealHelpers.cs
@@ -194,7 +194,6 @@ namespace Gauntlet
 		/// <summary>
 		/// Gets the filehost IP to provide to devkits by examining our local adapters and
 		/// returning the one that's active and on the local LAN (based on DNS assignment)
-		/// AG-TODO: PreferredDomain should be in the master config
 		/// </summary>
 		/// <returns></returns>
 		public static string GetHostIpAddress(string PreferredDomain="epicgames.net")
@@ -214,7 +213,7 @@ namespace Gauntlet
 				//
 				// When the above fails, attempt to extract the eth IP manually.
 				//
-				// TODO: Fallback to the wireless adapter if/when no eth device is
+				// TODO: Fallback to the wireless adapter if/when no device is
 				//       available/active.
 				//
 				LocalAddress = GetAllLocalIPv4(NetworkInterfaceType.Ethernet)
