Remove unacceptable strings from CopyBuildToStagingDirectory

#rnx

diff --git a/Engine/Source/Programs/AutomationTool/Scripts/CopyBuildToStagingDirectory.Automation.cs b/Engine/Source/Programs/AutomationTool/Scripts/CopyBuildToStagingDirectory.Automation.cs
--- a/Engine/Source/Programs/AutomationTool/Scripts/CopyBuildToStagingDirectory.Automation.cs
+++ b/Engine/Source/Programs/AutomationTool/Scripts/CopyBuildToStagingDirectory.Automation.cs
@@ -2632,9 +2632,9 @@ namespace AutomationScripts
 			string BulkOption = "";
 			if (bHasPlatformEngineConfig)
 			{
-				bool bMasterEnable = false;
-				PlatformEngineConfig.GetBool("MemoryMappedFiles", "MasterEnable", out bMasterEnable);
-				if (bMasterEnable)
+				bool bEnabled = false;
+				PlatformEngineConfig.GetBool("MemoryMappedFiles", "MasterEnable", out bEnabled);
+				if (bEnabled)
 				{
 					int Value = 0;
 					PlatformEngineConfig.GetInt32("MemoryMappedFiles", "Alignment", out Value);
@@ -3283,8 +3283,8 @@ namespace AutomationScripts
 
 						InternalUtils.SafeCopyFile(SourceManifestPath, DestManifestPath);
 
-						// generate the master manifest
-						GenerateMasterChunkManifest(CloudDir, Params.ChunkInstallVersionString, SC.FinalCookPlatform);
+						// generate the primary manifest
+						GeneratePrimaryChunkManifest(CloudDir, Params.ChunkInstallVersionString, SC.FinalCookPlatform);
 					}
 					else
 					{
@@ -3831,7 +3831,7 @@ namespace AutomationScripts
 			File.Copy(ChunkLayerFilename, OutputChunkLayerFilename, true);
 		}
 
-		private static void GenerateMasterChunkManifest(string Dir, string Version, string PlatformStr)
+		private static void GeneratePrimaryChunkManifest(string Dir, string Version, string PlatformStr)
 		{
 			//Create the directory if it doesn't exist
 			InternalUtils.SafeCreateDirectory(Dir);
