UAT: Add a systemic error code for Perforce submit failures.

diff --git a/Engine/Source/Programs/AutomationTool/BuildGraph/Tasks/SubmitTask.cs b/Engine/Source/Programs/AutomationTool/BuildGraph/Tasks/SubmitTask.cs
--- a/Engine/Source/Programs/AutomationTool/BuildGraph/Tasks/SubmitTask.cs
+++ b/Engine/Source/Programs/AutomationTool/BuildGraph/Tasks/SubmitTask.cs
@@ -11,6 +11,7 @@ using System.Xml;
 using EpicGames.Core;
 using UnrealBuildBase;
 using UnrealBuildTool;
+using Microsoft.Extensions.Logging;
 
 namespace AutomationTool.Tasks
 {
@@ -119,65 +120,74 @@ namespace AutomationTool.Tasks
 			}
 			else
 			{
-				// Get the connection that we're going to submit with
-				P4Connection SubmitP4 = CommandUtils.P4;
-				if (Parameters.Workspace != null)
+				try
 				{
-					// Create a brand new workspace
-					P4ClientInfo Client = new P4ClientInfo();
-					Client.Owner = CommandUtils.P4Env.User;
-					Client.Host = Environment.MachineName;
-					Client.RootPath = Parameters.RootDir.FullName ?? Unreal.RootDirectory.FullName;
-					Client.Name = $"{Parameters.Workspace}_{Regex.Replace(Client.Host, "[^a-zA-Z0-9]", "-")}_{ContentHash.MD5((CommandUtils.P4Env.ServerAndPort ?? "").ToUpperInvariant())}";
-					Client.Options = P4ClientOption.NoAllWrite | P4ClientOption.Clobber | P4ClientOption.NoCompress | P4ClientOption.Unlocked | P4ClientOption.NoModTime | P4ClientOption.RmDir;
-					Client.LineEnd = P4LineEnd.Local;
-					if (!String.IsNullOrEmpty(Parameters.Branch))
+					// Get the connection that we're going to submit with
+					P4Connection SubmitP4 = CommandUtils.P4;
+					if (Parameters.Workspace != null)
 					{
-						Client.View.Add(new KeyValuePair<string, string>($"{Parameters.Branch}/...", $"/..."));
+						// Create a brand new workspace
+						P4ClientInfo Client = new P4ClientInfo();
+						Client.Owner = CommandUtils.P4Env.User;
+						Client.Host = Environment.MachineName;
+						Client.RootPath = Parameters.RootDir.FullName ?? Unreal.RootDirectory.FullName;
+						Client.Name = $"{Parameters.Workspace}_{Regex.Replace(Client.Host, "[^a-zA-Z0-9]", "-")}_{ContentHash.MD5((CommandUtils.P4Env.ServerAndPort ?? "").ToUpperInvariant())}";
+						Client.Options = P4ClientOption.NoAllWrite | P4ClientOption.Clobber | P4ClientOption.NoCompress | P4ClientOption.Unlocked | P4ClientOption.NoModTime | P4ClientOption.RmDir;
+						Client.LineEnd = P4LineEnd.Local;
+						if (!String.IsNullOrEmpty(Parameters.Branch))
+						{
+							Client.View.Add(new KeyValuePair<string, string>($"{Parameters.Branch}/...", $"/..."));
+						}
+						else
+						{
+							Client.Stream = Parameters.Stream ?? CommandUtils.P4Env.Branch;
+						}
+						CommandUtils.P4.CreateClient(Client, AllowSpew: Parameters.P4Verbose);
+
+						// Create a new connection for it
+						SubmitP4 = new P4Connection(Client.Owner, Client.Name);
 					}
-					else
+
+					// Get the latest version of it
+					int NewCL = SubmitP4.CreateChange(Description: Parameters.Description.Replace("\\n", "\n"));
+					foreach(FileReference File in Files)
 					{
-						Client.Stream = Parameters.Stream ?? CommandUtils.P4Env.Branch;
+						SubmitP4.Revert(String.Format("-k \"{0}\"", File.FullName), AllowSpew: Parameters.P4Verbose);
+						SubmitP4.Sync(String.Format("-k \"{0}\"", File.FullName), AllowSpew: Parameters.P4Verbose);
+						SubmitP4.Add(NewCL, String.Format("\"{0}\"", File.FullName));
+						SubmitP4.Edit(NewCL, String.Format("\"{0}\"", File.FullName), AllowSpew: Parameters.P4Verbose);
+						if (Parameters.FileType != null)
+						{
+							SubmitP4.P4(String.Format("reopen -t \"{0}\" \"{1}\"", Parameters.FileType, File.FullName), AllowSpew: Parameters.P4Verbose);
+						}
 					}
-					CommandUtils.P4.CreateClient(Client, AllowSpew: Parameters.P4Verbose);
-
-					// Create a new connection for it
-					SubmitP4 = new P4Connection(Client.Owner, Client.Name);
-				}
 
-				// Get the latest version of it
-				int NewCL = SubmitP4.CreateChange(Description: Parameters.Description.Replace("\\n", "\n"));
-				foreach(FileReference File in Files)
-				{
-					SubmitP4.Revert(String.Format("-k \"{0}\"", File.FullName), AllowSpew: Parameters.P4Verbose);
-					SubmitP4.Sync(String.Format("-k \"{0}\"", File.FullName), AllowSpew: Parameters.P4Verbose);
-					SubmitP4.Add(NewCL, String.Format("\"{0}\"", File.FullName));
-					SubmitP4.Edit(NewCL, String.Format("\"{0}\"", File.FullName), AllowSpew: Parameters.P4Verbose);
-					if (Parameters.FileType != null)
+					// Revert any unchanged files
+					if(Parameters.RevertUnchanged)
 					{
-						SubmitP4.P4(String.Format("reopen -t \"{0}\" \"{1}\"", Parameters.FileType, File.FullName), AllowSpew: Parameters.P4Verbose);
+						SubmitP4.RevertUnchanged(NewCL);
+						if(SubmitP4.TryDeleteEmptyChange(NewCL))
+						{
+							CommandUtils.LogInformation("No files to submit; ignored.");
+							return Task.CompletedTask;
+						}
 					}
-				}
 
-				// Revert any unchanged files
-				if(Parameters.RevertUnchanged)
-				{
-					SubmitP4.RevertUnchanged(NewCL);
-					if(SubmitP4.TryDeleteEmptyChange(NewCL))
+					// Submit it
+					int SubmittedCL;
+					SubmitP4.Submit(NewCL, out SubmittedCL, Force: Parameters.Force);
+					if (SubmittedCL <= 0)
 					{
-						CommandUtils.LogInformation("No files to submit; ignored.");
-						return Task.CompletedTask;
+						throw new AutomationException("Submit failed.");
 					}
-				}
 
-				// Submit it
-				int SubmittedCL;
-				SubmitP4.Submit(NewCL, out SubmittedCL, Force: Parameters.Force);
-				if (SubmittedCL <= 0)
+					CommandUtils.LogInformation("Submitted in changelist {0}", SubmittedCL);
+				}
+				catch (P4Exception Ex)
 				{
-					throw new AutomationException("Submit failed.");
+					Job.OwnerCommand.Logger.LogError(KnownLogEvents.Systemic_Perforce, "{Message}", Ex.Message);
+					throw new AutomationException(Ex.ErrorCode, Ex, "{0}", Ex.Message) { OutputFormat = AutomationExceptionOutputFormat.Silent };
 				}
-				CommandUtils.LogInformation("Submitted in changelist {0}", SubmittedCL);
 			}
 			return Task.CompletedTask;
 		}
diff --git a/Engine/Source/Programs/Shared/EpicGames.Core/KnownLogEvents.cs b/Engine/Source/Programs/Shared/EpicGames.Core/KnownLogEvents.cs
--- a/Engine/Source/Programs/Shared/EpicGames.Core/KnownLogEvents.cs
+++ b/Engine/Source/Programs/Shared/EpicGames.Core/KnownLogEvents.cs
@@ -145,6 +145,11 @@ namespace EpicGames.Core
 		/// </summary>
 		public static EventId Systemic { get; } = new EventId(700);
 
+		/// <summary>
+		/// A systemic event relating to Perforce
+		/// </summary>
+		public static EventId Systemic_Perforce { get; } = new EventId(701);
+
 		/// <summary>
 		/// A systemic event from XGE
 		/// </summary>
