Made Gauntlet issue with a locked log a warning vs an error

diff --git a/Engine/Source/Programs/AutomationTool/Gauntlet/Platform/Windows/Gauntlet.TargetDeviceWindows.cs b/Engine/Source/Programs/AutomationTool/Gauntlet/Platform/Windows/Gauntlet.TargetDeviceWindows.cs
--- a/Engine/Source/Programs/AutomationTool/Gauntlet/Platform/Windows/Gauntlet.TargetDeviceWindows.cs
+++ b/Engine/Source/Programs/AutomationTool/Gauntlet/Platform/Windows/Gauntlet.TargetDeviceWindows.cs
@@ -332,7 +332,8 @@ namespace Gauntlet
 				}
 				catch (Exception Ex)
 				{
-					throw new AutomationException("Unable to delete existing log file {0} {1}", ProcessLogFile, Ex.Message);
+					//throw new AutomationException("Unable to delete existing log file {0} {1}", ProcessLogFile, Ex.Message);
+					Log.Warning("Unable to delete existing log file {0}. {1}", ProcessLogFile, Ex.Message);
 				}
 
 				Log.Verbose("\t{0}", CmdLine);
