Gauntlet implement timestamp in output log messages relative to program start time

diff --git a/Engine/Source/Programs/AutomationTool/Gauntlet/Framework/Gauntlet.Utils.cs b/Engine/Source/Programs/AutomationTool/Gauntlet/Framework/Gauntlet.Utils.cs
--- a/Engine/Source/Programs/AutomationTool/Gauntlet/Framework/Gauntlet.Utils.cs
+++ b/Engine/Source/Programs/AutomationTool/Gauntlet/Framework/Gauntlet.Utils.cs
@@ -27,6 +27,8 @@ namespace Gauntlet
 			set { InnerParams = value; }
 		}
 
+		public static readonly DateTime ProgramStartTime = DateTime.UtcNow;
+
 		static string InnerTempDir;
 		static string InnerLogDir;
 		static string InnerUnrealRootDir;
@@ -339,6 +341,12 @@ namespace Gauntlet
 				}		
 			}
 
+			if (Globals.Params.ParseParam("timestamp"))
+			{
+				TimeSpan ElapsedTime = DateTime.UtcNow - Globals.ProgramStartTime;
+				Message = "[" + ElapsedTime.ToString() + "] " + Message;
+			}
+
 			if (string.IsNullOrEmpty(Prefix) == false)
 			{
 				Message = Prefix + ": " + Message;
