- Moved where we stage Mono files so it only happens when targeting Linux

diff --git a/Engine/Config/BaseGame.ini b/Engine/Config/BaseGame.ini
--- a/Engine/Config/BaseGame.ini
+++ b/Engine/Config/BaseGame.ini
@@ -245,7 +245,6 @@ FrameInterpolationProcessor=Class'/Script/LiveLink.LiveLinkBasicFrameInterpolati
 +EngineExtraStageFiles=(Path="Build", Files="Build.version", Recursive=false, NonUFS=true)
 +EngineExtraStageFiles=(Path="Build/BatchFiles", Recursive=true, NonUFS=true)
 +EngineExtraStageFiles=(Path="Source/ThirdParty/Licenses")
-+EngineExtraStageFiles=(Path="Binaries/ThirdParty/Mono/Linux", Recursive=true, NonUFS=true)
 
 ; default directories to stage
 bStageShaderDirs=true
diff --git a/Engine/Source/Programs/AutomationTool/CookedEditor/CookedEditor.Automation.cs b/Engine/Source/Programs/AutomationTool/CookedEditor/CookedEditor.Automation.cs
--- a/Engine/Source/Programs/AutomationTool/CookedEditor/CookedEditor.Automation.cs
+++ b/Engine/Source/Programs/AutomationTool/CookedEditor/CookedEditor.Automation.cs
@@ -564,8 +564,9 @@ public class MakeCookedEditor : BuildCommand
 
 		if (Context.IniPlatformName == "Linux")
 		{
-			// linux needs dotnet runtime
+			// linux needs dotnet runtime, as well as mono (at least for now until we can break our dependence)
 			SC.StageFiles(StagedFileType.NonUFS, Unreal.FindDotnetDirectoryForPlatform(RuntimePlatform.Type.Linux), StageFilesSearch.AllDirectories);
+			SC.StageFiles(StagedFileType.NonUFS, DirectoryReference.Combine(Context.EngineDirectory, "Binaries/ThirdParty/Mono/Linux"), StageFilesSearch.AllDirectories);
 		}
 
 		// not sure if we need this or not now
