Use robocopy/rsync to copy temp storage manifests locally before unzipping them, and use them to copy manifests off to the share at the end of a task.

improves platform reliability, m1 macs struggle to update file attributes when copying to an smb share. rsync ensures the attributes are correctly copied.

#rnx

diff --git a/Engine/Source/Programs/AutomationTool/BuildGraph/TempStorage.cs b/Engine/Source/Programs/AutomationTool/BuildGraph/TempStorage.cs
--- a/Engine/Source/Programs/AutomationTool/BuildGraph/TempStorage.cs
+++ b/Engine/Source/Programs/AutomationTool/BuildGraph/TempStorage.cs
@@ -10,11 +10,12 @@ using System.Diagnostics;
 using System.IO;
 using System.IO.Compression;
 using System.Linq;
+using System.Text.RegularExpressions;
 using System.Threading;
 using System.Threading.Tasks;
 using System.Xml;
 using System.Xml.Serialization;
-using Microsoft.Extensions.Logging;
+using UnrealBuildTool;
 
 namespace AutomationTool
 {
@@ -918,7 +919,13 @@ namespace AutomationTool
 					// Unzip all the build products
 					DirectoryReference SharedNodeDir = GetDirectoryForNode(SharedDir, NodeName);
 					FileInfo[] ZipFiles = Manifest.ZipFiles.Select(x => new FileInfo(FileReference.Combine(SharedNodeDir, x.Name).FullName)).ToArray();
-					ParallelUnzipFiles(ZipFiles, RootDir);
+					string Result = ParallelUnzipFiles(ZipFiles, RootDir);
+					if (!string.IsNullOrWhiteSpace(Result))
+					{
+						string LogPath = CommandUtils.CombinePaths(RootDir.FullName, "Engine/Programs/AutomationTool/Saved/Logs", $"Copy Manifest - {NodeName}.log");
+						CommandUtils.LogInformation("Saving copy log to {0}", LogPath);
+						File.WriteAllText(LogPath, Result);
+					}
 
 					// Update the timestamps to match the manifest. Zip files only use local time, and there's no guarantee it matches the local clock.
 					foreach(TempStorageFile ManifestFile in Manifest.Files)
@@ -1041,23 +1048,7 @@ namespace AutomationTool
 							throw new AutomationException("Unable to open file {0}", ZipFileName.FullName);
 						}
 
-						try
-						{
-							// if we are using a staging dir, copy to the final location and delete the staged copy.
-							string FinalZipFile = ZipFileName.FullName;
-							if (StagingDir != null)
-							{
-								FinalZipFile = CommandUtils.CombinePaths(OutputDir.FullName, ZipFileName.GetFileName());
-								CommandUtils.CopyFile(ZipFileName.FullName, FinalZipFile, bQuiet: true, bRetry: true);
-								FileReference.Delete(ZipFileName);
-							}
-							ZipFiles.Add(new FileInfo(FinalZipFile));
-						}
-						catch (IOException Ex)
-						{
-							CommandUtils.LogError("Unable to copy file \"{0}\" to TempStorage: {1}", ZipFileName.FullName, Ex);
-							throw new AutomationException("Unable to copy file {0}", ZipFileName.FullName);
-						}
+						ZipFiles.Add(new FileInfo(ZipFileName.FullName));
 					}
 				})).ToList();
 
@@ -1068,7 +1059,25 @@ namespace AutomationTool
 			}
 
 			ZipThreads.ForEach(thread => thread.Join());
-			
+
+			if (ZipFiles.Any() && !string.IsNullOrWhiteSpace(StagingDir.FullName))
+			{
+				try
+				{
+					string CopyResult = CopyDirectory(ZipDir, OutputDir);
+					string LogPath = CommandUtils.CombinePaths(RootDir.FullName, "Engine/Programs/AutomationTool/Saved/Logs", $"Copy Files to Temp Storage.log");
+					CommandUtils.LogInformation("Saving copy log to {0}", LogPath);
+					File.WriteAllText(LogPath, string.Join(Environment.NewLine, CopyResult));
+					Parallel.ForEach(ZipFiles, (z) => z.Delete());
+					ZipFiles = new ConcurrentBag<FileInfo>(ZipFiles.Select(z => new FileInfo(CommandUtils.MakeRerootedFilePath(z.FullName, StagingDir.FullName, OutputDir.FullName))));
+				}
+				catch (IOException Ex)
+				{
+					CommandUtils.LogError("Unable to copy staging directory {0} to {1}, Ex: {2}", ZipDir.FullName, OutputDir.FullName, Ex);
+					throw new AutomationException(Ex, "Unable to copy staging directory {0} to {1}", ZipDir.FullName, OutputDir.FullName);
+				}
+			}
+
 			return ZipFiles.OrderBy(x => x.Name).ToArray();
 		}
 
@@ -1081,18 +1090,15 @@ namespace AutomationTool
 		/// <remarks>
 		/// The code is expected to be the used as the symmetrical inverse of <see cref="ParallelZipFiles"/>, but could be used independently, as long as the files in the zip do not overlap.
 		/// </remarks>
-		static void ParallelUnzipFiles(FileInfo[] ZipFiles, DirectoryReference RootDir)
+		private static string ParallelUnzipFiles(FileInfo[] ZipFiles, DirectoryReference RootDir)
 		{
+			ConcurrentBag<string> CopyResults = new ConcurrentBag<string>();
 			Parallel.ForEach(ZipFiles,
 				(ZipFile) =>
 				{
-					// Copy the zip file locally before unzipped to harden against network hiccups.
-					string LocalZipPath = CommandUtils.CombinePaths(RootDir.FullName, "Engine/Intermediate/Manifests", ZipFile.Name);
-					CommandUtils.CopyFile(ZipFile.FullName, LocalZipPath, bQuiet: true, bRetry: true);
-					if (!File.Exists(LocalZipPath))
-					{
-						throw new AutomationException("Failed to copy manifest {0} to {1}", ZipFile, LocalZipPath);
-					}
+					// Copy the ZIP to the local drive before unzipping to harden against network issues.
+					CopyResults.Add(CopyFile(ZipFile, RootDir));
+					string LocalZipFile = CommandUtils.CombinePaths(RootDir.FullName, ZipFile.Name);
 
 					// unzip the files manually instead of caling ZipFile.ExtractToDirectory() because we need to overwrite readonly files. Because of this, creating the directories is up to us as well.
 					List<string> ExtractedPaths = new List<string>();
@@ -1101,7 +1107,7 @@ namespace AutomationTool
 					{
 						try
 						{
-							using (ZipArchive ZipArchive = System.IO.Compression.ZipFile.OpenRead(LocalZipPath))
+							using (ZipArchive ZipArchive = System.IO.Compression.ZipFile.OpenRead(LocalZipFile))
 							{
 								foreach (ZipArchiveEntry Entry in ZipArchive.Entries)
 								{
@@ -1149,7 +1155,7 @@ namespace AutomationTool
 													throw;
 												}
 
-												Log.TraceWarning("Failed to unzip '{0}' from '{1}' to '{2}', retrying.. (Error: {3})", Entry.FullName, LocalZipPath, ExtractedFilename, IOEx.Message);
+												Log.TraceWarning("Failed to unzip '{0}' from '{1}' to '{2}', retrying.. (Error: {3})", Entry.FullName, LocalZipFile, ExtractedFilename, IOEx.Message);
 											}
 										}
 									}
@@ -1162,22 +1168,30 @@ namespace AutomationTool
 						{
 							if (UnzipFileAttempts == 0)
 							{
-								Log.TraceError("All retries exhausted attempting to unzip entries from '{0}'. Terminating.", LocalZipPath);
+								Log.TraceError("All retries exhausted attempting to unzip entries from '{0}'. Terminating.", LocalZipFile);
+								string LogPath = CommandUtils.CombinePaths(RootDir.FullName, "Engine/Programs/AutomationTool/Saved/Logs", $"Copy Manifest - {ZipFile.Name}.log");
+								CommandUtils.LogInformation("Saving copy log to {0}", LogPath);
+								File.WriteAllText(LogPath, string.Join(Environment.NewLine, CopyResults));
 								throw;
 							}
 
 							// Some exceptions may be caused by networking hiccups. We want to retry in those cases.
 							if ((Ex is IOException || Ex is InvalidDataException))
 							{
-								Log.TraceWarning("Failed to unzip entries from '{0}' to '{1}', retrying.. (Error: {2})", LocalZipPath, RootDir.FullName, Ex.Message);
+								Log.TraceWarning("Failed to unzip entries from '{0}' to '{1}', retrying.. (Error: {2})", LocalZipFile, RootDir.FullName, Ex.Message);
 							}
 						}
 						finally
 						{
-							File.Delete(LocalZipPath);
+							if (File.Exists(LocalZipFile))
+							{
+								File.Delete(LocalZipFile);
+							}
 						}
 					}
 				});
+
+			return string.Join(Environment.NewLine, CopyResults);
 		}
 
 		/// <summary>
@@ -1280,6 +1294,46 @@ namespace AutomationTool
 			}
 			return false;
 		}
+
+		static string CopyFile(FileInfo ZipFile, DirectoryReference RootDir)
+		{
+			if (HostPlatform.Current.HostEditorPlatform == UnrealTargetPlatform.Win64)
+			{
+				return CommandUtils.RunAndLog(GetRoboCopyExe(), $"\"{ZipFile.DirectoryName}\" \"{RootDir}\" \"{ZipFile.Name}\" /w:5 /r:10", MaxSuccessCode: 3, Options: CommandUtils.ERunOptions.AppMustExist | CommandUtils.ERunOptions.NoLoggingOfRunCommand);
+			}
+			else
+			{
+				return CommandUtils.RunAndLog("rsync", $"-v \"{ZipFile.DirectoryName}/{ZipFile.Name}\" \"{RootDir}/{ZipFile.Name}\"", Options: CommandUtils.ERunOptions.AppMustExist | CommandUtils.ERunOptions.NoLoggingOfRunCommand);
+			}
+		}
+
+		static string CopyDirectory(DirectoryReference SourceDir, DirectoryReference DestinationDir)
+		{
+			if (HostPlatform.Current.HostEditorPlatform == UnrealTargetPlatform.Win64)
+			{
+				return CommandUtils.RunAndLog(GetRoboCopyExe(), $"\"{SourceDir}\" \"{DestinationDir}\" * /S /w:5 /r:10", MaxSuccessCode: 3, Options: CommandUtils.ERunOptions.AppMustExist | CommandUtils.ERunOptions.NoLoggingOfRunCommand);
+			}
+			else
+			{
+				return CommandUtils.RunAndLog("rsync", $"-vam --include=\"**\" \"{SourceDir}/\" \"{DestinationDir}/\"", Options: CommandUtils.ERunOptions.AppMustExist | CommandUtils.ERunOptions.NoLoggingOfRunCommand);
+			}
+		}
+
+		static string GetRoboCopyExe()
+		{
+			var Result = CommandUtils.CombinePaths(Environment.SystemDirectory, "robocopy.exe");
+			if (!CommandUtils.FileExists(Result))
+			{
+				// Use Regex.Replace so we can do a case-insensitive replacement of System32
+				var SysNativeDirectory = Regex.Replace(Environment.SystemDirectory, "System32", "Sysnative", RegexOptions.IgnoreCase);
+				var SysNativeExe = CommandUtils.CombinePaths(SysNativeDirectory, "robocopy.exe");
+				if (CommandUtils.FileExists(SysNativeExe))
+				{
+					Result = SysNativeExe;
+				}
+			}
+			return Result;
+		}
 	}
 
 	/// <summary>
