Stage project shaders for cooked cookers

diff --git a/Engine/Source/Programs/AutomationTool/CookedEditor/CookedEditor.Automation.cs b/Engine/Source/Programs/AutomationTool/CookedEditor/CookedEditor.Automation.cs
--- a/Engine/Source/Programs/AutomationTool/CookedEditor/CookedEditor.Automation.cs
+++ b/Engine/Source/Programs/AutomationTool/CookedEditor/CookedEditor.Automation.cs
@@ -478,6 +478,11 @@ public class MakeCookedEditor : BuildCommand
 			GatherTargetDependencies(Params, SC, Context, SC.StageExecutables[0]);
 		}
 
+		// project shaders
+		if (Context.bStageShaderDirs)
+		{
+			Context.NonUFSFilesToStage.AddRange(DirectoryReference.EnumerateFiles(DirectoryReference.Combine(Context.ProjectDirectory, "Shaders"), "*", SearchOption.AllDirectories));
+		}
 
 		StagePlatformExtensionFiles(Params, SC, Context, Context.ProjectDirectory);
 		StagePluginFiles(Params, SC, Context, false);
