PR #9268: Turnkey: Use `User_LastPerforceClient` for streams (Contributed by [REDACTED_PERSON])

diff --git a/Engine/Source/Programs/AutomationTool/Turnkey/Providers/PerforceCopyProvider.cs b/Engine/Source/Programs/AutomationTool/Turnkey/Providers/PerforceCopyProvider.cs
--- a/Engine/Source/Programs/AutomationTool/Turnkey/Providers/PerforceCopyProvider.cs
+++ b/Engine/Source/Programs/AutomationTool/Turnkey/Providers/PerforceCopyProvider.cs
@@ -111,11 +111,22 @@ namespace Turnkey
 			// find clients for this user in the given stream
 			P4ClientInfo[] StreamClients = PerforceConnection.GetClientsForUser(Username, null, Stream);
 
+			if (TurnkeySettings.HasSetUserSetting("User_LastPerforceClient"))
+			{
+				string ClientName = TurnkeySettings.GetUserSetting("User_LastPerforceClient");
+				P4ClientInfo Client = StreamClients.FirstOrDefault(Client => string.Compare(Client.Name, ClientName, StringComparison.InvariantCultureIgnoreCase) == 0);
+				if (Client != null && (string.IsNullOrEmpty(Client.Host) || string.Compare(Client.Host, Hostname) == 0))
+				{
+					return Client;
+				}
+			}
+
 			// find the first one usable on this host
 			foreach (P4ClientInfo Client in StreamClients)
 			{
 				if (string.IsNullOrEmpty(Client.Host) || string.Compare(Client.Host, Hostname) == 0)
 				{
+					TurnkeySettings.SetUserSetting("User_LastPerforceClient", Client.Name);
 					return Client;
 				}
 			}
