Platform SDK "Force Reinstall Sdk" is only shown in the editor if the 'best' SDK version is already selected, otherwise it shows 'Update Sdk'.

 - note: 'best' is approximated as either MainVersion or MaxVersion to avoid touching the sdk file sources on editor startup.

#rnx

diff --git a/Engine/Source/Editor/TurnkeySupport/Private/TurnkeySupportModule.cpp b/Engine/Source/Editor/TurnkeySupport/Private/TurnkeySupportModule.cpp
--- a/Engine/Source/Editor/TurnkeySupport/Private/TurnkeySupportModule.cpp
+++ b/Engine/Source/Editor/TurnkeySupport/Private/TurnkeySupportModule.cpp
@@ -1387,7 +1387,7 @@ static void MakeTurnkeyPlatformMenu(UToolMenu* ToolMenu, FName IniPlatformName,
 	FString NoDevice;
 	if (SdkInfo.bCanInstallFullSdk || SdkInfo.bCanInstallAutoSdk)
 	{
-		if (SdkInfo.Status == ETurnkeyPlatformSdkStatus::OutOfDate)
+		if (SdkInfo.Status == ETurnkeyPlatformSdkStatus::OutOfDate || (SdkInfo.Status == ETurnkeyPlatformSdkStatus::Valid && SdkInfo.bHasBestSdk == false))
 		{
 			SdkSection.AddMenuEntry(
 				NAME_None,
@@ -2071,6 +2071,7 @@ bool GetSdkInfoFromTurnkey(FString Line, FName& PlatformName, FString& DeviceId,
 	}
 	SdkInfo.bCanInstallFullSdk = FlagsString.Contains(TEXT("Support_FullSdk"));
 	SdkInfo.bCanInstallAutoSdk = FlagsString.Contains(TEXT("Support_AutoSdk"));
+	SdkInfo.bHasBestSdk = FlagsString.Contains(TEXT("Sdk_HasBestVersion"));
 
 	SdkInfo.DeviceStatus = ETurnkeyDeviceStatus::Unknown;
 	if (FlagsString.Contains(TEXT("Device_InvalidPrerequisites")))
diff --git a/Engine/Source/Editor/TurnkeySupport/Public/Interfaces/ITurnkeySupportModule.h b/Engine/Source/Editor/TurnkeySupport/Public/Interfaces/ITurnkeySupportModule.h
--- a/Engine/Source/Editor/TurnkeySupport/Public/Interfaces/ITurnkeySupportModule.h
+++ b/Engine/Source/Editor/TurnkeySupport/Public/Interfaces/ITurnkeySupportModule.h
@@ -37,6 +37,7 @@ struct FTurnkeySdkInfo
 	TMap<FString, Version> SDKVersions;
 	bool bCanInstallFullSdk;
 	bool bCanInstallAutoSdk;
+	bool bHasBestSdk;
 };
 
 
diff --git a/Engine/Source/Programs/AutomationTool/Turnkey/Commands/VerifySdk.cs b/Engine/Source/Programs/AutomationTool/Turnkey/Commands/VerifySdk.cs
--- a/Engine/Source/Programs/AutomationTool/Turnkey/Commands/VerifySdk.cs
+++ b/Engine/Source/Programs/AutomationTool/Turnkey/Commands/VerifySdk.cs
@@ -81,11 +81,11 @@ namespace Turnkey.Commands
 				else
 				{
 					StatusString = "Valid";
-					ReportedState &= (SdkUtils.LocalAvailability.AutoSdk_ValidVersionExists | SdkUtils.LocalAvailability.InstalledSdk_ValidVersionExists | SdkUtils.LocalAvailability.Support_FullSdk | SdkUtils.LocalAvailability.Support_AutoSdk);
+					ReportedState &= (SdkUtils.LocalAvailability.AutoSdk_ValidVersionExists | SdkUtils.LocalAvailability.InstalledSdk_ValidVersionExists | SdkUtils.LocalAvailability.Support_FullSdk | SdkUtils.LocalAvailability.Support_AutoSdk | SdkUtils.LocalAvailability.Sdk_HasBestVersion);
 				}
 
 
-				// gather lisat of tags/values to write out
+				// gather list of tags/values to write out
 				SDKCollection SDKInfo = null;
 				List<string> PlatformProperties = new List<string>();
 				{
diff --git a/Engine/Source/Programs/AutomationTool/Turnkey/SdkInfo.cs b/Engine/Source/Programs/AutomationTool/Turnkey/SdkInfo.cs
--- a/Engine/Source/Programs/AutomationTool/Turnkey/SdkInfo.cs
+++ b/Engine/Source/Programs/AutomationTool/Turnkey/SdkInfo.cs
@@ -18,24 +18,26 @@ namespace Turnkey
 		public enum LocalAvailability
 		{
 			None = 0,
-			AutoSdk_VariableExists = 1,
-			AutoSdk_ValidVersionExists = 2,
-			AutoSdk_InvalidVersionExists = 4,
+			AutoSdk_VariableExists                  = (1 << 0),
+			AutoSdk_ValidVersionExists              = (1 << 1),
+			AutoSdk_InvalidVersionExists            = (1 << 2),
 			
-			InstalledSdk_ValidInactiveVersionExists = 8,
-			InstalledSdk_ValidVersionExists = 16,
-			InstalledSdk_InvalidVersionExists = 32,
+			InstalledSdk_ValidInactiveVersionExists = (1 << 3),
+			InstalledSdk_ValidVersionExists         = (1 << 4),
+			InstalledSdk_InvalidVersionExists       = (1 << 5),
 
-			Platform_ValidHostPrerequisites = 64,
-			Platform_InvalidHostPrerequisites = 128,
+			Platform_ValidHostPrerequisites         = (1 << 6),
+			Platform_InvalidHostPrerequisites       = (1 << 7),
 
-			Device_InvalidPrerequisites = 256,
-			Device_InstallSoftwareValid = 512,
-			Device_InstallSoftwareInvalid = 1024,
-			Device_CannotConnect = 2048,
+			Device_InvalidPrerequisites             = (1 << 8),
+			Device_InstallSoftwareValid             = (1 << 9),
+			Device_InstallSoftwareInvalid           = (1 << 10),
+			Device_CannotConnect                    = (1 << 11),
 
-			Support_FullSdk = 4096,
-			Support_AutoSdk = 8192,
+			Support_FullSdk                         = (1 << 12),
+			Support_AutoSdk                         = (1 << 13),
+
+			Sdk_HasBestVersion                      = (1 << 14),
 		}
 
 		static public LocalAvailability GetLocalAvailability(AutomationTool.Platform AutomationPlatform, bool bAllowUpdatingPrerequisites, TurnkeyContextImpl TurnkeyContext)
@@ -114,6 +116,18 @@ namespace Turnkey
 				}
 			}
 
+			// see if we have the best version available (note: this is an approximation of FileSource.ChooseBest() without hitting the file sources)
+			if ((Result & (LocalAvailability.InstalledSdk_ValidVersionExists | LocalAvailability.AutoSdk_ValidVersionExists)) != 0)
+			{
+				string MainVersion = SDK.GetMainVersion();
+				if (SDKInfo.Sdks.All(x => (x.Current == null) || (x.Current == MainVersion) || (x.Current == x.Max)))
+				{
+					// only set this when all valid Current sdks are either MainVersion or MaxVersion
+					Result |= LocalAvailability.Sdk_HasBestVersion;
+				}
+			}
+
+
 			string FullSupportedPlatforms = TurnkeyUtils.GetVariableValue("Studio_FullInstallPlatforms");
 			string AutoSdkSupportedPlatforms = TurnkeyUtils.GetVariableValue("Studio_AutoSdkPlatforms");
 			if (!string.IsNullOrEmpty(FullSupportedPlatforms))
