Turnkey indicates to the platform whether an SDK is already installed when getting the install command, for platforms that need to handle this as a special case (uninstalling the current version, for example).

#rnx

diff --git a/Engine/Source/Programs/AutomationTool/AutomationUtils/Platform.cs b/Engine/Source/Programs/AutomationTool/AutomationUtils/Platform.cs
--- a/Engine/Source/Programs/AutomationTool/AutomationUtils/Platform.cs
+++ b/Engine/Source/Programs/AutomationTool/AutomationUtils/Platform.cs
@@ -231,7 +231,7 @@ namespace AutomationTool
 
 		}
 
-		public virtual bool InstallSDK(BuildCommand BuildCommand, ITurnkeyContext TurnkeyContext, DeviceInfo Device, bool bUnattended)
+		public virtual bool InstallSDK(BuildCommand BuildCommand, ITurnkeyContext TurnkeyContext, DeviceInfo Device, bool bUnattended, bool bSdkAlreadyInstalled)
 		{
 			string Command, Params;
 
@@ -242,7 +242,7 @@ namespace AutomationTool
 				int ExitCode = TurnkeyContext.RunExternalCommand(Command, Params, bRequiresPrivilegeElevation, bUnattended, bCreateWindow);
 				return OnSDKInstallComplete(ExitCode, TurnkeyContext, Device);
 			}
-			else if (Device == null && GetSDKInstallCommand(out Command, out Params, ref bRequiresPrivilegeElevation, ref bCreateWindow, TurnkeyContext))
+			else if (Device == null && GetSDKInstallCommand(out Command, out Params, ref bRequiresPrivilegeElevation, ref bCreateWindow, TurnkeyContext, bSdkAlreadyInstalled))
 			{
 				int ExitCode = TurnkeyContext.RunExternalCommand(Command, Params, bRequiresPrivilegeElevation, bUnattended, bCreateWindow);
 				return OnSDKInstallComplete(ExitCode, TurnkeyContext, null);
@@ -267,7 +267,10 @@ namespace AutomationTool
 			Params = null;
 			return false;
 		}
-
+		public virtual bool GetSDKInstallCommand(out string Command, out string Params, ref bool bRequiresPrivilegeElevation, ref bool bCreateWindow, ITurnkeyContext TurnkeyContext, bool bSdkAlreadyInstalled)
+		{
+			return GetSDKInstallCommand(out Command, out Params, ref bRequiresPrivilegeElevation, ref bCreateWindow, TurnkeyContext);
+		}
 		public virtual bool GetDeviceUpdateSoftwareCommand(out string Command, out string Params, ref bool bRequiresPrivilegeElevation, ref bool bCreateWindow, ITurnkeyContext TurnkeyContext, DeviceInfo Device = null)
 		{
 			Command = null;
diff --git a/Engine/Source/Programs/AutomationTool/Turnkey/Commands/InstallSdk.cs b/Engine/Source/Programs/AutomationTool/Turnkey/Commands/InstallSdk.cs
--- a/Engine/Source/Programs/AutomationTool/Turnkey/Commands/InstallSdk.cs
+++ b/Engine/Source/Programs/AutomationTool/Turnkey/Commands/InstallSdk.cs
@@ -314,7 +314,9 @@ namespace Turnkey.Commands
 						continue;
 					}
 
-					Sdk.DownloadOrInstall(Platform, TurnkeyContext, InstallDevice, false);
+					bool bSdkAlreadyInstalled = (InstallDevice == null) && (Sdk.Type == FileSource.SourceType.Full) && SDK.GetAllInstalledSDKVersions().Contains(Sdk.Version);
+
+					Sdk.DownloadOrInstall(Platform, TurnkeyContext, InstallDevice, false, bSdkAlreadyInstalled);
 				}
 			}
 		}
diff --git a/Engine/Source/Programs/AutomationTool/Turnkey/Commands/VerifySdk.cs b/Engine/Source/Programs/AutomationTool/Turnkey/Commands/VerifySdk.cs
--- a/Engine/Source/Programs/AutomationTool/Turnkey/Commands/VerifySdk.cs
+++ b/Engine/Source/Programs/AutomationTool/Turnkey/Commands/VerifySdk.cs
@@ -199,8 +199,10 @@ namespace Turnkey.Commands
 
 						foreach (FileSource Sdk in SdksToInstall)
 						{
+							bool bSdkAlreadyInstalled = Sdk.Type == FileSource.SourceType.Full && PlatformSDK.GetAllInstalledSDKVersions().Contains(Sdk.Version);
+
 							// attempt to fast switch to the best one if it's already fully installed, unless we are forcing a reinstall
-							if (!bForceSdkInstall && Sdk.Type == FileSource.SourceType.Full && PlatformSDK.GetAllInstalledSDKVersions().Contains(Sdk.Version))
+							if (!bForceSdkInstall && bSdkAlreadyInstalled)
 							{
 								bool bWasSwitched = PlatformSDK.SwitchToAlternateSDK(Sdk.Version, false);
 								if (bWasSwitched == true)
@@ -212,7 +214,7 @@ namespace Turnkey.Commands
 								}
 							}
 
-							if (Sdk.DownloadOrInstall(Platform, TurnkeyContext, null, bUnattended) == false)
+							if (Sdk.DownloadOrInstall(Platform, TurnkeyContext, null, bUnattended, bSdkAlreadyInstalled) == false)
 							{
 								TurnkeyUtils.Log("Failed to install {0}", Sdk.Name);
 								TurnkeyUtils.ExitCode = ExitCode.Error_SDKInstallFailed;
@@ -302,7 +304,7 @@ namespace Turnkey.Commands
 									}
 									else
 									{
-										if (MatchingInstallableSdk.DownloadOrInstall(Platform, TurnkeyContext, Device, bUnattended) == false)
+										if (MatchingInstallableSdk.DownloadOrInstall(Platform, TurnkeyContext, Device, bUnattended, false) == false)
 										{
 											TurnkeyUtils.Log("Failed to update Device '{0}' with '{0}'", Device.Name, MatchingInstallableSdk.Name);
 											TurnkeyUtils.ExitCode = ExitCode.Error_DeviceUpdateFailed;
diff --git a/Engine/Source/Programs/AutomationTool/Turnkey/FileSource.cs b/Engine/Source/Programs/AutomationTool/Turnkey/FileSource.cs
--- a/Engine/Source/Programs/AutomationTool/Turnkey/FileSource.cs
+++ b/Engine/Source/Programs/AutomationTool/Turnkey/FileSource.cs
@@ -390,7 +390,7 @@ namespace Turnkey
 			return null;
 		}
 
-		public bool DownloadOrInstall(UnrealTargetPlatform Platform, ITurnkeyContext TurnkeyContext, DeviceInfo Device, bool bUnattended)
+		public bool DownloadOrInstall(UnrealTargetPlatform Platform, ITurnkeyContext TurnkeyContext, DeviceInfo Device, bool bUnattended, bool bSdkAlreadyInstalled)
 		{
 			// standard variables
 			TurnkeyUtils.SetVariable("Platform", Platform.ToString());
@@ -417,7 +417,7 @@ namespace Turnkey
 			}
 
 			// let the platform decide how to install
-			return AutomationTool.Platform.GetPlatform(Platform).InstallSDK(TurnkeyUtils.CommandUtilHelper, TurnkeyContext, Device, bUnattended);
+			return AutomationTool.Platform.GetPlatform(Platform).InstallSDK(TurnkeyUtils.CommandUtilHelper, TurnkeyContext, Device, bUnattended, bSdkAlreadyInstalled);
 		}
 
 
