Turnkey uses "p4 where" to identify the AutoSDK root directory.

 - should make it more reliable when working with non-trivial workspace mappings.

#rnx

diff --git a/Engine/Source/Programs/AutomationTool/Turnkey/Providers/PerforceCopyProvider.cs b/Engine/Source/Programs/AutomationTool/Turnkey/Providers/PerforceCopyProvider.cs
--- a/Engine/Source/Programs/AutomationTool/Turnkey/Providers/PerforceCopyProvider.cs
+++ b/Engine/Source/Programs/AutomationTool/Turnkey/Providers/PerforceCopyProvider.cs
@@ -22,9 +22,39 @@ namespace Turnkey
 				return null;
 			}
 
-			// calculate the output path
-			// @todo turnkey: use p4 where, because a depot mapping could be strange and make this invlalid
-			string OutputPath = Operation.Replace(ConnectedRoot, PerforceClient.RootPath);
+			// find the output path using p4 where because a depot mapping could be strange
+			string OutputPath = null;
+			IProcessResult P4Result = PerforceConnection.P4(string.Format("where {0}", Operation), AllowSpew:false);
+			if (P4Result.ExitCode == 0)
+			{
+				// sadly, this doesn't set an ExitCode, so we have to look at output
+				if (!P4Result.Output.Trim().EndsWith("not in client view."))
+				{
+					// p4 where returns a trio of Depot, Client and Local paths, separated by spaces and potentially quoted.
+					// more complicated mappings may have multiple trios so walk them to find one that matches
+					string[] Items = SharedUtils.ParseCommandLine(P4Result.Output);
+					if ((Items.Length % 3) == 0)
+					{
+						for (int Item = 0; Item < Items.Length; Item += 3)
+						{
+							// find a matching Depot path
+							if (Items[Item].Equals(Operation, StringComparison.OrdinalIgnoreCase))
+							{
+								// use the Local path and stop searching
+								OutputPath = Items[Item + 2];
+								break;
+							}
+						}
+					}
+				}
+			}
+
+			// no luck with p4 where
+			if (OutputPath == null)
+			{
+				OutputPath = Operation.Replace(ConnectedRoot, PerforceClient.RootPath);
+				TurnkeyUtils.Log("Unable to discover local path using p4 where {0}. Falling back to best guess {1}", Operation, OutputPath);
+			}
 
 			// turn //a/b/c/foo*/... to //a/b/c
 			OutputPath = GetDirectoryBeforeWildcard(OutputPath);
