Fix Gauntlet test search - Include original test name as first search critiria

#rnx

diff --git a/Engine/Source/Programs/AutomationTool/Gauntlet/Framework/Gauntlet.Utils.cs b/Engine/Source/Programs/AutomationTool/Gauntlet/Framework/Gauntlet.Utils.cs
--- a/Engine/Source/Programs/AutomationTool/Gauntlet/Framework/Gauntlet.Utils.cs
+++ b/Engine/Source/Programs/AutomationTool/Gauntlet/Framework/Gauntlet.Utils.cs
@@ -524,25 +524,21 @@ namespace Gauntlet
 		{
 
 			/// <summary>
-			/// Helper function that returns the type of an object based on namespace and name
+			/// Helper function that returns the type of an object based on namespaces and name
 			/// </summary>
-			/// <param name="Namespace"></param>
+			/// <param name="Namespaces"></param>
 			/// <param name="TestName"></param>
 			/// <returns></returns>
 			private static Type GetTypeForTest(string TestName, IEnumerable<string> Namespaces)
 			{
 				var SearchAssemblies = AppDomain.CurrentDomain.GetAssemblies();
 
-				// turn foo into [n1.foo, n2.foo, foo]
-				IEnumerable<string> FullNames;
+				// turn foo into [foo, n1.foo, n2.foo]
+				IEnumerable<string> FullNames = new[] { TestName };
 
 				if (Namespaces != null)
 				{
-					FullNames = Namespaces.Select(N => N + "." + TestName);
-				}
-				else
-				{
-					FullNames = new[] { TestName };
+					FullNames = FullNames.Concat(Namespaces.Select(N => N + "." + TestName));
 				}
 
 				Log.VeryVerbose("Will search {0} for test {1}", string.Join(" ", FullNames), TestName);
@@ -568,6 +564,8 @@ namespace Gauntlet
 				{
 					// Even tho the user might have specified N1.Foo it still might be Other.N1.Foo so only
 					// compare based on the number of namespaces that were specified.
+
+
 					foreach (var Type in CandidateTypes)
 					{
 						string[] UserNameComponents = UserTypeName.Split('.');
@@ -577,7 +575,6 @@ namespace Gauntlet
 
 						if (MissingUserComponents > 0)
 						{
-							// 
 							TypeNameComponents = TypeNameComponents.Skip(MissingUserComponents).ToArray();
 						}
 
