- Fixed issue with directory not existing when adding project shaders

diff --git a/Engine/Source/Programs/AutomationTool/CookedEditor/CookedEditor.Automation.cs b/Engine/Source/Programs/AutomationTool/CookedEditor/CookedEditor.Automation.cs
--- a/Engine/Source/Programs/AutomationTool/CookedEditor/CookedEditor.Automation.cs
+++ b/Engine/Source/Programs/AutomationTool/CookedEditor/CookedEditor.Automation.cs
@@ -481,7 +481,11 @@ public class MakeCookedEditor : BuildCommand
 		// project shaders
 		if (Context.bStageShaderDirs)
 		{
-			Context.NonUFSFilesToStage.AddRange(DirectoryReference.EnumerateFiles(DirectoryReference.Combine(Context.ProjectDirectory, "Shaders"), "*", SearchOption.AllDirectories));
+			DirectoryReference ProjectShaders = DirectoryReference.Combine(Context.ProjectDirectory, "Shaders");
+			if (DirectoryReference.Exists(ProjectShaders))
+			{
+				Context.NonUFSFilesToStage.AddRange(DirectoryReference.EnumerateFiles(ProjectShaders, "*", SearchOption.AllDirectories));
+			}		
 		}
 
 		StagePlatformExtensionFiles(Params, SC, Context, Context.ProjectDirectory);
