SyncProject.Automation: Correctly pass params to BuildCmd & OpenCmd

#rnx

diff --git a/Engine/Source/Programs/AutomationTool/Scripts/SyncProject.Automation.cs b/Engine/Source/Programs/AutomationTool/Scripts/SyncProject.Automation.cs
--- a/Engine/Source/Programs/AutomationTool/Scripts/SyncProject.Automation.cs
+++ b/Engine/Source/Programs/AutomationTool/Scripts/SyncProject.Automation.cs
@@ -315,8 +315,7 @@ class SyncProject : SyncProjectBase
 				}
 
 				BuildEditor BuildCmd = new BuildEditor();
-				BuildCmd.Clean = ParseParam("clean");
-				BuildCmd.ProjectName = ProjectArgForEditor;
+				BuildCmd.Params = this.Params;
 				ExitStatus = BuildCmd.Execute();
 			}
 
@@ -325,7 +324,7 @@ class SyncProject : SyncProjectBase
 				Log.TraceVerbose("Opening Editor for {0}", ProjectArgForEditor);
 
 				OpenEditor OpenCmd = new OpenEditor();
-				OpenCmd.ProjectName = ProjectArgForEditor;
+				OpenCmd.Params = this.Params;
 				ExitStatus = OpenCmd.Execute();
 			}
 
