Only adjust the CookedEditor default staged dir if it's not already been overriden.

diff --git a/Engine/Source/Programs/AutomationTool/CookedEditor/CookedEditor.Automation.cs b/Engine/Source/Programs/AutomationTool/CookedEditor/CookedEditor.Automation.cs
--- a/Engine/Source/Programs/AutomationTool/CookedEditor/CookedEditor.Automation.cs
+++ b/Engine/Source/Programs/AutomationTool/CookedEditor/CookedEditor.Automation.cs
@@ -1063,8 +1063,15 @@ public class MakeCookedEditor : BuildCommand
 			// cook and stage into our project, instead of the Engine's plugins
 			DirectoryReference BaseOutputDirectory = DirectoryReference.Combine(ProjectFile.Directory, "Saved", "CookedEditor");
 			string TargetPlatformName = ConfigHierarchy.GetIniPlatformName(Platform) + TargetPlatformList[0].CookFlavor;
-			Params.CookOutputDir = DirectoryReference.Combine(BaseOutputDirectory, "Cooked", TargetPlatformName).FullName;
-			Params.StageDirectoryParam = DirectoryReference.Combine(BaseOutputDirectory, "Staged").FullName;
+			// Only override the defaults if something hasn't already overridden the potentially from the commandline.
+			if (string.IsNullOrEmpty(Params.CookOutputDir))
+			{
+				Params.CookOutputDir = DirectoryReference.Combine(BaseOutputDirectory, "Cooked", TargetPlatformName).FullName;
+			}
+			if (string.IsNullOrEmpty(Params.StageDirectoryParam))
+			{
+				Params.StageDirectoryParam = DirectoryReference.Combine(BaseOutputDirectory, "Staged").FullName;
+			}
 
 			// make WindowsClient or LinuxGame, etc
 			string ReleaseTargetName = GetReleaseTargetName(Platform, ReleaseType);
