Fix a potential warning when two UAT commands are ran back to back in the build graph. This can result in files potentially not being immediately deleted.

diff --git a/Engine/Source/Programs/AutomationTool/Gauntlet/Platform/Windows/Gauntlet.TargetDeviceWindows.cs b/Engine/Source/Programs/AutomationTool/Gauntlet/Platform/Windows/Gauntlet.TargetDeviceWindows.cs
--- a/Engine/Source/Programs/AutomationTool/Gauntlet/Platform/Windows/Gauntlet.TargetDeviceWindows.cs
+++ b/Engine/Source/Programs/AutomationTool/Gauntlet/Platform/Windows/Gauntlet.TargetDeviceWindows.cs
@@ -327,7 +327,7 @@ namespace Gauntlet
 				{
 					if (File.Exists(ProcessLogFile))
 					{
-						File.Delete(ProcessLogFile);
+						EpicGames.Core.FileUtils.ForceDeleteFile(ProcessLogFile);
 					}
 				}
 				catch (Exception Ex)
