Fix UserDir path handling for Gauntlet Linux host

#rnx

Details
- Previously UserDir was set in such a way that multiple gauntlet processes launched on the same machine would have their log directories collide.  In the observed failure case, the client would launch after the server, stomping the server's log.  Then gauntlet would be unable to find the test results and report as "No tests executed".  This change just makes the path be the same as what Windows uses.

diff --git a/Engine/Source/Programs/AutomationTool/Gauntlet/Platform/Linux/Gauntlet.TargetDeviceLinux.cs b/Engine/Source/Programs/AutomationTool/Gauntlet/Platform/Linux/Gauntlet.TargetDeviceLinux.cs
--- a/Engine/Source/Programs/AutomationTool/Gauntlet/Platform/Linux/Gauntlet.TargetDeviceLinux.cs
+++ b/Engine/Source/Programs/AutomationTool/Gauntlet/Platform/Linux/Gauntlet.TargetDeviceLinux.cs
@@ -123,7 +123,7 @@ namespace Gauntlet
 			LocalCachePath = InCacheDir;
 			RunOptions = CommandUtils.ERunOptions.NoWaitForExit | CommandUtils.ERunOptions.NoLoggingOfRunCommand;
 
-			UserDir = Path.Combine(LocalCachePath, "UserDir");
+			UserDir = Path.Combine(LocalCachePath, string.Format("{0}_UserDir", Name));
             LocalDirectoryMappings = new Dictionary<EIntendedBaseCopyDirectory, string>();
 		}
 
