Explicitly close the log file reader when an application terminates or has exited. This is make sure all log file handles are closed when another device comes along to potentially delete this log file.

diff --git a/Engine/Source/Programs/AutomationTool/Gauntlet/Platform/Windows/Gauntlet.TargetDeviceWindows.cs b/Engine/Source/Programs/AutomationTool/Gauntlet/Platform/Windows/Gauntlet.TargetDeviceWindows.cs
--- a/Engine/Source/Programs/AutomationTool/Gauntlet/Platform/Windows/Gauntlet.TargetDeviceWindows.cs
+++ b/Engine/Source/Programs/AutomationTool/Gauntlet/Platform/Windows/Gauntlet.TargetDeviceWindows.cs
@@ -81,7 +81,6 @@ namespace Gauntlet
 			using (FileStream ProcessLog = File.Open(ProcessLogFile, FileMode.Open, FileAccess.Read, FileShare.ReadWrite))
 			{
 				StreamReader LogReader = new StreamReader(ProcessLog);
-
 				// Read until the process has exited
 				do
 				{
@@ -98,6 +97,10 @@ namespace Gauntlet
 					}
 				}
 				while (!HasExited);
+
+				LogReader.Close();
+				ProcessLog.Close();
+				ProcessLog.Dispose();
 			}
 		}
 
