Fix steam deck GetDeviceInfo returning true on devices that are not part of the ini list, which was causing issues in Linux automations for real devices (non-steam deck devices)

diff --git a/Engine/Source/Programs/AutomationTool/SteamDeck/SteamDeckSupport.cs b/Engine/Source/Programs/AutomationTool/SteamDeck/SteamDeckSupport.cs
--- a/Engine/Source/Programs/AutomationTool/SteamDeck/SteamDeckSupport.cs
+++ b/Engine/Source/Programs/AutomationTool/SteamDeck/SteamDeckSupport.cs
@@ -98,8 +98,12 @@ public static class SteamDeckSupport
 		{
 			Device = Devices.FirstOrDefault(x => x.Name.Equals(DeviceId, StringComparison.InvariantCultureIgnoreCase));
 		}
-		// if the device wasn't found, just use the DeviceId (it may be a DNS name, so we can't really verify it here)
-		IpAddr = Device == null ? DeviceId : Device.Id;
+
+		// if the device is still null, dont assume this is a valid steam deck device
+		if (Device == null)
+		{
+			return false;
+		}
 
 
 		// if -deviceuser was specified, prefer that, otherwise use what's in the registry
