Use platform-dependent LocalPrint and make sure output is redirected to stdout when running LowLevelTests so that local and build machine test runs print debug and fault information directly to stdout.
AssertionMacrosTest is printing out failed asserts using LocalPrint which is undesirable - implemented GBlockLocalOutput as a way to supress LocalPrint output.
Other fixes included:
- ability to select a low level test build based on any configuration, not just Development: useful for Debug builds
- New "printoutput" LLT Gauntlet parameter used for certain platforms to control output redirect to stdout

diff --git a/Engine/Source/Developer/LowLevelTestsRunner/Private/TestListeners/ConsoleListener.cpp b/Engine/Source/Developer/LowLevelTestsRunner/Private/TestListeners/ConsoleListener.cpp
--- a/Engine/Source/Developer/LowLevelTestsRunner/Private/TestListeners/ConsoleListener.cpp
+++ b/Engine/Source/Developer/LowLevelTestsRunner/Private/TestListeners/ConsoleListener.cpp
@@ -1,12 +1,17 @@
 // Copyright Epic Games, Inc. All Rights Reserved.
 
 #include "TestListeners/ConsoleListener.h"
+#include "Containers/UnrealString.h"
 
 using namespace Catch;
 
 class ConsoleListener : public EventListenerBase, public FOutputDevice
 {
-	using EventListenerBase::EventListenerBase; // inherit constructor
+public:
+	ConsoleListener(IConfig const* config): EventListenerBase(config)
+	{
+		m_preferences.shouldRedirectStdOut = true;
+	}
 private:
 	void testRunStarting(TestRunInfo const& testRunInfo) override {
 		// Register this event listener as an output device to enable reporting of UE_LOG, ensure etc
@@ -18,6 +23,27 @@ private:
 		GLog->RemoveOutputDevice(this);
 	}
 
+	void platformPrint(const TCHAR* Message, bool ForceOutput = false)
+	{
+		bool SavedBlockPrintValue = GBlockLocalPrint;
+		// Force: always print this event
+		if (ForceOutput)
+		{
+			TGuardValue<bool> BlockLocalPrint(GBlockLocalPrint, false);
+		}
+
+#if defined(PLATFORM_WINDOWS)
+		FGenericPlatformMisc::LocalPrint(Message);
+#else
+		FPlatformMisc::LocalPrint(Message);
+#endif
+
+		if (ForceOutput)
+		{
+			TGuardValue<bool> RestoreBlockLocalPrint(GBlockLocalPrint, SavedBlockPrintValue);
+		}
+	}
+
 	// FOutputDevice interface
 	virtual void Serialize(const TCHAR* V, ELogVerbosity::Type LogVerbosity, const class FName& Category) override
 	{
@@ -33,14 +59,10 @@ private:
 		// I tried this but unfortunately it didn't work, passing that option makes catch complain with the error message:
 		// "Verbosity level not supported by this reporter"
 
-		if (LogVerbosity <= DesiredLogVerbosity && !GLog->GetSuppressEventTag())
+		if (LogVerbosity <= DesiredLogVerbosity)
 		{
-			std::cout << TCHAR_TO_UTF8(*FText::FromName(Category).ToString())
-				<< "("
-				<< TCHAR_TO_UTF8(ToString(LogVerbosity))
-				<< "): "
-				<< TCHAR_TO_UTF8(V)
-				<< "\n";
+			FString OutputMessage = FString::Printf(TEXT("%s(%s): %s\n"), *FText::FromName(Category).ToString(), ToString(LogVerbosity), V);
+			platformPrint(OutputMessage.GetCharArray().GetData());
 		}
 	}
 
@@ -58,20 +80,25 @@ private:
 	void testCaseStarting(TestCaseInfo  const& TestInfo) override {
 		if (bGDebug)
 		{
-			std::cout << TestInfo.lineInfo.file << ":" << TestInfo.lineInfo.line << " with tags " << TestInfo.tagsAsString() << " \n";
+			FString InfoTestCaseStarting = FString::Printf(TEXT("%s:%d with tags %s\n"), *FString(TestInfo.lineInfo.file), TestInfo.lineInfo.line, *FString(TestInfo.tagsAsString().c_str()));			
+			platformPrint(InfoTestCaseStarting.GetCharArray().GetData(), true);
 		}
 	}
 
+	// TODO: Use UE structured logging for build machines
 	void testCaseEnded(TestCaseStats const& testCaseStats) override {
 		if (testCaseStats.totals.testCases.failed > 0)
 		{
-			std::cout << "* Error: Test case \"" << testCaseStats.testInfo->name << "\" failed \n";
+			FString ErrorTestCase = FString::Printf(TEXT("* Error: Test case \"%s\" failed \n"), *FString(testCaseStats.testInfo->name.c_str()));
+			platformPrint(ErrorTestCase.GetCharArray().GetData(), true);
 		}
 	}
 
+	// TODO: Use UE structured logging for build machines
 	void assertionEnded(AssertionStats const& assertionStats) override {
 		if (!assertionStats.assertionResult.succeeded()) {
-			std::cout << "* Error: Assertion \"" << assertionStats.assertionResult.getExpression() << "\" failed at " << assertionStats.assertionResult.getSourceInfo().file << ": " << assertionStats.assertionResult.getSourceInfo().line << "\n";
+			FString ErrorAssertion = FString::Printf(TEXT("* Error: Assertion \"%s\" failed at %s:%d\n"), *FString(assertionStats.assertionResult.getExpression().c_str()), *FString(assertionStats.assertionResult.getSourceInfo().file), assertionStats.assertionResult.getSourceInfo().line);
+			platformPrint(ErrorAssertion.GetCharArray().GetData(), true);
 		}
 	}
 };
diff --git a/Engine/Source/Developer/LowLevelTestsRunner/Public/TestListeners/ConsoleListener.h b/Engine/Source/Developer/LowLevelTestsRunner/Public/TestListeners/ConsoleListener.h
--- a/Engine/Source/Developer/LowLevelTestsRunner/Public/TestListeners/ConsoleListener.h
+++ b/Engine/Source/Developer/LowLevelTestsRunner/Public/TestListeners/ConsoleListener.h
@@ -6,11 +6,10 @@
 #include "CoreGlobals.h"
 #include "Logging/LogVerbosity.h"
 #include "Misc/OutputDeviceRedirector.h"
+#include "HAL/PlatformMisc.h"
 
 #include "TestGlobals.h"
 
 #include <catch2/catch_test_case_info.hpp>
 #include <catch2/reporters/catch_reporter_event_listener.hpp>
-#include <catch2/reporters/catch_reporter_registrars.hpp>
-
-#include <iostream>
\ No newline at end of file
+#include <catch2/reporters/catch_reporter_registrars.hpp>
\ No newline at end of file
diff --git a/Engine/Source/Programs/AutomationTool/LowLevelTests/RunLowLevelTests.cs b/Engine/Source/Programs/AutomationTool/LowLevelTests/RunLowLevelTests.cs
--- a/Engine/Source/Programs/AutomationTool/LowLevelTests/RunLowLevelTests.cs
+++ b/Engine/Source/Programs/AutomationTool/LowLevelTests/RunLowLevelTests.cs
@@ -54,7 +54,8 @@ namespace LowLevelTests
 			LowLevelTestsBuildSource BuildSource = new LowLevelTestsBuildSource(
 				ContextOptions.TestApp,
 				ContextOptions.Build,
-				ContextOptions.Platform);
+				ContextOptions.Platform,
+				ContextOptions.Configuration);
 
 			SetupDevices(TestPlatform, ContextOptions);
 
@@ -426,9 +427,9 @@ namespace LowLevelTests
 	{
 		bool CanSupportPlatform(UnrealTargetPlatform InPlatform);
 
-		LowLevelTestsBuild CreateBuild(UnrealTargetPlatform InPlatform, string InTestApp, string InBuildPath);
+		LowLevelTestsBuild CreateBuild(UnrealTargetPlatform InPlatform, UnrealTargetConfiguration InConfiguration, string InTestApp, string InBuildPath);
 
-		protected static string GetExecutable(UnrealTargetPlatform InPlatform, string InTestApp, string InBuildPath, string FileRegEx)
+		protected static string GetExecutable(UnrealTargetPlatform InPlatform, UnrealTargetConfiguration InConfiguration, string InTestApp, string InBuildPath, string FileRegEx)
 		{
 			IEnumerable<string> Executables = DirectoryUtils.FindFiles(InBuildPath, new Regex(FileRegEx));
 			foreach (string Executable in Executables)
@@ -437,21 +438,24 @@ namespace LowLevelTests
 				{
 					if (InBuildPath.ToLower().Contains(InTestApp.ToString().ToLower()))
 					{
-						return Path.GetRelativePath(InBuildPath, Executable);
+						if (InConfiguration == UnrealTargetConfiguration.Development || InBuildPath.ToLower().Contains(InConfiguration.ToString().ToLower()))
+						{
+							return Path.GetRelativePath(InBuildPath, Executable);
+						}
 					}
 				}
 			}
 			throw new AutomationException("Cannot find low level test executable for {0} in build path {1} for {2} using regex \"{3}\"", InPlatform, InBuildPath, InTestApp, FileRegEx);
 		}
 
-		protected static void CleanupUnusedFiles(UnrealTargetPlatform InPlatform, string InBuildPath)
+		protected static void CleanupUnusedFiles(UnrealTargetPlatform InPlatform, UnrealTargetConfiguration InConfiguration, string InBuildPath)
 		{
 			try
 			{
 				string[] BuildFiles = Directory.GetFiles(InBuildPath);
 				foreach (string BuildFile in BuildFiles)
 				{
-					if (new FileInfo(BuildFile).Extension == ".pdb")
+					if (new FileInfo(BuildFile).Extension == ".pdb" && InConfiguration != UnrealTargetConfiguration.Debug)
 					{
 						File.Delete(BuildFile);
 					}
@@ -471,12 +475,12 @@ namespace LowLevelTests
 			return InPlatform.IsInGroup(UnrealPlatformGroup.Desktop);
 		}
 
-		public LowLevelTestsBuild CreateBuild(UnrealTargetPlatform InPlatform, string InTestApp, string InBuildPath)
+		public LowLevelTestsBuild CreateBuild(UnrealTargetPlatform InPlatform, UnrealTargetConfiguration InConfiguration, string InTestApp, string InBuildPath)
 		{
 			string DesktopExecutableRegEx;
 			if (InPlatform.IsInGroup(UnrealPlatformGroup.Windows))
 			{
-				DesktopExecutableRegEx = @"\w+Tests.exe$";
+				DesktopExecutableRegEx = @"\w+Tests(?:-\w+)?(?:-\w+)?.exe$";
 			}
 			else if (InPlatform == UnrealTargetPlatform.Linux || InPlatform == UnrealTargetPlatform.Mac)
 			{
@@ -486,8 +490,8 @@ namespace LowLevelTests
 			{
 				throw new AutomationException("Cannot create build for non-desktop platform " + InPlatform);
 			}
-			string ExecutablePath = ILowLevelTestsBuildFactory.GetExecutable(InPlatform, InTestApp, InBuildPath, DesktopExecutableRegEx);
-			return new LowLevelTestsBuild(InPlatform, InBuildPath, ExecutablePath);
+			string ExecutablePath = ILowLevelTestsBuildFactory.GetExecutable(InPlatform, InConfiguration, InTestApp, InBuildPath, DesktopExecutableRegEx);
+			return new LowLevelTestsBuild(InPlatform, InConfiguration, InBuildPath, ExecutablePath);
 		}
 	}
 
@@ -523,22 +527,24 @@ namespace LowLevelTests
 		private ILowLevelTestsReporting LowLevelTestsReporting;
 
 		public UnrealTargetPlatform Platform { get; protected set; }
+		public UnrealTargetConfiguration Configuration { get; protected set; }
 		public LowLevelTestsBuild DiscoveredBuild { get; protected set; }
 
-		public LowLevelTestsBuildSource(string InTestApp, string InBuildPath, UnrealTargetPlatform InTargetPlatform)
+		public LowLevelTestsBuildSource(string InTestApp, string InBuildPath, UnrealTargetPlatform InTargetPlatform, UnrealTargetConfiguration InConfiguration)
 		{
 			TestApp = InTestApp;
 			Platform = InTargetPlatform;
 			BuildPath = InBuildPath;
-			InitBuildSource(InTestApp, InBuildPath, InTargetPlatform);
+			Configuration = InConfiguration;
+			InitBuildSource(InTestApp, InBuildPath, InTargetPlatform, InConfiguration);
 		}
 
-		protected void InitBuildSource(string InTestApp, string InBuildPath, UnrealTargetPlatform InTargetPlatform)
+		protected void InitBuildSource(string InTestApp, string InBuildPath, UnrealTargetPlatform InTargetPlatform, UnrealTargetConfiguration InConfiguration)
 		{
 			LowLevelTestsBuildFactory = Gauntlet.Utils.InterfaceHelpers.FindImplementations<ILowLevelTestsBuildFactory>(true)
 				.Where(B => B.CanSupportPlatform(InTargetPlatform))
 				.First();
-			DiscoveredBuild = LowLevelTestsBuildFactory.CreateBuild(InTargetPlatform, InTestApp, InBuildPath);
+			DiscoveredBuild = LowLevelTestsBuildFactory.CreateBuild(InTargetPlatform, InConfiguration, InTestApp, InBuildPath);
 			if (DiscoveredBuild == null)
 			{
 				throw new AutomationException("No builds were discovered at path {0} matching test app name {1} and target platform {2}", InBuildPath, InTestApp, InTargetPlatform);
@@ -589,12 +595,9 @@ namespace LowLevelTests
 
 	public class LowLevelTestsBuild : StagedBuild
 	{
-		public LowLevelTestsBuild(UnrealTargetPlatform InPlatform, string InBuildPath, string InExecutablePath)
-			: base(InPlatform, UnrealTargetConfiguration.Development, UnrealTargetRole.Client, InBuildPath, InExecutablePath)
+		public LowLevelTestsBuild(UnrealTargetPlatform InPlatform, UnrealTargetConfiguration InConfiguration, string InBuildPath, string InExecutablePath)
+			: base(InPlatform, InConfiguration, UnrealTargetRole.Client, InBuildPath, InExecutablePath)
 		{
-			Platform = InPlatform;
-			BuildPath = InBuildPath;
-			ExecutablePath = InExecutablePath;
 			Flags = BuildFlags.CanReplaceExecutable | BuildFlags.Loose;
 		}
 	}
diff --git a/Engine/Source/Runtime/Core/Private/Apple/ApplePlatformMisc.cpp b/Engine/Source/Runtime/Core/Private/Apple/ApplePlatformMisc.cpp
--- a/Engine/Source/Runtime/Core/Private/Apple/ApplePlatformMisc.cpp
+++ b/Engine/Source/Runtime/Core/Private/Apple/ApplePlatformMisc.cpp
@@ -73,8 +73,11 @@ FString FApplePlatformMisc::GetEnvironmentVariable(const TCHAR* VariableName)
 
 void FApplePlatformMisc::LocalPrint(const TCHAR* Message)
 {
-	//NsLog will out to all iOS output consoles, instead of just the Xcode console.
-	NSLog(@"[UE] %s", TCHAR_TO_UTF8(Message));
+	if (!GBlockLocalPrint)
+	{
+		//NsLog will out to all iOS output consoles, instead of just the Xcode console.
+		NSLog(@"[UE] %s", TCHAR_TO_UTF8(Message));
+	}
 }
 
 const TCHAR* FApplePlatformMisc::GetSystemErrorMessage(TCHAR* OutBuffer, int32 BufferCount, int32 Error)
diff --git a/Engine/Source/Runtime/Core/Private/GenericPlatform/GenericPlatformMisc.cpp b/Engine/Source/Runtime/Core/Private/GenericPlatform/GenericPlatformMisc.cpp
--- a/Engine/Source/Runtime/Core/Private/GenericPlatform/GenericPlatformMisc.cpp
+++ b/Engine/Source/Runtime/Core/Private/GenericPlatform/GenericPlatformMisc.cpp
@@ -887,13 +887,16 @@ void FGenericPlatformMisc::SetUTF8Output()
 
 void FGenericPlatformMisc::LocalPrint( const TCHAR* Str )
 {
+	if (!GBlockLocalPrint)
+	{
 #if PLATFORM_TCHAR_IS_CHAR16
-	printf("%s", TCHAR_TO_UTF8(Str));
+		printf("%s", TCHAR_TO_UTF8(Str));
 #elif PLATFORM_USE_LS_SPEC_FOR_WIDECHAR
-	printf("%ls", Str);
+		printf("%ls", Str);
 #else
-	wprintf(TEXT("%s"), Str);
+		wprintf(TEXT("%s"), Str);
 #endif
+	}
 }
 
 bool FGenericPlatformMisc::HasSeparateChannelForDebugOutput()
diff --git a/Engine/Source/Runtime/Core/Private/Misc/CoreGlobals.cpp b/Engine/Source/Runtime/Core/Private/Misc/CoreGlobals.cpp
--- a/Engine/Source/Runtime/Core/Private/Misc/CoreGlobals.cpp
+++ b/Engine/Source/Runtime/Core/Private/Misc/CoreGlobals.cpp
@@ -422,6 +422,9 @@ bool CORE_API			GIsGPUCrashed = false;
 /** Whether we should ignore the attached debugger. */
 CORE_API bool			GIgnoreDebugger = false;
 
+/** Whether local print will produce output. */
+CORE_API bool			GBlockLocalPrint = false;
+
 #endif // #if !UE_BUILD_SHIPPING
 
 bool GetEmitDrawEvents()
diff --git a/Engine/Source/Runtime/Core/Private/Windows/WindowsPlatformMisc.cpp b/Engine/Source/Runtime/Core/Private/Windows/WindowsPlatformMisc.cpp
--- a/Engine/Source/Runtime/Core/Private/Windows/WindowsPlatformMisc.cpp
+++ b/Engine/Source/Runtime/Core/Private/Windows/WindowsPlatformMisc.cpp
@@ -937,7 +937,10 @@ void FWindowsPlatformMisc::SetUTF8Output()
 void FWindowsPlatformMisc::LocalPrint( const TCHAR *Message )
 {
 #if USE_DEBUG_LOGGING
-	OutputDebugString(Message);
+	if (!GBlockLocalPrint)
+	{
+		OutputDebugString(Message);
+	}
 #endif
 }
 
diff --git a/Engine/Source/Runtime/Core/Public/CoreGlobals.h b/Engine/Source/Runtime/Core/Public/CoreGlobals.h
--- a/Engine/Source/Runtime/Core/Public/CoreGlobals.h
+++ b/Engine/Source/Runtime/Core/Public/CoreGlobals.h
@@ -543,6 +543,9 @@ extern CORE_API bool GEnableVREditorHacks;
 /** Whether we should ignore the attached debugger. */
 extern CORE_API bool GIgnoreDebugger;
 
+/** Whether local print will produce output. */
+extern CORE_API bool GBlockLocalPrint;
+
 #endif // #if !UE_BUILD_SHIPPING
 
 enum class ETaskTag : int32
diff --git a/Engine/Source/Runtime/Core/Public/GenericPlatform/GenericPlatformMisc.h b/Engine/Source/Runtime/Core/Public/GenericPlatform/GenericPlatformMisc.h
--- a/Engine/Source/Runtime/Core/Public/GenericPlatform/GenericPlatformMisc.h
+++ b/Engine/Source/Runtime/Core/Public/GenericPlatform/GenericPlatformMisc.h
@@ -948,7 +948,7 @@ public:
 	/** Sets the default output to UTF8 */
 	static void SetUTF8Output();
 
-	/** Prints string to the default output */
+	/** Prints string to the default output, unless GBlockLocalPrint is set to true */
 	static void LocalPrint( const TCHAR* Str );
 
 	/** Whether LocalPrint can be called from any thread without overlapping */
diff --git a/Engine/Source/Runtime/Core/Tests/Misc/AssertionMacrosTest.cpp b/Engine/Source/Runtime/Core/Tests/Misc/AssertionMacrosTest.cpp
--- a/Engine/Source/Runtime/Core/Tests/Misc/AssertionMacrosTest.cpp
+++ b/Engine/Source/Runtime/Core/Tests/Misc/AssertionMacrosTest.cpp
@@ -3,6 +3,7 @@
 #if WITH_LOW_LEVEL_TESTS && DO_ENSURE
 
 #include "Misc/AssertionMacros.h"
+#include "GenericPlatform/GenericPlatformMisc.h"
 
 #include "TestHarness.h"
 #include "TestMacros/Assertions.h"
@@ -23,7 +24,7 @@ static void MultipleEnsuresFailedDifferent();
 TEST_CASE("Core::Misc::AssertionMacros", "[Core][Misc][AssertionMacros][Ensure]")
 {
 	TGuardValue<bool> IgnoreDebugger(GIgnoreDebugger, true);
-	GLog->SetSuppressEventTag(true);
+	TGuardValue<bool> BlockLocalPrint(GBlockLocalPrint, true);
 	SECTION("Ensure not triggered with REQUIRE_NOENSURE")
 	{
 		REQUIRE_NOENSURE(OneEnsureNonFailed());
@@ -69,7 +70,6 @@ TEST_CASE("Core::Misc::AssertionMacros", "[Core][Misc][AssertionMacros][Ensure]"
 	{
 		CHECK_ENSURE(MultipleEnsuresFailedDifferent());
 	}
-	GLog->SetSuppressEventTag(false);
 }
 
 static void OneEnsureNonFailed()
