AutomationTool: Fix passing the target type as the

#rnx

diff --git a/Engine/Source/Programs/AutomationTool/Scripts/BuildTarget.Automation.cs b/Engine/Source/Programs/AutomationTool/Scripts/BuildTarget.Automation.cs
--- a/Engine/Source/Programs/AutomationTool/Scripts/BuildTarget.Automation.cs
+++ b/Engine/Source/Programs/AutomationTool/Scripts/BuildTarget.Automation.cs
@@ -123,7 +123,7 @@ namespace AutomationTool
 			if (!string.IsNullOrEmpty(ProjectName))
 			{
 				// find the project
-				ProjectFile = ProjectUtils.FindProjectFileFromName(ProjectName);				
+				ProjectFile = ProjectUtils.FindProjectFileFromName(ProjectName);
 
 				if (ProjectFile == null)
 				{
@@ -248,7 +248,7 @@ namespace AutomationTool
 					new SimpleTargetInfo("UnrealGame", TargetType.Game),
 					new SimpleTargetInfo("UnrealClient", TargetType.Client),
 					new SimpleTargetInfo("UnrealServer", TargetType.Server),
-				};			
+				};
 			}
 
 
@@ -275,6 +275,14 @@ namespace AutomationTool
 										.Where(T => string.CompareOrdinal(T.TargetName, 0, ProjectName, 0, 1) == 0)
 										.Where(T => T.TargetName.IndexOf(InTargetName, StringComparison.OrdinalIgnoreCase) > 0)
 										.FirstOrDefault();
+
+						// Try to find a target where the target type matches the target name exactly such as "Editor" or "Game"
+						if (ProjectTarget == null && string.CompareOrdinal(InTargetName, "Program") != 0)
+						{
+							ProjectTarget = MatchingTargetTypes
+											.Where(T => string.CompareOrdinal(T.Type.ToString(), InTargetName) == 0)
+											.FirstOrDefault();
+						}
 					}
 				}
 			}		
