When looking for a Gauntlet exit code check for -gauntlet= and not -gauntlet. It's the former that indicates there will be code in the process that is orchestrating the test.

diff --git a/Engine/Source/Programs/AutomationTool/Gauntlet/Unreal/Base/Gauntlet.UnrealTestNode.cs b/Engine/Source/Programs/AutomationTool/Gauntlet/Unreal/Base/Gauntlet.UnrealTestNode.cs
--- a/Engine/Source/Programs/AutomationTool/Gauntlet/Unreal/Base/Gauntlet.UnrealTestNode.cs
+++ b/Engine/Source/Programs/AutomationTool/Gauntlet/Unreal/Base/Gauntlet.UnrealTestNode.cs
@@ -1475,8 +1475,8 @@ namespace Gauntlet
 				return UnrealProcessResult.ExitOk;
 			}
 
-			bool WasGauntletTest = InArtifacts.SessionRole.CommandLine.ToLower().Contains("-gauntlet");
-			// ok, process appears to have exited for no good reason so try to divine a result...
+			bool WasGauntletTest = InArtifacts.SessionRole.CommandLine.ToLower().Contains("-gauntlet=");
+			// ok, process was using a gauntlet controller so see if there's a result divine a result...
 			if (WasGauntletTest)
 			{
 				if (InLog.HasTestExitCode == false)
