ensure EditorOpenOrder.log files are not used for IoStore ordering.

#rnx

diff --git a/Engine/Source/Programs/AutomationTool/Scripts/CopyBuildToStagingDirectory.Automation.cs b/Engine/Source/Programs/AutomationTool/Scripts/CopyBuildToStagingDirectory.Automation.cs
--- a/Engine/Source/Programs/AutomationTool/Scripts/CopyBuildToStagingDirectory.Automation.cs
+++ b/Engine/Source/Programs/AutomationTool/Scripts/CopyBuildToStagingDirectory.Automation.cs
@@ -3085,7 +3085,7 @@ namespace AutomationScripts
 					}
 				}
 
-				List<OrderFile> FinalOrderFiles = OrderFiles.FindAll(x => (x.OrderType != OrderFile.OrderFileType.Cooker));
+				List<OrderFile> FinalOrderFiles = OrderFiles.FindAll(x => (x.OrderType != OrderFile.OrderFileType.Cooker && x.OrderType != OrderFile.OrderFileType.Editor));
 				if (bUseSecondaryOrder && OrderFiles.Count >= 1)
 				{
 					FinalOrderFiles.AddRange(OrderFiles.FindAll(x => x.OrderType == OrderFile.OrderFileType.Cooker));
