Xbox:  Increase initial wait time on PLM test to account for long startup times experienced by XB1 on Horde

- PLM test is failing sporadically on Horde due to the first suspend attempt being made before any frames have been displayed
    - locally I see ~5 secs before the app is displaying frames but this is ~20 more on Horde for some runs
    - even successful runs are nearing the current timeout (30s) on XB1
    - unclear why some Horde startups are so slow so just increase the initial wait to 45 secs to give us some headroom

diff --git a/Engine/Source/Programs/AutomationTool/Gauntlet/Unreal/Automation/UE.PLMTest.cs b/Engine/Source/Programs/AutomationTool/Gauntlet/Unreal/Automation/UE.PLMTest.cs
--- a/Engine/Source/Programs/AutomationTool/Gauntlet/Unreal/Automation/UE.PLMTest.cs
+++ b/Engine/Source/Programs/AutomationTool/Gauntlet/Unreal/Automation/UE.PLMTest.cs
@@ -15,7 +15,7 @@ namespace UE
 	{
 		/// How long to wait before starting
 		[AutoParam]
-		public int InitialWait = 30;
+		public int InitialWait = 45;
 
 		/// How long to wait between loops
 		[AutoParam]
