Add logging to launch of tests in UnrealSession that display the role type

diff --git a/Engine/Source/Programs/AutomationTool/Gauntlet/Unreal/Base/Gauntlet.UnrealSession.cs b/Engine/Source/Programs/AutomationTool/Gauntlet/Unreal/Base/Gauntlet.UnrealSession.cs
--- a/Engine/Source/Programs/AutomationTool/Gauntlet/Unreal/Base/Gauntlet.UnrealSession.cs
+++ b/Engine/Source/Programs/AutomationTool/Gauntlet/Unreal/Base/Gauntlet.UnrealSession.cs
@@ -813,6 +813,7 @@ namespace Gauntlet
 
 						try
 						{
+							Log.Info("Starting {0} on {1}", InstallRoleKV.Value, CurrentInstall.Device);
 							IAppInstance Instance = CurrentInstall.Run();
 							IDeviceUsageReporter.RecordStart(Instance.Device.Name, Instance.Device.Platform, IDeviceUsageReporter.EventType.Test);
 
