UAT: Limit BuildCookRun time to wait for a dedicated server process to start logging to 10 seconds since servers (in shipping) may run with logging disabled (improvement for 19645492).

#rnx

diff --git a/Engine/Source/Programs/AutomationTool/Scripts/RunProjectCommand.Automation.cs b/Engine/Source/Programs/AutomationTool/Scripts/RunProjectCommand.Automation.cs
--- a/Engine/Source/Programs/AutomationTool/Scripts/RunProjectCommand.Automation.cs
+++ b/Engine/Source/Programs/AutomationTool/Scripts/RunProjectCommand.Automation.cs
@@ -90,11 +90,29 @@ namespace AutomationScripts
 			}
 		}
 
-		private static bool WaitForProcessReady(IProcessResult Process, string Name, string LogFile, string[] ReadyTexts = null)
+		private static bool WaitForProcessReady(IProcessResult Process, string Name, string LogFile, string[] ReadyTexts, int MaxLogWaitTimeInSeconds = -1)
 		{
+			var StartTime = DateTime.UtcNow;
 			while (!FileExists(LogFile) && !Process.HasExited)
 			{
-				LogInformation("Waiting for {0} logging process to start...", Name);
+				if (MaxLogWaitTimeInSeconds > 0)
+				{
+					var Duration = (DateTime.UtcNow - StartTime).Seconds;
+					var TimeLeft = MaxLogWaitTimeInSeconds - Duration;
+					if (TimeLeft <= 0)
+					{
+						LogWarning("Giving up waiting for {0} to start logging, it may run with logging disabled.", Name);
+						return false;
+					}
+					else
+					{
+						LogInformation("Waiting for {0} seconds for {1} logging process to start...", TimeLeft, Name);
+					}
+				}
+				else
+				{
+					LogInformation("Waiting for {0} logging process to start...", Name);
+				}
 				Thread.Sleep(2000);
 			}
 
@@ -204,7 +222,7 @@ namespace AutomationScripts
 			if (CookServerProcess != null)
 			{
 				WaitForProcessReady(CookServerProcess, "Cook server", CookServerLogFile,
-						new string[] {"Unreal Network File Server is ready", "COTF server is ready"});
+					new string[] {"Unreal Network File Server is ready"});
 			}
 
 			if (Params.DedicatedServer && !Params.SkipServer)
@@ -224,8 +242,9 @@ namespace AutomationScripts
 				}
 
 				DedicatedServerProcess = RunDedicatedServer(Params, DedicatedServerLogFile, Params.RunCommandline);
+				int MaxLogCreationWaitTimeInSeconds = 10;
 				WaitForProcessReady(DedicatedServerProcess, "Dedicated server", DedicatedServerLogFile,
-						new string[] {"Game Engine Initialized"});
+					new string[] {"Game Engine Initialized"}, MaxLogCreationWaitTimeInSeconds);
 			}
 
 			if (!Params.NoClient)
