Add "host is down" I/O exception to systemic errors.

#rnx

diff --git a/Engine/Source/Programs/AutomationTool/AutomationUtils/Matchers/SystemicEventMatcher.cs b/Engine/Source/Programs/AutomationTool/AutomationUtils/Matchers/SystemicEventMatcher.cs
--- a/Engine/Source/Programs/AutomationTool/AutomationUtils/Matchers/SystemicEventMatcher.cs
+++ b/Engine/Source/Programs/AutomationTool/AutomationUtils/Matchers/SystemicEventMatcher.cs
@@ -17,6 +17,8 @@ namespace AutomationUtils.Matchers
 
 		static readonly Regex s_roboMerge = new Regex(@"RoboMerge\/gates.*already locked on Commit Server by buildmachine");
 
+		static readonly Regex s_hostDown = new Regex(@"ERROR: System\.IO\.IOException: Host is down");
+
 		public LogEventMatch? Match(ILogCursor cursor)
 		{
 			if (cursor.IsMatch(s_ddc))
@@ -36,6 +38,10 @@ namespace AutomationUtils.Matchers
 			{
 				return new LogEventBuilder(cursor).ToMatch(LogEventPriority.Low, LogLevel.Information, KnownLogEvents.Systemic_RoboMergeGateLocked);
 			}
+			if (cursor.IsMatch(s_hostDown))
+			{
+				return new LogEventBuilder(cursor).ToMatch(LogEventPriority.Low, LogLevel.Information, KnownLogEvents.Systemic_HostDownIOException);
+			}
 			if (cursor.Contains("LogXGEController: Warning: XGEControlWorker.exe does not exist"))
 			{
 				return new LogEventBuilder(cursor).ToMatch(LogEventPriority.High, LogLevel.Information, KnownLogEvents.Systemic_MissingXgeControlWorker);
diff --git a/Engine/Source/Programs/Shared/EpicGames.Core/KnownLogEvents.cs b/Engine/Source/Programs/Shared/EpicGames.Core/KnownLogEvents.cs
--- a/Engine/Source/Programs/Shared/EpicGames.Core/KnownLogEvents.cs
+++ b/Engine/Source/Programs/Shared/EpicGames.Core/KnownLogEvents.cs
@@ -220,6 +220,11 @@ namespace EpicGames.Core
 		/// </summary>
 		public static EventId Systemic_LogEventMatcher { get; } = new EventId(763);
 
+		/// <summary>
+		/// Host is down during I/O operation, usually meaning a share isn't accessible.
+		/// </summary>
+		public static EventId Systemic_HostDownIOException { get; } = new EventId(764);
+
 		/// <summary>
 		/// Maximum code for systemic events. Add new events in the 700-799 range.
 		/// </summary>
