Retry various file system operations that could fail for PS5 file system.

diff --git a/Engine/Source/Programs/AutomationTool/Gauntlet/Framework/Gauntlet.Utils.cs b/Engine/Source/Programs/AutomationTool/Gauntlet/Framework/Gauntlet.Utils.cs
--- a/Engine/Source/Programs/AutomationTool/Gauntlet/Framework/Gauntlet.Utils.cs
+++ b/Engine/Source/Programs/AutomationTool/Gauntlet/Framework/Gauntlet.Utils.cs
@@ -932,17 +932,17 @@ namespace Gauntlet
 				// find all files. If a directory get them all, else use the pattern/regex
 				if (Options.IsDirectoryPattern)
 				{
-					SourceFiles = SourceDir.GetFiles("*", Options.Recursive ? SearchOption.AllDirectories : SearchOption.TopDirectoryOnly);
+					SourceFiles = GetFiles(SourceDir, "*", Options.Recursive ? SearchOption.AllDirectories : SearchOption.TopDirectoryOnly);
 				}
 				else
 				{
 					if (Options.Regex == null)
 					{
-						SourceFiles = SourceDir.GetFiles(Options.Pattern, Options.Recursive ? SearchOption.AllDirectories : SearchOption.TopDirectoryOnly);
+						SourceFiles = GetFiles(SourceDir, Options.Pattern, Options.Recursive ? SearchOption.AllDirectories : SearchOption.TopDirectoryOnly);
 					}
 					else
 					{
-						SourceFiles = SourceDir.GetFiles("*", Options.Recursive ? SearchOption.AllDirectories : SearchOption.TopDirectoryOnly);
+						SourceFiles = GetFiles(SourceDir, "*", Options.Recursive ? SearchOption.AllDirectories : SearchOption.TopDirectoryOnly);
 
 						SourceFiles = SourceFiles.Where(F => Options.Regex.IsMatch(F.Name));
 					}
@@ -953,7 +953,7 @@ namespace Gauntlet
 
 				if (IsMirroring)
 				{
-					DestFiles = DestDir.GetFiles("*", SearchOption.AllDirectories);
+					DestFiles = GetFiles(DestDir, "*", SearchOption.AllDirectories);
 
 					foreach (FileInfo Info in DestFiles)
 					{
@@ -1069,7 +1069,7 @@ namespace Gauntlet
 						{
 							// avoid an UnauthorizedAccessException by making sure file isn't read only
 							DestInfo.IsReadOnly = false;
-							DestInfo.Delete();
+							Delete(DestInfo);
 						}
 						catch (Exception Ex)
 						{
@@ -1080,13 +1080,13 @@ namespace Gauntlet
 					// delete empty directories
 					DirectoryInfo DestDirInfo = new DirectoryInfo(DestDirPath);
 
-					DirectoryInfo[] AllSubDirs = DestDirInfo.GetDirectories("*", SearchOption.AllDirectories);
+					DirectoryInfo[] AllSubDirs = GetDirectories(DestDirInfo, "*", SearchOption.AllDirectories);
 
 					foreach (DirectoryInfo SubDir in AllSubDirs)
 					{
 						try
 						{
-							if (SubDir.GetFiles().Length == 0 && SubDir.GetDirectories().Length == 0)
+							if (GetFiles(SubDir).Length == 0 && GetDirectories(SubDir).Length == 0)
 							{
 								if (Options.Verbose)
 								{
@@ -1097,7 +1097,7 @@ namespace Gauntlet
 									Log.Verbose("Deleting empty dir {0}", SubDir.FullName);
 								}
 
-								SubDir.Delete(true);
+								Delete(SubDir, true);
 							}
 						}
 						catch (Exception Ex)
@@ -1315,10 +1315,10 @@ namespace Gauntlet
 					return;
 				}
 
-				foreach (DirectoryInfo SubDir in Di.GetDirectories())
+				foreach (DirectoryInfo SubDir in GetDirectories(Di))
 				{
 					bool HasFile = 
-						SubDir.GetFiles().Where(F => {
+						GetFiles(SubDir).Where(F => {
 							int DaysOld = (DateTime.Now - F.LastWriteTime).Days;				
 							
 							if (DaysOld >= Days)
@@ -1336,7 +1336,7 @@ namespace Gauntlet
 						Log.Info("Removing old directory {0}", SubDir.Name);
 						try
 						{
-							SubDir.Delete(true);
+							Delete(SubDir, true);
 						}
 						catch (Exception Ex)
 						{
@@ -1350,6 +1350,154 @@ namespace Gauntlet
 				}				
 			}
 
+			/// <summary>
+			/// Wrap DirectoryInfo.GetFiles() and retry once if exception is "A retry should be performed."
+			/// </summary>
+			public static FileInfo[] GetFiles(DirectoryInfo Directory)
+			{
+				FileInfo[] Files = null;
+				try
+				{
+					Files = Directory.GetFiles();
+				}
+				catch (Exception ex)
+				{
+					if (ex.ToString().Contains("A retry should be performed"))
+					{
+						Log.Info("Retrying Directory.GetFiles() once.");
+						Files = Directory.GetFiles();
+					}
+				}
+				return Files;
+			}
+
+			/// <summary>
+			/// Wrap DirectoryInfo.GetFiles(Pattern, SearchOption) and retry once if exception is "A retry should be performed."
+			/// </summary>
+			public static FileInfo[] GetFiles(DirectoryInfo Directory, string Pattern, SearchOption SearchOption)
+			{
+				FileInfo[] Files = null;
+				try
+				{
+					Files = Directory.GetFiles(Pattern, SearchOption);
+				}
+				catch (Exception ex)
+				{
+					if (ex.ToString().Contains("A retry should be performed"))
+					{
+						Log.Info("Retrying Directory.GetFiles(Pattern, SearchOption) once.");
+						Files = Directory.GetFiles(Pattern, SearchOption);
+					}
+				}
+				return Files;
+			}
+
+			/// <summary>
+			/// Wrap DirectoryInfo.GetDirectories() and retry once if exception is "A retry should be performed."
+			/// </summary>
+			public static DirectoryInfo[] GetDirectories(DirectoryInfo Directory)
+			{
+				DirectoryInfo[] Directories = null;
+				try
+				{
+					Directories = Directory.GetDirectories();
+				}
+				catch (Exception ex)
+				{
+					if (ex.ToString().Contains("A retry should be performed"))
+					{
+						Log.Info("Retrying Directory.GetDirectories() once.");
+						Directories = Directory.GetDirectories();
+					}
+				}
+				return Directories;
+			}
+
+			/// <summary>
+			/// Wrap DirectoryInfo.GetDirectories(Pattern, SearchOption) and retry once if exception is "A retry should be performed."
+			/// </summary>
+			public static DirectoryInfo[] GetDirectories(DirectoryInfo Directory, string Pattern, SearchOption SearchOption)
+			{
+				DirectoryInfo[] Directories = null;
+				try
+				{
+					Directories = Directory.GetDirectories(Pattern, SearchOption);
+				}
+				catch (Exception ex)
+				{
+					if (ex.ToString().Contains("A retry should be performed"))
+					{
+						Log.Info("Retrying Directory.GetDirectories(Pattern, SearchOption) once.");
+						Directories = Directory.GetDirectories(Pattern, SearchOption);
+					}
+				}
+				return Directories;
+			}
+
+
+			/// <summary>
+			/// Wrap DirectoryInfo.GetFileSystemInfos(Pattern, SearchOption) and retry once if exception is "A retry should be performed."
+			/// </summary>
+			/// <param name="Directory"></param>
+			/// <param name="Pattern"></param>
+			/// <param name="SearchOption"></param>
+			/// <returns></returns>
+			public static FileSystemInfo[] GetFileSystemInfos(DirectoryInfo Directory, string Pattern, SearchOption SearchOption)
+			{
+				FileSystemInfo[] FileInfos = null;
+				try
+				{
+					FileInfos = Directory.GetFileSystemInfos(Pattern, SearchOption);
+				}
+				catch (Exception ex)
+				{
+					if (ex.ToString().Contains("A retry should be performed"))
+					{
+						Log.Info("Retrying Directory.GetFileSystemInfos(Pattern, SearchOption) once.");
+						FileInfos = Directory.GetFileSystemInfos(Pattern, SearchOption);
+					}
+				}
+				return FileInfos;
+			}
+
+			/// <summary>
+			/// Wrap DirectoryInfo.Delete(Force) and retry once if exception is "A retry should be performed."
+			/// </summary>
+			public static void Delete(DirectoryInfo Directory, bool Force)
+			{
+				try
+				{
+					Directory.Delete(Force);
+				}
+				catch (Exception ex)
+				{
+					if (ex.ToString().Contains("A retry should be performed"))
+					{
+						Log.Info("Retrying Directory.Delete(Force) once.");
+						Directory.Delete(Force);
+					}
+				}
+			}
+
+			/// <summary>
+			/// Wrap FileInfo.Delete() and retry once if exception is "A retry should be performed."
+			/// </summary>
+			public static void Delete(FileInfo File)
+			{
+				try
+				{
+					File.Delete();
+				}
+				catch (Exception ex)
+				{
+					if (ex.ToString().Contains("A retry should be performed"))
+					{
+						Log.Info("Retrying File.Delete() once.");
+						File.Delete();
+					}
+				}
+			}
+
 			/// <summary>
 			/// Return the fully qualified location of file path, including the Long path prefix if necessary
 			/// </summary>
@@ -1374,7 +1522,7 @@ namespace Gauntlet
 
 				DirectoryInfo Di = new DirectoryInfo(InPath);
 
-				var Files = Di.GetFiles().Where(f => Extensions.Contains(f.Extension.ToLower()));
+				var Files = SystemHelpers.GetFiles(Di).Where(f => Extensions.Contains(f.Extension.ToLower()));
 
 				return Files;
 			}
@@ -1513,20 +1661,20 @@ namespace Gauntlet
 
 					foreach (FileInfo File in FilesToCleanUp)
 					{
-						File.Delete();
+						SystemHelpers.Delete(File);
 					}
 				}
-				catch (System.Exception Ex)
+				catch (Exception Ex)
 				{
 					Log.Warning("ConvertImages failed: {0}", Ex);
 					try
 					{
 						if (DeleteOriginals)
 						{
-							Files.ToList().ForEach(F => F.Delete());
+							Files.ToList().ForEach(F => SystemHelpers.Delete(F));
 						}
 					}
-					catch (System.Exception e)
+					catch (Exception e)
 					{
 						Log.Warning("Cleaning up original files failed: {0}", e);
 					}
@@ -1607,7 +1755,7 @@ namespace Gauntlet
 		{
 			List<DirectoryInfo> Found = new List<DirectoryInfo>();
 
-			List<DirectoryInfo> Candidates = new List<DirectoryInfo>(new DirectoryInfo(BaseDir).GetDirectories());
+			List<DirectoryInfo> Candidates = new List<DirectoryInfo>(Utils.SystemHelpers.GetDirectories(new DirectoryInfo(BaseDir)));
 
 			Regex Pattern = new Regex(RegexPattern, RegexOptions.IgnoreCase);
 
@@ -1620,7 +1768,7 @@ namespace Gauntlet
 				Found.AddRange(MatchingDirs);
 
 				// recurse
-				Candidates = Candidates.SelectMany(D => D.GetDirectories()).ToList();
+				Candidates = Candidates.SelectMany(D => Utils.SystemHelpers.GetDirectories(D)).ToList();
 
 			} while (Candidates.Any() && (CurrentDepth++ < RecursionDepth || RecursionDepth == -1)); ;
 
@@ -1647,12 +1795,12 @@ namespace Gauntlet
 			do
 			{
 				// check for matching files in this set of directories
-				IEnumerable<FileInfo> MatchingFiles = CandidateDirs.SelectMany(D => D.GetFiles()).Where(F => Pattern.IsMatch(F.Name));
+				IEnumerable<FileInfo> MatchingFiles = CandidateDirs.SelectMany(D => Utils.SystemHelpers.GetFiles(D)).Where(F => Pattern.IsMatch(F.Name));
 
 				Found.AddRange(MatchingFiles);
 
 				// recurse into this set of directories
-				CandidateDirs = CandidateDirs.SelectMany(D => D.GetDirectories()).ToList();
+				CandidateDirs = CandidateDirs.SelectMany(D => Utils.SystemHelpers.GetDirectories(D)).ToList();
 
 			} while (CandidateDirs.Any() && (CurrentDepth++ < RecursionDepth || RecursionDepth == -1)); ;
 
