- Added architectures flag to BuildPlugin, mostly for Mac users to choose what architecture their plugin supports

diff --git a/Engine/Source/Programs/AutomationTool/Scripts/BuildPluginCommand.Automation.cs b/Engine/Source/Programs/AutomationTool/Scripts/BuildPluginCommand.Automation.cs
--- a/Engine/Source/Programs/AutomationTool/Scripts/BuildPluginCommand.Automation.cs
+++ b/Engine/Source/Programs/AutomationTool/Scripts/BuildPluginCommand.Automation.cs
@@ -1,4 +1,4 @@
-// Copyright Epic Games, Inc. All Rights Reserved.
+﻿// Copyright Epic Games, Inc. All Rights Reserved.
 using System;
 using System.Collections.Generic;
 using System.IO;
@@ -20,11 +20,13 @@ using UnrealBuildBase;
 [Help("StrictIncludes", "Disables precompiled headers and unity build in order to check all source files have self-contained headers.")]
 [Help("EngineDir=<RootDirectory>", "Root Directory of the engine that will be used to build plugin(s) (optional)")]
 [Help("Unversioned", "Do not embed the current engine version into the descriptor")]
+[Help("Architecture_<Platform>=<Architecture[s]>", "Control architecture to compile for a platform (eg. -Architecture_Mac=arm64+x86). Default is to use UBT defaults for the platform.")]
 public sealed class BuildPlugin : BuildCommand
 {
 	const string AndroidArchitectures = "armv7+arm64";
 	string UnrealBuildToolDllRelativePath = @"Engine/Binaries/DotNET/UnrealBuildTool/UnrealBuildTool.dll";
 	FileReference UnrealBuildToolDll;
+	static private Dictionary<UnrealTargetPlatform, string> PlatformToArchitectureMap = new Dictionary<UnrealTargetPlatform, string>();
 
 	public override void ExecuteBuild()
 	{
@@ -110,6 +112,12 @@ public sealed class BuildPlugin : BuildCommand
 			AdditionalArgs.Append(" -NoPCH -NoSharedPCH -DisableUnity");
 		}
 
+		// check if any architectures were specified
+		foreach (UnrealTargetPlatform Platform in UnrealTargetPlatform.GetValidPlatforms())
+		{
+			PlatformToArchitectureMap[Platform] = ParseParamValue($"architecture_{Platform}");
+		}
+
 		// Compile the plugin for all the target platforms
 		List<UnrealTargetPlatform> HostPlatforms = GetHostPlatforms(this);
 		List<UnrealTargetPlatform> TargetPlatforms = GetTargetPlatforms(this, BuildHostPlatform.Current.Platform);
@@ -354,9 +362,15 @@ public sealed class BuildPlugin : BuildCommand
 				ManifestFileNames.Add(ManifestFileName);
 				
 				string Arguments = String.Format("-plugin={0} -iwyu -noubtmakefiles -manifest={1} -nohotreload", CommandUtils.MakePathSafeToUseWithCommandLine(HostProjectPluginFile.FullName), CommandUtils.MakePathSafeToUseWithCommandLine(ManifestFileName.FullName));
-				if (Platform == UnrealTargetPlatform.Android)
+
+				string SpecifiedArchitecture = PlatformToArchitectureMap[Platform];
+				if (SpecifiedArchitecture != null)
 				{
-					Arguments += String.Format(" -architectures={0}", AndroidArchitectures);
+					Arguments += String.Format(" -architecture={0}", SpecifiedArchitecture);
+				}
+				else if (Platform == UnrealTargetPlatform.Android)
+				{
+					Arguments += String.Format(" -architecture={0}", AndroidArchitectures);
 				}
 
 				if (!String.IsNullOrEmpty(InAdditionalArgs))
