Re-throw if not a "retry" exception.

diff --git a/Engine/Source/Programs/AutomationTool/Gauntlet/Framework/Gauntlet.Utils.cs b/Engine/Source/Programs/AutomationTool/Gauntlet/Framework/Gauntlet.Utils.cs
--- a/Engine/Source/Programs/AutomationTool/Gauntlet/Framework/Gauntlet.Utils.cs
+++ b/Engine/Source/Programs/AutomationTool/Gauntlet/Framework/Gauntlet.Utils.cs
@@ -1367,6 +1367,10 @@ namespace Gauntlet
 						Log.Info("Retrying Directory.GetFiles() once.");
 						Files = Directory.GetFiles();
 					}
+					else
+					{
+						throw;
+					}
 				}
 				return Files;
 			}
@@ -1388,6 +1392,10 @@ namespace Gauntlet
 						Log.Info("Retrying Directory.GetFiles(Pattern, SearchOption) once.");
 						Files = Directory.GetFiles(Pattern, SearchOption);
 					}
+					else
+					{
+						throw;
+					}
 				}
 				return Files;
 			}
@@ -1409,6 +1417,10 @@ namespace Gauntlet
 						Log.Info("Retrying Directory.GetDirectories() once.");
 						Directories = Directory.GetDirectories();
 					}
+					else
+					{
+						throw;
+					}
 				}
 				return Directories;
 			}
@@ -1430,6 +1442,10 @@ namespace Gauntlet
 						Log.Info("Retrying Directory.GetDirectories(Pattern, SearchOption) once.");
 						Directories = Directory.GetDirectories(Pattern, SearchOption);
 					}
+					else
+					{
+						throw;
+					}
 				}
 				return Directories;
 			}
@@ -1456,6 +1472,10 @@ namespace Gauntlet
 						Log.Info("Retrying Directory.GetFileSystemInfos(Pattern, SearchOption) once.");
 						FileInfos = Directory.GetFileSystemInfos(Pattern, SearchOption);
 					}
+					else
+					{
+						throw;
+					}
 				}
 				return FileInfos;
 			}
@@ -1476,6 +1496,10 @@ namespace Gauntlet
 						Log.Info("Retrying Directory.Delete(Force) once.");
 						Directory.Delete(Force);
 					}
+					else
+					{
+						throw;
+					}
 				}
 			}
 
@@ -1495,6 +1519,10 @@ namespace Gauntlet
 						Log.Info("Retrying File.Delete() once.");
 						File.Delete();
 					}
+					else
+					{
+						throw;
+					}
 				}
 			}
 
