- Fixing dictionary lookup error

diff --git a/Engine/Source/Programs/AutomationTool/Scripts/BuildPluginCommand.Automation.cs b/Engine/Source/Programs/AutomationTool/Scripts/BuildPluginCommand.Automation.cs
--- a/Engine/Source/Programs/AutomationTool/Scripts/BuildPluginCommand.Automation.cs
+++ b/Engine/Source/Programs/AutomationTool/Scripts/BuildPluginCommand.Automation.cs
@@ -363,8 +363,8 @@ public sealed class BuildPlugin : BuildCommand
 				
 				string Arguments = String.Format("-plugin={0} -iwyu -noubtmakefiles -manifest={1} -nohotreload", CommandUtils.MakePathSafeToUseWithCommandLine(HostProjectPluginFile.FullName), CommandUtils.MakePathSafeToUseWithCommandLine(ManifestFileName.FullName));
 
-				string SpecifiedArchitecture = PlatformToArchitectureMap[Platform];
-				if (SpecifiedArchitecture != null)
+				string SpecifiedArchitecture;
+				if (PlatformToArchitectureMap.TryGetValue(Platform, out SpecifiedArchitecture))
 				{
 					Arguments += String.Format(" -architecture={0}", SpecifiedArchitecture);
 				}
