UAT: Add separate log output for each line of mismatched Perforce case.

diff --git a/Engine/Source/Programs/AutomationTool/Scripts/CheckPerforceCase.cs b/Engine/Source/Programs/AutomationTool/Scripts/CheckPerforceCase.cs
--- a/Engine/Source/Programs/AutomationTool/Scripts/CheckPerforceCase.cs
+++ b/Engine/Source/Programs/AutomationTool/Scripts/CheckPerforceCase.cs
@@ -73,26 +73,25 @@ namespace BuildScripts.Automation
 				{
 					if (Group.Count() > 1)
 					{
-						StringBuilder Message = new StringBuilder();
+						Logger.LogWarning(KnownLogEvents.AutomationTool_PerforceCase, "Inconsistent casing for {File}:", Group.First().Value.Path);
+
 						foreach (KeyValuePair<string, TreeNode> Pair in Group)
 						{
-							Message.AppendFormat("\n  Could be '{0}':", Pair.Key, Pair.Value.Count, (Pair.Value.Count == 1)? "file" : "files");
+							Logger.LogWarning(KnownLogEvents.AutomationTool_PerforceCase, "  Could be '{Rendering}':", Pair.Key);
 
 							int NumFiles = 0;
 							foreach (string File in Pair.Value.EnumerateFiles())
 							{
-								Message.AppendFormat("\n    {0}", File);
+								Logger.LogWarning(KnownLogEvents.AutomationTool_PerforceCase, "    {DepotFile}", new LogValue(LogValueType.DepotPath, File));
 								if (++NumFiles >= 10)
 								{
 									break;
 								}
 							}
 
-							Message.AppendFormat("\n    ({0} {1})", Pair.Value.Count, (Pair.Value.Count == 1) ? "file" : "files");
+							Logger.LogWarning(KnownLogEvents.AutomationTool_PerforceCase, "    ({NumFile} file(s))", Pair.Value.Count);
 						}
 
-						Logger.LogWarning(KnownLogEvents.AutomationTool_PerforceCase, "Inconsistent casing for {File}: {Message}", new LogValue(LogValueType.DepotPath, Group.First().Value.Path), Message.ToString());
-
 						NumIssues++;
 					}
 
