SyncBinariesFromUGS: Retrieve correct compatible changelist using the existing P4Env
Move list of extensions treated as code to common location

#rnx

diff --git a/Engine/Source/Programs/AutomationTool/AutomationUtils/P4Environment.cs b/Engine/Source/Programs/AutomationTool/AutomationUtils/P4Environment.cs
--- a/Engine/Source/Programs/AutomationTool/AutomationUtils/P4Environment.cs
+++ b/Engine/Source/Programs/AutomationTool/AutomationUtils/P4Environment.cs
@@ -458,8 +458,7 @@ namespace AutomationTool
 			// Retrieve the current changelist
 			StringBuilder P4Cmd = new StringBuilder("changes -m 1");
 
-			string[] CodeExtensions = { ".cs", ".h", ".cpp", ".inl", ".usf", ".ush", ".uproject", ".uplugin" };
-			foreach(string CodeExtension in CodeExtensions)
+			foreach (string CodeExtension in EpicGames.Perforce.PerforceUtils.CodeExtensions)
 			{
 				P4Cmd.AppendFormat(" \"{0}/...{1}#have\"", CommandUtils.CombinePaths(PathSeparator.Depot, ClientRootPath), CodeExtension);
 			}
diff --git a/Engine/Source/Programs/AutomationTool/Scripts/SyncBinariesFromUGS.Automation.cs b/Engine/Source/Programs/AutomationTool/Scripts/SyncBinariesFromUGS.Automation.cs
--- a/Engine/Source/Programs/AutomationTool/Scripts/SyncBinariesFromUGS.Automation.cs
+++ b/Engine/Source/Programs/AutomationTool/Scripts/SyncBinariesFromUGS.Automation.cs
@@ -55,10 +55,8 @@ class SyncBinariesFromUGS : SyncProjectBase
 		string DepotBinaryPath = GetZippedBinaryPathFromUGSConfig(ProjectFile.Directory, ProjectRecord.DepotFile);
 
 		// get the CL
-		int CompatibleChangeList = 0;
-		int CurrentChangeList = 0;
-		
-		GetCurrentAndCompatibleChangeLists(ProjectRecord.DepotFile, out CurrentChangeList, out CompatibleChangeList);
+		int CurrentChangeList = P4Env.Changelist;
+		int CompatibleChangeList = P4Env.CodeChangelist;
 
 		LogInformation("Current CL: {0}, Required Compatible CL: {1}", CurrentChangeList, CompatibleChangeList);
 
@@ -205,91 +203,4 @@ class SyncBinariesFromUGS : SyncProjectBase
 
 		return BinaryPath;
 	}
-
-	bool GetCurrentAndCompatibleChangeLists(string P4ProjectFilePath, out int OutCurrentChangeList, out int OutCompatibleChangeList)
-	{
-		OutCurrentChangeList = 0;
-		OutCompatibleChangeList = 0;
-
-		string[] CurrentChangeListPaths =
-		{
-			"<project>/...#have",
-			"<engine>/...#have"
-		};
-
-		string[] CompatibleChangeListPaths =
-		{
-			"<project>/Plugins/...#have",
-			"<project>/Source/...#have",
-			"<engine>/...#have"
-		};
-
-
-		string P4ProjectRoot = null;
-
-		if (P4ProjectFilePath.IndexOf(".uproject", StringComparison.CurrentCultureIgnoreCase) < 0)
-		{
-			throw new AutomationException("P4ProjectFilePath should end in .uproject");
-		}
-
-		P4WhereRecord Record = P4.Where(P4ProjectFilePath, AllowSpew: false).FirstOrDefault(x => x.DepotFile != null && !x.bUnmap);
-
-		P4ProjectRoot = Record.DepotFile.Substring(0, Record.DepotFile.LastIndexOf("/"));
-
-		// assume for now that if the project is in //depot/<some_path> that the engine is in //Engine;
-		string P4EngineRoot = "//" + P4ProjectRoot.Split(new char[] { '/' }, StringSplitOptions.RemoveEmptyEntries).First() + "/Engine";
-
-		// Transform the compatible changelist paths queries with the now-known project and engine paths
-		CompatibleChangeListPaths = CompatibleChangeListPaths.Select(S => {
-			string Result = S.Replace("<project>", P4ProjectRoot);
-			Result = Result.Replace("<engine>", P4EngineRoot);
-			return Result;
-		})
-		.ToArray();
-
-		// Transform the current changelist paths queries with the now-known project and engine paths
-		CurrentChangeListPaths = CurrentChangeListPaths.Select(S =>
-		{
-			string Result = S.Replace("<project>", P4ProjectRoot);
-			Result = Result.Replace("<engine>", P4EngineRoot);
-			return Result;
-		})
-		.ToArray();
-	
-		// get the most recent changes for all the compatible paths
-		List<P4Connection.ChangeRecord> ChangeRecords;
-		string CompatiblePathArg = string.Join(" ", CompatibleChangeListPaths);
-		P4.Changes(out ChangeRecords, "-m1 " + CompatiblePathArg);
-
-		if (ChangeRecords.Any() == false)
-		{
-			throw new AutomationException("No changes returned for P4 paths {0}", CompatiblePathArg);
-		}
-
-		int SortRecordsDescending (P4Connection.ChangeRecord lhs, P4Connection.ChangeRecord rhs)
-		{
-			return lhs.CL.CompareTo(rhs.CL) * -1;
-		};
-
-		// sort descending to get the most recent CL first
-		ChangeRecords.Sort(SortRecordsDescending);
-
-		// this is the CL we need for compatibility
-		OutCompatibleChangeList = ChangeRecords.First().CL;
-
-		// get records for the entire project and engine paths
-		ChangeRecords.Clear();
-		string CurrentPathArg = string.Join(" ", CurrentChangeListPaths);
-		P4.Changes(out ChangeRecords, "-m1 " + CurrentPathArg);
-
-		if (ChangeRecords.Any() == false)
-		{
-			throw new AutomationException("No changes returned for P4 paths {0}", CurrentPathArg);
-		}
-
-		ChangeRecords.Sort(SortRecordsDescending);
-		OutCurrentChangeList = ChangeRecords.First().CL;
-
-		return OutCompatibleChangeList != 0 && OutCurrentChangeList != 0;
-	}
 }
diff --git a/Engine/Source/Programs/Shared/EpicGames.Perforce/PerforceUtils.cs b/Engine/Source/Programs/Shared/EpicGames.Perforce/PerforceUtils.cs
--- a/Engine/Source/Programs/Shared/EpicGames.Perforce/PerforceUtils.cs
+++ b/Engine/Source/Programs/Shared/EpicGames.Perforce/PerforceUtils.cs
@@ -1,6 +1,7 @@
 // Copyright Epic Games, Inc. All Rights Reserved.
 
 using System;
+using System.Collections.Generic;
 using System.Diagnostics.CodeAnalysis;
 using System.IO;
 using EpicGames.Core;
@@ -12,6 +13,30 @@ namespace EpicGames.Perforce
 	/// </summary>
 	public static class PerforceUtils
 	{
+		/// <summary>
+		/// Set of extensions to treat as code
+		/// </summary>
+		public static readonly HashSet<string> CodeExtensions = new HashSet<string>
+		{
+			".c",
+			".cc",
+			".cpp",
+			".inl",
+			".m",
+			".mm",
+			".rc",
+			".cs",
+			".csproj",
+			".h",
+			".hpp",
+			".inl",
+			".usf",
+			".ush",
+			".uproject",
+			".uplugin",
+			".sln"
+		};
+
 		/// <summary>
 		/// Escape a path to Perforce syntax
 		/// </summary>
diff --git a/Engine/Source/Programs/UnrealGameSync/UnrealGameSyncShared/Utility.cs b/Engine/Source/Programs/UnrealGameSync/UnrealGameSyncShared/Utility.cs
--- a/Engine/Source/Programs/UnrealGameSync/UnrealGameSyncShared/Utility.cs
+++ b/Engine/Source/Programs/UnrealGameSync/UnrealGameSyncShared/Utility.cs
@@ -32,30 +32,6 @@ namespace UnrealGameSync
 
 	public class PerforceChangeDetails
 	{
-		/// <summary>
-		/// Set of extensions to treat as code
-		/// </summary>
-		static readonly HashSet<string> CodeExtensions = new HashSet<string>
-		{
-			".c",
-			".cc",
-			".cpp",
-			".inl",
-			".m",
-			".mm",
-			".rc",
-			".cs",
-			".csproj",
-			".h",
-			".hpp",
-			".inl",
-			".usf",
-			".ush",
-			".uproject",
-			".uplugin",
-			".sln"
-		};
-
 		public string Description;
 		public bool ContainsCode;
 		public bool ContainsContent;
@@ -67,7 +43,7 @@ namespace UnrealGameSync
 			// Check whether the files are code or content
 			foreach (DescribeFileRecord file in describeRecord.Files)
 			{
-				if (CodeExtensions.Any(extension => file.DepotFile.EndsWith(extension, StringComparison.OrdinalIgnoreCase)))
+				if (PerforceUtils.CodeExtensions.Any(extension => file.DepotFile.EndsWith(extension, StringComparison.OrdinalIgnoreCase)))
 				{
 					ContainsCode = true;
 				}
