Extended ParamHelp attribute for BuildCommand to better describe the metadata of parameters for UAT scripts.
Update BuildGraph with ParamHelp attributes.

diff --git a/Engine/Source/Programs/AutomationTool/AutomationUtils/HelpUtils.cs b/Engine/Source/Programs/AutomationTool/AutomationUtils/HelpUtils.cs
--- a/Engine/Source/Programs/AutomationTool/AutomationUtils/HelpUtils.cs
+++ b/Engine/Source/Programs/AutomationTool/AutomationUtils/HelpUtils.cs
@@ -1,4 +1,5 @@
 // Copyright Epic Games, Inc. All Rights Reserved.
+
 using System;
 using System.Collections;
 using System.Collections.Generic;
@@ -73,9 +74,163 @@ namespace AutomationTool
 		}
 	}
 
-	/// <summary>
-	/// Base utility function for script commands.
-	/// </summary>
+	[AttributeUsage(AttributeTargets.Class, AllowMultiple = true)]
+	public class ParamHelpAttribute : HelpAttribute
+	{
+		public enum ParamAction
+		{
+			Store,
+			Store_True,
+			Store_False,
+			Append
+		}
+
+		private ParamHelpAttribute(string Description) : base(Description)
+		{
+		}
+
+		public ParamHelpAttribute(string Name, string Description) : base(Name, Description)
+		{
+			ParamName = Name.TrimStart('-');
+			int EqualsInParamNameIndex = ParamName.IndexOf("=");
+			if (EqualsInParamNameIndex > 0)
+			{
+				ParamName = ParamName.Substring(0, EqualsInParamNameIndex);
+			}
+			_Action = ParamAction.Store;
+			Required = false;
+			ParamType = typeof(string);
+			DefaultValue = null;
+			_Choices = null;
+			MultiSelectSeparator = null;
+			Deprecated = false;
+		}
+
+		public string ParamName
+		{
+			get;
+		}
+		
+		public string ParamDescription
+		{
+			get
+			{
+				string RequiredStr = Required ? ". Required" : ". Optional";
+				string ChoicesStr = "";
+				var ValidChoices = Choices as IEnumerable<object>;
+				if (ValidChoices != null)
+				{
+					ChoicesStr += ". Choices=[";
+					foreach (var Value in ValidChoices)
+					{
+						ChoicesStr += Value.ToString() + (MultiSelectSeparator != null ? MultiSelectSeparator : "|");
+					}
+					ChoicesStr = ChoicesStr.Remove(ChoicesStr.Length - 1);
+					ChoicesStr += "]";
+				}
+				string DefaultStr = DefaultValue != null ? ". DefaultValue=" + DefaultValue.ToString() : "";
+				return base.Description + string.Format("{0}{1}{2}",
+						RequiredStr,
+						DefaultStr,
+						ChoicesStr
+					);
+			}
+		}
+
+		private ParamAction _Action = ParamAction.Store;
+		public ParamAction Action
+		{
+			get
+			{
+				return _Action;
+			}
+			set
+			{
+				if (value == ParamAction.Store_True || value == ParamAction.Store_False)
+				{
+					ParamType = typeof(bool);
+				}
+				_Action = value;
+			}
+		}
+
+		public object DefaultValue
+		{
+			get;
+			set;
+		}
+
+		public Type ParamType
+		{
+			get;
+			set;
+		}
+
+		object _Choices;
+		public object Choices
+		{
+			get
+			{
+				return _Choices;
+			}
+			set
+			{
+				var Enumerable = value as System.Collections.IEnumerable;
+				if (Enumerable != null)
+				{
+					_Choices = Enumerable;
+				}
+			}
+		}
+
+		public bool Required
+		{
+			get;
+			set;
+		}
+
+		public string Help
+		{
+			get;
+			set;
+		}
+
+		public string MultiSelectSeparator
+		{
+			get;
+			set;
+		}
+
+		public bool Deprecated
+		{
+			get;
+			set;
+		}
+
+		private string _Flag = null;
+		public string Flag
+		{
+			get
+			{
+				if (_Flag == null)
+				{
+					return ParamName;
+				}
+				return _Flag;
+			}
+			set
+			{
+				if (!string.IsNullOrEmpty(value))
+				{
+					_Flag = value.TrimStart('-');
+				}
+			}
+		}
+	}
+
+/// <summary>
+/// Base utility function for script commands.
+/// </summary>
 	public partial class CommandUtils
 	{
 		/// <summary>
diff --git a/Engine/Source/Programs/AutomationTool/BuildGraph/BuildGraph.cs b/Engine/Source/Programs/AutomationTool/BuildGraph/BuildGraph.cs
--- a/Engine/Source/Programs/AutomationTool/BuildGraph/BuildGraph.cs
+++ b/Engine/Source/Programs/AutomationTool/BuildGraph/BuildGraph.cs
@@ -180,32 +180,32 @@ namespace AutomationTool
 	/// elements as you type.
 	/// </summary>
 	[Help("Tool for creating extensible build processes in UE which can be run locally or in parallel across a build farm.")]
-	[Help("Script=<FileName>", "Path to the script describing the graph")]
-	[Help("Target=<Name>", "Name of the node or output tag to be built")]
-	[Help("Schema", "Generates a schema to the default location")]
-	[Help("Schema=<FileName>", "Generate a schema describing valid script documents, including all the known tasks")]
-	[Help("ImportSchema=<FileName>", "Imports a schema from an existing schema file")]
-	[Help("Set:<Property>=<Value>", "Sets a named property to the given value")]
-	[Help("Branch=<Value>", "Overrides the auto-detection of the current branch")]
-	[Help("Clean", "Cleans all cached state of completed build nodes before running")]
-	[Help("CleanNode=<Name>[+<Name>...]", "Cleans just the given nodes before running")]
-	[Help("Resume", "Resumes a local build from the last node that completed successfully")]
-	[Help("ListOnly", "Shows the contents of the preprocessed graph, but does not execute it")]
-	[Help("ShowDiagnostics", "When running with -ListOnly, causes diagnostic messages entered when parsing the graph to be shown")]
-	[Help("ShowDeps", "Show node dependencies in the graph output")]
-	[Help("ShowNotifications", "Show notifications that will be sent for each node in the output")]
-	[Help("Trigger=<Name>", "Executes only nodes behind the given trigger")]
-	[Help("SkipTrigger=<Name>[+<Name>...]", "Skips the given triggers, including all the nodes behind them in the graph")]
-	[Help("SkipTriggers", "Skips all triggers")]
-	[Help("TokenSignature=<Name>", "Specifies the signature identifying the current job, to be written to tokens for nodes that require them. Tokens are ignored if this parameter is not specified.")]
-	[Help("SkipTargetsWithoutTokens", "Excludes targets which we can't acquire tokens for, rather than failing")]
-	[Help("Preprocess=<FileName>", "Writes the preprocessed graph to the given file")]
-	[Help("Export=<FileName>", "Exports a JSON file containing the preprocessed build graph, for use as part of a build system")]
-	[Help("HordeExport=<FileName>", "Exports a JSON file containing the full build graph for use by Horde.")]
-	[Help("PublicTasksOnly", "Only include built-in tasks in the schema, excluding any other UAT modules")]
-	[Help("SharedStorageDir=<DirName>", "Sets the directory to use to transfer build products between agents in a build farm")]
-	[Help("SingleNode=<Name>", "Run only the given node. Intended for use on a build system after running with -Export.")]
-	[Help("WriteToSharedStorage", "Allow writing to shared storage. If not set, but -SharedStorageDir is specified, build products will read but not written")]
+	[ParamHelp("Script=<FileName>", "Path to the script describing the graph", ParamType = typeof(File), Required = true)]
+	[ParamHelp("Target=<Name>", "Name of the node or output tag to be built", Required = true)]
+	[ParamHelp("Schema", "Generates a schema to the default location", ParamType = typeof(bool))]
+	[ParamHelp("Schema=<FileName>", "Generate a schema describing valid script documents, including all the known tasks", ParamType = typeof(string))]
+	[ParamHelp("ImportSchema=<FileName>", "Imports a schema from an existing schema file", ParamType = typeof(File))]
+	[ParamHelp("Set:<Property>=<Value>", "Sets a named property to the given value", Action = ParamHelpAttribute.ParamAction.Append, Flag = "-Set")]
+	[ParamHelp("Branch=<Value>", "Overrides the auto-detection of the current branch")]
+	[ParamHelp("Clean", "Cleans all cached state of completed build nodes before running", ParamType = typeof(bool), Deprecated = true)]
+	[ParamHelp("CleanNode=<Name>[+<Name>...]", "Cleans just the given nodes before running", MultiSelectSeparator = "+")]
+	[ParamHelp("Resume", "Resumes a local build from the last node that completed successfully", ParamType = typeof(bool))]
+	[ParamHelp("ListOnly", "Shows the contents of the preprocessed graph, but does not execute it", ParamType = typeof(bool), DefaultValue = false)]
+	[ParamHelp("ShowDiagnostics", "When running with -ListOnly, causes diagnostic messages entered when parsing the graph to be shown", ParamType = typeof(bool))]
+	[ParamHelp("ShowDeps", "Show node dependencies in the graph output", ParamType = typeof(bool))]
+	[ParamHelp("ShowNotifications", "Show notifications that will be sent for each node in the output", ParamType = typeof(bool))]
+	[ParamHelp("Trigger=<Name>", "Executes only nodes behind the given trigger")]
+	[ParamHelp("SkipTrigger=<Name>[+<Name>...]", "Skips the given triggers, including all the nodes behind them in the graph", MultiSelectSeparator = "+")]
+	[ParamHelp("SkipTriggers", "Skips all triggers", ParamType = typeof(bool))]
+	[ParamHelp("TokenSignature=<Name>", "Specifies the signature identifying the current job, to be written to tokens for nodes that require them. Tokens are ignored if this parameter is not specified.")]
+	[ParamHelp("SkipTargetsWithoutTokens", "Excludes targets which we can't acquire tokens for, rather than failing", ParamType = typeof(bool))]
+	[ParamHelp("Preprocess=<FileName>", "Writes the preprocessed graph to the given file", ParamType = typeof(File))]
+	[ParamHelp("Export=<FileName>", "Exports a JSON file containing the preprocessed build graph, for use as part of a build system", ParamType = typeof(File))]
+	[ParamHelp("HordeExport=<FileName>", "Exports a JSON file containing the full build graph for use by Horde.", ParamType = typeof(File))]
+	[ParamHelp("PublicTasksOnly", "Only include built-in tasks in the schema, excluding any other UAT modules", ParamType = typeof(bool))]
+	[ParamHelp("SharedStorageDir=<DirName>", "Sets the directory to use to transfer build products between agents in a build farm", ParamType = typeof(Directory))]
+	[ParamHelp("SingleNode=<Name>", "Run only the given node. Intended for use on a build system after running with -Export.")]
+	[ParamHelp("WriteToSharedStorage", "Allow writing to shared storage. If not set, but -SharedStorageDir is specified, build products will read but not written", ParamType = typeof(bool))]
 	public class BuildGraph : BuildCommand
 	{
 		/// <summary>
