Workaround for tbb.psym colliding from multiple clients copying tbb.psym

diff --git a/Engine/Source/Programs/AutomationTool/BuildGraph/TempStorage.cs b/Engine/Source/Programs/AutomationTool/BuildGraph/TempStorage.cs
--- a/Engine/Source/Programs/AutomationTool/BuildGraph/TempStorage.cs
+++ b/Engine/Source/Programs/AutomationTool/BuildGraph/TempStorage.cs
@@ -1263,7 +1263,7 @@ namespace AutomationTool
 			{
 				return true;
 			}
-			if (FileName.Equals("tbb.dll", StringComparison.OrdinalIgnoreCase) || FileName.Equals("tbb.pdb", StringComparison.OrdinalIgnoreCase) || FileName.Equals("libtbb.dylib", StringComparison.OrdinalIgnoreCase))
+			if (FileName.Equals("tbb.dll", StringComparison.OrdinalIgnoreCase) || FileName.Equals("tbb.pdb", StringComparison.OrdinalIgnoreCase) || FileName.Equals("libtbb.dylib", StringComparison.OrdinalIgnoreCase) || FileName.Equals("tbb.psym", StringComparison.OrdinalIgnoreCase))
 			{
 				return true;
 			}
