Fix typo in string constant for Attribute metadata on BuildGraph BuildCommand

diff --git a/Engine/Source/Programs/AutomationTool/BuildGraph/BuildGraph.cs b/Engine/Source/Programs/AutomationTool/BuildGraph/BuildGraph.cs
--- a/Engine/Source/Programs/AutomationTool/BuildGraph/BuildGraph.cs
+++ b/Engine/Source/Programs/AutomationTool/BuildGraph/BuildGraph.cs
@@ -185,7 +185,7 @@ namespace AutomationTool
 	[ParamHelp("Schema", "Generates a schema to the default location", ParamType = typeof(bool))]
 	[ParamHelp("Schema=<FileName>", "Generate a schema describing valid script documents, including all the known tasks", ParamType = typeof(string))]
 	[ParamHelp("ImportSchema=<FileName>", "Imports a schema from an existing schema file", ParamType = typeof(File))]
-	[ParamHelp("Set:<Property>=<Value>", "Sets a named property to the given value", Action = ParamHelpAttribute.ParamAction.Append, Flag = "-Set")]
+	[ParamHelp("Set:<Property>=<Value>", "Sets a named property to the given value", Action = ParamHelpAttribute.ParamAction.Append, Flag = "-Set:")]
 	[ParamHelp("Branch=<Value>", "Overrides the auto-detection of the current branch")]
 	[ParamHelp("Clean", "Cleans all cached state of completed build nodes before running", ParamType = typeof(bool), Deprecated = true)]
 	[ParamHelp("CleanNode=<Name>[+<Name>...]", "Cleans just the given nodes before running", MultiSelectSeparator = "+")]
