PR #9417: Fix NullReferenceException when running Turnkey on installed engine build without Android support (Contributed by [REDACTED_PERSON])

#android

diff --git a/Engine/Source/Programs/AutomationTool/Android/AndroidPlatform.Automation.cs b/Engine/Source/Programs/AutomationTool/Android/AndroidPlatform.Automation.cs
--- a/Engine/Source/Programs/AutomationTool/Android/AndroidPlatform.Automation.cs
+++ b/Engine/Source/Programs/AutomationTool/Android/AndroidPlatform.Automation.cs
@@ -167,7 +167,7 @@ public class AndroidPlatform : Platform
 	{
 		UEBuildPlatformSDK AndroidSDK = UEBuildPlatformSDK.GetSDKForPlatform("Android");
 
-		return new string[] { AndroidSDK.GetMainVersion() };
+		return AndroidSDK != null ? new string[] { AndroidSDK.GetMainVersion() } : Array.Empty<string>();
 	}
 
 	public override bool GetSDKInstallCommand(out string Command, out string Params, ref bool bRequiresPrivilegeElevation, ref bool bCreateWindow, ITurnkeyContext TurnkeyContext)
