P4Utils: Update output scan to correctly detect submitted changelist if additional output is found due to a p4 extension

#rnx

diff --git a/Engine/Source/Programs/AutomationTool/AutomationUtils/P4Utils.cs b/Engine/Source/Programs/AutomationTool/AutomationUtils/P4Utils.cs
--- a/Engine/Source/Programs/AutomationTool/AutomationUtils/P4Utils.cs
+++ b/Engine/Source/Programs/AutomationTool/AutomationUtils/P4Utils.cs
@@ -2395,31 +2395,18 @@ namespace AutomationTool
 				else
 				{
 					LastCmdOutput = CmdOutput;
-					if (CmdOutput.Trim().EndsWith("submitted."))
+
+					Regex SubmitRegex = new Regex(@"Change \d+ renamed change (?<number>\d+) and submitted");
+					Match SubmitMatch = SubmitRegex.Match(CmdOutput);
+					if (!SubmitMatch.Success)
 					{
-						if (CmdOutput.Trim().EndsWith(" and submitted."))
-						{
-							string EndStr = " and submitted.";
-							string ChangeStr = "renamed change ";
-							int Offset = CmdOutput.LastIndexOf(ChangeStr);
-							int EndOffset = CmdOutput.LastIndexOf(EndStr);
-							if (Offset >= 0 && Offset < EndOffset)
-							{
-								SubmittedCL = int.Parse(CmdOutput.Substring(Offset + ChangeStr.Length, EndOffset - Offset - ChangeStr.Length));
-							}
-						}
-						else
-						{
-							string EndStr = " submitted.";
-							string ChangeStr = "Change ";
-							int Offset = CmdOutput.LastIndexOf(ChangeStr);
-							int EndOffset = CmdOutput.LastIndexOf(EndStr);
-							if (Offset >= 0 && Offset < EndOffset)
-							{
-								SubmittedCL = int.Parse(CmdOutput.Substring(Offset + ChangeStr.Length, EndOffset - Offset - ChangeStr.Length));
-							}
-						}
+						SubmitRegex = new Regex(@"Change (?<number>\d+) submitted");
+						SubmitMatch = SubmitRegex.Match(CmdOutput);
+					}
 
+					if (SubmitMatch.Success)
+					{
+						SubmittedCL = int.Parse(SubmitMatch.Groups["number"].Value);
 						CommandUtils.LogInformation("Submitted CL {0} which became CL {1}\n", CL, SubmittedCL);
 					}
 
