Automation: Fixed staging of blueprint-only projects on some platforms.

The problem was that blueprint-only projects would return true for IsCodeBasedUProjectFile during staging. The issue was triggered by CL 21214930 where PropertyCacheKey (ProjectUtils.cs) was added. Before this change GetProjectProperties worked based only on the path to .uproject. After the change, it also takes the platform and configurations into account. It means that:

GetProjectProperties("...", SomePlatform, Debug)
GetProjectProperties("...", null, Debug)

may return different results. While it's probably the core or the fix this change implemented, it may given unexpected results because when null is passed as a list of platforms, internally all available platforms are taken into account and if only one of them requires the project to be compiled IsCodeBasedUProjectFile will return true, which may not be what our code is expected to see. That's what happened during staging - IsCodeBasedUProjectFile was invoked two times in a slightly different way leading to the automation code handling blueprint-only projects in an inconsistent-way during the whole process.

Fixed by modifying DeploymentContext to initialize IsCodeBasedProject by passing both the current platform and configurations.

#rnx

diff --git a/Engine/Source/Programs/AutomationTool/AutomationUtils/DeploymentContext.cs b/Engine/Source/Programs/AutomationTool/AutomationUtils/DeploymentContext.cs
--- a/Engine/Source/Programs/AutomationTool/AutomationUtils/DeploymentContext.cs
+++ b/Engine/Source/Programs/AutomationTool/AutomationUtils/DeploymentContext.cs
@@ -421,7 +421,7 @@ public class DeploymentContext //: ProjectParams
 		StageTargetConfigurations = new List<UnrealTargetConfiguration>(InTargetConfigurations);
 		StageTargets = new List<StageTarget>(InStageTargets);
 		StageExecutables = InStageExecutables;
-        IsCodeBasedProject = ProjectUtils.IsCodeBasedUProjectFile(RawProjectPath, StageTargetConfigurations);
+        IsCodeBasedProject = ProjectUtils.IsCodeBasedUProjectFile(RawProjectPath, StageTargetPlatform.PlatformType, StageTargetConfigurations);
 		ShortProjectName = ProjectUtils.GetShortProjectName(RawProjectPath);
 		Stage = InStage;
 		Archive = InArchive;
diff --git a/Engine/Source/Programs/AutomationTool/AutomationUtils/ProjectUtils.cs b/Engine/Source/Programs/AutomationTool/AutomationUtils/ProjectUtils.cs
--- a/Engine/Source/Programs/AutomationTool/AutomationUtils/ProjectUtils.cs
+++ b/Engine/Source/Programs/AutomationTool/AutomationUtils/ProjectUtils.cs
@@ -152,9 +152,19 @@ namespace AutomationTool
 		/// </summary>
 		/// <param name="RawProjectPath">Full project path.</param>
 		/// <returns>True if the project is a UProject file with source code.</returns>
-		public static bool IsCodeBasedUProjectFile(FileReference RawProjectPath, List<UnrealTargetConfiguration> ClientTargetConfigurations = null)
+		public static bool IsCodeBasedUProjectFile(FileReference RawProjectPath, List<UnrealTargetPlatform> ClientTargetPlatforms = null, List < UnrealTargetConfiguration> ClientTargetConfigurations = null)
 		{
-			return GetProjectProperties(RawProjectPath, null, ClientTargetConfigurations).bIsCodeBasedProject;
+			return GetProjectProperties(RawProjectPath, ClientTargetPlatforms, ClientTargetConfigurations).bIsCodeBasedProject;
+		}
+
+		/// <summary>
+		/// Checks if the project is a UProject file with source code.
+		/// </summary>
+		/// <param name="RawProjectPath">Full project path.</param>
+		/// <returns>True if the project is a UProject file with source code.</returns>
+		public static bool IsCodeBasedUProjectFile(FileReference RawProjectPath, UnrealTargetPlatform ClientTargetPlatform, List<UnrealTargetConfiguration> ClientTargetConfigurations = null)
+		{
+			return GetProjectProperties(RawProjectPath, new List<UnrealTargetPlatform>() { ClientTargetPlatform }, ClientTargetConfigurations).bIsCodeBasedProject;
 		}
 
 		/// <summary>
