Automation: Fixed PropertiesCache cache with ProjectProperties that could add multiple values for the same key.

GetProjectProperties uses PropertiesCache collection (hashmap) to store project properties indexed by PropertyCacheKey struct (.uproject path + platforms + configurations). However, it didn't work as expected because PropertyCacheKey has array members that are reference types. Without a custom hashing implemented, this key would generate different hash codes even if elements stored in the arrays were the same. As a result, querying the same properties multiple times would generate multiple entries in the cache e.g.

GetProjectProperties("foo.uproject", {UnrealTargetPlatform.PC}, {UnrealTargetConfiguration.Debug})
GetProjectProperties("foo.uproject", {UnrealTargetPlatform.PC}, {UnrealTargetConfiguration.Debug})

would create two entries for PC/Debug.

Resolved by implementing the full set of Equals, GetHashCode and comparison operators.

#rnx

diff --git a/Engine/Source/Programs/AutomationTool/AutomationUtils/ProjectUtils.cs b/Engine/Source/Programs/AutomationTool/AutomationUtils/ProjectUtils.cs
--- a/Engine/Source/Programs/AutomationTool/AutomationUtils/ProjectUtils.cs
+++ b/Engine/Source/Programs/AutomationTool/AutomationUtils/ProjectUtils.cs
@@ -11,6 +11,7 @@ using EpicGames.Core;
 using System.Reflection;
 using UnrealBuildBase;
 using System.Runtime.Serialization;
+using System.Collections;
 
 namespace AutomationTool
 {
@@ -82,7 +83,7 @@ namespace AutomationTool
 		/// DetectProjectProperties may return different answers, e.g. Some platforms require a 
 		///  codebased project for targets
 		/// </summary>
-		struct PropertyCacheKey
+		struct PropertyCacheKey : IEquatable<PropertyCacheKey>
 		{
 			string ProjectName;
 
@@ -96,6 +97,36 @@ namespace AutomationTool
 				TargetPlatforms = InTargetPlatforms != null ? InTargetPlatforms.ToArray() : new UnrealTargetPlatform[0];
 				TargetConfigurations = InTargetConfigurations != null ? InTargetConfigurations.ToArray() : new UnrealTargetConfiguration[0];
 			}
+
+			public bool Equals(PropertyCacheKey Other)
+			{
+				return ProjectName == Other.ProjectName &&
+						StructuralComparisons.StructuralEqualityComparer.Equals(TargetPlatforms, Other.TargetPlatforms) &&
+						StructuralComparisons.StructuralEqualityComparer.Equals(TargetConfigurations, Other.TargetConfigurations);
+			}
+
+			public override bool Equals(object Other)
+			{
+				return Other is PropertyCacheKey OtherKey && Equals(OtherKey);
+			}
+
+			public override int GetHashCode()
+			{
+				return HashCode.Combine(
+					ProjectName.GetHashCode(),
+					StructuralComparisons.StructuralEqualityComparer.GetHashCode(TargetPlatforms),
+					StructuralComparisons.StructuralEqualityComparer.GetHashCode(TargetConfigurations));
+			}
+
+			public static bool operator==(PropertyCacheKey A, PropertyCacheKey B)
+			{
+				return A.Equals(B);
+			}
+
+			public static bool operator!=(PropertyCacheKey A, PropertyCacheKey B)
+			{
+				return !(A == B);
+			}
 		}
 
 
