Fix ParamHelpAttribute serialization issues with public setters
Make Flag getter more consistent

diff --git a/Engine/Source/Programs/AutomationTool/AutomationUtils/HelpUtils.cs b/Engine/Source/Programs/AutomationTool/AutomationUtils/HelpUtils.cs
--- a/Engine/Source/Programs/AutomationTool/AutomationUtils/HelpUtils.cs
+++ b/Engine/Source/Programs/AutomationTool/AutomationUtils/HelpUtils.cs
@@ -89,6 +89,11 @@ namespace AutomationTool
 		{
 		}
 
+		public ParamHelpAttribute() : this("", "")
+		{
+			// Needed for de-serialization
+		}
+
 		public ParamHelpAttribute(string Name, string Description) : base(Name, Description)
 		{
 			ParamName = Name.TrimStart('-');
@@ -109,12 +114,19 @@ namespace AutomationTool
 		public string ParamName
 		{
 			get;
+			set;
 		}
 		
+		private string _ParamDescription = string.Empty;
 		public string ParamDescription
 		{
 			get
 			{
+				if (!string.IsNullOrEmpty(_ParamDescription))
+				{
+					return _ParamDescription;
+				}
+
 				string RequiredStr = Required ? ". Required" : ". Optional";
 				string ChoicesStr = "";
 				var ValidChoices = Choices as IEnumerable<object>;
@@ -135,6 +147,10 @@ namespace AutomationTool
 						ChoicesStr
 					);
 			}
+			set
+			{
+				_ParamDescription = value;
+			}
 		}
 
 		private ParamAction _Action = ParamAction.Store;
@@ -189,7 +205,7 @@ namespace AutomationTool
 			set;
 		}
 
-		public string Help
+		public bool IsArgument
 		{
 			get;
 			set;
@@ -214,7 +230,7 @@ namespace AutomationTool
 			{
 				if (_Flag == null)
 				{
-					return ParamName;
+					return "-" + ParamName;
 				}
 				return _Flag;
 			}
@@ -222,7 +238,7 @@ namespace AutomationTool
 			{
 				if (!string.IsNullOrEmpty(value))
 				{
-					_Flag = value.TrimStart('-');
+					_Flag = value;
 				}
 			}
 		}
