Fix error message when multiple Client or Server target.cs files exist, but -Target= is not specified.

#rnx

diff --git a/Engine/Source/Programs/AutomationTool/AutomationUtils/ProjectParams.cs b/Engine/Source/Programs/AutomationTool/AutomationUtils/ProjectParams.cs
--- a/Engine/Source/Programs/AutomationTool/AutomationUtils/ProjectParams.cs
+++ b/Engine/Source/Programs/AutomationTool/AutomationUtils/ProjectParams.cs
@@ -2245,7 +2245,7 @@ namespace AutomationTool
 				case 0:
 					throw new AutomationException("{0} target not found!", Type);
 				default:
-					throw new AutomationException("More than one {0} target found. Specify which one to use with the -{1}= option.", Type, Type);
+					throw new AutomationException("More than one {0} target found. Specify which one to use with the -Target= option.", Type);
 			}
 		}
 
