Fix Gauntlet throwing an InitializationFailed error when a process fails to log in.

diff --git a/Engine/Source/Programs/AutomationTool/Gauntlet/Unreal/Base/Gauntlet.UnrealProcessResult.cs b/Engine/Source/Programs/AutomationTool/Gauntlet/Unreal/Base/Gauntlet.UnrealProcessResult.cs
--- a/Engine/Source/Programs/AutomationTool/Gauntlet/Unreal/Base/Gauntlet.UnrealProcessResult.cs
+++ b/Engine/Source/Programs/AutomationTool/Gauntlet/Unreal/Base/Gauntlet.UnrealProcessResult.cs
@@ -11,6 +11,7 @@ namespace Gauntlet
 		InitializationFailure,          // Process failed to initialize (e.g. the editor or game failed to load)
 		EncounteredFatalError,          // A fatal error occurred
 		EncounteredEnsure,              // An ensure occurred (will only be returned if the test considers ensures as fatal)
+		LoginFailed,					// Client never successfully logged in.
 		TimeOut,                        // A timeout occurred
 		TestFailure,                    // A test is known to have failed
 		UnrealError,                    // Unreal exited with an error code unrelated to test issues
diff --git a/Engine/Source/Programs/AutomationTool/Gauntlet/Unreal/Base/Gauntlet.UnrealTestNode.cs b/Engine/Source/Programs/AutomationTool/Gauntlet/Unreal/Base/Gauntlet.UnrealTestNode.cs
--- a/Engine/Source/Programs/AutomationTool/Gauntlet/Unreal/Base/Gauntlet.UnrealTestNode.cs
+++ b/Engine/Source/Programs/AutomationTool/Gauntlet/Unreal/Base/Gauntlet.UnrealTestNode.cs
@@ -1093,6 +1093,10 @@ namespace Gauntlet
 				{
 					TestNodeEvents.Add(new UnrealTestEvent(EventSeverity.Fatal, "Engine Initialization Failed", new List<string> { "Engine failed to initialize in one or more roles. This is likely a bad build." }));
 				}
+				if (DistinctResults.Contains(UnrealProcessResult.LoginFailed))
+				{
+					TestNodeEvents.Add(new UnrealTestEvent(EventSeverity.Fatal, "Login Unsuccessful", new List<string> { "User account never successfully finished logging in." }));
+				}
 				if (DistinctResults.Contains(UnrealProcessResult.EncounteredFatalError))
 				{
 					TestNodeEvents.Add(new UnrealTestEvent(EventSeverity.Fatal, "Fatal Error Encountered", new List<string> { "Test encountered a fatal error. Check ClientOutput.log and ServerOutput.log for details." }));
