- Allow staging/packaging of programs (requires the program have a basic .uproject, see the new SlateUGS.uproject)
- Added code to patch paths for handling Programs having source in /Engine/Source/Programs/XYZ, their config fies, etc in /Engine/Programs/XYZ, and staging using /XYZ
- Mac can now apply -specifiedarchitecture to UAT builds of programs
- Added support to SlateUGS to load Pak files (programs have to opt-in to it in their main Build.cs)

diff --git a/Engine/Programs/SlateUGS/SlateUGS.uproject b/Engine/Programs/SlateUGS/SlateUGS.uproject
new file mode 100644
--- /dev/null
+++ b/Engine/Programs/SlateUGS/SlateUGS.uproject
@@ -0,0 +1,11 @@
+{
+	"FileVersion": 3,
+	"EngineAssociation": "",
+	"Category": "",
+	"Description": "",
+	"TargetPlatforms": [
+		"Windows",
+		"Mac",
+		"Linux"
+	]
+}
diff --git a/Engine/Source/Programs/AutomationTool/AutomationUtils/DeploymentContext.cs b/Engine/Source/Programs/AutomationTool/AutomationUtils/DeploymentContext.cs
--- a/Engine/Source/Programs/AutomationTool/AutomationUtils/DeploymentContext.cs
+++ b/Engine/Source/Programs/AutomationTool/AutomationUtils/DeploymentContext.cs
@@ -707,6 +707,12 @@ public class DeploymentContext //: ProjectParams
 	/// <param name="Files">List to receive the enumerated files</param>
 	private void FindFilesToStageInternal(DirectoryReference BaseDir, string Pattern, StageFilesSearch Option, List<FileReference> Files)
 	{
+		// if the directory doesn't exist, this will crash in EnumerateFiles
+		if (!DirectoryReference.Exists(BaseDir))
+		{
+			return;
+		}	
+
 		// Enumerate all the files in this directory
 		Files.AddRange(DirectoryReference.EnumerateFiles(BaseDir, Pattern));
 
diff --git a/Engine/Source/Programs/AutomationTool/AutomationUtils/ProjectParams.cs b/Engine/Source/Programs/AutomationTool/AutomationUtils/ProjectParams.cs
--- a/Engine/Source/Programs/AutomationTool/AutomationUtils/ProjectParams.cs
+++ b/Engine/Source/Programs/AutomationTool/AutomationUtils/ProjectParams.cs
@@ -702,7 +702,7 @@ namespace AutomationTool
 			this.Clean = GetOptionalParamValueIfNotSpecified(Command, Clean, this.Clean, "clean", null);
 			this.SignPak = ParseParamValueIfNotSpecified(Command, SignPak, "signpak", String.Empty);
 			this.SignedPak = !String.IsNullOrEmpty(this.SignPak) || GetParamValueIfNotSpecified(Command, SignedPak, this.SignedPak, "signedpak");
-			if (string.IsNullOrEmpty(this.SignPak))
+			if (string.IsNullOrEmpty(this.SignPak) && RawProjectPath != null)
 			{
 				this.SignPak = Path.Combine(RawProjectPath.Directory.FullName, @"Restricted\NoRedist\Build\Keys.txt");
 				if (!File.Exists(this.SignPak))
@@ -3073,7 +3073,7 @@ namespace AutomationTool
 				CommandUtils.LogLog("DirectoriesToCook={0}", DirectoriesToCook.ToString());
                 CommandUtils.LogLog("DDCGraph={0}", DDCGraph);
                 CommandUtils.LogLog("CulturesToCook={0}", CommandUtils.IsNullOrEmpty(CulturesToCook) ? "<Not Specified> (Use Defaults)" : CulturesToCook.ToString());
-				CommandUtils.LogLog("EditorTargets={0}", EditorTargets.ToString());
+				CommandUtils.LogLog("EditorTargets={0}", EditorTargets?.ToString());
 				CommandUtils.LogLog("Foreign={0}", Foreign);
 				CommandUtils.LogLog("IsCodeBasedProject={0}", IsCodeBasedProject.ToString());
 				CommandUtils.LogLog("IsProgramTarget={0}", IsProgramTarget.ToString());
diff --git a/Engine/Source/Programs/AutomationTool/AutomationUtils/ProjectUtils.cs b/Engine/Source/Programs/AutomationTool/AutomationUtils/ProjectUtils.cs
--- a/Engine/Source/Programs/AutomationTool/AutomationUtils/ProjectUtils.cs
+++ b/Engine/Source/Programs/AutomationTool/AutomationUtils/ProjectUtils.cs
@@ -580,7 +580,7 @@ namespace AutomationTool
 			Properties.RawProjectPath = RawProjectPath;
 
 			// detect if the project is content only, but has non-default build settings
-			List<string> ExtraSearchPaths = null;
+			List<string> ExtraSearchPaths = new();
 			if (RawProjectPath != null)
 			{
 				// no Target file, now check to see if build settings have changed
@@ -614,7 +614,6 @@ namespace AutomationTool
 				{
 					GenerateTempTarget(RawProjectPath);
 					Properties.bWasGenerated = true;
-					ExtraSearchPaths = new List<string>();
                     ExtraSearchPaths.Add(TempTargetDir);
 				}
 				else if (File.Exists(Path.Combine(Path.GetDirectoryName(RawProjectPath.FullName), "Intermediate", "Source", Path.GetFileNameWithoutExtension(RawProjectPath.FullName) + ".Target.cs")))
@@ -745,7 +744,15 @@ namespace AutomationTool
 
 				TargetsDllFilename = FileReference.Combine(RulesFolder, String.Format("UATRules-{0}.dll", ContentHash.MD5(Properties.RawProjectPath.FullName.ToUpperInvariant()).ToString()));
 
-				FullProjectPath = CommandUtils.GetDirectoryName(Properties.RawProjectPath.FullName);
+				FullProjectPath = CommandUtils.GetDirectoryName(Properties.RawProjectPath.FullName).Replace("\\", "/");
+
+				// there is a special case of Programs, where the uproject doesn't align with the Source directory, so we redirect to where
+				// the program's target.cs file(s) are
+				if (FullProjectPath.Contains("/Programs/"))
+				{
+					FullProjectPath = FullProjectPath.Replace("/Programs/", "/Source/Programs/");
+				}
+
 				GameFolders.Add(new DirectoryReference(FullProjectPath));
 				CommandUtils.LogVerbose("Searching for target rule files in {0}", FullProjectPath);
 			}
diff --git a/Engine/Source/Programs/AutomationTool/Mac/MacPlatform.Automation.cs b/Engine/Source/Programs/AutomationTool/Mac/MacPlatform.Automation.cs
--- a/Engine/Source/Programs/AutomationTool/Mac/MacPlatform.Automation.cs
+++ b/Engine/Source/Programs/AutomationTool/Mac/MacPlatform.Automation.cs
@@ -1,4 +1,4 @@
-// Copyright Epic Games, Inc. All Rights Reserved.
+﻿// Copyright Epic Games, Inc. All Rights Reserved.
 using System;
 using System.Collections.Generic;
 using System.Linq;
@@ -121,19 +121,19 @@ public class MacPlatform : Platform
 			// Default to Intel. 
 			string UBTArchitectureParam = MacExports.IntelArchitecture;
 
-			// Targets are easy. 
-			// - If the project is set to Intel/Apple we do that
-			// - If it's universal we build the local architecture unless distributing
-			// - SpecifiedArchitecture overrides these
-			if (IsTarget)
+			// If an architecture was specified, use that
+			if (!string.IsNullOrEmpty(Params.SpecifiedArchitecture))
 			{
-				// If an architecture was specified, use that
-				if (!string.IsNullOrEmpty(Params.SpecifiedArchitecture))
-				{
-					UBTArchitectureParam = Params.SpecifiedArchitecture;
-					Log.TraceInformation("Building {0} as {1} due to -specifiedarchitecture", Target.TargetName, UBTArchitectureParam);
-				}
-				else
+				UBTArchitectureParam = Params.SpecifiedArchitecture;
+				Log.TraceInformation("Building {0} as {1} due to -specifiedarchitecture", Target.TargetName, UBTArchitectureParam);
+			}
+			else
+			{
+				// Targets are easy. 
+				// - If the project is set to Intel/Apple we do that
+				// - If it's universal we build the local architecture unless distributing
+				// - SpecifiedArchitecture overrides these
+				if (IsTarget)
 				{
 					// If the project isn't marked as universal built what it's set to
 					if (!ProjectIsUniversal)
@@ -156,20 +156,20 @@ public class MacPlatform : Platform
 						}
 					}
 				}
-			}
-			else
-			{
-				// We build tools for the local architecture if possible
-				if (CanBuildTargetForAllArchitectures(Target, Params) || LocalArchitecture == MacExports.IntelArchitecture)
-				{
-					UBTArchitectureParam = LocalArchitecture;
-					Log.TraceInformation("Building {0} as {1} for host", Target.TargetName, UBTArchitectureParam);
-				}
-				else if (MacExports.IsRunningOnAppleArchitecture)
+				else
 				{
-					// Tell them why to avoid confusion
-					Log.TraceInformation("Building {0} as {1} (arm64 not currently supported)", Target.TargetName, UBTArchitectureParam);
-					UBTArchitectureParam = MacExports.IntelArchitecture;
+					// We build tools for the local architecture if possible
+					if (CanBuildTargetForAllArchitectures(Target, Params) || LocalArchitecture == MacExports.IntelArchitecture)
+					{
+						UBTArchitectureParam = LocalArchitecture;
+						Log.TraceInformation("Building {0} as {1} for host", Target.TargetName, UBTArchitectureParam);
+					}
+					else if (MacExports.IsRunningOnAppleArchitecture)
+					{
+						// Tell them why to avoid confusion
+						Log.TraceInformation("Building {0} as {1} (arm64 not currently supported)", Target.TargetName, UBTArchitectureParam);
+						UBTArchitectureParam = MacExports.IntelArchitecture;
+					}
 				}
 			}
 
diff --git a/Engine/Source/Programs/AutomationTool/Scripts/BuildProjectCommand.Automation.cs b/Engine/Source/Programs/AutomationTool/Scripts/BuildProjectCommand.Automation.cs
--- a/Engine/Source/Programs/AutomationTool/Scripts/BuildProjectCommand.Automation.cs
+++ b/Engine/Source/Programs/AutomationTool/Scripts/BuildProjectCommand.Automation.cs
@@ -107,7 +107,7 @@ namespace AutomationScripts
 			// Build any tools we need to stage
 			if ((TargetMask & ProjectBuildTargets.UnrealPak) == ProjectBuildTargets.UnrealPak && !Unreal.IsEngineInstalled())
 			{
-				if (Params.EditorTargets.Contains("UnrealPak") == false)
+				if (!Params.HasEditorTargets || Params.EditorTargets.Contains("UnrealPak") == false)
 				{
 					Agenda.AddTargets(new string[] { "UnrealPak" }, HostPlatform.Current.HostEditorPlatform, UnrealTargetConfiguration.Development, Params.CodeBasedUprojectPath);
 				}
diff --git a/Engine/Source/Programs/AutomationTool/Scripts/CopyBuildToStagingDirectory.Automation.cs b/Engine/Source/Programs/AutomationTool/Scripts/CopyBuildToStagingDirectory.Automation.cs
--- a/Engine/Source/Programs/AutomationTool/Scripts/CopyBuildToStagingDirectory.Automation.cs
+++ b/Engine/Source/Programs/AutomationTool/Scripts/CopyBuildToStagingDirectory.Automation.cs
@@ -3905,6 +3905,11 @@ namespace AutomationScripts
 				return false;
 			}
 
+			if (Params.IsProgramTarget)
+			{
+				return false;
+			}
+
 			if (Params.SkipIoStore)
 			{
 				return false;
@@ -4321,7 +4326,19 @@ namespace AutomationScripts
 		//@todo move this
 		public static List<DeploymentContext> CreateDeploymentContext(ProjectParams Params, bool InDedicatedServer, bool DoCleanStage = false)
 		{
-			ParamList<string> ListToProcess = InDedicatedServer && (Params.Cook || Params.CookOnTheFly) ? Params.ServerCookedTargets : Params.ClientCookedTargets;
+			ParamList<string> ListToProcess;
+			if (Params.IsProgramTarget)
+			{
+				ListToProcess = Params.ProgramTargets;
+			}
+			else if (InDedicatedServer && (Params.Cook || Params.CookOnTheFly))
+			{
+				ListToProcess = Params.ServerCookedTargets;
+			}
+			else
+			{
+				ListToProcess = Params.ClientCookedTargets;
+			} 
 			var ConfigsToProcess = InDedicatedServer && (Params.Cook || Params.CookOnTheFly) ? Params.ServerConfigsToBuild : Params.ClientConfigsToBuild;
 
 			List<Tuple<string, UnrealTargetConfiguration>> TargetAndConfigPairs = new List<Tuple<string, UnrealTargetConfiguration>>();
@@ -4331,7 +4348,8 @@ namespace AutomationScripts
 				// If we are staging a client and have been asked to include editor targets, we currently only want to
 				// include a single Development editor target. Ideally, we should have shipping editor configs and then
 				// just include the requested configs for all targets
-				if (Params.EditorTargets.Contains(Target) || Params.ProgramTargets.Contains(Target))
+				if ((Params.HasEditorTargets && Params.EditorTargets.Contains(Target)) || 
+					(Params.HasProgramTargets && Params.ProgramTargets.Contains(Target)))
 				{
 					TargetAndConfigPairs.Add(new Tuple<string, UnrealTargetConfiguration>(Target, UnrealTargetConfiguration.Development));
 				}
diff --git a/Engine/Source/Programs/SlateUGS/SlateUGS.Build.cs b/Engine/Source/Programs/SlateUGS/SlateUGS.Build.cs
--- a/Engine/Source/Programs/SlateUGS/SlateUGS.Build.cs
+++ b/Engine/Source/Programs/SlateUGS/SlateUGS.Build.cs
@@ -19,6 +19,7 @@ public class SlateUGS : ModuleRules
 				"SlateCore",
 				"StandaloneRenderer",
 				"ToolWidgets",
+				"PakFile",
 				"DesktopPlatform", // Open file dialog
 			}
 		);
diff --git a/Engine/Source/Programs/UnrealBuildTool/Platform/Mac/MacExports.cs b/Engine/Source/Programs/UnrealBuildTool/Platform/Mac/MacExports.cs
--- a/Engine/Source/Programs/UnrealBuildTool/Platform/Mac/MacExports.cs
+++ b/Engine/Source/Programs/UnrealBuildTool/Platform/Mac/MacExports.cs
@@ -1,4 +1,4 @@
-// Copyright Epic Games, Inc. All Rights Reserved.
+﻿// Copyright Epic Games, Inc. All Rights Reserved.
 
 using System;
 using System.Collections.Generic;
@@ -67,7 +67,8 @@ namespace UnrealBuildTool
 					"UnrealGame", 
 					"UnrealServer",
 					"UnrealHeaderTool", 
-					"UnrealPak"
+					"UnrealPak",
+					"SlateUGS",
 				};
 			}
 		}
diff --git a/Engine/Source/Runtime/Core/Private/GenericPlatform/GenericPlatformMisc.cpp b/Engine/Source/Runtime/Core/Private/GenericPlatform/GenericPlatformMisc.cpp
--- a/Engine/Source/Runtime/Core/Private/GenericPlatform/GenericPlatformMisc.cpp
+++ b/Engine/Source/Runtime/Core/Private/GenericPlatform/GenericPlatformMisc.cpp
@@ -1279,6 +1279,12 @@ const TCHAR* FGenericPlatformMisc::ProjectDir()
 		{
 			// monolithic, game-agnostic executables, the ini is in Engine/Config/Platform
 			ProjectDir = FString::Printf(TEXT("../../../Engine/Programs/%s/"), FApp::GetProjectName());
+
+			// however, if it was staged, that directory won't exist, so look in the normal staged location
+			if (!FPlatformFileManager::Get().GetPlatformFile().DirectoryExists(*ProjectDir))
+			{
+				ProjectDir = FString::Printf(TEXT("../../../%s/"), FApp::GetProjectName());
+			}
 		}
 		else
 		{
