Retry save role artifacts once with device reboot on exception containing message "A retry should be performed"

diff --git a/Engine/Source/Programs/AutomationTool/Gauntlet/Unreal/Base/Gauntlet.UnrealSession.cs b/Engine/Source/Programs/AutomationTool/Gauntlet/Unreal/Base/Gauntlet.UnrealSession.cs
--- a/Engine/Source/Programs/AutomationTool/Gauntlet/Unreal/Base/Gauntlet.UnrealSession.cs
+++ b/Engine/Source/Programs/AutomationTool/Gauntlet/Unreal/Base/Gauntlet.UnrealSession.cs
@@ -1196,12 +1196,23 @@ namespace Gauntlet
 					{
 						Artifacts = SaveRoleArtifacts(Context, App, DestPath);
 					}
-					catch (Exception)
+					catch (Exception SaveArtifactsException)
 					{
 						// Caught an exception -> report failure
 						IDeviceUsageReporter.RecordEnd(device.Name, device.Platform, IDeviceUsageReporter.EventType.SavingArtifacts, IDeviceUsageReporter.EventState.Failure);
-						// Pass exception to the surrounding try/catch in UnrealTestNode while preserving the original callstack
-						throw;
+
+						// Retry once only, after rebooting device, if artifacts couldn't be saved.
+						if (SaveArtifactsException.Message.Contains("A retry should be performed"))
+						{
+							Log.Info("Rebooting device and retrying save role artifacts once.");
+							App.AppInstance.Device.Reboot();
+							Artifacts = SaveRoleArtifacts(Context, App, DestPath);
+						}
+						else
+						{
+							// Pass exception to the surrounding try/catch in UnrealTestNode while preserving the original callstack
+							throw;
+						}
 					}
 					// Did not catch -> successful reporting
 					IDeviceUsageReporter.RecordEnd(device.Name, device.Platform, IDeviceUsageReporter.EventType.SavingArtifacts, IDeviceUsageReporter.EventState.Success);
