- Allow AutoSDK installation via Turnkey from any copy provider

diff --git a/Engine/Source/Programs/AutomationTool/Turnkey/CopyProvider.cs b/Engine/Source/Programs/AutomationTool/Turnkey/CopyProvider.cs
--- a/Engine/Source/Programs/AutomationTool/Turnkey/CopyProvider.cs
+++ b/Engine/Source/Programs/AutomationTool/Turnkey/CopyProvider.cs
@@ -1,4 +1,4 @@
-// Copyright Epic Games, Inc. All Rights Reserved.
+﻿// Copyright Epic Games, Inc. All Rights Reserved.
 
 using System;
 using System.Collections.Generic;
@@ -109,8 +109,12 @@ namespace Turnkey
 
 			ParseOperation(CopyOperation, out Provider, out ProviderParam, false);
 
-			// execute what comes after the colon
-			string OutputPath = Provider.Execute(ProviderParam, SpecialMode, SpecialModeHint);
+			string OperationExt = Path.GetExtension(ProviderParam).ToLower();
+			bool bWillDecompressToLocation = (OperationExt == ".zip" || OperationExt == ".7z") && SpecialMode == CopyExecuteSpecialMode.UsePermanentStorage;
+
+			// execute what comes after the colon (if we guess we are going to decompress, don't pass in a specified location, we will decompress to it after)
+			string OutputPath = bWillDecompressToLocation ? Provider.Execute(ProviderParam, CopyExecuteSpecialMode.None, null) : 
+				Provider.Execute(ProviderParam, SpecialMode, SpecialModeHint);
 
 			// always unzip the file into a temp directory (or downloadpath if it's a permanent download), and make that directory be the outputpath
 			if (OutputPath != null)
@@ -137,7 +141,7 @@ namespace Turnkey
 						string DecompressLocation;
 						if (SpecialMode == CopyExecuteSpecialMode.UsePermanentStorage)
 						{
-							DecompressLocation = Path.Combine(Path.GetDirectoryName(OutputPath), Path.GetFileNameWithoutExtension(OutputPath) + "_Uncompressed");
+							DecompressLocation = SpecialModeHint; // Path.Combine(Path.GetDirectoryName(OutputPath), Path.GetFileNameWithoutExtension(OutputPath) + "_Uncompressed");
 						}
 						else
 						{
diff --git a/Engine/Source/Programs/AutomationTool/Turnkey/LocalCache.cs b/Engine/Source/Programs/AutomationTool/Turnkey/LocalCache.cs
--- a/Engine/Source/Programs/AutomationTool/Turnkey/LocalCache.cs
+++ b/Engine/Source/Programs/AutomationTool/Turnkey/LocalCache.cs
@@ -1,4 +1,4 @@
-// Copyright Epic Games, Inc. All Rights Reserved.
+﻿// Copyright Epic Games, Inc. All Rights Reserved.
 
 using System;
 using System.Collections.Generic;
@@ -145,6 +145,7 @@ namespace Turnkey
 		static void SerializeObject(SavedCache Object, string Filename)
 		{
 			string Str = JsonSerializer.Serialize(Object);
+			Directory.CreateDirectory(Path.GetDirectoryName(Filename));
 			File.WriteAllText(Filename, Str);
 
 // 			using (FileStream Stream = new FileStream(Filename, FileMode.Create))
diff --git a/Engine/Source/Programs/AutomationTool/Turnkey/Providers/GoogleDriveCopyProvider.cs b/Engine/Source/Programs/AutomationTool/Turnkey/Providers/GoogleDriveCopyProvider.cs
--- a/Engine/Source/Programs/AutomationTool/Turnkey/Providers/GoogleDriveCopyProvider.cs
+++ b/Engine/Source/Programs/AutomationTool/Turnkey/Providers/GoogleDriveCopyProvider.cs
@@ -1,4 +1,4 @@
-// Copyright Epic Games, Inc. All Rights Reserved.
+﻿// Copyright Epic Games, Inc. All Rights Reserved.
 
 using System;
 using System.Collections.Generic;
@@ -111,7 +111,7 @@ namespace Turnkey
 				// if we are just downloading a quick switch SDK, or similar, then we are given a Hint for a location to download to
 				if (SpecialMode == CopyExecuteSpecialMode.UsePermanentStorage && CachedOperationLocation == null)
 				{
-					DownloadDirectory = Path.Combine(LocalCache.GetInstallCacheDirectory(), SpecialModeHint);
+					DownloadDirectory = SpecialModeHint;// Path.Combine(LocalCache.GetInstallCacheDirectory(), SpecialModeHint);
 					AutomationTool.InternalUtils.SafeDeleteDirectory(DownloadDirectory);
 				}
 
diff --git a/Engine/Source/Programs/AutomationTool/Turnkey/Providers/HttpCopyProvider.cs b/Engine/Source/Programs/AutomationTool/Turnkey/Providers/HttpCopyProvider.cs
--- a/Engine/Source/Programs/AutomationTool/Turnkey/Providers/HttpCopyProvider.cs
+++ b/Engine/Source/Programs/AutomationTool/Turnkey/Providers/HttpCopyProvider.cs
@@ -1,4 +1,4 @@
-// Copyright Epic Games, Inc. All Rights Reserved.
+﻿// Copyright Epic Games, Inc. All Rights Reserved.
 
 using System;
 using System.IO;
@@ -64,7 +64,7 @@ namespace Turnkey
 				// if we are just downloading a quick switch SDK, or similar, then we are given a Hint for a location to download to
 				if (SpecialMode == CopyExecuteSpecialMode.UsePermanentStorage && CachedOperationLocation == null)
 				{
-					DownloadDirectory = Path.Combine(LocalCache.GetInstallCacheDirectory(), SpecialModeHint);
+					DownloadDirectory = SpecialModeHint;// Path.Combine(LocalCache.GetInstallCacheDirectory(), SpecialModeHint);
 					AutomationTool.InternalUtils.SafeDeleteDirectory(DownloadDirectory);
 				}
 				else
diff --git a/Engine/Source/Programs/AutomationTool/Turnkey/Providers/NullCopyProvider.cs b/Engine/Source/Programs/AutomationTool/Turnkey/Providers/NullCopyProvider.cs
--- a/Engine/Source/Programs/AutomationTool/Turnkey/Providers/NullCopyProvider.cs
+++ b/Engine/Source/Programs/AutomationTool/Turnkey/Providers/NullCopyProvider.cs
@@ -1,4 +1,4 @@
-// Copyright Epic Games, Inc. All Rights Reserved.
+﻿// Copyright Epic Games, Inc. All Rights Reserved.
 
 using System.IO;
 using System.Collections.Generic;
@@ -81,7 +81,7 @@ namespace Turnkey
 			bool bIsLargeFile = File.Exists(OutputPath) && new FileInfo(OutputPath).Length > 5 * 1024 * 1024;
 			bool bIsDirectory = Directory.Exists(OutputPath);
 
-			if (bIsRemotePath && (bIsLargeFile || bIsDirectory))
+			if ((bIsRemotePath && (bIsLargeFile || bIsDirectory)) || (SpecialMode == CopyExecuteSpecialMode.UsePermanentStorage))
 			{
 				// look in out local cache
 				string OperationTag = string.Format("file_op:{0}", Operation);
@@ -92,7 +92,7 @@ namespace Turnkey
 					return CachedLocation;
 				}
 
-				string CopyLocation = Path.Combine(LocalCache.CreateTempDirectory(), Path.GetFileName(OutputPath));
+				string CopyLocation = (SpecialMode == CopyExecuteSpecialMode.UsePermanentStorage) ? SpecialModeHint : Path.Combine(LocalCache.CreateTempDirectory(), Path.GetFileName(OutputPath));
 
 				if (bIsDirectory)
 				{
diff --git a/Engine/Source/Programs/AutomationTool/Turnkey/SdkInfo.cs b/Engine/Source/Programs/AutomationTool/Turnkey/SdkInfo.cs
--- a/Engine/Source/Programs/AutomationTool/Turnkey/SdkInfo.cs
+++ b/Engine/Source/Programs/AutomationTool/Turnkey/SdkInfo.cs
@@ -1,4 +1,4 @@
-// Copyright Epic Games, Inc. All Rights Reserved.
+﻿// Copyright Epic Games, Inc. All Rights Reserved.
 
 using System;
 using System.Collections.Generic;
@@ -151,6 +151,8 @@ namespace Turnkey
 
 		public static bool SetupAutoSdk(FileSource Source, ITurnkeyContext TurnkeyContext, UnrealTargetPlatform Platform, bool bUnattended)
 		{
+			AutomationTool.Platform AutomationPlatform = AutomationTool.Platform.GetPlatform(Platform);
+
 			bool bSetupEnvVarAfterInstall = false;
 			if (Environment.GetEnvironmentVariable("UE_SDKS_ROOT") == null)
 			{
@@ -166,10 +168,11 @@ namespace Turnkey
 					bSetupEnvVarAfterInstall = true;
 				}
 			}
+			else
+			{
+				TurnkeyUtils.Log("{0}: AutoSdk is setup on this computer, will look for available AutoSdk to download", Platform);
+			}
 
-			AutomationTool.Platform AutomationPlatform = AutomationTool.Platform.GetPlatform(Platform);
-
-			TurnkeyUtils.Log("{0}: AutoSdk is setup on this computer, will look for available AutoSdk to download", Platform);
 
 			// make sure this is unset so that we can know if it worked or not after install
 			TurnkeyUtils.ClearVariable("CopyOutputPath");
@@ -183,8 +186,44 @@ namespace Turnkey
 				return false;
 			}
 
-			// download the AutoSDK
-			string DownloadedRoot = CopyProvider.ExecuteCopy(CopyOperation);
+			string DownloadedRoot;
+
+			// perforce is special case in that it will setup UE_SDKS_ROOT after we download, because p4 already has a mapping to a certain location
+			// and we don't need to ask user (if perforce needs to ask, it will)
+			if (!CopyOperation.StartsWith(new PerforceCopyProvider().ProviderToken))
+			{
+				if (bSetupEnvVarAfterInstall)
+				{
+					string Response = TurnkeyUtils.ReadInput("Enter directory to use for root of AutoSDKs", Path.Combine(System.Environment.GetFolderPath(Environment.SpecialFolder.MyDocuments), "AutoSDK"));
+					if (string.IsNullOrEmpty(Response))
+					{
+						return false;
+					}
+
+					// set the env var, globally
+					TurnkeyUtils.StartTrackingExternalEnvVarChanges();
+					Environment.SetEnvironmentVariable("UE_SDKS_ROOT", Response);
+					Environment.SetEnvironmentVariable("UE_SDKS_ROOT", Response, EnvironmentVariableTarget.User);
+					TurnkeyUtils.EndTrackingExternalEnvVarChanges();
+
+					bSetupEnvVarAfterInstall = false;
+				}
+
+				// download the AutoSDK, to the UE_SDKS_ROOT dir
+				string AutoSDKDownloadedRoot = Path.Combine(Environment.GetEnvironmentVariable("UE_SDKS_ROOT"), $"Host{HostPlatform.Platform}", Platform.ToString(), Source.Version);
+				// it should return what we pass into ig
+				DownloadedRoot = CopyProvider.ExecuteCopy(CopyOperation, CopyExecuteSpecialMode.UsePermanentStorage, AutoSDKDownloadedRoot);
+				if (DownloadedRoot != AutoSDKDownloadedRoot)
+				{
+					throw new BuildException($"Turnkey did not download to the expected location, the copy operation failed: {CopyOperation} did not respect CopyExecuteSpecialMode.UsePermanentStorage with {AutoSDKDownloadedRoot}, it copied to {DownloadedRoot} instead");
+				}
+			}
+			else
+			{
+				// download the AutoSDK using perforce, where it doesn't need a Hint passed to it
+				DownloadedRoot = CopyProvider.ExecuteCopy(CopyOperation);
+			}
+
 			if (string.IsNullOrEmpty(DownloadedRoot))
 			{
 				TurnkeyContext.ReportError($"Unable to download the AutoSDK for {Platform}. Your Studio's TurnkeyManifest.xml file(s) may need to be fixed.");
