Fixes issue with -open and -build being specified on the command line with SyncProject.

diff --git a/Engine/Source/Programs/AutomationTool/Scripts/SyncProject.Automation.cs b/Engine/Source/Programs/AutomationTool/Scripts/SyncProject.Automation.cs
--- a/Engine/Source/Programs/AutomationTool/Scripts/SyncProject.Automation.cs
+++ b/Engine/Source/Programs/AutomationTool/Scripts/SyncProject.Automation.cs
@@ -318,8 +318,11 @@ class SyncProject : SyncProjectBase
 				BuildCmd.Params = this.Params;
 				ExitStatus = BuildCmd.Execute();
 			}
-
-			if (OpenProject && ExitStatus == ExitCode.Success)
+			
+			// If both -build and -open are specified
+			// on the commandline, then we will rely on the successful build to call
+			// it's open of the editor. Otherwise just open the editor
+			if (!BuildProject && OpenProject && ExitStatus == ExitCode.Success)
 			{
 				Log.TraceVerbose("Opening Editor for {0}", ProjectArgForEditor);
 
