- Used new ConfigRedirects to rename the following ini keys:
  WhitelistConfigFiles=AllowedConfigFiles
  BlacklistConfigFiles=DisallowedConfigFiles
  WhitelistDirectories=AllowedDirectories
  BlacklistLocalizationTargets=DisallowedLocalizationTargets

diff --git a/Engine/Config/BaseGame.ini b/Engine/Config/BaseGame.ini
--- a/Engine/Config/BaseGame.ini
+++ b/Engine/Config/BaseGame.ini
@@ -220,10 +220,10 @@ FrameInterpolationProcessor=Class'/Script/LiveLink.LiveLinkBasicFrameInterpolati
 ; Games can set options used during the Staging process by adding lines in this section to DefaultGame.ini or platform-specific ini files
 ; These are examples, the full list is described in DeploymentContext.cs:
 ; +RemapDirectories=(From="Engine/Platforms/SecretPlatform/SpecificPlugin", To="Engine/Plugins/SpecificPlugin")
-; +WhitelistDirectories=GameName/SuspiciousFolderName
-; +BlacklistConfigFiles=GameName/Config/EditorOnlySystem.ini
-; +WhiitelistConfigFiles=GameName/Config/SuspiciouslyNamedSystem.ini
-; +BlacklistLocalizationTargets=Engine
+; +AllowedDirectories=GameName/SuspiciousFolderName
+; +DisallowedConfigFiles=GameName/Config/EditorOnlySystem.ini
+; +AllowedConfigFiles=GameName/Config/SuspiciouslyNamedSystem.ini
+; +DisallowedLocalizationTargets=Engine
 
 [CookedEditorSettings]
 +EngineAssetPaths="/Engine"
diff --git a/Engine/Config/ConfigRedirects.ini b/Engine/Config/ConfigRedirects.ini
--- a/Engine/Config/ConfigRedirects.ini
+++ b/Engine/Config/ConfigRedirects.ini
@@ -20,8 +20,8 @@
 
 ;; This will look in sections named [Staging] and remap WhitelistConfigFiles to AllowedConfigFiles, etc
 [Staging]
-;WhitelistConfigFiles=AllowedConfigFiles
-;BlacklistConfigFiles=DisallowedConfigFiles
-;WhitelistDirectories=AllowedDirectories
-;BlacklistLocalizationTargets=DisallowedLocalizationTargets
+WhitelistConfigFiles=AllowedConfigFiles
+BlacklistConfigFiles=DisallowedConfigFiles
+WhitelistDirectories=AllowedDirectories
+BlacklistLocalizationTargets=DisallowedLocalizationTargets
 
diff --git a/Engine/Config/HoloLens/HoloLensGame.ini b/Engine/Config/HoloLens/HoloLensGame.ini
--- a/Engine/Config/HoloLens/HoloLensGame.ini
+++ b/Engine/Config/HoloLens/HoloLensGame.ini
@@ -1,2 +1,2 @@
 [Staging]
-+WhitelistDirectories=Engine/Binaries/ThirdParty/Windows
\ No newline at end of file
++AllowedlistDirectories=Engine/Binaries/ThirdParty/Windows
diff --git a/Engine/Source/Editor/UnrealEd/Private/CookOnTheFlyServer.cpp b/Engine/Source/Editor/UnrealEd/Private/CookOnTheFlyServer.cpp
--- a/Engine/Source/Editor/UnrealEd/Private/CookOnTheFlyServer.cpp
+++ b/Engine/Source/Editor/UnrealEd/Private/CookOnTheFlyServer.cpp
@@ -10763,7 +10763,7 @@ void UCookOnTheFlyServer::RegisterLocalizationChunkDataGenerator()
 	TArray<FString> LocalizationTargetsToChunk = PackagingSettings->LocalizationTargetsToChunk;
 	{
 		TArray<FString> BlocklistLocalizationTargets;
-		GConfig->GetArray(TEXT("Staging"), TEXT("BlacklistLocalizationTargets"), BlocklistLocalizationTargets, GGameIni);
+		GConfig->GetArray(TEXT("Staging"), TEXT("DisallowedLocalizationTargets"), BlocklistLocalizationTargets, GGameIni);
 		if (BlocklistLocalizationTargets.Num() > 0)
 		{
 			LocalizationTargetsToChunk.RemoveAll([&BlocklistLocalizationTargets](const FString& InLocalizationTarget)
diff --git a/Engine/Source/Programs/AutomationTool/AutomationUtils/DeploymentContext.cs b/Engine/Source/Programs/AutomationTool/AutomationUtils/DeploymentContext.cs
--- a/Engine/Source/Programs/AutomationTool/AutomationUtils/DeploymentContext.cs
+++ b/Engine/Source/Programs/AutomationTool/AutomationUtils/DeploymentContext.cs
@@ -1,4 +1,4 @@
-// Copyright Epic Games, Inc. All Rights Reserved.
+﻿// Copyright Epic Games, Inc. All Rights Reserved.
 using System;
 using System.Collections.Generic;
 using System.IO;
@@ -289,19 +289,19 @@ public class DeploymentContext //: ProjectParams
 
 	/// <summary>
 	/// List of directories to allow staging, even if they contain restricted folder names
-	/// This list is read from the +WhitelistDirectories=... array in the [Staging] section of *Game.ini files
+	/// This list is read from the +AllowedDirectories=... array in the [Staging] section of *Game.ini files
 	/// </summary>
 	public List<StagedDirectoryReference> DirectoriesAllowList = new List<StagedDirectoryReference>();
 
 	/// <summary>
 	/// Set of config files which are allow listed to be staged. By default, we warn for config files which are not well known to prevent internal data (eg. editor/server settings)
-	/// leaking in packaged builds. This list is read from the +WhitelistConfigFiles=... array in the [Staging] section of *Game.ini files.
+	/// leaking in packaged builds. This list is read from the +AllowedConfigFiles=... array in the [Staging] section of *Game.ini files.
 	/// </summary>
 	public HashSet<StagedFileReference> ConfigFilesAllowList = new HashSet<StagedFileReference>();
 
 	/// <summary>
 	/// Set of config files which are denied from staging. By default, we warn for config files which are not well known to prevent internal data (eg. editor/server settings)
-	/// leaking in packaged builds. This list is read from the +BlacklistConfigFiles=... array in the [Staging] section of *Game.ini files.
+	/// leaking in packaged builds. This list is read from the +DeniedConfigFiles=... array in the [Staging] section of *Game.ini files.
 	/// </summary>
 	public HashSet<StagedFileReference> ConfigFilesDenyList = new HashSet<StagedFileReference>();
 
@@ -328,7 +328,7 @@ public class DeploymentContext //: ProjectParams
 
 	/// <summary>
 	/// List of localization targets that are not included in staged build. By default, all project Content/Localization targets are automatically staged.
-	/// This list is read from the +BlacklistLocalizationTargets=... array in the [Staging] section of *Game.ini files.
+	/// This list is read from the +DisallowedLocalizationTargets=... array in the [Staging] section of *Game.ini files.
 	/// </summary>
 	public List<string> LocalizationTargetsDenyList = new List<string>();
 
@@ -573,7 +573,7 @@ public class DeploymentContext //: ProjectParams
 
 		// Read the list of directories to allow (prevent from warning about from restricted folders)
 		List<string> DirectoriesAllowListStrings;
-		if (GameConfig.GetArray("Staging", "WhitelistDirectories", out DirectoriesAllowListStrings))
+		if (GameConfig.GetArray("Staging", "AllowedDirectories", out DirectoriesAllowListStrings))
 		{
 			foreach(string AllowedDir in DirectoriesAllowListStrings)
 			{
@@ -582,7 +582,7 @@ public class DeploymentContext //: ProjectParams
 		}
 
 		List<string> LocTargetsDenyListStrings;
-		if (GameConfig.GetArray("Staging", "BlacklistLocalizationTargets", out LocTargetsDenyListStrings))
+		if (GameConfig.GetArray("Staging", "DisallowedLocalizationTargets", out LocTargetsDenyListStrings))
 		{
 			foreach (string DeniedLocTarget in LocTargetsDenyListStrings)
 			{
@@ -591,8 +591,8 @@ public class DeploymentContext //: ProjectParams
 		}
 
 		// Read the list of files which are allow listed to be staged
-		ReadConfigFileList(GameConfig, "Staging", "WhitelistConfigFiles", ConfigFilesAllowList);
-		ReadConfigFileList(GameConfig, "Staging", "BlacklistConfigFiles", ConfigFilesDenyList);
+		ReadConfigFileList(GameConfig, "Staging", "AllowedConfigFiles", ConfigFilesAllowList);
+		ReadConfigFileList(GameConfig, "Staging", "DeniedConfigFiles", ConfigFilesDenyList);
 
 		// Grab the game ini data
 		String PackagingIniPath = "/Script/UnrealEd.ProjectPackagingSettings";
diff --git a/Engine/Source/Programs/AutomationTool/Scripts/CopyBuildToStagingDirectory.Automation.cs b/Engine/Source/Programs/AutomationTool/Scripts/CopyBuildToStagingDirectory.Automation.cs
--- a/Engine/Source/Programs/AutomationTool/Scripts/CopyBuildToStagingDirectory.Automation.cs
+++ b/Engine/Source/Programs/AutomationTool/Scripts/CopyBuildToStagingDirectory.Automation.cs
@@ -1,4 +1,4 @@
-// Copyright Epic Games, Inc. All Rights Reserved.
+﻿// Copyright Epic Games, Inc. All Rights Reserved.
 
 using AutomationTool;
 using System;
@@ -1389,7 +1389,7 @@ namespace AutomationScripts
 				{
 					Message.Append("\nAlternatively, allow list them using this syntax in DefaultGame.ini:");
 					Message.Append("\n[Staging]");
-					Message.Append("\n+WhitelistDirectories=MyGame/Content/Foo");
+					Message.Append("\n+AllowedDirectories=MyGame/Content/Foo");
 				}
 				throw new AutomationException(Message.ToString());
 			}
@@ -1444,7 +1444,7 @@ namespace AutomationScripts
 				Nullable<bool> ShouldStage = ShouldStageConfigFile(SC, ConfigDir, ConfigFile, PlatformExtensionName);
 				if (ShouldStage == null)
 				{
-					CommandUtils.LogWarning("The config file '{0}' will be staged, but is not explicitly allowed or denied. Add +WhitelistConfigFiles={0} or +BlacklistConfigFiles={0} to the [Staging] section of DefaultGame.ini", SC.GetStagedFileLocation(ConfigFile));
+					CommandUtils.LogWarning("The config file '{0}' will be staged, but is not explicitly allowed or denied. Add +AllowedConfigFiles={0} or +DeniedConfigFiles={0} to the [Staging] section of DefaultGame.ini", SC.GetStagedFileLocation(ConfigFile));
 				}
 
 				if (ShouldStage ?? true)
diff --git a/Templates/TP_AEC_HandheldARBP/Config/Android/AndroidGame.ini b/Templates/TP_AEC_HandheldARBP/Config/Android/AndroidGame.ini
--- a/Templates/TP_AEC_HandheldARBP/Config/Android/AndroidGame.ini
+++ b/Templates/TP_AEC_HandheldARBP/Config/Android/AndroidGame.ini
@@ -1,2 +1,2 @@
 [Staging]
-+WhitelistDirectories=Engine/Plugins/Runtime/AR/Apple/AppleARKit
\ No newline at end of file
++AllowedDirectories=Engine/Plugins/Runtime/AR/Apple/AppleARKit
\ No newline at end of file
diff --git a/Templates/TP_HandheldARBP/Config/Android/AndroidGame.ini b/Templates/TP_HandheldARBP/Config/Android/AndroidGame.ini
--- a/Templates/TP_HandheldARBP/Config/Android/AndroidGame.ini
+++ b/Templates/TP_HandheldARBP/Config/Android/AndroidGame.ini
@@ -1,2 +1,2 @@
 [Staging]
-+WhitelistDirectories=Engine/Plugins/Runtime/AR/Apple/AppleARKit
\ No newline at end of file
++AllowedDirectories=Engine/Plugins/Runtime/AR/Apple/AppleARKit
\ No newline at end of file
diff --git a/Templates/TP_MFG_HandheldARBP/Config/Android/AndroidGame.ini b/Templates/TP_MFG_HandheldARBP/Config/Android/AndroidGame.ini
--- a/Templates/TP_MFG_HandheldARBP/Config/Android/AndroidGame.ini
+++ b/Templates/TP_MFG_HandheldARBP/Config/Android/AndroidGame.ini
@@ -1,2 +1,2 @@
 [Staging]
-+WhitelistDirectories=Engine/Plugins/Runtime/AR/Apple/AppleARKit
\ No newline at end of file
++AllowedDirectories=Engine/Plugins/Runtime/AR/Apple/AppleARKit
\ No newline at end of file
