Setup a systemic match for SignTool timestamp errors

Don't raise to an error as retries kick in and will throw when there are no retries left.

#rnx

diff --git a/Engine/Source/Programs/AutomationTool/AutomationUtils/Matchers/SystemicEventMatcher.cs b/Engine/Source/Programs/AutomationTool/AutomationUtils/Matchers/SystemicEventMatcher.cs
--- a/Engine/Source/Programs/AutomationTool/AutomationUtils/Matchers/SystemicEventMatcher.cs
+++ b/Engine/Source/Programs/AutomationTool/AutomationUtils/Matchers/SystemicEventMatcher.cs
@@ -46,6 +46,10 @@ namespace AutomationUtils.Matchers
 			{
 				return new LogEventBuilder(cursor).ToMatch(LogEventPriority.High, LogLevel.Information, KnownLogEvents.Systemic_MissingXgeControlWorker);
 			}
+			if (cursor.Contains("SignTool Error: The specified timestamp server either could not be reached or"))
+			{
+				return new LogEventBuilder(cursor).ToMatch(LogEventPriority.Normal, LogLevel.Information, KnownLogEvents.Systemic_SignToolTimeStampServer);
+			}
 
 			return null;
 		}
