Fix passing -WriteBackMetadataToAssetRegistry via RunUAT. Now you can pass this in ushell after the -- when doing .stage game.

diff --git a/Engine/Source/Developer/IoStoreUtilities/Private/IoStoreUtilities.cpp b/Engine/Source/Developer/IoStoreUtilities/Private/IoStoreUtilities.cpp
--- a/Engine/Source/Developer/IoStoreUtilities/Private/IoStoreUtilities.cpp
+++ b/Engine/Source/Developer/IoStoreUtilities/Private/IoStoreUtilities.cpp
@@ -6634,6 +6634,7 @@ int32 CreateIoStoreContainerFiles(const TCHAR* CmdLine)
 		else
 		{
 			UE_LOG(LogIoStore, Error, TEXT("Invalid WriteBackMetdataToAssetRegistry value: %s - check setting in ProjectSettings -> Packaging"), *WriteBackMetadataToAssetRegistry);
+			UE_LOG(LogIoStore, Error, TEXT("Valid options are: AdjacentFile, OriginalFile, Disabled."), *WriteBackMetadataToAssetRegistry);
 			return -1;
 		}
 	}
diff --git a/Engine/Source/Programs/AutomationTool/AutomationUtils/ProjectParams.cs b/Engine/Source/Programs/AutomationTool/AutomationUtils/ProjectParams.cs
--- a/Engine/Source/Programs/AutomationTool/AutomationUtils/ProjectParams.cs
+++ b/Engine/Source/Programs/AutomationTool/AutomationUtils/ProjectParams.cs
@@ -373,6 +373,7 @@ namespace AutomationTool
 			this.DeviceNames = InParams.DeviceNames;
 			this.ServerDevice = InParams.ServerDevice;
             this.NullRHI = InParams.NullRHI;
+			this.WriteBackMetadataToAssetRegistry = InParams.WriteBackMetadataToAssetRegistry;
             this.FakeClient = InParams.FakeClient;
             this.EditorTest = InParams.EditorTest;
             this.RunAutomationTests = InParams.RunAutomationTests;
@@ -934,6 +935,7 @@ namespace AutomationTool
 			this.SpecifiedArchitecture = ParseParamValueIfNotSpecified(Command, SpecifiedArchitecture, "specifiedarchitecture", String.Empty);
 			this.UbtArgs = ParseParamValueIfNotSpecified(Command, UbtArgs, "ubtargs", String.Empty);
 			this.AdditionalPackageOptions = ParseParamValueIfNotSpecified(Command, AdditionalPackageOptions, "AdditionalPackageOptions", String.Empty);
+			this.WriteBackMetadataToAssetRegistry = ParseParamValueIfNotSpecified(Command, WriteBackMetadataToAssetRegistry, "WriteBackMetadataToAssetRegistry", String.Empty);
 
 			// -trace can be used with or without a value
 			if (Trace != null || GetParamValueIfNotSpecified(Command, null, false, "trace"))
@@ -2109,6 +2111,9 @@ namespace AutomationTool
         [Help("nullrhi", "add -nullrhi to the client commandlines")]
         public bool NullRHI;
 
+        [Help("WriteBackMetadataToAssetRegistry", "Passthru to iostore staging, see IoStoreUtilities.cpp")]
+        public string WriteBackMetadataToAssetRegistry;
+
         /// <summary>
         /// Run:adds ?fake to the server URL
         /// </summary>
@@ -3113,6 +3118,7 @@ namespace AutomationTool
 				CommandUtils.LogLog("Package={0}", Package);
 				CommandUtils.LogLog("ForcePackageData={0}", ForcePackageData);
 				CommandUtils.LogLog("NullRHI={0}", NullRHI);
+				CommandUtils.LogLog("WriteBackMetadataToAssetRegistry={0}", WriteBackMetadataToAssetRegistry);
 				CommandUtils.LogLog("FakeClient={0}", FakeClient);
                 CommandUtils.LogLog("EditorTest={0}", EditorTest);
                 CommandUtils.LogLog("RunAutomationTests={0}", RunAutomationTests); 
diff --git a/Engine/Source/Programs/AutomationTool/Scripts/CopyBuildToStagingDirectory.Automation.cs b/Engine/Source/Programs/AutomationTool/Scripts/CopyBuildToStagingDirectory.Automation.cs
--- a/Engine/Source/Programs/AutomationTool/Scripts/CopyBuildToStagingDirectory.Automation.cs
+++ b/Engine/Source/Programs/AutomationTool/Scripts/CopyBuildToStagingDirectory.Automation.cs
@@ -3161,10 +3161,18 @@ namespace AutomationScripts
 					}
 				}
 
-				string WriteBackMetadataToAssetRegistry;
-				if (PlatformGameConfig.GetString("/Script/UnrealEd.ProjectPackagingSettings", "WriteBackMetadataToAssetRegistry", out WriteBackMetadataToAssetRegistry))
+
+				if (Params.WriteBackMetadataToAssetRegistry.Length != 0)
+				{
+					AdditionalArgs += " -WriteBackMetadataToAssetRegistry=" + Params.WriteBackMetadataToAssetRegistry;
+				}
+				else
 				{
-					AdditionalArgs += " -WriteBackMetadataToAssetRegistry=" + WriteBackMetadataToAssetRegistry;
+					string WriteBackMetadataToAssetRegistry;
+					if (PlatformGameConfig.GetString("/Script/UnrealEd.ProjectPackagingSettings", "WriteBackMetadataToAssetRegistry", out WriteBackMetadataToAssetRegistry))
+					{
+						AdditionalArgs += " -WriteBackMetadataToAssetRegistry=" + WriteBackMetadataToAssetRegistry;
+					}
 				}
 
 
