Low level tests quality of life changes:
- bump "Started test case" log messages to Display when --debug is used
- additional duration log messages. Catch2 durations were never printed in sync with --debug log messages
- ability to specify a device and attach to a debugger when running the LowLevelTests script; ability to specify configuration type

diff --git a/Engine/Source/Developer/LowLevelTestsRunner/Private/TestListeners/EventListener.cpp b/Engine/Source/Developer/LowLevelTestsRunner/Private/TestListeners/EventListener.cpp
--- a/Engine/Source/Developer/LowLevelTestsRunner/Private/TestListeners/EventListener.cpp
+++ b/Engine/Source/Developer/LowLevelTestsRunner/Private/TestListeners/EventListener.cpp
@@ -1,6 +1,8 @@
 // Copyright Epic Games, Inc. All Rights Reserved.
 
 #include "Logging/LogMacros.h"
+#include "Misc/DateTime.h"
+#include "Misc/Timespan.h"
 #include "TestRunnerPrivate.h"
 
 #include <catch2/catch_test_case_info.hpp>
@@ -15,6 +17,8 @@ DEFINE_LOG_CATEGORY_STATIC(LogLowLevelTests, Log, VeryVerbose);
 
 class FEventListener final : public Catch::EventListenerBase
 {
+private:
+	FDateTime TestStarted;
 public:
 	using EventListenerBase::EventListenerBase;
 
@@ -22,7 +26,8 @@ public:
 	{
 		if (ITestRunner* TestRunner = ITestRunner::Get())
 		{
-			UE_CLOG(TestRunner->IsDebugMode(), LogLowLevelTests, Log,
+			TestStarted = FDateTime::Now();
+			UE_CLOG(TestRunner->IsDebugMode(), LogLowLevelTests, Display,
 				TEXT("Started test case \"%hs\" at %hs(%" SIZE_T_FMT ") with tags %hs"),
 				TestInfo.name.c_str(),
 				TestInfo.lineInfo.file,
@@ -35,6 +40,33 @@ public:
 	{
 		if (ITestRunner* TestRunner = ITestRunner::Get())
 		{
+			FTimespan Duration = FDateTime::Now() - TestStarted;
+			
+			TStringBuilder<64> DurationText;
+			const int32 Hours = Duration.GetHours();
+			const int32 Minutes = Duration.GetMinutes();
+			const double Seconds = Duration.GetTotalSeconds();
+			if (Hours > 0)
+			{
+				DurationText << Hours << TEXTVIEW("h ");
+			}
+			if (Hours > 0 || Minutes > 0)
+			{
+				DurationText << Minutes << TEXTVIEW("m ");
+			}
+			if (Hours > 0 || Minutes > 0 || Seconds >= 1.0)
+			{
+				DurationText.Appendf(TEXT("%.3fs "), Seconds);
+			}
+			else
+			{
+				DurationText << Duration.GetFractionMilli() << TEXTVIEW("ms");
+			}
+
+			UE_CLOG(TestRunner->IsDebugMode(), LogLowLevelTests, Display,
+				TEXT("Test took %s"),
+				*DurationText);
+
 			UE_CLOG(TestRunner->IsDebugMode() && TestCaseStats.totals.testCases.failed > 0,
 				LogLowLevelTests, Error, TEXT("Failed test case \"%hs\" at %hs(%" SIZE_T_FMT ")"),
 				TestCaseStats.testInfo->name.c_str(),
diff --git a/Engine/Source/Developer/LowLevelTestsRunner/Private/TestRunner.cpp b/Engine/Source/Developer/LowLevelTestsRunner/Private/TestRunner.cpp
--- a/Engine/Source/Developer/LowLevelTestsRunner/Private/TestRunner.cpp
+++ b/Engine/Source/Developer/LowLevelTestsRunner/Private/TestRunner.cpp
@@ -78,6 +78,7 @@ private:
 	bool bDebugMode = false;
 	bool bMultiThreaded = false;
 	bool bWaitForInputToTerminate = false;
+	bool bAttachToDebugger = false;
 	int32 SleepOnInitSeconds = 0;
 };
 
@@ -153,6 +154,10 @@ void FTestRunner::ParseCommandLine(TConstArrayView<const ANSICHAR*> Args)
 		{
 			bWaitForInputToTerminate = true;
 		}
+		else if (Arg == ANSITEXTVIEW("--attach-to-debugger"))
+		{
+			bAttachToDebugger = true;
+		}
 		else
 		{
 			CatchArgs.Add(Arg.GetData());
@@ -174,6 +179,16 @@ void FTestRunner::SleepOnInit() const
 
 void FTestRunner::GlobalSetup() const
 {
+	if (bAttachToDebugger)
+	{
+		FPlatformMisc::LocalPrint(TEXT("Waiting for debugger..."));
+		while (!FPlatformMisc::IsDebuggerPresent())
+		{
+			FPlatformProcess::Sleep(0.1f);
+		}
+		UE_DEBUG_BREAK();
+	}
+
 	if (!bGlobalSetup)
 	{
 		return;
diff --git a/Engine/Source/Programs/AutomationTool/LowLevelTests/RunLowLevelTests.cs b/Engine/Source/Programs/AutomationTool/LowLevelTests/RunLowLevelTests.cs
--- a/Engine/Source/Programs/AutomationTool/LowLevelTests/RunLowLevelTests.cs
+++ b/Engine/Source/Programs/AutomationTool/LowLevelTests/RunLowLevelTests.cs
@@ -160,6 +160,8 @@ namespace LowLevelTests
 		[AutoParam(0)]
 		public int Sleep;
 
+		public bool AttachToDebugger;
+
 		public string Build;
 
 		[AutoParam("")]
@@ -206,6 +208,7 @@ namespace LowLevelTests
 			TestApp = Globals.Params.ParseValue("testapp=", "");
 
 			Tags = Params.ParseValue("tags=", null);
+			AttachToDebugger = Params.ParseParam("attachtodebugger");
 
 			string PlatformArgString = Params.ParseValue("platform=", null);
 			Platform = string.IsNullOrEmpty(PlatformArgString) ? BuildHostPlatform.Current.Platform : UnrealTargetPlatform.Parse(PlatformArgString);
@@ -231,14 +234,16 @@ namespace LowLevelTests
 		private LowLevelTestsBuildSource BuildSource { get; set; }
 		private string Tags { get; set; }
 		private int Sleep { get; set; }
+		private bool AttachToDebugger { get; set; }
 
 		public UnrealDeviceReservation UnrealDeviceReservation { get; private set; }
 
-		public LowLevelTestsSession(LowLevelTestsBuildSource InBuildSource, string InTags, int InSleep)
+		public LowLevelTestsSession(LowLevelTestsBuildSource InBuildSource, string InTags, int InSleep, bool InAttachToDebugger)
 		{
 			BuildSource = InBuildSource;
 			Tags = InTags;
 			Sleep = InSleep;
+			AttachToDebugger = InAttachToDebugger;
 			UnrealDeviceReservation = new UnrealDeviceReservation();
 		}
 
@@ -262,7 +267,7 @@ namespace LowLevelTests
 
 			// TargetDevice<Platform> classes have a hard dependency on UnrealAppConfig instead of IAppConfig.
 			// More refactoring needed to support non-packaged applications that can be run natively from a path on the device.
-			UnrealAppConfig AppConfig = BuildSource.GetUnrealAppConfig(Tags, Sleep);
+			UnrealAppConfig AppConfig = BuildSource.GetUnrealAppConfig(Tags, Sleep, AttachToDebugger);
 
 			IEnumerable<ITargetDevice> DevicesToInstallOn = UnrealDeviceReservation.ReservedDevices.ToArray();
 			ITargetDevice Device = DevicesToInstallOn.Where(D => D.IsConnected && D.Platform == BuildSource.Platform).First();
@@ -437,13 +442,26 @@ namespace LowLevelTests
 		protected static string GetExecutable(UnrealTargetPlatform InPlatform, UnrealTargetConfiguration InConfiguration, string InTestApp, string InBuildPath, string FileRegEx)
 		{
 			IEnumerable<string> Executables = DirectoryUtils.FindFiles(InBuildPath, new Regex(FileRegEx));
+			string ParentDirPath;
+			string BuildExecutableName;
 			foreach (string Executable in Executables)
 			{
-				if (InBuildPath.ToLower().Contains(InPlatform.ToString().ToLower()))
+				ParentDirPath = Directory.GetParent(Executable).FullName;
+				BuildExecutableName = Path.GetFileNameWithoutExtension(Executable);
+
+				if (ParentDirPath.ToLower().Contains(InPlatform.ToString().ToLower()))
 				{
-					if (InBuildPath.ToLower().Contains(InTestApp.ToString().ToLower()))
+					if (ParentDirPath.ToLower().Contains(InTestApp.ToString().ToLower()))
 					{
-						if (InConfiguration == UnrealTargetConfiguration.Development || InBuildPath.ToLower().Contains(InConfiguration.ToString().ToLower()))
+						// Executable name must not contain any configuration or platform name
+						if (InConfiguration == UnrealTargetConfiguration.Development)
+						{
+							if (BuildExecutableName.CompareTo(InTestApp) == 0 && !BuildExecutableName.Contains(InPlatform.ToString()))
+							{
+								return Path.GetRelativePath(InBuildPath, Executable);
+							}
+						}
+						else if (BuildExecutableName.Contains(InTestApp) && BuildExecutableName.Contains(InPlatform.ToString()))
 						{
 							return Path.GetRelativePath(InBuildPath, Executable);
 						}
@@ -560,7 +578,7 @@ namespace LowLevelTests
 				.First();
 		}
 
-		public UnrealAppConfig GetUnrealAppConfig(string InTags, int InSleep)
+		public UnrealAppConfig GetUnrealAppConfig(string InTags, int InSleep, bool InAttachToDebugger)
 		{
 			if (CachedConfig == null)
 			{
@@ -574,10 +592,7 @@ namespace LowLevelTests
 				CachedConfig.Sandbox = "LowLevelTests";
 				CachedConfig.FilesToCopy = new List<UnrealFileToCopy>();
 				// Set reporting options, filters etc
-				if (TestApp == "LowLevelTests" || TestApp == "LowLevelTestsTests")
-				{
-					CachedConfig.CommandLineParams.AddRawCommandline("--durations=yes");
-				}
+				CachedConfig.CommandLineParams.AddRawCommandline("--durations=no");
 				CachedConfig.CommandLineParams.AddRawCommandline("--reporter=console");
 				CachedConfig.CommandLineParams.AddRawCommandline(string.Format("--out={0}", LowLevelTestsReporting.GetTargetReportPath(Platform, TestApp, BuildPath)));
 				CachedConfig.CommandLineParams.AddRawCommandline("--filenames-as-tags");
@@ -590,6 +605,11 @@ namespace LowLevelTests
 					CachedConfig.CommandLineParams.AddRawCommandline(String.Format("--sleep={0}", InSleep));
 				}
 				CachedConfig.CommandLineParams.AddRawCommandline("--debug");
+				CachedConfig.CommandLineParams.AddRawCommandline("--log");
+				if (InAttachToDebugger)
+				{
+					CachedConfig.CommandLineParams.AddRawCommandline("--attach-to-debugger");
+				}
 			}
 			return CachedConfig;
 		}
diff --git a/Engine/Source/Programs/AutomationTool/LowLevelTests/Tests/Gauntlet.LowLevelTestNode.cs b/Engine/Source/Programs/AutomationTool/LowLevelTests/Tests/Gauntlet.LowLevelTestNode.cs
--- a/Engine/Source/Programs/AutomationTool/LowLevelTests/Tests/Gauntlet.LowLevelTestNode.cs
+++ b/Engine/Source/Programs/AutomationTool/LowLevelTests/Tests/Gauntlet.LowLevelTestNode.cs
@@ -41,7 +41,7 @@ namespace LowLevelTests
 		{
 			if (LowLevelTestsApp == null)
 			{
-				LowLevelTestsApp = new LowLevelTestsSession(Context.BuildInfo, Context.Options.Tags, Context.Options.Sleep);
+				LowLevelTestsApp = new LowLevelTestsSession(Context.BuildInfo, Context.Options.Tags, Context.Options.Sleep, Context.Options.AttachToDebugger);
 			}
 
 			return LowLevelTestsApp.TryReserveDevices();
