remove unnecessary platform references from automation code.

#rnx

diff --git a/Engine/Source/Programs/AutomationTool/AutomationUtils/CommandUtils.cs b/Engine/Source/Programs/AutomationTool/AutomationUtils/CommandUtils.cs
--- a/Engine/Source/Programs/AutomationTool/AutomationUtils/CommandUtils.cs
+++ b/Engine/Source/Programs/AutomationTool/AutomationUtils/CommandUtils.cs
@@ -2414,31 +2414,6 @@ namespace AutomationTool
 			return Files;
 		}
 
-		static public Dictionary<Version, string> GetBuildVersionPathMap(string[] BuildFolders)
-		{
-			Dictionary<Version, string> BuildVersionPaths = new Dictionary<Version, string>();
-
-			foreach (string buildFolder in BuildFolders)
-			{
-				string XboxoneReleaseVersion = Path.GetFileName(buildFolder);
-
-				// The name of all recent archived builds either start with "EA" or only contain version
-				string VersionPrefix = "EA";
-				if (XboxoneReleaseVersion.StartsWith(VersionPrefix))
-				{
-					XboxoneReleaseVersion.Substring(VersionPrefix.Length);  // Length of "EA"
-				}
-
-				Version ReleaseVersion;
-				if (Version.TryParse(XboxoneReleaseVersion, out ReleaseVersion))
-				{
-					BuildVersionPaths.Add(ReleaseVersion, buildFolder);
-				}
-			}
-
-			return BuildVersionPaths;
-		}
-		
 		public static string FormatSizeString(long Size)
 		{
 			string[] Units = { "bytes", "KB", "MB", "GB", "TB", "PB", "EB", "ZB", "YB" };
diff --git a/Engine/Source/Programs/AutomationTool/AutomationUtils/LinuxHostPlatform.cs b/Engine/Source/Programs/AutomationTool/AutomationUtils/LinuxHostPlatform.cs
--- a/Engine/Source/Programs/AutomationTool/AutomationUtils/LinuxHostPlatform.cs
+++ b/Engine/Source/Programs/AutomationTool/AutomationUtils/LinuxHostPlatform.cs
@@ -102,7 +102,7 @@ namespace AutomationTool
 			{
 				AppName = "xbuild";
 				CommandLine = (String.IsNullOrEmpty(CommandLine) ? "" : CommandLine) + " /verbosity:quiet /nologo";
-				// Pass #define MONO to all the automation scripts (see XboxOne)
+				// Pass #define MONO to all the automation scripts
 				CommandLine += " /p:DefineConstants=MONO";
 				CommandLine += " /p:DefineConstants=__MonoCS__";
 				// Some projects have TargetFrameworkProfile=Client which causes warnings on Linux
