- Fixed warning about ConfigRedirects.ini not being allowed or denied

diff --git a/Engine/Source/Programs/AutomationTool/Scripts/CopyBuildToStagingDirectory.Automation.cs b/Engine/Source/Programs/AutomationTool/Scripts/CopyBuildToStagingDirectory.Automation.cs
--- a/Engine/Source/Programs/AutomationTool/Scripts/CopyBuildToStagingDirectory.Automation.cs
+++ b/Engine/Source/Programs/AutomationTool/Scripts/CopyBuildToStagingDirectory.Automation.cs
@@ -1589,6 +1589,12 @@ namespace AutomationScripts
 					return SC.StageTargetConfigurations.Any(x => x != UnrealTargetConfiguration.Test && x != UnrealTargetConfiguration.Shipping);
 				}
 
+				// this is needed for all targets
+				if (NormalizedPath == "configredirects.ini")
+				{
+					return true;
+				}
+
 				if (NormalizedPath == "locgatherconfig.ini")
 				{
 					return false;
