Don't pass an empty architecture when building a plugin

#rnx

diff --git a/Engine/Source/Programs/AutomationTool/Scripts/BuildPluginCommand.Automation.cs b/Engine/Source/Programs/AutomationTool/Scripts/BuildPluginCommand.Automation.cs
--- a/Engine/Source/Programs/AutomationTool/Scripts/BuildPluginCommand.Automation.cs
+++ b/Engine/Source/Programs/AutomationTool/Scripts/BuildPluginCommand.Automation.cs
@@ -1,4 +1,4 @@
-﻿// Copyright Epic Games, Inc. All Rights Reserved.
+// Copyright Epic Games, Inc. All Rights Reserved.
 using System;
 using System.Collections.Generic;
 using System.IO;
@@ -363,8 +363,7 @@ public sealed class BuildPlugin : BuildCommand
 				
 				string Arguments = String.Format("-plugin={0} -iwyu -noubtmakefiles -manifest={1} -nohotreload", CommandUtils.MakePathSafeToUseWithCommandLine(HostProjectPluginFile.FullName), CommandUtils.MakePathSafeToUseWithCommandLine(ManifestFileName.FullName));
 
-				string SpecifiedArchitecture;
-				if (PlatformToArchitectureMap.TryGetValue(Platform, out SpecifiedArchitecture))
+				if (PlatformToArchitectureMap.TryGetValue(Platform, out string SpecifiedArchitecture) && !string.IsNullOrEmpty(SpecifiedArchitecture))
 				{
 					Arguments += String.Format(" -architecture={0}", SpecifiedArchitecture);
 				}
