Smartling localization provider integration for Unreal Engine.
- Introduced the SmartlingLocalizationProvider abstract class to interface with Smartling backend for authentication, uploading and downloading of .po files.
- Subclasses of SmartlingLocalizationProvider can now be used as part of the Localize UAT command for automation.

diff --git a/Engine/Source/Programs/AutomationTool/SmartlingLocalization/Properties/AssemblyInfo.cs b/Engine/Source/Programs/AutomationTool/SmartlingLocalization/Properties/AssemblyInfo.cs
new file mode 100644
--- /dev/null
+++ b/Engine/Source/Programs/AutomationTool/SmartlingLocalization/Properties/AssemblyInfo.cs
@@ -0,0 +1,12 @@
+// Copyright Epic Games, Inc. All Rights Reserved.
+
+using System.Reflection;
+using System.Runtime.CompilerServices;
+using System.Runtime.InteropServices;
+
+// General Information about an assembly is controlled through the following 
+// set of attributes. Change these attribute values to modify the information
+// associated with an assembly.
+[assembly: AssemblyTitle("SmartlingLocalization.Automation")]
+[assembly: AssemblyDescription("")]
+[assembly: AssemblyConfiguration("")]
diff --git a/Engine/Source/Programs/AutomationTool/SmartlingLocalization/SmartlingLocalization.Automation.csproj b/Engine/Source/Programs/AutomationTool/SmartlingLocalization/SmartlingLocalization.Automation.csproj
new file mode 100644
--- /dev/null
+++ b/Engine/Source/Programs/AutomationTool/SmartlingLocalization/SmartlingLocalization.Automation.csproj
@@ -0,0 +1,47 @@
+﻿<Project Sdk="Microsoft.NET.Sdk">
+  <Import Project="..\..\Shared\UnrealEngine.csproj.props" />
+  
+  <PropertyGroup>
+    <TargetFramework>net6.0</TargetFramework>
+    <Configuration Condition=" '$(Configuration)' == '' ">Development</Configuration>
+    <OutputType>Library</OutputType>
+    <AppendTargetFrameworkToOutputPath>false</AppendTargetFrameworkToOutputPath>
+    <GenerateAssemblyInfo>false</GenerateAssemblyInfo>
+    <GenerateTargetFrameworkAttribute>false</GenerateTargetFrameworkAttribute>
+    <Configurations>Debug;Release;Development</Configurations>
+    <RootNamespace>AutomationTool</RootNamespace>
+    <AssemblyName>SmartlingLocalization.Automation</AssemblyName>
+    <WarningsNotAsErrors>612,618</WarningsNotAsErrors>
+    <OutputPath>..\..\..\..\Binaries\DotNET\AutomationTool\AutomationScripts\SmartlingLocalization</OutputPath>
+    <TreatWarningsAsErrors>true</TreatWarningsAsErrors>
+    <CopyLocalLockFileAssemblies>true</CopyLocalLockFileAssemblies>
+    <SatelliteResourceLanguages>en</SatelliteResourceLanguages> <!-- remove non english resource languages -->
+    <DebugType>pdbonly</DebugType>
+  </PropertyGroup>
+  <PropertyGroup Condition=" '$(Configuration)|$(Platform)' == 'Development|AnyCPU' ">
+    <DefineConstants>$(DefineConstants);TRACE</DefineConstants>
+    <Optimize>true</Optimize>
+    <NoWarn>$(NoWarn);1701;1702;CS1998;CS0649</NoWarn>
+  </PropertyGroup>
+  <PropertyGroup Condition=" '$(Configuration)|$(Platform)' == 'Release|AnyCPU' ">
+    <DefineConstants>$(DefineConstants);TRACE</DefineConstants>
+    <Optimize>true</Optimize>
+    <NoWarn>1701;1702;CS1998;CS0649</NoWarn>
+  </PropertyGroup>
+  <PropertyGroup Condition=" '$(Configuration)|$(Platform)' == 'Debug|AnyCPU' ">
+    <DefineConstants>$(DefineConstants);DEBUG;TRACE</DefineConstants>
+    <DebugSymbols>true</DebugSymbols>
+    <DebugType>full</DebugType>
+    <NoWarn>1701;1702;CS1998;CS0649</NoWarn>
+  </PropertyGroup>
+
+  <ItemGroup>
+    <ProjectReference Include="..\..\Shared\EpicGames.Core\EpicGames.Core.csproj" PrivateAssets="All"><Private>false</Private></ProjectReference>
+    <ProjectReference Include="..\..\UnrealBuildTool\UnrealBuildTool.csproj" PrivateAssets="All"><Private>false</Private></ProjectReference>
+    <ProjectReference Include="..\AutomationUtils\AutomationUtils.Automation.csproj" PrivateAssets="All"><Private>false</Private></ProjectReference>
+    <ProjectReference Include="..\Localization\Localization.Automation.csproj" PrivateAssets="All"><Private>false</Private></ProjectReference>
+  </ItemGroup>
+  <ItemGroup>
+    <PackageReference Include="Microsoft.CSharp" Version="4.7.0" />
+  </ItemGroup>
+</Project>
\ No newline at end of file
diff --git a/Engine/Source/Programs/AutomationTool/SmartlingLocalization/SmartlingLocalizationProvider.cs b/Engine/Source/Programs/AutomationTool/SmartlingLocalization/SmartlingLocalizationProvider.cs
new file mode 100644
--- /dev/null
+++ b/Engine/Source/Programs/AutomationTool/SmartlingLocalization/SmartlingLocalizationProvider.cs
@@ -0,0 +1,372 @@
+// Copyright Epic Games, Inc. All Rights Reserved.
+using System;
+using System.Collections.Generic;
+using System.IO;
+using System.Threading.Tasks;
+using System.Text;
+using System.Text.Json;
+using System.Text.Json.Serialization;
+using System.Net;
+using System.Net.Http;
+using System.Net.Http.Headers;
+using AutomationTool;
+using UnrealBuildTool;
+using EpicGames.Localization;
+
+
+#pragma warning disable SYSLIB0014
+
+namespace EpicGames.SmartlingLocalization
+{
+	public class SmartlingConfig
+	{
+		public string ProjectId;
+		public string UserId;
+		public string UserSecret;
+	}
+
+	// A series of model classes to make deserialization of HTTP responses cleaner and easier 
+	public class SmartlingResponseEnvelope<T>
+	{
+		public T Response { get; set; }
+	}
+
+	public class SmartlingDataEnvelope<T>
+	{
+		public T Data { get; set; }
+	}
+
+	public class SmartlingAuthenticationToken
+	{
+		public string AccessToken { get; set; }
+		public int ExpiresIn { get; set; }
+		public int RefreshExpiresIn { get; set; }
+		public string RefreshToken { get; set; }
+		public string TokenType { get; set; }
+	}
+
+	public class SmartlingRequestError
+	{
+		public string Key { get; set; }
+		public string Message { get; set; }
+		
+		public override string ToString()
+		{
+			return $"Key: {Key}\nMessage: {Message}";
+		}
+	}
+
+	public class SmartlingRequestErrorsEnvelope
+	{
+		public List<SmartlingRequestError> Errors { get; set; }
+	}
+
+	public abstract class SmartlingLocalizationProvider : LocalizationProvider
+	{
+		public SmartlingLocalizationProvider(LocalizationProviderArgs InArgs)
+			: base(InArgs)
+		{
+			if (String.IsNullOrEmpty(LocalizationBranchName))
+			{
+				Console.WriteLine("LocalizationBranchName is null or empty. The branch name is used to create the correct file URIs for Smartling. Wrong files may be downloaded from Smartling and files may be clobbered in Smartling. Please append a value to LocalizationName as a command line argument to fix this.");
+			}
+			Config = new SmartlingConfig();
+			Client = new HttpClient();
+			JsonOptions = new JsonSerializerOptions();
+			JsonOptions.PropertyNameCaseInsensitive = true;
+			JsonOptions.DefaultIgnoreCondition = JsonIgnoreCondition.WhenWritingNull;
+			AuthenticationToken = null;
+		}
+
+		public async override Task InitializeProjectWithLocalizationProvider(string ProjectName, ProjectImportExportInfo ProjectImportInfo)
+		{
+			// @TODOLocalization: Throw an exeption if the config data is null/invalid 
+		}
+
+		public async override Task DownloadProjectFromLocalizationProvider(string ProjectName, ProjectImportExportInfo ProjectImportInfo)
+		{
+			// Get the latest files for each culture.
+			foreach (var Culture in ProjectImportInfo.CulturesToGenerate)
+			{
+				// Skip the native culture as we don't need the original source from Smartling 
+				if (Culture == ProjectImportInfo.NativeCulture)
+				{
+					continue;
+				}
+
+				await DownloadLatestPOFile(Culture, null, ProjectImportInfo);
+				foreach (var Platform in ProjectImportInfo.SplitPlatformNames)
+				{
+					await DownloadLatestPOFile(Culture, Platform, ProjectImportInfo);
+				}
+			}
+		}
+
+		private async Task DownloadLatestPOFile(string EpicLocale, string Platform, ProjectImportExportInfo ProjectImportInfo)
+		{
+			await GetAuthenticationToken();
+			Client.DefaultRequestHeaders.Authorization = new AuthenticationHeaderValue(AuthenticationToken.TokenType, AuthenticationToken.AccessToken);
+
+			var DestinationDirectory = String.IsNullOrEmpty(Platform)
+			? new DirectoryInfo(CommandUtils.CombinePaths(RootWorkingDirectory, ProjectImportInfo.DestinationPath))
+			: new DirectoryInfo(CommandUtils.CombinePaths(RootWorkingDirectory, ProjectImportInfo.DestinationPath, ProjectImportExportInfo.PlatformLocalizationFolderName, Platform));
+			var CultureDirectory = (ProjectImportInfo.bUseCultureDirectory) ? new DirectoryInfo(Path.Combine(DestinationDirectory.FullName, EpicLocale)) : DestinationDirectory;
+
+			var ExportFile = new FileInfo(Path.Combine(CultureDirectory.FullName, ProjectImportInfo.PortableObjectName));
+			// note that the Smartling Filename and Uri may be different 
+			var SmartlingFilename = GetSmartlingFilename(ProjectImportInfo.PortableObjectName, Platform);
+
+			Console.WriteLine($"Exporting: '{SmartlingFilename}' as '{ExportFile.FullName}' ({EpicLocale})");
+			// @TODOLocalization: Refactor this long try block into multiple try-catch blocks. 
+			try
+			{
+				string SmartlingLocale = ConvertEpicLocaleToSmartlingLocale(EpicLocale);
+				string DownloadEndpoint = $"https://api.smartling.com/files-api/v2/projects/{Config.ProjectId}/locales/{SmartlingLocale}/file";
+				string SmartlingFileUri = GetSmartlingFileUri(SmartlingFilename);
+				var DownloadUriBuilder = new UriBuilder(DownloadEndpoint);
+				// Changing the retrieval type to pseudo is a good way to test downloads
+				// @TODOLocalization: Not sure if the retrieval type should be pending or published to get all of the relevant strings 
+				string DownloadQueryString = $"fileUri={SmartlingFileUri}&retrievalType=published&includeOriginalStrings=false";
+				DownloadUriBuilder.Query = DownloadQueryString;
+				var DownloadResponse = await Client.GetAsync(DownloadUriBuilder.Uri);
+				if (DownloadResponse.IsSuccessStatusCode)
+				{
+					if (!CultureDirectory.Exists)
+					{
+						CultureDirectory.Create();
+					}
+					bool ExportFileWasReadOnly = true;
+					if (ExportFile.Exists)
+					{
+						// We're going to clobber the existing PO file, so make sure it's writable (it may be read-only if in Perforce)
+						ExportFileWasReadOnly = ExportFile.IsReadOnly;
+						ExportFile.IsReadOnly = false;
+					}
+					// Download the file 
+					using (var DownloadStream = await DownloadResponse.Content.ReadAsStreamAsync())
+					{
+						using (var DownloadFileStream = ExportFile.OpenWrite())
+						{
+							await DownloadStream.CopyToAsync(DownloadFileStream);
+						}
+					}
+
+					Console.WriteLine($"[SUCCESS] Exporting: '{SmartlingFileUri}' as '{ExportFile.FullName}' ({EpicLocale})");
+					// Reset the write status of the file
+					if (ExportFileWasReadOnly)
+					{
+						ExportFile.IsReadOnly = true;
+					}
+
+					// Update the back-up copy so we can diff against what we got from Smartling, and what the gather commandlet produced
+					// This can be used to verify if the Smartling download is corrupt or if our export file is corrupted.
+					// Set the flag to true for debug. 
+					bool bCreateBackupCopy = false;
+					if (bCreateBackupCopy)
+					{
+						string ExportFileCopyPath = Path.Combine(ExportFile.DirectoryName, $"{Path.GetFileNameWithoutExtension(ExportFile.Name)}_FromSmartling{ExportFile.Extension}");
+						Console.WriteLine($"Updating Smartling copy '{ExportFileCopyPath}'");
+						var ExportFileCopy = new FileInfo(ExportFileCopyPath);
+
+						var ExportFileCopyWasReadOnly = false;
+						if (ExportFileCopy.Exists)
+						{
+							// We're going to clobber the existing PO file, so make sure it's writable (it may be read-only if in Perforce)
+							ExportFileCopyWasReadOnly = ExportFileCopy.IsReadOnly;
+							ExportFileCopy.IsReadOnly = false;
+						}
+
+						ExportFile.CopyTo(ExportFileCopy.FullName, true);
+						// reset the read/write status 
+						if (ExportFileCopyWasReadOnly)
+						{
+							ExportFileCopy.IsReadOnly = true;
+						}
+						// We don't upload the ocpy to p4 as it's no longer strictly necessary. Just use the file as a lcoal debug. 
+					}
+				}
+				else
+				{
+					// The file may not currently exist in Smartling and will need to be uploaded first via the Upload step later on. 
+					Console.WriteLine($"[FAILED] Exporting: '{ExportFile.FullName}' ({EpicLocale}. The file may need to be uploaded first.)");
+					await PrintRequestErrors(DownloadResponse);
+				}
+			}
+			catch (Exception Ex)
+			{
+				Console.WriteLine($"[FAILED] Exporting: '{ExportFile.FullName}' ({EpicLocale}) - {Ex}");
+			}
+		}
+
+		private async Task PrintRequestErrors(HttpResponseMessage Response)
+		{
+			string ResponseString = await Response.Content.ReadAsStringAsync();
+			var ResponseEnvelope = JsonSerializer.Deserialize<SmartlingResponseEnvelope<SmartlingRequestErrorsEnvelope>>(ResponseString, JsonOptions);
+			var RequestErrors = ResponseEnvelope.Response.Errors;
+			foreach (SmartlingRequestError RequestError in RequestErrors)
+			{
+				Console.WriteLine($"Error:\n {RequestError.ToString()}");
+			}
+		}
+
+		public async override Task UploadProjectToLocalizationProvider(string ProjectName, ProjectImportExportInfo ProjectExportInfo)
+		{
+			// Upload the .po file for the native culture first
+			await UploadLatestPOFile(ProjectExportInfo.NativeCulture, null, ProjectExportInfo);
+			foreach (var Platform in ProjectExportInfo.SplitPlatformNames)
+			{
+				await UploadLatestPOFile(ProjectExportInfo.NativeCulture, Platform, ProjectExportInfo);
+			}
+
+			// Uploading all cultures is a legacy behavior from Onesky that isn't currently needed within our workflow. We can support this if there is a need. 
+		}
+
+		private async Task UploadLatestPOFile(string EpicLocale, string Platform, ProjectImportExportInfo ProjectExportInfo)
+		{
+			await GetAuthenticationToken();
+			Client.DefaultRequestHeaders.Authorization = new AuthenticationHeaderValue(AuthenticationToken.TokenType, AuthenticationToken.AccessToken);
+
+			var SourceDirectory = String.IsNullOrEmpty(Platform)
+				? new DirectoryInfo(CommandUtils.CombinePaths(RootWorkingDirectory, ProjectExportInfo.DestinationPath))
+				: new DirectoryInfo(CommandUtils.CombinePaths(RootWorkingDirectory, ProjectExportInfo.DestinationPath, ProjectImportExportInfo.PlatformLocalizationFolderName, Platform));
+			var CultureDirectory = (ProjectExportInfo.bUseCultureDirectory) ? new DirectoryInfo(Path.Combine(SourceDirectory.FullName, EpicLocale)) : SourceDirectory;
+			string FileToUploadPath = Path.Combine(CultureDirectory.FullName, ProjectExportInfo.PortableObjectName);
+			var FileToUpload = new FileInfo(FileToUploadPath);
+			if (!FileToUpload.Exists)
+			{
+				Console.WriteLine($"Unable to upload '{FileToUploadPath}'. File does not exist.");
+				return;
+			}
+			string SmartlingFilename = GetSmartlingFilename(ProjectExportInfo.PortableObjectName, Platform);
+			string SmartlingFileUri = GetSmartlingFileUri(SmartlingFilename);
+			bool bIsNative = EpicLocale == ProjectExportInfo.NativeCulture;
+
+			Console.WriteLine($"Uploading: '{FileToUpload.FullName}' as '{SmartlingFileUri}' ({EpicLocale})");
+			// For now we only upload the file in the native locale
+			// It is dissuaded to import translations to Smartling. Perform all translations in the Smartling dashboard instead.
+			try
+			{
+				string UploadEndpoint = $"https://api.smartling.com/files-api/v2/projects/{Config.ProjectId}/file";
+				using (var UploadMultipartFormDataContent = new MultipartFormDataContent())
+				{
+					var UploadFileStreamContent = new StreamContent(File.OpenRead(FileToUpload.FullName));
+					// generic binary stream we will simply consider the octet-stream mime type 
+					UploadFileStreamContent.Headers.ContentType = new MediaTypeHeaderValue("application/octet-stream");
+					UploadMultipartFormDataContent.Add(UploadFileStreamContent, "file", SmartlingFileUri);	
+					UploadMultipartFormDataContent.Add(new StringContent(SmartlingFileUri, Encoding.UTF8, "application/json"), "fileUri");
+					UploadMultipartFormDataContent.Add(new StringContent("gettext", Encoding.UTF8, "application/json"), "fileType");
+					var UploadResponse = await Client.PostAsync(UploadEndpoint, UploadMultipartFormDataContent);
+					if (UploadResponse.IsSuccessStatusCode)
+					{
+						Console.WriteLine($"[SUCCESS] Uploading: '{FileToUpload.FullName}' ({EpicLocale})");
+					}
+					else
+					{
+						Console.WriteLine($"[FAILED] Uploading: '{FileToUpload.FullName}' ({EpicLocale})");
+						await PrintRequestErrors(UploadResponse);
+					}
+				}
+			}
+			catch (Exception Ex)
+			{
+				Console.WriteLine($"[FAILED] Uploading: '{FileToUpload.FullName}' ({EpicLocale}) - {Ex}");
+			}
+		}
+
+		// Override in child classes if there is a different set of languages that you need translations for 
+		virtual protected Dictionary<string,string> GetEpicLocaleToSmartlingLocaleDictionary()
+		{
+			// Smartling mostly uses the same language codes as Unreal, however there are some differences
+			// Login to the Smartling dashboard and check the languages you are translating to for each project. 
+			return new Dictionary<string, string>
+			{
+				{ "es",  "es-ES" },	// Spain Spanish
+				{ "es-419",  "es-MX" },	// Mexico Spanish
+				{ "zh-Hans", "zh-CN" },	// Simp Chinese
+				{ "zh-Hant", "zh-TW" },	// Trad Chinese
+			};
+		}
+
+		private string ConvertEpicLocaleToSmartlingLocale(string EpicLocale)
+		{
+			string SmartlingLocale = null;
+			var CultureMappings = GetEpicLocaleToSmartlingLocaleDictionary();
+			if (CultureMappings.TryGetValue(EpicLocale, out SmartlingLocale))
+			{
+				return SmartlingLocale;
+			}
+			return EpicLocale;
+		}
+
+		private string GetSmartlingFilename(string BaseFilename, string Platform)
+		{
+			string SmartlingFilename = BaseFilename;
+			if (!String.IsNullOrEmpty(Platform))
+			{
+				// Apply the platform suffix.
+				string SmartlingFilenameWithSuffix = Path.GetFileNameWithoutExtension(SmartlingFilename) + "_" + Platform + Path.GetExtension(SmartlingFilename);
+				SmartlingFilename = SmartlingFilenameWithSuffix;
+			}
+			if (!String.IsNullOrEmpty(RemoteFilenamePrefix))
+			{
+				// Apply the prefix (this is used to avoid collisions with plugins that use the same name for their PO files)
+				// E.g Say you have both PluginA and PluginB with a loc target called Plugin. This would cause the remote names to be PluginA_Plugin and PluginB_Plugin rather than conflict.
+				SmartlingFilename = RemoteFilenamePrefix + "_" + SmartlingFilename;
+			}
+			return SmartlingFilename;
+		}
+
+		private string GetSmartlingFileUri(string SmartlingFilename)
+		{
+			// The LocalizationBranchName variable is populated by the LocalizationBranch command line argument 
+			// It needs to be valid to ensure that the correct file is interacted with on Smartling. 
+			// We already raise a warning in the constructor, but we dont' crash as we don't want the automation jobs to end with localization
+			if (!String.IsNullOrEmpty(LocalizationBranchName))
+			{
+				// Apply the localization branch suffix.
+				string SmartlingFilenameWithSuffix = Path.GetFileNameWithoutExtension(SmartlingFilename) + "_" + LocalizationBranchName + Path.GetExtension(SmartlingFilename);
+				return SmartlingFilenameWithSuffix;
+			}
+			return SmartlingFilename;
+		}
+
+		private async Task GetAuthenticationToken()
+		{
+			// @TODOLocalization: Implement refreshing instead of retrieving a new token each time 
+			string AuthenticateEndpoint = "https://api.smartling.com/auth-api/v2/authenticate";
+			// serialize from Dictionary as MultipartFormDataContent not supported as a content-type 
+			var AuthenticateRequestContent = new Dictionary<string, string>()
+			{
+				{ "userIdentifier", Config.UserId},
+				{ "userSecret", Config.UserSecret}
+			};
+			var AuthenticateRequestContentJson = JsonSerializer.Serialize(AuthenticateRequestContent, JsonOptions);
+			var AuthenticateRequestStringContent = new StringContent(AuthenticateRequestContentJson, Encoding.UTF8, "application/json");
+			try
+			{
+				var AuthenticateResponse = Client.PostAsync(AuthenticateEndpoint, AuthenticateRequestStringContent).Result;
+				if (AuthenticateResponse.IsSuccessStatusCode)
+				{
+					string AuthenticateResponseString = await AuthenticateResponse.Content.ReadAsStringAsync();
+					var AuthenticateResponseEnvelope = JsonSerializer.Deserialize<SmartlingResponseEnvelope<SmartlingDataEnvelope<SmartlingAuthenticationToken>>>(AuthenticateResponseString, JsonOptions);
+					AuthenticationToken = AuthenticateResponseEnvelope.Response.Data;
+					Console.WriteLine("Successfully retrieved authentication token!");
+				}
+				else
+				{
+					Console.WriteLine("Failed to retrieve authentication token.");
+					await PrintRequestErrors(AuthenticateResponse);
+				}
+			}
+			catch (Exception Ex)
+			{
+				Console.WriteLine($"Failed to retrieve authentication token. {Ex}");
+			}
+		}
+		
+		protected SmartlingConfig Config;
+		protected HttpClient Client;
+		private JsonSerializerOptions JsonOptions;
+		private SmartlingAuthenticationToken AuthenticationToken;
+	}
+}
\ No newline at end of file
