- Fixing an ini key from the config redirect stuff recently

diff --git a/Engine/Source/Programs/AutomationTool/AutomationUtils/DeploymentContext.cs b/Engine/Source/Programs/AutomationTool/AutomationUtils/DeploymentContext.cs
--- a/Engine/Source/Programs/AutomationTool/AutomationUtils/DeploymentContext.cs
+++ b/Engine/Source/Programs/AutomationTool/AutomationUtils/DeploymentContext.cs
@@ -1,4 +1,4 @@
-﻿// Copyright Epic Games, Inc. All Rights Reserved.
+// Copyright Epic Games, Inc. All Rights Reserved.
 using System;
 using System.Collections.Generic;
 using System.IO;
@@ -301,7 +301,7 @@ public class DeploymentContext //: ProjectParams
 
 	/// <summary>
 	/// Set of config files which are denied from staging. By default, we warn for config files which are not well known to prevent internal data (eg. editor/server settings)
-	/// leaking in packaged builds. This list is read from the +DeniedConfigFiles=... array in the [Staging] section of *Game.ini files.
+	/// leaking in packaged builds. This list is read from the +DisallowedConfigFiles=... array in the [Staging] section of *Game.ini files.
 	/// </summary>
 	public HashSet<StagedFileReference> ConfigFilesDenyList = new HashSet<StagedFileReference>();
 
@@ -592,7 +592,7 @@ public class DeploymentContext //: ProjectParams
 
 		// Read the list of files which are allow listed to be staged
 		ReadConfigFileList(GameConfig, "Staging", "AllowedConfigFiles", ConfigFilesAllowList);
-		ReadConfigFileList(GameConfig, "Staging", "DeniedConfigFiles", ConfigFilesDenyList);
+		ReadConfigFileList(GameConfig, "Staging", "DisallowedConfigFiles", ConfigFilesDenyList);
 
 		// Grab the game ini data
 		String PackagingIniPath = "/Script/UnrealEd.ProjectPackagingSettings";
diff --git a/Engine/Source/Programs/AutomationTool/Scripts/CopyBuildToStagingDirectory.Automation.cs b/Engine/Source/Programs/AutomationTool/Scripts/CopyBuildToStagingDirectory.Automation.cs
--- a/Engine/Source/Programs/AutomationTool/Scripts/CopyBuildToStagingDirectory.Automation.cs
+++ b/Engine/Source/Programs/AutomationTool/Scripts/CopyBuildToStagingDirectory.Automation.cs
@@ -1444,7 +1444,7 @@ namespace AutomationScripts
 				Nullable<bool> ShouldStage = ShouldStageConfigFile(SC, ConfigDir, ConfigFile, PlatformExtensionName);
 				if (ShouldStage == null)
 				{
-					CommandUtils.LogWarning("The config file '{0}' will be staged, but is not explicitly allowed or denied. Add +AllowedConfigFiles={0} or +DeniedConfigFiles={0} to the [Staging] section of DefaultGame.ini", SC.GetStagedFileLocation(ConfigFile));
+					CommandUtils.LogWarning("The config file '{0}' will be staged, but is not explicitly allowed or denied. Add +AllowedConfigFiles={0} or +DisallowedConfigFiles={0} to the [Staging] section of DefaultGame.ini", SC.GetStagedFileLocation(ConfigFile));
 				}
 
 				if (ShouldStage ?? true)
