Fix BuildGraph ParamHelpAttribute for -target parameter is not really required because -SingleNode could be used instead.

diff --git a/Engine/Source/Programs/AutomationTool/BuildGraph/BuildGraph.cs b/Engine/Source/Programs/AutomationTool/BuildGraph/BuildGraph.cs
--- a/Engine/Source/Programs/AutomationTool/BuildGraph/BuildGraph.cs
+++ b/Engine/Source/Programs/AutomationTool/BuildGraph/BuildGraph.cs
@@ -181,7 +181,7 @@ namespace AutomationTool
 	/// </summary>
 	[Help("Tool for creating extensible build processes in UE which can be run locally or in parallel across a build farm.")]
 	[ParamHelp("Script=<FileName>", "Path to the script describing the graph", ParamType = typeof(File), Required = true)]
-	[ParamHelp("Target=<Name>", "Name of the node or output tag to be built", Required = true)]
+	[ParamHelp("Target=<Name>", "Name of the node or output tag to be built")]
 	[ParamHelp("Schema", "Generates a schema to the default location", ParamType = typeof(bool))]
 	[ParamHelp("Schema=<FileName>", "Generate a schema describing valid script documents, including all the known tasks", ParamType = typeof(string))]
 	[ParamHelp("ImportSchema=<FileName>", "Imports a schema from an existing schema file", ParamType = typeof(File))]
