Disable splash screen when running unattended

EngineTest simply does not have a splash screen and we warn on Linux. Lets just disable it for these tests

diff --git a/Engine/Source/Programs/AutomationTool/Gauntlet/Unreal/Base/Gauntlet.UnrealTestConfiguration.cs b/Engine/Source/Programs/AutomationTool/Gauntlet/Unreal/Base/Gauntlet.UnrealTestConfiguration.cs
--- a/Engine/Source/Programs/AutomationTool/Gauntlet/Unreal/Base/Gauntlet.UnrealTestConfiguration.cs
+++ b/Engine/Source/Programs/AutomationTool/Gauntlet/Unreal/Base/Gauntlet.UnrealTestConfiguration.cs
@@ -891,7 +891,7 @@ namespace Gauntlet
 
 			if (Attended == false)
 			{
-				AppConfig.CommandLine += " -unattended";
+				AppConfig.CommandLine += " -unattended -nosplash";
 
 				// if we are unattended but still may need access to Vulkan passing renderoffscreen to allow not depending on
 				// the X11/Wayland display server to be around and use a dummy/offscreen rendering mode
