remove reference to non-existent legacy platform.

#rnx

diff --git a/Engine/Source/Programs/AutomationTool/AutomationUtils/CommandUtils.cs b/Engine/Source/Programs/AutomationTool/AutomationUtils/CommandUtils.cs
--- a/Engine/Source/Programs/AutomationTool/AutomationUtils/CommandUtils.cs
+++ b/Engine/Source/Programs/AutomationTool/AutomationUtils/CommandUtils.cs
@@ -3161,10 +3161,7 @@ namespace AutomationTool
 					List<FileReference> FilesToSign = new List<FileReference>();
 					foreach (string File in Files)
 					{
-						if (!(Path.GetDirectoryName(File).Replace("\\", "/")).Contains("Binaries/XboxOne"))
-						{
-							FilesToSign.Add(new FileReference(File));
-						}						
+						FilesToSign.Add(new FileReference(File));
 					}
 					SignMultipleFilesIfEXEOrDLL(FilesToSign);
 				}
