UAT: CopyBuildToStagingDirectory: Make sure to include global.utoc/.ucas when copying .pak, .utoc and .ucas files from another build.

#rnx

diff --git a/Engine/Source/Programs/AutomationTool/Scripts/CopyBuildToStagingDirectory.Automation.cs b/Engine/Source/Programs/AutomationTool/Scripts/CopyBuildToStagingDirectory.Automation.cs
--- a/Engine/Source/Programs/AutomationTool/Scripts/CopyBuildToStagingDirectory.Automation.cs
+++ b/Engine/Source/Programs/AutomationTool/Scripts/CopyBuildToStagingDirectory.Automation.cs
@@ -2424,6 +2424,44 @@ namespace AutomationScripts
 			return String.Empty;
 		}
 
+		private static bool CopyExistingContainer(FileReference InUtocFile, FileReference OutputLocation)
+		{
+			bool bCopiedExistingPak = true;
+			FileReference InUcasFile = InUtocFile.ChangeExtension(".ucas");
+			if (FileReference.Exists(InUcasFile))
+			{
+				FileReference OutUcasFile = OutputLocation.ChangeExtension(".ucas");
+				LogInformation("Copying ucas from {0} to {1}", InUcasFile, OutUcasFile);
+				if (!InternalUtils.SafeCopyFile(InUcasFile.FullName, OutUcasFile.FullName))
+				{
+					LogInformation("Failed to copy ucas {0} to {1}, creating new pak", InUcasFile, OutUcasFile);
+					bCopiedExistingPak = false;
+				}
+			}
+			else
+			{
+				LogInformation("Missing ucas file {0}, creating new pak", InUcasFile);
+				bCopiedExistingPak = false;
+			}
+
+			if (FileReference.Exists(InUtocFile))
+			{
+				FileReference OutUtocFile = OutputLocation.ChangeExtension(".utoc");
+				LogInformation("Copying utoc from {0} to {1}", InUtocFile, OutUtocFile);
+				if (!InternalUtils.SafeCopyFile(InUtocFile.FullName, OutUtocFile.FullName))
+				{
+					LogInformation("Failed to copy utoc {0} to {1}, creating new pak", InUtocFile, OutUtocFile);
+					bCopiedExistingPak = false;
+				}
+			}
+			else
+			{
+				LogInformation("Missing utoc file {0}, creating new pak", InUtocFile);
+				bCopiedExistingPak = false;
+			}
+			return bCopiedExistingPak;
+		}
+
 		/// <summary>
 		/// Creates a pak file using response file.
 		/// </summary>
@@ -2835,6 +2873,7 @@ namespace AutomationScripts
 
 			CommonAdditionalArgs += " " + AdditionalCompressionOptionsOnCommandLine;
 
+			bool bCopiedExistingGlobalUtoc = false;
 			List<Tuple<FileReference, StagedFileReference, string>> Outputs = new List<Tuple<FileReference, StagedFileReference, string>>();
 			foreach (CreatePakParams PakParams in PakParamsList)
 			{
@@ -2918,38 +2957,10 @@ namespace AutomationScripts
 							}
 
 							FileReference InUtocFile = SourceOutputLocation.ChangeExtension(".utoc");
-							FileReference InUcasFile = SourceOutputLocation.ChangeExtension(".ucas");
-							if (FileReference.Exists(InUtocFile) || FileReference.Exists(InUcasFile))
+							if (FileExists_NoExceptions(InUtocFile.FullName))
 							{
-								if (FileReference.Exists(InUcasFile))
+								if (!CopyExistingContainer(InUtocFile, OutputLocation))
 								{
-									FileReference OutUcasFile = OutputLocation.ChangeExtension(".ucas");
-									LogInformation("Copying ucas from {0} to {1}", InUcasFile, OutUcasFile);
-									if (!InternalUtils.SafeCopyFile(InUcasFile.FullName, OutUcasFile.FullName))
-									{
-										LogInformation("Failed to copy ucas {0} to {1}, creating new pak", InUcasFile, OutUcasFile);
-										bCopiedExistingPak = false;
-									}
-								}
-								else
-								{
-									LogInformation("Missing ucas file {0}, creating new pak", InUcasFile);
-									bCopiedExistingPak = false;
-								}
-
-								if (FileReference.Exists(InUtocFile))
-								{
-									FileReference OutUtocFile = OutputLocation.ChangeExtension(".utoc");
-									LogInformation("Copying utoc from {0} to {1}", InUtocFile, OutUtocFile);
-									if (!InternalUtils.SafeCopyFile(InUtocFile.FullName, OutUtocFile.FullName))
-									{
-										LogInformation("Failed to copy utoc {0} to {1}, creating new pak", InUtocFile, OutUtocFile);
-										bCopiedExistingPak = false;
-									}
-								}
-								else
-								{
-									LogInformation("Missing utoc file {0}, creating new pak", InUtocFile);
 									bCopiedExistingPak = false;
 								}
 							}
@@ -2959,6 +2970,17 @@ namespace AutomationScripts
 					{
 						LogInformation("Failed to copy source pak from {0} to {1}, creating new pak", SourceOutputLocation, OutputLocation);
 					}
+
+					if (bCopiedExistingPak && !bCopiedExistingGlobalUtoc)
+					{
+						string GlobalFileName = "global.utoc";
+						FileReference InUtocFile = FileReference.Combine(SourceOutputLocation.Directory, GlobalFileName);
+						FileReference OutUtocFile = FileReference.Combine(OutputLocation.Directory, GlobalFileName);
+						if (FileExists_NoExceptions(InUtocFile.FullName))
+						{
+							bCopiedExistingGlobalUtoc = CopyExistingContainer(InUtocFile, OutUtocFile);
+						}
+					}
 				}
 
 				string PatchSourceContentPath = null;
