Unshelved from pending changelist '22074449':

Adding another SignTool matcher

SignTool outputs two errors for a failure

#rnx

diff --git a/Engine/Source/Programs/AutomationTool/AutomationUtils/Matchers/SystemicEventMatcher.cs b/Engine/Source/Programs/AutomationTool/AutomationUtils/Matchers/SystemicEventMatcher.cs
--- a/Engine/Source/Programs/AutomationTool/AutomationUtils/Matchers/SystemicEventMatcher.cs
+++ b/Engine/Source/Programs/AutomationTool/AutomationUtils/Matchers/SystemicEventMatcher.cs
@@ -46,10 +46,15 @@ namespace AutomationUtils.Matchers
 			{
 				return new LogEventBuilder(cursor).ToMatch(LogEventPriority.High, LogLevel.Information, KnownLogEvents.Systemic_MissingXgeControlWorker);
 			}
+
 			if (cursor.Contains("SignTool Error: The specified timestamp server either could not be reached or"))
 			{
 				return new LogEventBuilder(cursor).ToMatch(LogEventPriority.Normal, LogLevel.Information, KnownLogEvents.Systemic_SignToolTimeStampServer);
 			}
+			else if (cursor.Contains("SignTool Error: An error occurred while attempting to sign"))
+			{
+				return new LogEventBuilder(cursor).ToMatch(LogEventPriority.Normal, LogLevel.Information, KnownLogEvents.Systemic_SignTool);
+			}
 
 			return null;
 		}
diff --git a/Engine/Source/Programs/Shared/EpicGames.Core/KnownLogEvents.cs b/Engine/Source/Programs/Shared/EpicGames.Core/KnownLogEvents.cs
--- a/Engine/Source/Programs/Shared/EpicGames.Core/KnownLogEvents.cs
+++ b/Engine/Source/Programs/Shared/EpicGames.Core/KnownLogEvents.cs
@@ -248,10 +248,17 @@ namespace EpicGames.Core
 		/// <summary>
 		/// Microsoft SignTool cannot reach specified timestamp server 
 		/// 
-		/// Signing performs a looping retry using different servers and throws when all retries attemps have been exhausted.
+		/// Signing performs a looping retry using different servers and throws when all retries attempts have been exhausted.
 		/// </summary>
 		public static EventId Systemic_SignToolTimeStampServer { get; } = new EventId(765);
 
+		/// <summary>
+		/// Microsoft SignTool generic error
+		/// 
+		/// Usually preceded by another type of error.
+		/// </summary>
+		public static EventId Systemic_SignTool { get; } = new EventId(766);
+
 		/// <summary>
 		/// Maximum code for systemic events. Add new events in the 700-799 range.
 		/// </summary>
