For Mac InstallBuilds don't add in the -architecture flag for UBT, instead fall back to the default setting (x64_86). This will resolve an issue where the .precompile files are looked for in the wrong directory (arm64). This is not a long term solution, in the future we'll need to determine the arch the install build was built for and use that to add the appropriate arch flags.

diff --git a/Engine/Source/Programs/AutomationTool/Mac/MacPlatform.Automation.cs b/Engine/Source/Programs/AutomationTool/Mac/MacPlatform.Automation.cs
--- a/Engine/Source/Programs/AutomationTool/Mac/MacPlatform.Automation.cs
+++ b/Engine/Source/Programs/AutomationTool/Mac/MacPlatform.Automation.cs
@@ -7,6 +7,7 @@ using System.Text.RegularExpressions;
 using System.IO;
 using AutomationTool;
 using UnrealBuildTool;
+using UnrealBuildBase;
 using EpicGames.Core;
 
 public class MacPlatform : Platform
@@ -172,8 +173,13 @@ public class MacPlatform : Platform
 					}
 				}
 			}
-
-			Target.UBTArgs += string.Format(" -architecture={0}", UBTArchitectureParam);
+			
+			// TODO - This needs to be handled in a more graceful way, however for now
+			// when in an installedbuild just leave it to the default arch and don't add anything extra
+			if (!Unreal.IsEngineInstalled())
+			{
+				Target.UBTArgs += string.Format(" -architecture={0}", UBTArchitectureParam);
+			}
 		}
 	}
 
