Gauntlet: Make it so multi-pass tests will retain devices between each run

diff --git a/Engine/Source/Programs/AutomationTool/Gauntlet/Unreal/Base/Gauntlet.UnrealSession.cs b/Engine/Source/Programs/AutomationTool/Gauntlet/Unreal/Base/Gauntlet.UnrealSession.cs
--- a/Engine/Source/Programs/AutomationTool/Gauntlet/Unreal/Base/Gauntlet.UnrealSession.cs
+++ b/Engine/Source/Programs/AutomationTool/Gauntlet/Unreal/Base/Gauntlet.UnrealSession.cs
@@ -486,6 +486,16 @@ namespace Gauntlet
 		/// </summary>
 		public string Sandbox { get; set; }
 
+		/// <summary>
+		/// Whether or not we should retain our devices this pass
+		/// </summary>
+		public bool ShouldRetainDevices { get; set; }
+
+		/// <summary>
+		/// Record of our installations in case we want to re-use them in a later pass
+		/// </summary>
+		public Dictionary<UnrealSessionRole, IAppInstall> RolesToInstalls;
+
 		/// <summary>
 		/// Constructor that takes a build source and a number of roles
 		/// </summary>
@@ -566,7 +576,16 @@ namespace Gauntlet
 				RequiredDeviceTypes[C.Constraint]++;
 			});
 
-			return UnrealDeviceReservation.TryReserveDevices(RequiredDeviceTypes, RolesNeedingDevices.Count());
+
+			if (UnrealDeviceReservation.ReservedDevices != null && UnrealDeviceReservation.ReservedDevices.Count > 0
+				&& ShouldRetainDevices)
+			{
+				return true;
+			}
+			else
+			{
+				return UnrealDeviceReservation.TryReserveDevices(RequiredDeviceTypes, RolesNeedingDevices.Count());
+			}
 		}
 
 		/// <summary>
@@ -731,56 +750,77 @@ namespace Gauntlet
 
 					// todo - should this be elsewhere?
 					AppConfig.Sandbox = Sandbox;
-
-					// Tag the device for report result
-					if(BuildHostPlatform.Current.Platform != Device.Platform)
+					
+					IAppInstall Install = null;
+					bool bReinstallPerPass = Globals.Params.ParseParam("ReinstallPerPass");
+					if (RolesToInstalls == null || !RolesToInstalls.ContainsKey(Role) || bReinstallPerPass)
 					{
-						AppConfig.CommandLineParams.Add("DeviceTag", Device.Name);
-					}
 
-					IAppInstall Install = null;
+						// Tag the device for report result
+						if(BuildHostPlatform.Current.Platform != Device.Platform)
+						{
+							AppConfig.CommandLineParams.Add("DeviceTag", Device.Name);
+						}
 
-					IDeviceUsageReporter.RecordStart(Device.Name, Device.Platform, IDeviceUsageReporter.EventType.Device, IDeviceUsageReporter.EventState.Success);
-					IDeviceUsageReporter.RecordStart(Device.Name, Device.Platform, IDeviceUsageReporter.EventType.Install, IDeviceUsageReporter.EventState.Success, BuildSource.BuildName);
-					try
-					{
-						Install = Device.InstallApplication(AppConfig);
-						IDeviceUsageReporter.RecordEnd(Device.Name, Device.Platform, IDeviceUsageReporter.EventType.Install, IDeviceUsageReporter.EventState.Success);
-					}
-					catch (System.Exception Ex)
-					{
-						// Warn, ignore the device, and do not continue
-						string ErrorMessage = string.Format("Failed to install app onto device {0} for role {1}. {2}. Will retry with new device", Device, Role, Ex);
-						if (ErrorMessage.Contains("not enough space"))
+						IAppInstall Install = null;
+
+						IDeviceUsageReporter.RecordStart(Device.Name, Device.Platform, IDeviceUsageReporter.EventType.Device, IDeviceUsageReporter.EventState.Success);
+						IDeviceUsageReporter.RecordStart(Device.Name, Device.Platform, IDeviceUsageReporter.EventType.Install, IDeviceUsageReporter.EventState.Success, BuildSource.BuildName);
+						try
+						{
+							Install = Device.InstallApplication(AppConfig);
+							IDeviceUsageReporter.RecordEnd(Device.Name, Device.Platform, IDeviceUsageReporter.EventType.Install, IDeviceUsageReporter.EventState.Success);
+						}
+						catch (System.Exception Ex)
 						{
-							Log.Warning(ErrorMessage);
-							if(Device.Platform == BuildHostPlatform.Current.Platform)
+							// Warn, ignore the device, and do not continue
+							string ErrorMessage = string.Format("Failed to install app onto device {0} for role {1}. {2}. Will retry with new device", Device, Role, Ex);
+							if (ErrorMessage.Contains("not enough space"))
+							{
+								Log.Warning(ErrorMessage);
+								if (Device.Platform == BuildHostPlatform.Current.Platform)
+								{
+									// If on desktop platform, we are not retrying.
+									// It is unlikely that space is going to be made and InstallBuildParallel has marked the build path as problematic.
+									SessionRetries = 0;
+								}
+							}
+							else
 							{
-								// If on desktop platform, we are not retrying.
-								// It is unlikely that space is going to be made and InstallBuildParallel has marked the build path as problematic.
-								SessionRetries = 0;
+								Log.Info(ErrorMessage);
 							}
+							UnrealDeviceReservation.MarkProblemDevice(Device);
+							InstallSuccess = false;
+							IDeviceUsageReporter.RecordEnd(Device.Name, Device.Platform, IDeviceUsageReporter.EventType.Install, IDeviceUsageReporter.EventState.Failure);
+							break;
 						}
-						else
+
+
+						if (Globals.CancelSignalled)
 						{
-							Log.Info(ErrorMessage);
+							break;
 						}
-						UnrealDeviceReservation.MarkProblemDevice(Device);
-						InstallSuccess = false;
-						IDeviceUsageReporter.RecordEnd(Device.Name, Device.Platform, IDeviceUsageReporter.EventType.Install, IDeviceUsageReporter.EventState.Failure);
-						break;
-					}
-					
 
-					if (Globals.CancelSignalled)
-					{
-						break;
-					}
+						// Device has app installed, give role a chance to configure device
+						Role.ConfigureDevice?.Invoke(Device);
 
-					// Device has app installed, give role a chance to configure device
-					Role.ConfigureDevice?.Invoke(Device);
+						InstallsToRoles[Install] = Role;
 
-					InstallsToRoles[Install] = Role;
+						if(RolesToInstalls == null)
+						{
+							RolesToInstalls = new Dictionary<UnrealSessionRole, IAppInstall>();
+						}
+						if (!bReinstallPerPass)
+						{
+							RolesToInstalls[Role] = Install;
+						}
+					}
+					else
+					{
+						Install = RolesToInstalls[Role];
+						InstallsToRoles[Install] = Role;
+						Log.Info("Using previous install of {0} on {1}", Install.Name, Install.Device.Name);
+					}
 				}
 
 				if (InstallSuccess == false)
@@ -883,19 +923,30 @@ namespace Gauntlet
 			return LaunchSession();
 		}
 
-		/// <summary>
-		/// Shuts down the current session (if any)
+		///<summary>
+		/// Shuts down any running apps
 		/// </summary>
-		/// <returns></returns>
-		public void ShutdownSession()
+		public void ShutdownInstance()
 		{
 			if (SessionInstance != null)
 			{
 				SessionInstance.Dispose();
 				SessionInstance = null;
 			}
+		}
+
+		/// <summary>
+		/// Shuts down the current session (if any)
+		/// </summary>
+		/// <returns></returns>
+		public void ShutdownSession()
+		{
+			ShutdownInstance();
 
-			UnrealDeviceReservation.ReleaseDevices();
+			if (!ShouldRetainDevices)
+			{
+				UnrealDeviceReservation.ReleaseDevices();
+			}
 		}
 
 		/// <summary>
diff --git a/Engine/Source/Programs/AutomationTool/Gauntlet/Unreal/Base/Gauntlet.UnrealTestNode.cs b/Engine/Source/Programs/AutomationTool/Gauntlet/Unreal/Base/Gauntlet.UnrealTestNode.cs
--- a/Engine/Source/Programs/AutomationTool/Gauntlet/Unreal/Base/Gauntlet.UnrealTestNode.cs
+++ b/Engine/Source/Programs/AutomationTool/Gauntlet/Unreal/Base/Gauntlet.UnrealTestNode.cs
@@ -243,10 +243,40 @@ namespace Gauntlet
 			TestNodeEvents.Add(InEvent);
 		}
 
+		/// <summary>
+		/// Dictionary of Sessions for each test, by test name
+		/// </summary>
+		static protected Dictionary<string, UnrealSession> TestSessions = new Dictionary<string, UnrealSession>();
+
 		/// <summary>
 		/// Helper class that turns our wishes into reality
+		/// Indexed from the dictionary of Sessions by Test Name
 		/// </summary>
-		protected UnrealSession UnrealApp;
+		protected UnrealSession UnrealApp
+		{
+			get
+			{
+				if (TestSessions.ContainsKey(Name))
+				{
+					return TestSessions[Name];
+				}
+				else
+				{
+					return null;
+				}
+			}
+			set
+			{
+				if (TestSessions.ContainsKey(Name))
+				{
+					TestSessions[Name] = value;
+				}
+				else
+				{
+					TestSessions.Add(Name, value);
+				}
+			}
+		}
 
 		/// <summary>
 		/// Used to track how much of our app log has been written out
@@ -511,7 +541,15 @@ namespace Gauntlet
 
 			if (UnrealApp != null)
 			{
-				throw new AutomationException("Node already has an UnrealApp, was PrepareUnrealSession called twice?");
+				if (NumPasses > 1 && CurrentPass + 1 < NumPasses)
+				{
+					// Already prepared from a previous pass
+					return true;
+				}
+				else
+				{
+					throw new AutomationException("Node already has an UnrealApp, was PrepareUnrealSession called twice?");
+				}
 			}
 
 			// pass through any arguments such as -TestNameArg or -TestNameArg=Value
@@ -702,9 +740,20 @@ namespace Gauntlet
 		/// <returns></returns>
 		public override bool StartTest(int Pass, int InNumPasses)
 		{
-			if (UnrealApp == null)
+			if (InNumPasses > 1
+				&& Pass + 1 < InNumPasses)
 			{
-				throw new AutomationException("Node already has a null UnrealApp, was PrepareUnrealSession or IsReadyToStart called?");
+				if (UnrealApp == null)
+				{
+					throw new AutomationException("Node already has a null UnrealApp, was PrepareUnrealSession or IsReadyToStart called?");
+				}
+				else
+				{
+					bool bReacquireDevicesPerPass = Globals.Params.ParseParam("ReacquireDevicesPerPass");
+					bool ShouldRetain = (!bReacquireDevicesPerPass && InNumPasses > 1 && Pass + 1 < InNumPasses) ? true : false;
+					UnrealApp.ShouldRetainDevices = ShouldRetain;
+					Log.Verbose("ShouldRetainDevices: {0}", UnrealApp.ShouldRetainDevices);
+				}
 			}
 
 			// ensure we reset things
@@ -716,7 +765,7 @@ namespace Gauntlet
 			LastAppLogCount = 0;
 			LastEditorLogCount = 0;
 			LastHeartbeatTime = DateTime.MinValue;
-			LastActiveHeartbeatTime = DateTime.MinValue;			
+			LastActiveHeartbeatTime = DateTime.MinValue;
 
 			TConfigClass Config = GetCachedConfiguration();
 
@@ -819,9 +868,16 @@ namespace Gauntlet
 
 			if (UnrealApp != null)
 			{
-				UnrealApp.Dispose();
-				UnrealApp = null;
-			}			
+				if (CurrentPass + 1 == NumPasses)
+				{
+					UnrealApp.Dispose();
+					UnrealApp = null;
+				}
+				else
+				{
+					UnrealApp.ShutdownInstance();
+				}
+			}
 		}
 
 		/// <summary>
@@ -992,14 +1048,17 @@ namespace Gauntlet
 				Log.Warning("Failed to save artifacts. {0}", Ex);
 			}
 
-			try
-			{
-				// Artifacts have been saved, release devices back to pool for other tests to use
-				UnrealApp.UnrealDeviceReservation.ReleaseDevices();
-			}
-			catch (Exception Ex)
+			if (CurrentPass + 1 == NumPasses)
 			{
-				Log.Warning("Failed to release devices. {0}", Ex);
+				try
+				{
+					// Artifacts have been saved, release devices back to pool for other tests to use
+					UnrealApp.UnrealDeviceReservation.ReleaseDevices();
+				}
+				catch (Exception Ex)
+				{
+					Log.Warning("Failed to release devices. {0}", Ex);
+				}
 			}
 
 
