Gauntlet: Make DeviceReservations forget problem devices between passes, so that the DeviceService is authoritative

diff --git a/Engine/Source/Programs/AutomationTool/Gauntlet/Framework/Devices/Gauntlet.UnrealDeviceReservation.cs b/Engine/Source/Programs/AutomationTool/Gauntlet/Framework/Devices/Gauntlet.UnrealDeviceReservation.cs
--- a/Engine/Source/Programs/AutomationTool/Gauntlet/Framework/Devices/Gauntlet.UnrealDeviceReservation.cs
+++ b/Engine/Source/Programs/AutomationTool/Gauntlet/Framework/Devices/Gauntlet.UnrealDeviceReservation.cs
@@ -11,7 +11,7 @@ namespace Gauntlet
 	/// </summary>
 	public class UnrealDeviceReservation
 	{
-		public List<ProblemDevice> ProblemDevices { get; protected set; }
+		protected List<ProblemDevice> ProblemDevices { get; private set; } = new List<ProblemDevice>();
 		public List<ITargetDevice> ReservedDevices { get; protected set; }
 
 		public bool TryReserveDevices(Dictionary<UnrealDeviceTargetConstraint, int> RequiredDeviceTypes, int ExpectedNumberOfDevices)
@@ -21,10 +21,7 @@ namespace Gauntlet
 
 			ReleaseDevices();
 
-			if (ProblemDevices == null)
-			{
-				ProblemDevices = new List<ProblemDevice>();
-			}
+			ProblemDevices.Clear();
 
 			// check whether pool can accommodate devices
 			if (!DevicePool.Instance.CheckAvailableDevices(RequiredDeviceTypes, ProblemDevices))
