Enable cached package data in the asset registry for cooked cookers, so that they can correctly find all the data it needs

diff --git a/Engine/Source/Programs/AutomationTool/CookedEditor/CookedEditor.Automation.cs b/Engine/Source/Programs/AutomationTool/CookedEditor/CookedEditor.Automation.cs
--- a/Engine/Source/Programs/AutomationTool/CookedEditor/CookedEditor.Automation.cs
+++ b/Engine/Source/Programs/AutomationTool/CookedEditor/CookedEditor.Automation.cs
@@ -1047,6 +1047,10 @@ public class MakeCookedEditor : BuildCommand
 
 		// when making cooked editors, we some special commandline options to override some assumptions about editor data
 		Params.AdditionalCookerOptions += " -ini:Engine:[Core.System]:CanStripEditorOnlyExportsAndImports=False";
+
+		// when making cooked editors, we need to generate package data in our asset registry so that the cooker can actually find various things on disc
+		Params.AdditionalCookerOptions += " -ini:Engine:[AssetRegistry]:bSerializePackageData=True";
+
 		// We tend to "over-cook" packages to get everything we might need, so some non-editor BPs that are referencing editor BPs may
 		// get cooked. This is okay, because the editor stuff should exist. We may want to revist this, and not cook anything that would
 		// cause the issues
