Gauntlet: Make sure we don't erroneously try to reuse an install on a previous device if we had trouble launching on that device, or if we released devices at any point

diff --git a/Engine/Source/Programs/AutomationTool/Gauntlet/Unreal/Base/Gauntlet.UnrealSession.cs b/Engine/Source/Programs/AutomationTool/Gauntlet/Unreal/Base/Gauntlet.UnrealSession.cs
--- a/Engine/Source/Programs/AutomationTool/Gauntlet/Unreal/Base/Gauntlet.UnrealSession.cs
+++ b/Engine/Source/Programs/AutomationTool/Gauntlet/Unreal/Base/Gauntlet.UnrealSession.cs
@@ -822,8 +822,7 @@ namespace Gauntlet
 
 				if (InstallSuccess == false)
 				{
-					// release all devices
-					UnrealDeviceReservation.ReleaseDevices();
+					ReleaseSessionDevices();
 
 					if (SessionRetries == 0)
 					{
@@ -883,8 +882,7 @@ namespace Gauntlet
 							// mark that device as a problem
 							UnrealDeviceReservation.MarkProblemDevice(CurrentInstall.Device);
 
-							// release all devices
-							UnrealDeviceReservation.ReleaseDevices();
+							ReleaseSessionDevices();
 
 							if (SessionRetries == 0)
 							{
@@ -942,10 +940,16 @@ namespace Gauntlet
 
 			if (!ShouldRetainDevices)
 			{
-				UnrealDeviceReservation.ReleaseDevices();
+				ReleaseSessionDevices();
 			}
 		}
 
+		public void ReleaseSessionDevices()
+		{
+			UnrealDeviceReservation.ReleaseDevices();
+			RolesToInstalls.Clear();
+		}
+
 		/// <summary>
 		/// Retrieves and saves all artifacts from the provided session role. Artifacts are saved to the destination path 
 		/// </summary>
diff --git a/Engine/Source/Programs/AutomationTool/Gauntlet/Unreal/Base/Gauntlet.UnrealTestNode.cs b/Engine/Source/Programs/AutomationTool/Gauntlet/Unreal/Base/Gauntlet.UnrealTestNode.cs
--- a/Engine/Source/Programs/AutomationTool/Gauntlet/Unreal/Base/Gauntlet.UnrealTestNode.cs
+++ b/Engine/Source/Programs/AutomationTool/Gauntlet/Unreal/Base/Gauntlet.UnrealTestNode.cs
@@ -1053,7 +1053,7 @@ namespace Gauntlet
 				try
 				{
 					// Artifacts have been saved, release devices back to pool for other tests to use
-					UnrealApp.UnrealDeviceReservation.ReleaseDevices();
+					UnrealApp.ReleaseSessionDevices();
 				}
 				catch (Exception Ex)
 				{
