AutomationTool: Support running ParallelExecutor with all cores

#rnx

diff --git a/Engine/Source/Programs/AutomationTool/AutomationUtils/ParallelExecutor.cs b/Engine/Source/Programs/AutomationTool/AutomationUtils/ParallelExecutor.cs
--- a/Engine/Source/Programs/AutomationTool/AutomationUtils/ParallelExecutor.cs
+++ b/Engine/Source/Programs/AutomationTool/AutomationUtils/ParallelExecutor.cs
@@ -76,9 +76,9 @@ namespace AutomationTool
 			}
 		}
 
-		public static int Execute(string ActionsFileName, bool bStopOnErrors)
+		public static int Execute(string ActionsFileName, bool bStopOnErrors, bool bAllCores)
 		{
-			return Execute(ActionsFileName, UnrealBuildTool.ParallelExecutorConfiguration.GetMaxParallelProcesses(Log.Logger), bStopOnErrors);
+			return Execute(ActionsFileName, UnrealBuildTool.ParallelExecutorConfiguration.GetMaxParallelProcesses(0, bAllCores, Log.Logger), bStopOnErrors);
 		}
 
 		public static int Execute(string ActionsFileName, int MaxProcesses, bool bStopOnErrors)
diff --git a/Engine/Source/Programs/AutomationTool/AutomationUtils/UnrealBuild.cs b/Engine/Source/Programs/AutomationTool/AutomationUtils/UnrealBuild.cs
--- a/Engine/Source/Programs/AutomationTool/AutomationUtils/UnrealBuild.cs
+++ b/Engine/Source/Programs/AutomationTool/AutomationUtils/UnrealBuild.cs
@@ -583,7 +583,7 @@ namespace AutomationTool
 					CommandUtils.PushDir(CommandUtils.CombinePaths(CommandUtils.CmdEnv.LocalRoot, @"\Engine\Source"));
 					try
 					{
-						int ExitCode = ParallelExecutor.Execute(TaskFilePath, OwnerCommand.ParseParam("StopOnErrors"));
+						int ExitCode = ParallelExecutor.Execute(TaskFilePath, OwnerCommand.ParseParam("StopOnErrors"), OwnerCommand.ParseParam("AllCores"));
 						if(ExitCode != 0)
 						{
 							return false;
diff --git a/Engine/Source/Programs/UnrealBuildTool/Executors/ParallelExecutor.cs b/Engine/Source/Programs/UnrealBuildTool/Executors/ParallelExecutor.cs
--- a/Engine/Source/Programs/UnrealBuildTool/Executors/ParallelExecutor.cs
+++ b/Engine/Source/Programs/UnrealBuildTool/Executors/ParallelExecutor.cs
@@ -481,5 +481,10 @@ namespace UnrealBuildTool
 		/// Maximum number of processes that should be used for execution
 		/// </summary>
 		public static int GetMaxParallelProcesses(ILogger Logger) => ParallelExecutor.GetDefaultNumParallelProcesses(0, false, Logger);
+
+		/// <summary>
+		/// Maximum number of processes that should be used for execution
+		/// </summary>
+		public static int GetMaxParallelProcesses(int MaxLocalActions, bool bAllCores, ILogger Logger) => ParallelExecutor.GetDefaultNumParallelProcesses(MaxLocalActions, bAllCores, Logger);
 	}
 }
