Marking up Gauntlet classes with parameter help attributes

diff --git a/Engine/Source/Programs/AutomationTool/Gauntlet/Unreal/RunUnreal.cs b/Engine/Source/Programs/AutomationTool/Gauntlet/Unreal/RunUnreal.cs
--- a/Engine/Source/Programs/AutomationTool/Gauntlet/Unreal/RunUnreal.cs
+++ b/Engine/Source/Programs/AutomationTool/Gauntlet/Unreal/RunUnreal.cs
@@ -30,6 +30,36 @@ namespace Gauntlet
 	/// For a full list of options see UnrealTestContextOption
 	/// 
 	/// </summary>
+	[Help("Run Unreal tests using Gauntlet")]
+	[ParamHelp("Tests", "List of gauntlet tests to run", Required = true, MultiSelectSeparator = ",")]
+	[ParamHelp("ExecCmds", "List commands to execute", MultiSelectSeparator = "+")]
+	[ParamHelp("Build", "Reference to the build that is being tested")]
+	[ParamHelp("Configuration", "Configuration to perform tests on", Choices = new string[] { "Debug", "DebugGame", "Development", "Test", "Shipping"})]
+	[ParamHelp("Platform", "Platforms to perform tests on and their params")]
+	[ParamHelp("Dev", "Run in Dev mode", ParamType = typeof(bool))]
+	[ParamHelp("MaxDuration", "Maximum duration for test in sections", ParamType = typeof(int), DefaultValue = 3600)]
+	[ParamHelp("NoTimeout", "No maximum timeout", ParamType = typeof(bool))]
+	[ParamHelp("TestIterations", "Number of iterations to repeat this test", ParamType = typeof(int), DefaultValue = 1)]
+	[ParamHelp("CookedEditor", "Restricts usage of uncooked editor role", ParamType = typeof(bool))]
+	[ParamHelp("Device", "List of devices to use for tests", Action = ParamHelpAttribute.ParamAction.Append)]
+	[ParamHelp("NumClients", "Number of clients to run test with", ParamType = typeof(int), DefaultValue = 1)]
+	[ParamHelp("Server", "Run test with server", ParamType = typeof(bool))]
+	[ParamHelp("NullRHI", "Null Rendering Hardware Interface (run headless)", ParamType = typeof(bool))]
+	[ParamHelp("Windowed", "Run in Windowed mode", ParamType = typeof(bool))]
+	[ParamHelp("ResX", "Horizontal resolution", ParamType = typeof(int), DefaultValue = 1920)]
+	[ParamHelp("ResY", "Vertical resolution", ParamType = typeof(int), DefaultValue = 1080)]
+	[ParamHelp("Unattended", "Run in Unattended mode", ParamType = typeof(bool))]
+	[ParamHelp("Reboot", "Reboot device before starting test", ParamType = typeof(bool))]
+	[ParamHelp("SkipDeploy", "Skip deployment of build packages to devices", ParamType = typeof(bool))]
+	[ParamHelp("Log", "Output Logs", ParamType = typeof(bool))]
+	[ParamHelp("LogDir", "Location to store log files. Defaults to TempDir/Logs")]
+	[ParamHelp("Timestamp", "Print timestamp prefix to Log lines", ParamType = typeof(bool))]
+	[ParamHelp("TempDir", "Location to store temporary files. Defaults to GauntletTemp")]
+	[ParamHelp("Verbose", "Verbose logging", ParamType = typeof(bool))]
+	[ParamHelp("VeryVerbose", "Very Verbose logging", ParamType = typeof(bool))]
+	[ParamHelp("HeartbeatPeriod", "Gauntlet heartbeat period", ParamType = typeof(float))]
+	[ParamHelp("Args", "Extra arguments to pass role(s)")]
+	[ParamHelp("Namespaces", "Comma-separated list of namespaces to check for tests.", MultiSelectSeparator = ",")]
 	public class RunUnreal : BuildCommand
 	{
 
