Added "ReportType=" parameter to RunLowLevelTests.cs to support enabling report type.

diff --git a/Engine/Source/Programs/AutomationTool/LowLevelTests/RunLowLevelTests.cs b/Engine/Source/Programs/AutomationTool/LowLevelTests/RunLowLevelTests.cs
--- a/Engine/Source/Programs/AutomationTool/LowLevelTests/RunLowLevelTests.cs
+++ b/Engine/Source/Programs/AutomationTool/LowLevelTests/RunLowLevelTests.cs
@@ -167,6 +167,9 @@ namespace LowLevelTests
 		[AutoParam("")]
 		public string LogDir;
 
+		[AutoParam("console")]
+		public string ReportType;
+
 		public Type BuildSourceType { get; protected set; }
 
 		[AutoParam(UnrealTargetConfiguration.Development)]
@@ -210,6 +213,8 @@ namespace LowLevelTests
 			Tags = Params.ParseValue("tags=", null);
 			AttachToDebugger = Params.ParseParam("attachtodebugger");
 
+			ReportType = Params.ParseValue("reporttype=", "console");
+
 			string PlatformArgString = Params.ParseValue("platform=", null);
 			Platform = string.IsNullOrEmpty(PlatformArgString) ? BuildHostPlatform.Current.Platform : UnrealTargetPlatform.Parse(PlatformArgString);
 
@@ -235,15 +240,17 @@ namespace LowLevelTests
 		private string Tags { get; set; }
 		private int Sleep { get; set; }
 		private bool AttachToDebugger { get; set; }
+		private string ReportType { get; set; }
 
 		public UnrealDeviceReservation UnrealDeviceReservation { get; private set; }
 
-		public LowLevelTestsSession(LowLevelTestsBuildSource InBuildSource, string InTags, int InSleep, bool InAttachToDebugger)
+		public LowLevelTestsSession(LowLevelTestsBuildSource InBuildSource, string InTags, int InSleep, bool InAttachToDebugger, string InReportType)
 		{
 			BuildSource = InBuildSource;
 			Tags = InTags;
 			Sleep = InSleep;
 			AttachToDebugger = InAttachToDebugger;
+			ReportType = InReportType;
 			UnrealDeviceReservation = new UnrealDeviceReservation();
 		}
 
@@ -267,7 +274,7 @@ namespace LowLevelTests
 
 			// TargetDevice<Platform> classes have a hard dependency on UnrealAppConfig instead of IAppConfig.
 			// More refactoring needed to support non-packaged applications that can be run natively from a path on the device.
-			UnrealAppConfig AppConfig = BuildSource.GetUnrealAppConfig(Tags, Sleep, AttachToDebugger);
+			UnrealAppConfig AppConfig = BuildSource.GetUnrealAppConfig(Tags, Sleep, AttachToDebugger, ReportType);
 
 			IEnumerable<ITargetDevice> DevicesToInstallOn = UnrealDeviceReservation.ReservedDevices.ToArray();
 			ITargetDevice Device = DevicesToInstallOn.Where(D => D.IsConnected && D.Platform == BuildSource.Platform).First();
@@ -578,7 +585,7 @@ namespace LowLevelTests
 				.First();
 		}
 
-		public UnrealAppConfig GetUnrealAppConfig(string InTags, int InSleep, bool InAttachToDebugger)
+		public UnrealAppConfig GetUnrealAppConfig(string InTags, int InSleep, bool InAttachToDebugger, string InReportType)
 		{
 			if (CachedConfig == null)
 			{
@@ -593,7 +600,7 @@ namespace LowLevelTests
 				CachedConfig.FilesToCopy = new List<UnrealFileToCopy>();
 				// Set reporting options, filters etc
 				CachedConfig.CommandLineParams.AddRawCommandline("--durations=no");
-				CachedConfig.CommandLineParams.AddRawCommandline("--reporter=console");
+				CachedConfig.CommandLineParams.AddRawCommandline(string.Format("--reporter={0}", InReportType));
 				CachedConfig.CommandLineParams.AddRawCommandline(string.Format("--out={0}", LowLevelTestsReporting.GetTargetReportPath(Platform, TestApp, BuildPath)));
 				CachedConfig.CommandLineParams.AddRawCommandline("--filenames-as-tags");
 				if (!string.IsNullOrEmpty(InTags))
diff --git a/Engine/Source/Programs/AutomationTool/LowLevelTests/Tests/Gauntlet.LowLevelTestNode.cs b/Engine/Source/Programs/AutomationTool/LowLevelTests/Tests/Gauntlet.LowLevelTestNode.cs
--- a/Engine/Source/Programs/AutomationTool/LowLevelTests/Tests/Gauntlet.LowLevelTestNode.cs
+++ b/Engine/Source/Programs/AutomationTool/LowLevelTests/Tests/Gauntlet.LowLevelTestNode.cs
@@ -41,7 +41,7 @@ namespace LowLevelTests
 		{
 			if (LowLevelTestsApp == null)
 			{
-				LowLevelTestsApp = new LowLevelTestsSession(Context.BuildInfo, Context.Options.Tags, Context.Options.Sleep, Context.Options.AttachToDebugger);
+				LowLevelTestsApp = new LowLevelTestsSession(Context.BuildInfo, Context.Options.Tags, Context.Options.Sleep, Context.Options.AttachToDebugger, Context.Options.ReportType);
 			}
 
 			return LowLevelTestsApp.TryReserveDevices();
