Allow the CompileTask to enable AllCores for UBT

Turns on automatically if CommandUtils.IsBuildMachine is true or the CompileTask sets AllowAllCores to true

#rnx

diff --git a/Engine/Source/Programs/AutomationTool/AutomationUtils/UnrealBuild.cs b/Engine/Source/Programs/AutomationTool/AutomationUtils/UnrealBuild.cs
--- a/Engine/Source/Programs/AutomationTool/AutomationUtils/UnrealBuild.cs
+++ b/Engine/Source/Programs/AutomationTool/AutomationUtils/UnrealBuild.cs
@@ -205,7 +205,7 @@ namespace AutomationTool
 			}
         }
 
-		BuildManifest BuildWithUBT(string TargetName, UnrealTargetPlatform TargetPlatform, UnrealTargetConfiguration Config, FileReference UprojectPath, bool ForceFlushMac = false, bool DisableXGE = false, string InAddArgs = "")
+		BuildManifest BuildWithUBT(string TargetName, UnrealTargetPlatform TargetPlatform, UnrealTargetConfiguration Config, FileReference UprojectPath, bool ForceFlushMac = false, bool DisableXGE = false, bool AllCores = false, string InAddArgs = "")
 		{
 			string AddArgs = "";
 			if (UprojectPath != null)
@@ -222,6 +222,10 @@ namespace AutomationTool
 			{
 				AddArgs += " -noxge";
 			}
+			if (AllCores)
+			{
+				AddArgs += " -allcores";
+			}
 
 			PrepareUBT();
 
@@ -527,7 +531,7 @@ namespace AutomationTool
 			return Result;
 		}
 
-		public bool ProcessXGEItems(List<XGEItem> Actions, string XGETool, string Args, string TaskFilePath, bool ShowProgress)
+		public bool ProcessXGEItems(List<XGEItem> Actions, string XGETool, string Args, string TaskFilePath, bool ShowProgress, bool AllCores)
 		{
 			IScope CombineXGEScope = GlobalTracer.Instance.BuildSpan("CombineXGEItemFiles").WithTag("xgeTool", Path.GetFileNameWithoutExtension(XGETool)).StartActive();
 
@@ -583,7 +587,7 @@ namespace AutomationTool
 					CommandUtils.PushDir(CommandUtils.CombinePaths(CommandUtils.CmdEnv.LocalRoot, @"\Engine\Source"));
 					try
 					{
-						int ExitCode = ParallelExecutor.Execute(TaskFilePath, OwnerCommand.ParseParam("StopOnErrors"), OwnerCommand.ParseParam("AllCores"));
+						int ExitCode = ParallelExecutor.Execute(TaskFilePath, OwnerCommand.ParseParam("StopOnErrors"), AllCores);
 						if(ExitCode != 0)
 						{
 							return false;
@@ -982,7 +986,8 @@ namespace AutomationTool
 		/// <param name="InUpdateVersionFiles">True if the version files are to be updated </param>
 		/// <param name="InForceNoXGE">If true will force XGE off</param>
 		/// <param name="InUseParallelExecutor">If true AND XGE not present or not being used then use ParallelExecutor</param>
-		public void Build(BuildAgenda Agenda, bool? InDeleteBuildProducts = null, bool InUpdateVersionFiles = true, bool InForceNoXGE = false, bool InUseParallelExecutor = false, bool InShowProgress = false, int? InChangelistNumberOverride = null, Dictionary<BuildTarget, BuildManifest> InTargetToManifest = null)
+		/// <param name="InAllCores">If true AND XGE not present or not being used then ensure UBT uses all available cores</param>
+		public void Build(BuildAgenda Agenda, bool? InDeleteBuildProducts = null, bool InUpdateVersionFiles = true, bool InForceNoXGE = false, bool InUseParallelExecutor = false, bool InShowProgress = false, bool InAllCores = false, int? InChangelistNumberOverride = null, Dictionary<BuildTarget, BuildManifest> InTargetToManifest = null)
 		{
 			if (!CommandUtils.CmdEnv.HasCapabilityToCompile)
 			{
@@ -1076,7 +1081,7 @@ namespace AutomationTool
 			{
 				// When building a target for Mac or iOS, use UBT's -flushmac option to clean up the remote builder
 				bool bForceFlushMac = DeleteBuildProducts && (Target.Platform == UnrealTargetPlatform.Mac || Target.Platform == UnrealTargetPlatform.IOS);
-				BuildManifest Manifest = BuildWithUBT(Target.TargetName, Target.Platform, Target.Config, Target.UprojectPath, bForceFlushMac, bDisableXGE, Target.UBTArgs);
+				BuildManifest Manifest = BuildWithUBT(Target.TargetName, Target.Platform, Target.Config, Target.UprojectPath, bForceFlushMac, bDisableXGE, InAllCores, Target.UBTArgs);
 				if(InTargetToManifest != null)
 				{
 					InTargetToManifest[Target] = Manifest;
@@ -1086,17 +1091,17 @@ namespace AutomationTool
 			// Execute all the XGE targets
 			if(ParallelXgeTargets.Count > 0)
 			{
-				BuildParallelTargets(ParallelXgeTargets, InShowProgress, XGEConsole, InTargetToManifest);
+				BuildParallelTargets(ParallelXgeTargets, InShowProgress, InAllCores, XGEConsole, InTargetToManifest);
 			}
 
 			// Execute all the parallel targets
 			if(ParallelTargets.Count > 0)
 			{
-				BuildParallelTargets(ParallelTargets, InShowProgress, null, InTargetToManifest);
+				BuildParallelTargets(ParallelTargets, InShowProgress, InAllCores, null, InTargetToManifest);
 			}
 		}
 
-		private void BuildParallelTargets(List<BuildTarget> ParallelTargets, bool InShowProgress, string XGETool, Dictionary<BuildTarget, BuildManifest> InTargetToManifest)
+		private void BuildParallelTargets(List<BuildTarget> ParallelTargets, bool InShowProgress, bool InAllCores, string XGETool, Dictionary<BuildTarget, BuildManifest> InTargetToManifest)
 		{
 			string TaskFilePath = CommandUtils.CombinePaths(CommandUtils.CmdEnv.LogFolder, @"UAT_XGE.xml");
 
@@ -1138,7 +1143,7 @@ namespace AutomationTool
 			}
 
 			CommandUtils.LogSetProgress(InShowProgress, "Building...");
-			if (!ProcessXGEItems(XGEItems, XGETool, Args, TaskFilePath, InShowProgress))
+			if (!ProcessXGEItems(XGEItems, XGETool, Args, TaskFilePath, InShowProgress, InAllCores))
 			{
 				throw new UnrealBuildException("{0} failed, retries not enabled:", XGETool);
 			}
diff --git a/Engine/Source/Programs/AutomationTool/BuildGraph/Tasks/CompileTask.cs b/Engine/Source/Programs/AutomationTool/BuildGraph/Tasks/CompileTask.cs
--- a/Engine/Source/Programs/AutomationTool/BuildGraph/Tasks/CompileTask.cs
+++ b/Engine/Source/Programs/AutomationTool/BuildGraph/Tasks/CompileTask.cs
@@ -61,6 +61,12 @@ namespace AutomationTool.Tasks
 		[TaskParameter(Optional = true)]
 		public bool AllowParallelExecutor = true;
 
+		/// <summary>
+		/// Whether to allow UBT to use all available cores, when AllowXGE is disabled.
+		/// </summary>
+		[TaskParameter(Optional = true)]
+		public bool AllowAllCores = false;
+
 		/// <summary>
 		/// Whether to allow cleaning this target. If unspecified, targets are cleaned if the -Clean argument is passed on the command line.
 		/// </summary>
@@ -99,6 +105,11 @@ namespace AutomationTool.Tasks
 		/// </summary>
 		bool bAllowParallelExecutor = true;
 
+		/// <summary>
+		/// Whether to allow using all available cores for this job, when bAllowXGE is false
+		/// </summary>
+		bool bAllowAllCores = false;
+
 		/// <summary>
 		/// Constructor
 		/// </summary>
@@ -136,6 +147,7 @@ namespace AutomationTool.Tasks
 			CompileTaskParameters Parameters = CompileTask.Parameters;
 			bAllowXGE &= Parameters.AllowXGE;
 			bAllowParallelExecutor &= Parameters.AllowParallelExecutor;
+			bAllowAllCores &= Parameters.AllowAllCores;
 
 			UnrealBuild.BuildTarget Target = new UnrealBuild.BuildTarget { TargetName = Parameters.Target, Platform = Parameters.Platform, Config = Parameters.Configuration, UprojectPath = CompileTask.FindProjectFile(), UBTArgs = "-nobuilduht " + (Parameters.Arguments ?? ""), Clean = Parameters.Clean };
 			if(!String.IsNullOrEmpty(Parameters.Tag))
@@ -165,7 +177,8 @@ namespace AutomationTool.Tasks
 			UnrealBuild Builder = new UnrealBuild(Job.OwnerCommand);
 
 			bool bCanUseParallelExecutor = (BuildHostPlatform.Current.Platform == UnrealTargetPlatform.Win64 && bAllowParallelExecutor);	// parallel executor is only available on Windows as of 2016-09-22
-			Builder.Build(Agenda, InDeleteBuildProducts: null, InUpdateVersionFiles: false, InForceNoXGE: !bAllowXGE, InUseParallelExecutor: bCanUseParallelExecutor, InTargetToManifest: TargetToManifest);
+			bool bAllCores = (CommandUtils.IsBuildMachine || bAllowAllCores);	// Enable using all cores if this is a build agent or the flag was passed in to the task and XGE is disabled.
+			Builder.Build(Agenda, InDeleteBuildProducts: null, InUpdateVersionFiles: false, InForceNoXGE: !bAllowXGE, InUseParallelExecutor: bCanUseParallelExecutor, InAllCores: bAllCores, InTargetToManifest: TargetToManifest);
 
 			UnrealBuild.CheckBuildProducts(Builder.BuildProductFiles);
 
