BuildGraph: Allow overriding path to Docker executable through env var

Allows for using alternative but compatible Docker implementations.

diff --git a/Engine/Source/Programs/AutomationTool/BuildGraph/Tasks/DockerBuildTask.cs b/Engine/Source/Programs/AutomationTool/BuildGraph/Tasks/DockerBuildTask.cs
--- a/Engine/Source/Programs/AutomationTool/BuildGraph/Tasks/DockerBuildTask.cs
+++ b/Engine/Source/Programs/AutomationTool/BuildGraph/Tasks/DockerBuildTask.cs
@@ -173,7 +173,8 @@ namespace AutomationTool.Tasks
 				}
 				
 				string WorkingDir = isStagingEnabled ? StagingDir.FullName : BaseDir.FullName;
-				await SpawnTaskBase.ExecuteAsync("docker", Arguments.ToString(), EnvVars: EnvVars, WorkingDir: WorkingDir, SpewFilterCallback: FilterOutput);
+				string Exe = DockerTask.GetDockerExecutablePath();
+				await SpawnTaskBase.ExecuteAsync(Exe, Arguments.ToString(), EnvVars: EnvVars, WorkingDir: WorkingDir, SpewFilterCallback: FilterOutput);
 			}
 		}
 
diff --git a/Engine/Source/Programs/AutomationTool/BuildGraph/Tasks/DockerPushTask.cs b/Engine/Source/Programs/AutomationTool/BuildGraph/Tasks/DockerPushTask.cs
--- a/Engine/Source/Programs/AutomationTool/BuildGraph/Tasks/DockerPushTask.cs
+++ b/Engine/Source/Programs/AutomationTool/BuildGraph/Tasks/DockerPushTask.cs
@@ -86,17 +86,18 @@ namespace AutomationTool.Tasks
 			Log.TraceInformation("Pushing Docker image");
 			using (LogIndentScope Scope = new LogIndentScope("  "))
 			{
+				string Exe = DockerTask.GetDockerExecutablePath();
 				Dictionary<string, string> Environment = ParseEnvVars(Parameters.Environment, Parameters.EnvironmentFile);
 
 				if (Parameters.AwsEcr)
 				{
 					IProcessResult Result = await SpawnTaskBase.ExecuteAsync("aws", "ecr get-login-password", EnvVars: Environment, LogOutput: false);
-					await ExecuteAsync("docker", $"login {Parameters.Repository} --username AWS --password-stdin", Input: Result.Output);
+					await ExecuteAsync(Exe, $"login {Parameters.Repository} --username AWS --password-stdin", Input: Result.Output);
 				}
 
 				string TargetImage = Parameters.TargetImage ?? Parameters.Image;
-				await ExecuteAsync("docker", $"tag {Parameters.Image} {Parameters.Repository}/{TargetImage}", EnvVars: Environment);
-				await ExecuteAsync("docker", $"push {Parameters.Repository}/{TargetImage}", EnvVars: Environment);
+				await ExecuteAsync(Exe, $"tag {Parameters.Image} {Parameters.Repository}/{TargetImage}", EnvVars: Environment);
+				await ExecuteAsync(Exe, $"push {Parameters.Repository}/{TargetImage}", EnvVars: Environment);
 			}
 		}
 
diff --git a/Engine/Source/Programs/AutomationTool/BuildGraph/Tasks/DockerTask.cs b/Engine/Source/Programs/AutomationTool/BuildGraph/Tasks/DockerTask.cs
--- a/Engine/Source/Programs/AutomationTool/BuildGraph/Tasks/DockerTask.cs
+++ b/Engine/Source/Programs/AutomationTool/BuildGraph/Tasks/DockerTask.cs
@@ -69,7 +69,7 @@ namespace AutomationTool.Tasks
 		/// <param name="TagNameToFileSet">Mapping from tag names to the set of files they include</param>
 		public override async Task ExecuteAsync(JobContext Job, HashSet<FileReference> BuildProducts, Dictionary<string, HashSet<FileReference>> TagNameToFileSet)
 		{
-			await ExecuteAsync("docker", Parameters.Arguments, EnvVars: ParseEnvVars(Parameters.Environment, Parameters.EnvironmentFile), WorkingDir: Parameters.WorkingDir);
+			await ExecuteAsync(GetDockerExecutablePath(), Parameters.Arguments, EnvVars: ParseEnvVars(Parameters.Environment, Parameters.EnvironmentFile), WorkingDir: Parameters.WorkingDir);
 		}
 
 		/// <summary>
@@ -97,5 +97,15 @@ namespace AutomationTool.Tasks
 		{
 			yield break;
 		}
+
+		/// <summary>
+		/// Resolve path to Docker executable by using the optional env var "UE_DOCKER_EXEC_PATH"
+		/// Will default to "docker" if not set. Allows supporting alternative Docker implementations such as Podman.
+		/// </summary>
+		/// <returns>Path to Docker executable</returns>
+		public static string GetDockerExecutablePath()
+		{
+			return Environment.GetEnvironmentVariable("UE_DOCKER_EXEC_PATH") ?? "docker";
+		}
 	}
 }
