BuildGraph: Fix temp storage copy actions running 5 times, even if it succeeds on the first attempt.

diff --git a/Engine/Source/Programs/AutomationTool/BuildGraph/TempStorage.cs b/Engine/Source/Programs/AutomationTool/BuildGraph/TempStorage.cs
--- a/Engine/Source/Programs/AutomationTool/BuildGraph/TempStorage.cs
+++ b/Engine/Source/Programs/AutomationTool/BuildGraph/TempStorage.cs
@@ -967,6 +967,7 @@ namespace AutomationTool
 				try
 				{
 					retryAction();
+					break;
 				}
 				catch
 				{
