Marking up more parameter help attributes on Gauntlet commands

diff --git a/Engine/Source/Programs/AutomationTool/Gauntlet/Unreal/RunUnreal.cs b/Engine/Source/Programs/AutomationTool/Gauntlet/Unreal/RunUnreal.cs
--- a/Engine/Source/Programs/AutomationTool/Gauntlet/Unreal/RunUnreal.cs
+++ b/Engine/Source/Programs/AutomationTool/Gauntlet/Unreal/RunUnreal.cs
@@ -59,6 +59,9 @@ namespace Gauntlet
 	[ParamHelp("VeryVerbose", "Very Verbose logging", ParamType = typeof(bool))]
 	[ParamHelp("HeartbeatPeriod", "Gauntlet heartbeat period", ParamType = typeof(float))]
 	[ParamHelp("Args", "Extra arguments to pass role(s)")]
+	[ParamHelp("ClientArgs", "Extra arguments passed to client role(s)")]
+	[ParamHelp("ServerArgs", "Extra arguments passed to server role(s)")]
+	[ParamHelp("EditorArgs", "Extra arguments passed to editor role(s)")]
 	[ParamHelp("Namespaces", "Comma-separated list of namespaces to check for tests.", MultiSelectSeparator = ",")]
 	public class RunUnreal : BuildCommand
 	{
