Cooked editors now stage the editor localization targets

diff --git a/Engine/Source/Programs/AutomationTool/CookedEditor/CookedEditor.Automation.cs b/Engine/Source/Programs/AutomationTool/CookedEditor/CookedEditor.Automation.cs
--- a/Engine/Source/Programs/AutomationTool/CookedEditor/CookedEditor.Automation.cs
+++ b/Engine/Source/Programs/AutomationTool/CookedEditor/CookedEditor.Automation.cs
@@ -14,7 +14,7 @@ public class ConfigHelper
 {
 	private string SpecificConfigSection;
 	private string SharedConfigSection;
-	private ConfigHierarchy GameConfig;
+	public ConfigHierarchy GameConfig { get; }
 
 	public ConfigHelper(UnrealTargetPlatform Platform, FileReference ProjectFile, bool bIsCookedCooker)
 	{
@@ -465,6 +465,21 @@ public class MakeCookedEditor : BuildCommand
 			GatherTargetDependencies(Params, SC, Context, "UnrealPak");
 		}
 
+		// Stage the editor localization targets
+		if (!bIsCookedCooker)
+		{
+			List<string> CulturesToStage = Project.GetCulturesToStage(Params, ConfigHelper.GameConfig);
+
+			string[] EditorLocalizationTargetsToStage = { "Category", "Editor", "EditorTutorials", "Keywords", "PropertyNames", "ToolTips" };
+			foreach (string EditorLocalizationTargetToStage in EditorLocalizationTargetsToStage)
+			{
+				if (Project.ShouldStageLocalizationTarget(SC, null, EditorLocalizationTargetToStage))
+				{
+					Project.StageLocalizationDataForTarget(SC, CulturesToStage, DirectoryReference.Combine(Unreal.EngineDirectory, "Content", "Localization", EditorLocalizationTargetToStage));
+				}
+			}
+		}
+
 		StageIniPathArray(Params, SC, "EngineExtraStageFiles", Unreal.EngineDirectory, Context);
 
 		Context.FilesToUncook.Add(FileReference.Combine(Context.EngineDirectory, "Content", "EngineMaterials", "DefaultMaterial.uasset"));
diff --git a/Engine/Source/Programs/AutomationTool/Scripts/CopyBuildToStagingDirectory.Automation.cs b/Engine/Source/Programs/AutomationTool/Scripts/CopyBuildToStagingDirectory.Automation.cs
--- a/Engine/Source/Programs/AutomationTool/Scripts/CopyBuildToStagingDirectory.Automation.cs
+++ b/Engine/Source/Programs/AutomationTool/Scripts/CopyBuildToStagingDirectory.Automation.cs
@@ -381,7 +381,7 @@ namespace AutomationScripts
 			return InternationalizationPreset;
 		}
 
-		private static List<string> GetCulturesToStage(ProjectParams Params, ConfigHierarchy PlatformGameConfig, bool bMustExist = true)
+		public static List<string> GetCulturesToStage(ProjectParams Params, ConfigHierarchy PlatformGameConfig, bool bMustExist = true)
 		{
 			// Initialize cultures to stage.
 			List<string> CulturesToStage = null;
@@ -410,7 +410,7 @@ namespace AutomationScripts
 			return CulturesToStage;
 		}
 
-		private static bool ShouldStageLocalizationTarget(DeploymentContext SC, ConfigHierarchy PlatformGameConfig, string LocalizationTarget)
+		public static bool ShouldStageLocalizationTarget(DeploymentContext SC, ConfigHierarchy PlatformGameConfig, string LocalizationTarget)
 		{
 			if (SC.LocalizationTargetsDenyList.Contains(LocalizationTarget))
 			{
@@ -433,7 +433,7 @@ namespace AutomationScripts
 			return true;
 		}
 
-		private static void StageLocalizationDataForPlugin(DeploymentContext SC, List<string> CulturesToStage, FileReference Plugin)
+		public static void StageLocalizationDataForPlugin(DeploymentContext SC, List<string> CulturesToStage, FileReference Plugin)
 		{
 			PluginDescriptor Descriptor = PluginDescriptor.FromFile(Plugin);
 			if (Descriptor.LocalizationTargets != null)
@@ -452,7 +452,7 @@ namespace AutomationScripts
 			}
 		}
 
-		private static void StageLocalizationDataForTarget(DeploymentContext SC, List<string> CulturesToStage, DirectoryReference SourceDirectory)
+		public static void StageLocalizationDataForTarget(DeploymentContext SC, List<string> CulturesToStage, DirectoryReference SourceDirectory)
 		{
 			var PlatformSourceDirectory = new DirectoryReference(CombinePaths(SourceDirectory.FullName, "Platforms", ConfigHierarchy.GetIniPlatformName(SC.StageTargetPlatform.IniPlatformType)));
 			if (!DirectoryReference.Exists(PlatformSourceDirectory))
@@ -476,7 +476,7 @@ namespace AutomationScripts
 			}
 		}
 
-		private static void StageLocalizationDataForTargetsInDirectory(DeploymentContext SC, ConfigHierarchy PlatformGameConfig, List<string> CulturesToStage, DirectoryReference RootDirectory)
+		public static void StageLocalizationDataForTargetsInDirectory(DeploymentContext SC, ConfigHierarchy PlatformGameConfig, List<string> CulturesToStage, DirectoryReference RootDirectory)
 		{
 			if (DirectoryReference.Exists(RootDirectory))
 			{
