Horde: Add missing file list to systemic error matcher
#rnx

diff --git a/Engine/Source/Programs/AutomationTool/AutomationUtils/Matchers/SystemicEventMatcher.cs b/Engine/Source/Programs/AutomationTool/AutomationUtils/Matchers/SystemicEventMatcher.cs
--- a/Engine/Source/Programs/AutomationTool/AutomationUtils/Matchers/SystemicEventMatcher.cs
+++ b/Engine/Source/Programs/AutomationTool/AutomationUtils/Matchers/SystemicEventMatcher.cs
@@ -19,6 +19,8 @@ namespace AutomationUtils.Matchers
 
 		static readonly Regex s_hostDown = new Regex(@"ERROR: System\.IO\.IOException: Host is down");
 
+		static readonly Regex s_missingFileList = new Regex(@"^\s*ERROR: Missing local or shared file list.*\.xml");
+
 		public LogEventMatch? Match(ILogCursor cursor)
 		{
 			if (cursor.IsMatch(s_ddc))
@@ -47,6 +49,11 @@ namespace AutomationUtils.Matchers
 				return new LogEventBuilder(cursor).ToMatch(LogEventPriority.High, LogLevel.Information, KnownLogEvents.Systemic_MissingXgeControlWorker);
 			}
 
+			if (cursor.IsMatch(s_missingFileList))
+			{
+				return new LogEventBuilder(cursor).ToMatch(LogEventPriority.High, LogLevel.Error, KnownLogEvents.Systemic_MissingFileList);
+			}
+
 			if (cursor.Contains("SignTool Error: The specified timestamp server either could not be reached or"))
 			{
 				return new LogEventBuilder(cursor).ToMatch(LogEventPriority.Normal, LogLevel.Information, KnownLogEvents.Systemic_SignToolTimeStampServer);
diff --git a/Engine/Source/Programs/Shared/EpicGames.Core/KnownLogEvents.cs b/Engine/Source/Programs/Shared/EpicGames.Core/KnownLogEvents.cs
--- a/Engine/Source/Programs/Shared/EpicGames.Core/KnownLogEvents.cs
+++ b/Engine/Source/Programs/Shared/EpicGames.Core/KnownLogEvents.cs
@@ -259,6 +259,13 @@ namespace EpicGames.Core
 		/// </summary>
 		public static EventId Systemic_SignTool { get; } = new EventId(766);
 
+		/// <summary>
+		/// Missing file list error
+		/// 
+		/// Usually due to a volume that is not mounted or some of other IO error
+		/// </summary>
+		public static EventId Systemic_MissingFileList { get; } = new EventId(767);
+
 		/// <summary>
 		/// Maximum code for systemic events. Add new events in the 700-799 range.
 		/// </summary>
