Allow ParamHelpAttribute markup on interfaces

diff --git a/Engine/Source/Programs/AutomationTool/AutomationUtils/HelpUtils.cs b/Engine/Source/Programs/AutomationTool/AutomationUtils/HelpUtils.cs
--- a/Engine/Source/Programs/AutomationTool/AutomationUtils/HelpUtils.cs
+++ b/Engine/Source/Programs/AutomationTool/AutomationUtils/HelpUtils.cs
@@ -74,7 +74,7 @@ namespace AutomationTool
 		}
 	}
 
-	[AttributeUsage(AttributeTargets.Class, AllowMultiple = true)]
+	[AttributeUsage(AttributeTargets.Class | AttributeTargets.Interface, AllowMultiple = true)]
 	public class ParamHelpAttribute : HelpAttribute
 	{
 		public enum ParamAction
