Have gauntlet display the command a test was ran with
Add server output to the horde log so errors can be easier identified from a webpage.
Fixing a typo with Text which should have been Test

https://p4-swarm.epicgames.net/reviews/22408413/

diff --git a/Engine/Source/Programs/AutomationTool/Gauntlet/Framework/Base/Gauntlet.BaseTestNode.cs b/Engine/Source/Programs/AutomationTool/Gauntlet/Framework/Base/Gauntlet.BaseTestNode.cs
--- a/Engine/Source/Programs/AutomationTool/Gauntlet/Framework/Base/Gauntlet.BaseTestNode.cs
+++ b/Engine/Source/Programs/AutomationTool/Gauntlet/Framework/Base/Gauntlet.BaseTestNode.cs
@@ -113,6 +113,13 @@ namespace Gauntlet
 			return new string[0];
 		}
 
+		public virtual string GetRunLocalCommand(string LaunchingBuildCommand)
+		{
+			string CommandToRunLocally =
+				string.Format("RunUAT {0} -Test={1} ", LaunchingBuildCommand, GetType());
+			return CommandToRunLocally;
+		}
+
 		/// <summary>
 		/// Mark the test as started
 		/// </summary>
diff --git a/Engine/Source/Programs/AutomationTool/Gauntlet/Framework/Base/Gauntlet.TestNode.cs b/Engine/Source/Programs/AutomationTool/Gauntlet/Framework/Base/Gauntlet.TestNode.cs
--- a/Engine/Source/Programs/AutomationTool/Gauntlet/Framework/Base/Gauntlet.TestNode.cs
+++ b/Engine/Source/Programs/AutomationTool/Gauntlet/Framework/Base/Gauntlet.TestNode.cs
@@ -216,6 +216,11 @@ namespace Gauntlet
 		/// </summary>
 		IEnumerable<string> GetErrors();
 
+		/// <summary>
+		/// returns a string that will be how to run this test locally with RunUAT.bat
+		/// </summary>
+		string GetRunLocalCommand(string LaunchingBuildCommand);
+
 		/// <summary>
 		/// Called to request any that any necessary cleanup be performed. After CleanupTest is called no further calls will be
 		/// made to this test and thus all resources should be released.
diff --git a/Engine/Source/Programs/AutomationTool/Gauntlet/Framework/Gauntlet.TestExecutor.cs b/Engine/Source/Programs/AutomationTool/Gauntlet/Framework/Gauntlet.TestExecutor.cs
--- a/Engine/Source/Programs/AutomationTool/Gauntlet/Framework/Gauntlet.TestExecutor.cs
+++ b/Engine/Source/Programs/AutomationTool/Gauntlet/Framework/Gauntlet.TestExecutor.cs
@@ -34,7 +34,7 @@ namespace Gauntlet
 	/// <summary>
 	/// Class that is manages the creation and execution of one or more tests
 	/// </summary>
-	public class TextExecutor
+	public class TestExecutor
 	{
 		class TestExecutionInfo
 		{
@@ -54,7 +54,7 @@ namespace Gauntlet
 			public ExecutionResult	Result;
 			public TestResult		FinalResult;
 			public string			CancellationReason;
-
+			
 			/// <summary>
 			/// Time the test had to wait before running
 			/// </summary>
@@ -91,15 +91,20 @@ namespace Gauntlet
 
 		TestExecutorOptions Options;
 
+		protected string BuildCommandThatLaunchedExecutor;
+
 		public bool IsRunning { get; private set; }
 		public bool IsCancelled { get; private set; }
 		protected bool HaveReceivedPostAbort { get; private set; }
 
 		/// <summary>
 		/// Constructor that fills in some member variables
+		/// @Param the BuildCommand that was used to start this Test Executor.
+		/// Useful to log out or know where this test executor came from.
 		/// </summary>
-		public TextExecutor()
+		public TestExecutor(string InBuildCommand)
 		{
+			BuildCommandThatLaunchedExecutor = InBuildCommand;
 			RunningTests = new List<TestExecutionInfo>();
 		}
 
@@ -612,6 +617,9 @@ namespace Gauntlet
 				}
 			}
 
+			// display how the test could be ran locally
+			string RunLocalString = TestInfo.TestNode.GetRunLocalCommand(BuildCommandThatLaunchedExecutor);
+			Log.Info("How to run locally: ({0})", RunLocalString);
 		}
 
 		/// <summary>
diff --git a/Engine/Source/Programs/AutomationTool/Gauntlet/SelfTest/Framework/Gauntlet.SelfTest.OrderOfOpsTest.cs b/Engine/Source/Programs/AutomationTool/Gauntlet/SelfTest/Framework/Gauntlet.SelfTest.OrderOfOpsTest.cs
--- a/Engine/Source/Programs/AutomationTool/Gauntlet/SelfTest/Framework/Gauntlet.SelfTest.OrderOfOpsTest.cs
+++ b/Engine/Source/Programs/AutomationTool/Gauntlet/SelfTest/Framework/Gauntlet.SelfTest.OrderOfOpsTest.cs
@@ -175,6 +175,13 @@ namespace Gauntlet.SelfTest
 		{
 		}
 
+		public virtual string GetRunLocalCommand(string LaunchingBuildCommand)
+		{
+			string CommandToRunLocally =
+				string.Format("RunUAT {0} -Test={1} ", LaunchingBuildCommand, GetType());
+			return CommandToRunLocally;
+		}
+
 		public void CleanupTest()
 		{
 			SetNewStage(TestStages.CleanupTest);
diff --git a/Engine/Source/Programs/AutomationTool/Gauntlet/SelfTest/Gauntlet.SelfTest.BaseNode.cs b/Engine/Source/Programs/AutomationTool/Gauntlet/SelfTest/Gauntlet.SelfTest.BaseNode.cs
--- a/Engine/Source/Programs/AutomationTool/Gauntlet/SelfTest/Gauntlet.SelfTest.BaseNode.cs
+++ b/Engine/Source/Programs/AutomationTool/Gauntlet/SelfTest/Gauntlet.SelfTest.BaseNode.cs
@@ -87,6 +87,12 @@ namespace Gauntlet.SelfTest
 			return new string[0];
 		}
 
+		public virtual string GetRunLocalCommand(string LaunchingBuildCommand)
+		{
+			string CommandToRunLocally =
+				string.Format("RunUAT {0} -Test={1} ", LaunchingBuildCommand, GetType());
+			return CommandToRunLocally;
+		}
 
 		public virtual ITestNode[] GetSubTests()
 		{
diff --git a/Engine/Source/Programs/AutomationTool/Gauntlet/SelfTest/TestGauntlet.cs b/Engine/Source/Programs/AutomationTool/Gauntlet/SelfTest/TestGauntlet.cs
--- a/Engine/Source/Programs/AutomationTool/Gauntlet/SelfTest/TestGauntlet.cs
+++ b/Engine/Source/Programs/AutomationTool/Gauntlet/SelfTest/TestGauntlet.cs
@@ -55,7 +55,7 @@ namespace Gauntlet
 			IEnumerable<ITestNode> Nodes = Utils.TestConstructor.ConstructTests<ITestNode, string[]>(TestList, null, new[] { "Gauntlet.SelfTest" });
 
 			// Create the test executor
-			var Executor = new TextExecutor();
+			var Executor = new TestExecutor(ToString());
 
 			TestExecutorOptions Options = new TestExecutorOptions();
 			AutoParam.ApplyParamsAndDefaults(Options, this.Params);
diff --git a/Engine/Source/Programs/AutomationTool/Gauntlet/Unreal/Base/Gauntlet.UnrealTestNode.cs b/Engine/Source/Programs/AutomationTool/Gauntlet/Unreal/Base/Gauntlet.UnrealTestNode.cs
--- a/Engine/Source/Programs/AutomationTool/Gauntlet/Unreal/Base/Gauntlet.UnrealTestNode.cs
+++ b/Engine/Source/Programs/AutomationTool/Gauntlet/Unreal/Base/Gauntlet.UnrealTestNode.cs
@@ -279,14 +279,19 @@ namespace Gauntlet
 		}
 
 		/// <summary>
-		/// Used to track how much of our app log has been written out
+		/// Used to track how much of our server log has been written out
 		/// </summary>
-		private int LastAppLogCount;
+		private int LastServerLogCount = 0;
+
+		/// <summary>
+		/// Used to track how much of our client log has been written out
+		/// </summary>
+		private int LastClientLogCount = 0;
 
 		/// <summary>
 		/// Used to track how much of our editor log has been written out
 		/// </summary>
-		private int LastEditorLogCount;
+		private int LastEditorLogCount = 0;
 
 		private int CurrentPass;
 
@@ -367,7 +372,8 @@ namespace Gauntlet
 
 			UnrealTestResult = TestResult.Invalid;
 			TimeToWaitForProcesses = 5;
-			LastAppLogCount = 0;
+			LastServerLogCount = 0;
+			LastClientLogCount = 0;
 			LastEditorLogCount = 0;
 			CurrentPass = 0;
 			NumPasses = 0;
@@ -762,7 +768,8 @@ namespace Gauntlet
 			UnrealTestResult = TestResult.Invalid;
 			CurrentPass = Pass;
 			NumPasses = InNumPasses;
-			LastAppLogCount = 0;
+			LastServerLogCount = 0;
+			LastClientLogCount = 0;
 			LastEditorLogCount = 0;
 			LastHeartbeatTime = DateTime.MinValue;
 			LastActiveHeartbeatTime = DateTime.MinValue;
@@ -853,6 +860,59 @@ namespace Gauntlet
 			});
 		}
 
+		public override string GetRunLocalCommand(string LaunchingBuildCommand)
+		{
+			string[] ArgsToNotDisplay =
+			{
+					"test",
+					"tests",
+					"tempdir",
+					"logdir",
+					"branch",
+					"changelist",
+					"JobDetails",
+					"RecordDeviceUsage",
+					"uploadreport",
+					"skipdashboardsubmit",
+			};
+
+			bool ShouldArgBeDisplayed(string InArg)
+			{
+				return !ArgsToNotDisplay.Any(str => InArg.StartsWith(str));
+			}
+			
+			string WrapParameterInQuotes(string InString)
+			{
+				var positionOfEqual = InString.IndexOf("=");
+				if (positionOfEqual > 0)
+				{
+					var ParamStrName = InString.Substring(0, positionOfEqual);
+					var ParamStrValue = InString.Substring(positionOfEqual+1);
+					int n;
+					if (int.TryParse(ParamStrValue, out n))
+					{
+						return InString;
+					}
+					else
+					{
+						return string.Format("{0}=\"{1}\"", ParamStrName, ParamStrValue);
+					}
+				}
+				else
+				{
+					return InString;
+				}
+			}
+
+			var CleanedArgs = Context.TestParams.AllArguments.Where(arg => ShouldArgBeDisplayed(arg));
+			CleanedArgs = CleanedArgs.Select(str => WrapParameterInQuotes(str));
+
+			string StringOfCleanedArguments = string.Format("-{0}", string.Join(" -", CleanedArgs));
+			string CommandToRunLocally =
+				string.Format("RunUAT {0} {1} {2}", LaunchingBuildCommand, GetType(), StringOfCleanedArguments);
+			return CommandToRunLocally;
+		}
+
 		/// <summary>
 		/// Cleanup all resources
 		/// </summary>
@@ -892,7 +952,8 @@ namespace Gauntlet
 			SessionArtifacts = Enumerable.Empty<UnrealRoleArtifacts>();
 			RoleResults = Enumerable.Empty<UnrealRoleResult>();
 
-			LastAppLogCount = 0;
+			LastServerLogCount = 0;
+			LastClientLogCount = 0;
 			LastEditorLogCount = 0;
 			LastHeartbeatTime = DateTime.MinValue;
 			LastActiveHeartbeatTime = DateTime.MinValue;
@@ -921,22 +982,6 @@ namespace Gauntlet
 		/// <param name="InInstance">The Unrealnstance being tested</param>
 		public virtual void TickTest(UnrealSessionInstance InInstance)
 		{
-			IAppInstance App = null;
-			string AppInfoPrefix = "App";
-			if (InInstance.ClientApps == null)
-			{
-				App = InInstance.ServerApp;
-				AppInfoPrefix = "Server";
-			}
-			else
-			{
-				if (InInstance.ClientApps.Length > 0)
-				{
-					App = InInstance.ClientApps.First();
-					AppInfoPrefix = "Client";
-				}
-			}
-
 			List<string> LogCategories = new List<string>();
 			LogCategories.Add("Gauntlet");
 			{
@@ -946,50 +991,56 @@ namespace Gauntlet
 			}
 			LogCategories = LogCategories.Distinct().ToList();
 
-
-			if (App != null)
+			void LogHeartbeatCategories(ref int LastLogCount, IAppInstance App, string AppPrefix, bool bUpdateHeartbeatTime)
 			{
-				UnrealLogParser Parser = new UnrealLogParser(App.StdOut);
-				List<string> TestLines = new List<string>();
-				// ONLY ADD RANGE ONCE. Ordering is important and will be skewed if multiple ranges are added which can skew how logs are pulled out.
-				TestLines.AddRange(Parser.GetLogChannels(LogCategories, true));
-
-				for (int i = LastAppLogCount; i < TestLines.Count(); i++)
+				if (App != null)
 				{
-					Log.Info(string.Format("{0}: {1}", AppInfoPrefix, TestLines[i]));
+					UnrealLogParser Parser = new UnrealLogParser(App.StdOut);
+					List<string> TestLines = new List<string>();
+					// ONLY ADD RANGE ONCE. Ordering is important and will be skewed if multiple ranges are added which can skew how logs are pulled out.
+					TestLines.AddRange(Parser.GetLogChannels(LogCategories, true));
 
-					if (Regex.IsMatch(TestLines[i], @".*GauntletHeartbeat\: Active.*"))
-					{
-						LastHeartbeatTime = DateTime.Now;
-						LastActiveHeartbeatTime = DateTime.Now;
-					}
-					else if (Regex.IsMatch(TestLines[i], @".*GauntletHeartbeat\: Idle.*"))
+					for (int i = LastLogCount; i < TestLines.Count(); i++)
 					{
-						LastHeartbeatTime = DateTime.Now;
-					}
-				}
+						Log.Info(string.Format("{0}: {1}", AppPrefix, TestLines[i]));
 
-				LastAppLogCount = TestLines.Count();
+						if (bUpdateHeartbeatTime)
+						{
+							if (Regex.IsMatch(TestLines[i], @".*GauntletHeartbeat\: Active.*"))
+							{
+								LastHeartbeatTime = DateTime.Now;
+								LastActiveHeartbeatTime = DateTime.Now;
+							}
+							else if (Regex.IsMatch(TestLines[i], @".*GauntletHeartbeat\: Idle.*"))
+							{
+								LastHeartbeatTime = DateTime.Now;
+							}
+						}
+					}
 
-				// Detect missed heartbeats and fail the test
-				CheckHeartbeat();
+					LastLogCount = TestLines.Count();
+				}
 			}
 
-			IAppInstance EditorApp = InInstance.EditorApp;
-			if (EditorApp != null)
+			if (InInstance.ServerApp != null)
+			{
+				bool bUpdateHeartbeat = InInstance.ClientApps == null;
+				LogHeartbeatCategories(ref LastServerLogCount, InInstance.ServerApp, "Server", bUpdateHeartbeat);
+			}
+			if (InInstance.ClientApps.Length > 0)
 			{
-				UnrealLogParser Parser = new UnrealLogParser(EditorApp.StdOut);
-				List<string> TestLines = new List<string>();
-				// ONLY ADD RANGE ONCE. Ordering is important and will be skewed if multiple ranges are added which can skew how logs are pulled out.
-				TestLines.AddRange(Parser.GetLogChannels(LogCategories, true));
+				bool bUpdateHeartbeat = true;
+				LogHeartbeatCategories(ref LastClientLogCount, InInstance.ClientApps.First(), "Client", bUpdateHeartbeat);
+			}
 
-				for (int i = LastEditorLogCount; i < TestLines.Count(); i++)
-				{
-					Log.Info(string.Format("Editor: {0}", TestLines[i]));
-				}
+			if (InInstance.EditorApp != null)
+			{
+				bool bUpdateHeartbeat = false;
+				LogHeartbeatCategories(ref LastEditorLogCount, InInstance.EditorApp, "Editor", bUpdateHeartbeat);
+			}
 
-				LastEditorLogCount = TestLines.Count();
-			}			
+			// Detect missed heartbeats and fail the test
+			CheckHeartbeat();		
 		}
 
 		/// <summary>
diff --git a/Engine/Source/Programs/AutomationTool/Gauntlet/Unreal/RunUnreal.cs b/Engine/Source/Programs/AutomationTool/Gauntlet/Unreal/RunUnreal.cs
--- a/Engine/Source/Programs/AutomationTool/Gauntlet/Unreal/RunUnreal.cs
+++ b/Engine/Source/Programs/AutomationTool/Gauntlet/Unreal/RunUnreal.cs
@@ -4,15 +4,10 @@ using System;
 using System.Collections.Generic;
 using System.Linq;
 using System.Text;
-using System.Threading.Tasks;
 using AutomationTool;
 using AutomationTool.DeviceReservation;
 using UnrealBuildTool;
-using Gauntlet;
 using System.IO;
-using Newtonsoft.Json;
-using System.Reflection;
-using System.Text.RegularExpressions;
 using EpicGames.Core;
 using UnrealBuildBase;
 
@@ -65,7 +60,6 @@ namespace Gauntlet
 	[ParamHelp("Namespaces", "Comma-separated list of namespaces to check for tests.", MultiSelectSeparator = ",")]
 	public class RunUnreal : BuildCommand
 	{
-
 		/// <summary>
 		/// Test node to create if none were specified
 		/// </summary>
@@ -322,8 +316,8 @@ namespace Gauntlet
 		bool ExecuteTests(UnrealTestOptions Options, IEnumerable<ITestNode> TestList)
 		{
 			// Create the test executor
-			var Executor = new TextExecutor();
-			
+			var Executor = new TestExecutor(ToString());
+
 			try
 			{
 				bool Result = Executor.ExecuteTests(Options, TestList);
diff --git a/Engine/Source/Programs/AutomationTool/LowLevelTests/RunLowLevelTests.cs b/Engine/Source/Programs/AutomationTool/LowLevelTests/RunLowLevelTests.cs
--- a/Engine/Source/Programs/AutomationTool/LowLevelTests/RunLowLevelTests.cs
+++ b/Engine/Source/Programs/AutomationTool/LowLevelTests/RunLowLevelTests.cs
@@ -84,7 +84,7 @@ namespace LowLevelTests
 
 		private bool ExecuteTest(LowLevelTestExecutorOptions Options, ITestNode LowLevelTestNode)
 		{
-			var Executor = new TextExecutor();
+			var Executor = new TestExecutor(ToString());
 
 			try
 			{
