Adding an option to serialize output to a file instead of converting a StringBuilder to a string.

#rnx

diff --git a/Engine/Platforms/Hololens/Source/Programs/AutomationTool/Gauntlet/Gauntlet.TargetDeviceHoloLens.cs b/Engine/Platforms/Hololens/Source/Programs/AutomationTool/Gauntlet/Gauntlet.TargetDeviceHoloLens.cs
--- a/Engine/Platforms/Hololens/Source/Programs/AutomationTool/Gauntlet/Gauntlet.TargetDeviceHoloLens.cs
+++ b/Engine/Platforms/Hololens/Source/Programs/AutomationTool/Gauntlet/Gauntlet.TargetDeviceHoloLens.cs
@@ -311,6 +311,11 @@ namespace Gauntlet
 			}
 		}
 
+		public FileReference WriteOutputToFile(string FileName)
+		{
+			throw new NotImplementedException();
+		}
+
 		private void TailLogFile(string LogFilePath)
 		{
 			string LogName = GetProcessName();
@@ -580,6 +585,16 @@ namespace Gauntlet
 				System.Threading.Thread.Sleep(1000);
 			}
 		}
+
+		public FileReference WriteOutputToFile(string FileName)
+		{
+			using (StreamWriter writer = new StreamWriter(FileName))
+			{
+				writer.Write(OutputLog);
+			}
+
+			return new FileReference(FileName);
+		}
 	}
 #endif
 
diff --git a/Engine/Platforms/Hololens/Source/Programs/AutomationTool/HoloLensPlatform.Automation.cs b/Engine/Platforms/Hololens/Source/Programs/AutomationTool/HoloLensPlatform.Automation.cs
--- a/Engine/Platforms/Hololens/Source/Programs/AutomationTool/HoloLensPlatform.Automation.cs
+++ b/Engine/Platforms/Hololens/Source/Programs/AutomationTool/HoloLensPlatform.Automation.cs
@@ -167,6 +167,11 @@ namespace HoloLens.Automation
 			}
 		}
 
+		public FileReference WriteOutputToFile(string FileName)
+		{
+			throw new NotImplementedException();
+		}
+
 		private void TailLogFile(string LogFilePath)
 		{
 			string LogName = GetProcessName();
@@ -348,6 +353,11 @@ namespace HoloLens.Automation
 				System.Threading.Thread.Sleep(1000);
 			}
 		}
+
+		public FileReference WriteOutputToFile(string FileName)
+		{
+			throw new NotImplementedException();
+		}
 	}
 
 	[SupportedOSPlatform("windows10.0.10240.0")]
diff --git a/Engine/Source/Programs/AutomationTool/Android/AndroidPlatform.Automation.cs b/Engine/Source/Programs/AutomationTool/Android/AndroidPlatform.Automation.cs
--- a/Engine/Source/Programs/AutomationTool/Android/AndroidPlatform.Automation.cs
+++ b/Engine/Source/Programs/AutomationTool/Android/AndroidPlatform.Automation.cs
@@ -132,6 +132,11 @@ public class AndroidPlatform : Platform
 			StopProcess();
 		}
 
+		public FileReference WriteOutputToFile(string FileName)
+		{
+			return AdbLogProcess.WriteOutputToFile(FileName);
+		}
+
 		private bool IsPackageRunningOnDevice()
 		{
 			ERunOptions Options = ERunOptions.Default | ERunOptions.SpewIsVerbose | ERunOptions.NoLoggingOfRunCommand;
diff --git a/Engine/Source/Programs/AutomationTool/AutomationUtils/ProcessUtils.cs b/Engine/Source/Programs/AutomationTool/AutomationUtils/ProcessUtils.cs
--- a/Engine/Source/Programs/AutomationTool/AutomationUtils/ProcessUtils.cs
+++ b/Engine/Source/Programs/AutomationTool/AutomationUtils/ProcessUtils.cs
@@ -219,11 +219,12 @@ namespace AutomationTool
 		void DisposeProcess();
 		void StdOut(object sender, DataReceivedEventArgs e);
 		void StdErr(object sender, DataReceivedEventArgs e);
-		int ExitCode{ get;set; }
-		string Output {get;}
-		Process ProcessObject {get;}
+		int ExitCode { get; set; }
+		string Output { get; }
+		Process ProcessObject { get; }
 		string ToString();
 		void WaitForExit();
+		FileReference WriteOutputToFile(string FileName);
 	}
 	
 	/// <summary>
@@ -264,13 +265,13 @@ namespace AutomationTool
 			}
 		}
 
-        ~ProcessResult()
-        {
-            if(Proc != null)
-            {
-                Proc.Dispose();
-            }
-        }
+		~ProcessResult()
+		{
+			if (Proc != null)
+			{
+				Proc.Dispose();
+			}
+		}
 
 		/// <summary>
 		/// Removes a process from the list of tracked processes.
@@ -280,28 +281,28 @@ namespace AutomationTool
 			ProcessManager.RemoveProcess(this);
 		}
 
-        /// <summary>
-        /// Log output of a remote process at a given severity.
-        /// To pretty up the output, we use a custom source so it will say the source of the process instead of this method name.
-        /// </summary>
-        /// <param name="Verbosity"></param>
-        /// <param name="Message"></param>
+		/// <summary>
+		/// Log output of a remote process at a given severity.
+		/// To pretty up the output, we use a custom source so it will say the source of the process instead of this method name.
+		/// </summary>
+		/// <param name="Verbosity"></param>
+		/// <param name="Message"></param>
 		private void LogOutput(LogEventType Verbosity, string Message)
 		{
-            Log.WriteLine(Verbosity, Message);
+			Log.WriteLine(Verbosity, Message);
 		}
-       
+
 		/// <summary>
 		/// Manually dispose of Proc and set it to null.
 		/// </summary>
-        public void DisposeProcess()
-        {
+		public void DisposeProcess()
+		{
 			if(Proc != null)
 			{
 				Proc.Dispose();
 				Proc = null;
 			}
-        }
+		}
 
 		/// <summary>
 		/// Process.OutputDataReceived event handler.
@@ -504,10 +505,10 @@ namespace AutomationTool
 						MaxWaitUntilMessagesReceived--;
 					}
 				}
-                if (!(bStdOutSignalReceived && bStdErrSignalReceived))
-                {
+				if (!(bStdOutSignalReceived && bStdErrSignalReceived))
+				{
 					CommandUtils.LogLog("Waited for a long time for output of {0}, some output may be missing; we gave up.", AppName);
-                }
+				}
 
 				// Double-check if the process terminated
 				lock (ProcSyncObject)
@@ -528,6 +529,16 @@ namespace AutomationTool
 			}
 		}
 
+		public FileReference WriteOutputToFile(string FileName)
+		{
+			using (StreamWriter writer = new StreamWriter(FileName))
+			{
+				writer.Write(ProcessOutput);
+			}
+
+			return new FileReference(FileName);
+		}
+
 		/// <summary>
 		/// Finds child processes of the current process.
 		/// </summary>
@@ -610,15 +621,15 @@ namespace AutomationTool
 			} while (bKilledAChild);
 		}
 
-        /// <summary>
-        /// returns true if this process has any descendants
-        /// </summary>
-        /// <param name="ProcessToCheck">Process to check</param>
-        public static bool HasAnyDescendants(Process ProcessToCheck)
-        {
-            Process[] AllProcs = Process.GetProcesses();
-            foreach (Process KillCandidate in AllProcs)
-            {
+		/// <summary>
+		/// returns true if this process has any descendants
+		/// </summary>
+		/// <param name="ProcessToCheck">Process to check</param>
+		public static bool HasAnyDescendants(Process ProcessToCheck)
+		{
+			Process[] AllProcs = Process.GetProcesses();
+			foreach (Process KillCandidate in AllProcs)
+			{
 				// Silently skip InvalidOperationExceptions here, because it depends on the process still running. It may have terminated.
 				string ProcessName;
 				try
@@ -637,9 +648,9 @@ namespace AutomationTool
 					CommandUtils.LogLog("Descendant pid={0}, name={1}", KillCandidate.Id, ProcessName);
 					return true;
 				}
-            }
-            return false;
-        }
+			}
+			return false;
+		}
 
 		public void StopProcess(bool KillDescendants = true)
 		{
@@ -716,27 +727,27 @@ namespace AutomationTool
 		public enum ERunOptions
 		{
 			None = 0,
-            /// <summary>
+			/// <summary>
 			/// If AllowSpew is set, then the redirected output from StdOut/StdErr is logged.
 			/// Not relevant when NoStdOutRedirect is set.
-            /// </summary>
+			/// </summary>
 			AllowSpew = 1 << 0,
 			AppMustExist = 1 << 1,
 			NoWaitForExit = 1 << 2,
-            /// <summary>
+			/// <summary>
 			/// If NoStdOutRedirect is set, then StdOut/StdErr output is not redirected, logged or captured.
 			/// Else, StdOut/StdErr is redirected and captured by default.
-            /// </summary>
+			/// </summary>
 			NoStdOutRedirect = 1 << 3,
-            NoLoggingOfRunCommand = 1 << 4,
-            UTF8Output = 1 << 5,
+			NoLoggingOfRunCommand = 1 << 4,
+			UTF8Output = 1 << 5,
 
 			/// When specified with AllowSpew, the output will be TraceEventType.Verbose instead of TraceEventType.Information
 			SpewIsVerbose = 1 << 6,
-            /// <summary>
-            /// If NoLoggingOfRunCommand is set, it normally suppresses the run duration output. This turns it back on.
-            /// </summary>
-            LoggingOfRunDuration = 1 << 7,
+			/// <summary>
+			/// If NoLoggingOfRunCommand is set, it normally suppresses the run duration output. This turns it back on.
+			/// </summary>
+			LoggingOfRunDuration = 1 << 7,
 			/// <summary>
 			/// If set, a window is allowed to be created
 			/// </summary>
@@ -838,10 +849,10 @@ namespace AutomationTool
 			var StartTime = DateTime.UtcNow;
 
 			LogEventType SpewVerbosity = Options.HasFlag(ERunOptions.SpewIsVerbose) ? LogEventType.Verbose : LogEventType.Console;
-            if (!Options.HasFlag(ERunOptions.NoLoggingOfRunCommand))
-            {
-                LogWithVerbosity(SpewVerbosity,"Running: " + App + " " + (String.IsNullOrEmpty(CommandLine) ? "" : CommandLine));
-            }
+			if (!Options.HasFlag(ERunOptions.NoLoggingOfRunCommand))
+			{
+				LogWithVerbosity(SpewVerbosity,"Running: " + App + " " + (String.IsNullOrEmpty(CommandLine) ? "" : CommandLine));
+			}
 
 			bool bRedirectStdOut = !Options.HasFlag(ERunOptions.NoStdOutRedirect);
 			bool bAllowSpew = bRedirectStdOut && Options.HasFlag(ERunOptions.AllowSpew);
@@ -925,13 +936,13 @@ namespace AutomationTool
 		/// <param name="LogName">Name of the logfile ( if null, executable name is used )</param>
 		/// <returns>The log file name.</returns>
 		public static string GetRunAndLogOnlyName(CommandEnvironment Env, string App, string LogName = null)
-        {
-            if (LogName == null)
-            {
-                LogName = Path.GetFileNameWithoutExtension(App);
-            }
-            return LogUtils.GetUniqueLogName(CombinePaths(Env.LogFolder, LogName));
-        }
+		{
+			if (LogName == null)
+			{
+				LogName = Path.GetFileNameWithoutExtension(App);
+			}
+			return LogUtils.GetUniqueLogName(CombinePaths(Env.LogFolder, LogName));
+		}
 
 		/// <summary>
 		/// Runs external program and writes the output to a logfile.
@@ -972,32 +983,32 @@ namespace AutomationTool
 		/// <param name="Options">Defines the options how to run. See ERunOptions.</param>
 		/// <param name="FilterCallback">Callback to filter log spew before output.</param>
 		public static string RunAndLog(string App, string CommandLine, string Logfile = null, int MaxSuccessCode = 0, string Input = null, ERunOptions Options = ERunOptions.Default, Dictionary<string, string> EnvVars = null, ProcessResult.SpewFilterCallbackType SpewFilterCallback = null)
-        {
+		{
 			IProcessResult Result = Run(App, CommandLine, Input, Options, EnvVars, SpewFilterCallback);
-            if (!String.IsNullOrEmpty(Result.Output) && Logfile != null)
-            {
-                WriteToFile(Logfile, Result.Output);
-            }
-            else if (Logfile == null)
-            {
-                Logfile = "[No logfile specified]";
-            }
-            else
-            {
-                Logfile = "[None!, no output produced]";
-            }
-
-            if (Result.ExitCode > MaxSuccessCode || Result.ExitCode < 0)
-            {
-                throw new CommandFailedException((ExitCode)Result.ExitCode, String.Format("Command failed (Result:{3}): {0} {1}. See logfile for details: '{2}' ",
-                                                App, CommandLine, Path.GetFileName(Logfile), Result.ExitCode)){ OutputFormat = AutomationExceptionOutputFormat.Minimal };
-            }
-            if (!String.IsNullOrEmpty(Result.Output))
-            {
-                return Result.Output;
-            }
-            return "";
-        }
+			if (!String.IsNullOrEmpty(Result.Output) && Logfile != null)
+			{
+				WriteToFile(Logfile, Result.Output);
+			}
+			else if (Logfile == null)
+			{
+				Logfile = "[No logfile specified]";
+			}
+			else
+			{
+				Logfile = "[None!, no output produced]";
+			}
+
+			if (Result.ExitCode > MaxSuccessCode || Result.ExitCode < 0)
+			{
+				throw new CommandFailedException((ExitCode)Result.ExitCode, String.Format("Command failed (Result:{3}): {0} {1}. See logfile for details: '{2}' ",
+												App, CommandLine, Path.GetFileName(Logfile), Result.ExitCode)){ OutputFormat = AutomationExceptionOutputFormat.Minimal };
+			}
+			if (!String.IsNullOrEmpty(Result.Output))
+			{
+				return Result.Output;
+			}
+			return "";
+		}
 
 		/// <summary>
 		/// Runs external program and writes the output to a logfile.
@@ -1008,19 +1019,19 @@ namespace AutomationTool
 		/// <param name="FilterCallback">Callback to filter log spew before output.</param>
 		/// <returns>Whether the program executed successfully or not.</returns>
 		public static string RunAndLog(string App, string CommandLine, out int SuccessCode, string Logfile = null, Dictionary<string, string> EnvVars = null, ProcessResult.SpewFilterCallbackType SpewFilterCallback = null)
-        {
+		{
 			IProcessResult Result = Run(App, CommandLine, Env: EnvVars, SpewFilterCallback: SpewFilterCallback);
-            SuccessCode = Result.ExitCode;
-            if (Result.Output.Length > 0 && Logfile != null)
-            {
-                WriteToFile(Logfile, Result.Output);
-            }
-            if (!String.IsNullOrEmpty(Result.Output))
-            {
-                return Result.Output;
-            }
-            return "";
-        }
+			SuccessCode = Result.ExitCode;
+			if (Result.Output.Length > 0 && Logfile != null)
+			{
+				WriteToFile(Logfile, Result.Output);
+			}
+			if (!String.IsNullOrEmpty(Result.Output))
+			{
+				return Result.Output;
+			}
+			return "";
+		}
 
 		/// <summary>
 		/// Runs external program and writes the output to a logfile.
@@ -1032,9 +1043,9 @@ namespace AutomationTool
 		/// <param name="FilterCallback">Callback to filter log spew before output.</param>
 		/// <returns>Whether the program executed successfully or not.</returns>
 		public static string RunAndLog(CommandEnvironment Env, string App, string CommandLine, out int SuccessCode, string LogName = null, Dictionary<string, string> EnvVars = null, ProcessResult.SpewFilterCallbackType SpewFilterCallback = null)
-        {
+		{
 			return RunAndLog(App, CommandLine, out SuccessCode, GetRunAndLogOnlyName(Env, App, LogName), EnvVars, SpewFilterCallback);
-        }
+		}
 
 		/// <summary>
 		/// Runs UAT recursively
@@ -1049,11 +1060,11 @@ namespace AutomationTool
 			string LogSubdir = CombinePaths(CmdEnv.LogFolder, DirOnlyName, "");
 			for(int Attempt = 1;;Attempt++)
 			{
-                string[] ExistingFiles = FindFiles(DirOnlyName + "*", false, CmdEnv.LogFolder);
-                if (ExistingFiles.Length == 0)
-                {
-                    break;
-                }
+				string[] ExistingFiles = FindFiles(DirOnlyName + "*", false, CmdEnv.LogFolder);
+				if (ExistingFiles.Length == 0)
+				{
+					break;
+				}
 				if (Attempt == 1000)
 				{
 					throw new AutomationException("Couldn't seem to create a log subdir {0}", LogSubdir);
diff --git a/Engine/Source/Programs/AutomationTool/IOS/IOSPlatform.Automation.cs b/Engine/Source/Programs/AutomationTool/IOS/IOSPlatform.Automation.cs
--- a/Engine/Source/Programs/AutomationTool/IOS/IOSPlatform.Automation.cs
+++ b/Engine/Source/Programs/AutomationTool/IOS/IOSPlatform.Automation.cs
@@ -1,4 +1,4 @@
-﻿// Copyright Epic Games, Inc. All Rights Reserved.
+// Copyright Epic Games, Inc. All Rights Reserved.
 using System;
 using System.Collections.Generic;
 using System.Linq;
@@ -111,6 +111,11 @@ class IOSClientProcess : IProcessResult
 		childProcess.WaitForExit();
 	}
 
+	public FileReference WriteOutputToFile(string FileName)
+	{
+		return childProcess.WriteOutputToFile(FileName);
+	}
+
 	private void StopConsoleOutput()
 	{
 		//processConsoleLogs = false;
