Robomerge - Removed the use of "#disregardexcludedauthors" when "deadend" is specified

diff --git a/Engine/Source/Programs/AutomationTool/Scripts/RebuildHLODCommand.Automation.cs b/Engine/Source/Programs/AutomationTool/Scripts/RebuildHLODCommand.Automation.cs
--- a/Engine/Source/Programs/AutomationTool/Scripts/RebuildHLODCommand.Automation.cs
+++ b/Engine/Source/Programs/AutomationTool/Scripts/RebuildHLODCommand.Automation.cs
@@ -360,6 +360,10 @@ namespace AutomationScripts.Automation
 				{
 					return "#ROBOMERGE[ALL] #DisregardExcludedAuthors";
 				}
+				else if (Command == "deadend")
+				{
+					return "#ROBOMERGE " + Command;
+				}
 				else
 				{
 					return "#ROBOMERGE " + Command + " #DisregardExcludedAuthors";
