Warn rather than error if localization source data is missing when staging

This can be synced out of devs local builds, but shouldn't stop the cook/stage process from working for testing

#rnx

diff --git a/Engine/Source/Programs/AutomationTool/Scripts/CopyBuildToStagingDirectory.Automation.cs b/Engine/Source/Programs/AutomationTool/Scripts/CopyBuildToStagingDirectory.Automation.cs
--- a/Engine/Source/Programs/AutomationTool/Scripts/CopyBuildToStagingDirectory.Automation.cs
+++ b/Engine/Source/Programs/AutomationTool/Scripts/CopyBuildToStagingDirectory.Automation.cs
@@ -454,6 +454,12 @@ namespace AutomationScripts
 
 		public static void StageLocalizationDataForTarget(DeploymentContext SC, List<string> CulturesToStage, DirectoryReference SourceDirectory)
 		{
+			if (!DirectoryReference.Exists(SourceDirectory))
+			{
+				Log.TraceWarning("Failed to stage '{0}' localization target as the directory does not exist! Full path: {1}", SourceDirectory.GetDirectoryName(), SourceDirectory.FullName);
+				return;
+			}
+
 			var PlatformSourceDirectory = new DirectoryReference(CombinePaths(SourceDirectory.FullName, "Platforms", ConfigHierarchy.GetIniPlatformName(SC.StageTargetPlatform.IniPlatformType)));
 			if (!DirectoryReference.Exists(PlatformSourceDirectory))
 			{
