- Fix for bad source data in an editor resource

diff --git a/Engine/Source/Programs/AutomationTool/CookedEditor/CookedEditor.Automation.cs b/Engine/Source/Programs/AutomationTool/CookedEditor/CookedEditor.Automation.cs
--- a/Engine/Source/Programs/AutomationTool/CookedEditor/CookedEditor.Automation.cs
+++ b/Engine/Source/Programs/AutomationTool/CookedEditor/CookedEditor.Automation.cs
@@ -483,6 +483,7 @@ public class MakeCookedEditor : BuildCommand
 		StageIniPathArray(Params, SC, "EngineExtraStageFiles", Unreal.EngineDirectory, Context);
 
 		Context.FilesToUncook.Add(FileReference.Combine(Context.EngineDirectory, "Content", "EngineMaterials", "DefaultMaterial.uasset"));
+		Context.FilesToUncook.Add(FileReference.Combine(Context.EngineDirectory, "Content", "EditorLandscapeResources", "DefaultAlphaTexture.uasset"));
 	}
 
 	protected virtual void StageProjectEditorFiles(ProjectParams Params, DeploymentContext SC, ModifyStageContext Context)
