Adding an option to login to any repository by specifying a json with the user & token

#rnx

diff --git a/Engine/Source/Programs/AutomationTool/BuildGraph/Tasks/DockerPushTask.cs b/Engine/Source/Programs/AutomationTool/BuildGraph/Tasks/DockerPushTask.cs
--- a/Engine/Source/Programs/AutomationTool/BuildGraph/Tasks/DockerPushTask.cs
+++ b/Engine/Source/Programs/AutomationTool/BuildGraph/Tasks/DockerPushTask.cs
@@ -53,6 +53,12 @@ namespace AutomationTool.Tasks
 		/// </summary>
 		[TaskParameter(Optional = true)]
 		public bool AwsEcr;
+
+		/// <summary>
+		/// Path to a json file for authentication to the repository for pushing.
+		/// </summary>
+		[TaskParameter(Optional = true)]
+		public string RepositoryAuthFile;
 	}
 
 	/// <summary>
@@ -94,6 +100,12 @@ namespace AutomationTool.Tasks
 					IProcessResult Result = await SpawnTaskBase.ExecuteAsync("aws", "ecr get-login-password", EnvVars: Environment, LogOutput: false);
 					await ExecuteAsync(Exe, $"login {Parameters.Repository} --username AWS --password-stdin", Input: Result.Output);
 				}
+				if (!String.IsNullOrEmpty(Parameters.RepositoryAuthFile))
+				{
+					string RepositoryText = CommandUtils.ReadAllText(Parameters.RepositoryAuthFile);
+					Dictionary<string, string> AuthDict = JsonSerializer.Deserialize<Dictionary<string, string>>(RepositoryText);
+					await ExecuteAsync(Exe, $"login {Parameters.Repository} --username {AuthDict["Username"]} --password-stdin", Input: AuthDict["Token"]);
+				}
 
 				string TargetImage = Parameters.TargetImage ?? Parameters.Image;
 				await ExecuteAsync(Exe, $"tag {Parameters.Image} {Parameters.Repository}/{TargetImage}", EnvVars: Environment);
