Fixes for OpenEditor BuildCommand to support inheritance, parameter passing to UnrealEditor and specification of different UnrealEditor using fully rooted path.

diff --git a/Engine/Source/Programs/AutomationTool/Scripts/OpenEditor.Automation.cs b/Engine/Source/Programs/AutomationTool/Scripts/OpenEditor.Automation.cs
--- a/Engine/Source/Programs/AutomationTool/Scripts/OpenEditor.Automation.cs
+++ b/Engine/Source/Programs/AutomationTool/Scripts/OpenEditor.Automation.cs
@@ -14,44 +14,63 @@ namespace AutomationTool
 
 	[Help("Opens the specified project.")]
 	[Help("project=<QAGame>", "Project to open. Will search current path and paths in ueprojectdirs. If omitted will open vanilla UnrealEditor")]
-	class OpenEditor: BuildCommand
+	public class OpenEditor: BuildCommand
 	{
 		// exposed as a property so projects can derive and set this directly
 		public string ProjectName { get; set; }
 
+		private string _UnrealEditorApp = "UnrealEditor";
+		public string UnrealEditorApp
+		{
+			get
+			{
+				return _UnrealEditorApp;
+			}
+			set
+			{
+				_UnrealEditorApp = value;
+			}
+		}
+
 		public OpenEditor()
 		{
 		}
 		
 		public override ExitCode Execute()
 		{
-			string EditorPath = HostPlatform.Current.GetUnrealExePath("UnrealEditor");
+			string EditorPath = HostPlatform.Current.GetUnrealExePath(UnrealEditorApp);
 
 			string EditorArgs = "";
 
 			ProjectName = ParseParamValue("project", ProjectName);
+			IEnumerable<string> ParamList = null;
 
-			if (!String.IsNullOrEmpty(ProjectName))
+			if (ParseParam("NoProject"))
 			{
-				FileReference ProjectFile = ProjectUtils.FindProjectFileFromName(ProjectName);
-
-				if (ProjectFile == null)
+				ParamList = Params;
+			}
+			else
+			{
+				if (!String.IsNullOrEmpty(ProjectName))
 				{
-					throw new AutomationException("Unable to find uproject file for {0}", ProjectName);
-				}
+					FileReference ProjectFile = ProjectUtils.FindProjectFileFromName(ProjectName);
 
-				EditorArgs = ProjectFile.FullName;
-			}
+					if (ProjectFile == null)
+					{
+						throw new AutomationException("Unable to find uproject file for {0}", ProjectName);
+					}
 
-			// filter out any -project argument since we want it to be the first un-prefixed argument to the editor 
-			IEnumerable<string> ParamList = this.Params
-												.Where(P => P.StartsWith("project=", StringComparison.OrdinalIgnoreCase) == false);
+					EditorArgs = ProjectFile.FullName;
+				}
 
+				// filter out any -project argument since we want it to be the first un-prefixed argument to the editor 
+				ParamList = this.Params
+								.Where(P => P.StartsWith("project=", StringComparison.OrdinalIgnoreCase) == false);
+			}
 
 			ParamList = new[] { EditorArgs }.Concat(ParamList);
-	
 
-			bool bLaunched = RunUntrackedProcess(EditorPath, string.Join(" ", ParamList));
+			bool bLaunched = RunUntrackedProcess(EditorPath, string.Join(" -", ParamList));
 
 			return bLaunched ? ExitCode.Success : ExitCode.Error_UATLaunchFailure;
 		}
