[Automation] Add missing logic to handle the new -perfmodel= commandline to gauntlet.

#tests tested running a few tests with perfmodel and perfspec to make sure both functionalities worked
#rnx

diff --git a/Engine/Source/Programs/AutomationTool/AutomationUtils/DeviceReservation.cs b/Engine/Source/Programs/AutomationTool/AutomationUtils/DeviceReservation.cs
--- a/Engine/Source/Programs/AutomationTool/AutomationUtils/DeviceReservation.cs
+++ b/Engine/Source/Programs/AutomationTool/AutomationUtils/DeviceReservation.cs
@@ -407,6 +407,7 @@ namespace AutomationTool.DeviceReservation
 		public string Type { get; set; }
 		public string IPOrHostName { get; set; }
 		public string PerfSpec { get; set; }
+		public string Model { get; set; }
 		public TimeSpan AvailableStartTime { get; set; }
 		public TimeSpan AvailableEndTime { get; set; }
 		public bool Enabled { get; set; }
diff --git a/Engine/Source/Programs/AutomationTool/Gauntlet/Framework/Devices/Gauntlet.DevicePool.cs b/Engine/Source/Programs/AutomationTool/Gauntlet/Framework/Devices/Gauntlet.DevicePool.cs
--- a/Engine/Source/Programs/AutomationTool/Gauntlet/Framework/Devices/Gauntlet.DevicePool.cs
+++ b/Engine/Source/Programs/AutomationTool/Gauntlet/Framework/Devices/Gauntlet.DevicePool.cs
@@ -114,9 +114,22 @@ namespace Gauntlet
 				throw new AutomationException("Comparing null target constraint");
 			}
 
-			if (ReferenceEquals(this, Other)) return true;
+			if (ReferenceEquals(this, Other))
+			{
+				return true;
+			}
+
+			if (Other.Platform != Platform)
+			{
+				return false;
+			}
+
+			if (Other.Model.Equals(Model, StringComparison.InvariantCultureIgnoreCase))
+			{
+				return true;
+			}
 
-			return Other.Platform == Platform && Other.PerfSpec == PerfSpec && Other.Model.Equals(Model, StringComparison.InvariantCultureIgnoreCase);
+			return Other.PerfSpec == PerfSpec;
 		}
 
 		public override bool Equals(object Obj)
@@ -153,7 +166,7 @@ namespace Gauntlet
 				return string.Format("{0}", Platform);
 			}
 
-			return string.Format("{0}:{1}", Platform, PerfSpec == EPerfSpec.Unspecified ? Model : PerfSpec.ToString());
+			return string.Format("{0}:{1}", Platform, Model == string.Empty ? PerfSpec.ToString() : Model);
 		}
 
 		public override int GetHashCode()
@@ -685,7 +698,7 @@ namespace Gauntlet
 					Def.Name = Device.Name;
 					Def.Platform = DeviceMap.FirstOrDefault(Entry => Entry.Value == Device.Type.Replace("-DevKit", "", StringComparison.OrdinalIgnoreCase)).Key;
 					Def.DeviceData = Device.DeviceData;
-					Def.Model = string.Empty;
+					Def.Model = Device.Model;
 
 					EPerfSpec Out = EPerfSpec.Unspecified;
 					if (!String.IsNullOrEmpty(Device.PerfSpec) && !Enum.TryParse<EPerfSpec>(Device.PerfSpec, true, out Out))
