CookedEditor - Minor log change
#rnx

diff --git a/Engine/Source/Programs/AutomationTool/CookedEditor/CookedEditor.Automation.cs b/Engine/Source/Programs/AutomationTool/CookedEditor/CookedEditor.Automation.cs
--- a/Engine/Source/Programs/AutomationTool/CookedEditor/CookedEditor.Automation.cs
+++ b/Engine/Source/Programs/AutomationTool/CookedEditor/CookedEditor.Automation.cs
@@ -120,7 +120,7 @@ public class ModifyStageContext
 		IniPlatformName = ConfigHierarchy.GetIniPlatformName(SC.StageTargetPlatform.IniPlatformType);
 		bIsDLC = Params.DLCFile != null && SC.MetadataDir != null; // MetadataDir needs to be set for DLC
 
-		Console.WriteLine("---> releaeOverrideDir = {0}, MetadataDrir = {1}", Params.BasedOnReleaseVersionPathOverride, SC.MetadataDir);
+		CommandUtils.LogInformation("---> ReleaseOverrideDir = {0}, MetadataDir = {1}", Params.BasedOnReleaseVersionPathOverride, SC.MetadataDir);
 
 
 		// cache info for DLC against a release
