Add '-rehydrate' to .umap/.uasset files when creating a PakResponse file, if the '-rehydrate' command line was used.

#rnx

diff --git a/Engine/Source/Programs/AutomationTool/AutomationUtils/ProjectParams.cs b/Engine/Source/Programs/AutomationTool/AutomationUtils/ProjectParams.cs
--- a/Engine/Source/Programs/AutomationTool/AutomationUtils/ProjectParams.cs
+++ b/Engine/Source/Programs/AutomationTool/AutomationUtils/ProjectParams.cs
@@ -326,6 +326,7 @@ namespace AutomationTool
 			this.SignPak = InParams.SignPak;
 			this.SignedPak = InParams.SignedPak;
 			this.PakAlignForMemoryMapping = InParams.PakAlignForMemoryMapping;
+			this.RehydrateAssets = InParams.RehydrateAssets;
 			this.SkipPak = InParams.SkipPak;
             this.PrePak = InParams.PrePak;
             this.NoXGE = InParams.NoXGE;
@@ -529,6 +530,7 @@ namespace AutomationTool
 			bool? NoBootstrapExe = null,
             bool? SignedPak = null,
 			bool? PakAlignForMemoryMapping = null,
+			bool? RehydrateAssets = null,
 			bool? NullRHI = null,
             bool? FakeClient = null,
             bool? EditorTest = null,
@@ -712,6 +714,7 @@ namespace AutomationTool
 				}
 			}
 			this.PakAlignForMemoryMapping = GetParamValueIfNotSpecified(Command, PakAlignForMemoryMapping, this.PakAlignForMemoryMapping, "PakAlignForMemoryMapping");
+			this.RehydrateAssets = GetParamValueIfNotSpecified(Command, RehydrateAssets, this.RehydrateAssets, "RehydrateAssets");		
 			this.Pak = GetParamValueIfNotSpecified(Command, Pak, this.Pak, "pak");
 			this.IgnorePaksFromDifferentCookSource = GetParamValueIfNotSpecified(Command, IgnorePaksFromDifferentCookSource, this.IgnorePaksFromDifferentCookSource, "IgnorePaksFromDifferentCookSource");
 			this.IoStore = GetParamValueIfNotSpecified(Command, IoStore, this.IoStore, "iostore");
@@ -1425,6 +1428,12 @@ namespace AutomationTool
 		/// </summary>
 		[Help("PakAlignForMemoryMapping", "The game will be set up for memory mapping bulk data.")]
 		public bool PakAlignForMemoryMapping { private set; get; }
+		
+		/// <summary>
+		/// Shared: true if we want to rehydrate virtualized assets when staging.
+		/// </summary>
+		[Help("rehydrateassets", "Should virtualized assets be rehydrated?")]
+		public bool RehydrateAssets { get; set; }
 
 		/// <summary>
 		/// Shared: true if this build is staged, command line: -stage
@@ -3148,6 +3157,7 @@ namespace AutomationTool
                 CommandUtils.LogLog("PrePak={0}", PrePak);
                 CommandUtils.LogLog("SkipStage={0}", SkipStage);
 				CommandUtils.LogLog("Stage={0}", Stage);
+				CommandUtils.LogLog("RehydrateAssets={0}", RehydrateAssets);
 				CommandUtils.LogLog("bTreatNonShippingBinariesAsDebugFiles={0}", bTreatNonShippingBinariesAsDebugFiles);
 				CommandUtils.LogLog("bUseExtraFlavor={0}", bUseExtraFlavor);
                 CommandUtils.LogLog("StageDirectoryParam={0}", StageDirectoryParam);
diff --git a/Engine/Source/Programs/AutomationTool/Scripts/CopyBuildToStagingDirectory.Automation.cs b/Engine/Source/Programs/AutomationTool/Scripts/CopyBuildToStagingDirectory.Automation.cs
--- a/Engine/Source/Programs/AutomationTool/Scripts/CopyBuildToStagingDirectory.Automation.cs
+++ b/Engine/Source/Programs/AutomationTool/Scripts/CopyBuildToStagingDirectory.Automation.cs
@@ -140,25 +140,33 @@ namespace AutomationScripts
 		/// </summary>
 		/// <param name="Filename"></param>
 		/// <param name="ResponseFile"></param>
-		private static void WritePakResponseFile(string Filename, Dictionary<string, string> ResponseFile, bool Compressed, EncryptionAndSigning.CryptoSettings CryptoSettings, bool bForceFullEncryption)
+		private static void WritePakResponseFile(string Filename, Dictionary<string, string> ResponseFile, bool Compressed, bool RehydrateAssets, EncryptionAndSigning.CryptoSettings CryptoSettings, bool bForceFullEncryption)
 		{
 			using (var Writer = new StreamWriter(Filename, false, new System.Text.UTF8Encoding(true)))
 			{
 				foreach (var Entry in ResponseFile)
 				{
+					string Extension = Path.GetExtension(Entry.Key);
 					string Line = String.Format("\"{0}\" \"{1}\"", Entry.Key, Entry.Value);
 
 					// explicitly exclude some file types from compression
-					if (Compressed && !Path.GetExtension(Entry.Key).Contains(".mp4") && !Path.GetExtension(Entry.Key).Contains("ushaderbytecode") && !Path.GetExtension(Entry.Key).Contains("upipelinecache"))
+					if (Compressed && !Path.GetExtension(Entry.Key).Contains(".mp4") && !Extension.Contains("ushaderbytecode") && !Path.GetExtension(Entry.Key).Contains("upipelinecache"))
 					{
 						Line += " -compress";
 					}
 
+					// todo: Ideally we would know if the package is virtualized and only opt to rehydrate those packages, but we'd need to be able
+					// to pipe that info this far.
+					if(RehydrateAssets && (Extension.Contains(".uasset") || Extension.Contains(".umap")))
+					{
+						Line += " -rehydrate";
+					}
+
 					if (CryptoSettings != null)
 					{
 						bool bEncryptFile = bForceFullEncryption || CryptoSettings.bEnablePakFullAssetEncryption;
-						bEncryptFile = bEncryptFile || (CryptoSettings.bEnablePakUAssetEncryption && Path.GetExtension(Entry.Key).Contains(".uasset"));
-						bEncryptFile = bEncryptFile || (CryptoSettings.bEnablePakIniEncryption && Path.GetExtension(Entry.Key).Contains(".ini"));
+						bEncryptFile = bEncryptFile || (CryptoSettings.bEnablePakUAssetEncryption && Extension.Contains(".uasset"));
+						bEncryptFile = bEncryptFile || (CryptoSettings.bEnablePakIniEncryption && Extension.Contains(".ini"));
 
 						if (bEncryptFile)
 						{
@@ -222,7 +230,7 @@ namespace AutomationScripts
 			return CmdLine.ToString();
 		}
 
-		static private string GetPakFileSpecificUnrealPakArguments(Dictionary<string, string> UnrealPakResponseFile, FileReference OutputLocation, string AdditionalOptions, bool Compressed, EncryptionAndSigning.CryptoSettings CryptoSettings, String PatchSourceContentPath, string EncryptionKeyGuid)
+		static private string GetPakFileSpecificUnrealPakArguments(Dictionary<string, string> UnrealPakResponseFile, FileReference OutputLocation, string AdditionalOptions, bool Compressed, bool RehydrateAssets,  EncryptionAndSigning.CryptoSettings CryptoSettings, String PatchSourceContentPath, string EncryptionKeyGuid)
 		{
 			StringBuilder CmdLine = new StringBuilder(MakePathSafeToUseWithCommandLine(OutputLocation.FullName));
 
@@ -231,7 +239,7 @@ namespace AutomationScripts
 			bool bForceEncryption = !string.IsNullOrEmpty(EncryptionKeyGuid);
 			string PakName = Path.GetFileNameWithoutExtension(OutputLocation.FullName);
 			string UnrealPakResponseFileName = CombinePaths(CmdEnv.LogFolder, "PakList_" + PakName + ".txt");
-			WritePakResponseFile(UnrealPakResponseFileName, UnrealPakResponseFile, Compressed, CryptoSettings, bForceEncryption);
+			WritePakResponseFile(UnrealPakResponseFileName, UnrealPakResponseFile, Compressed, RehydrateAssets, CryptoSettings, bForceEncryption);
 			CmdLine.AppendFormat(" -create={0}", CommandUtils.MakePathSafeToUseWithCommandLine(UnrealPakResponseFileName));
 
 			if (!String.IsNullOrEmpty(PatchSourceContentPath))
@@ -251,9 +259,11 @@ namespace AutomationScripts
 
 		static public string GetUnrealPakArguments(FileReference ProjectPath, Dictionary<string, string> UnrealPakResponseFile, FileReference OutputLocation, List<OrderFile> PakOrderFileLocations, string AdditionalOptions, bool Compressed, EncryptionAndSigning.CryptoSettings CryptoSettings, FileReference CryptoKeysCacheFilename, String PatchSourceContentPath, string EncryptionKeyGuid, List<OrderFile> SecondaryPakOrderFileLocations, bool bUnattended)
 		{
+			bool RehydrateAssets = false;
+
 			StringBuilder CmdLine = new StringBuilder(MakePathSafeToUseWithCommandLine(ProjectPath.FullName));
 			CmdLine.Append(" ");
-			CmdLine.Append(GetPakFileSpecificUnrealPakArguments(UnrealPakResponseFile, OutputLocation, "", Compressed, CryptoSettings, PatchSourceContentPath, EncryptionKeyGuid));
+			CmdLine.Append(GetPakFileSpecificUnrealPakArguments(UnrealPakResponseFile, OutputLocation, "", Compressed, RehydrateAssets, CryptoSettings, PatchSourceContentPath, EncryptionKeyGuid));
 			CmdLine.Append(" ");
 			CmdLine.Append(GetCommonUnrealPakArguments(PakOrderFileLocations, AdditionalOptions, CryptoSettings, CryptoKeysCacheFilename, SecondaryPakOrderFileLocations, bUnattended));
 			return CmdLine.ToString();
@@ -290,8 +300,9 @@ namespace AutomationScripts
 			// Force encryption of ALL files if we're using specific encryption key. This should be made an option per encryption key in the settings, but for our initial
 			// implementation we will just assume that we require maximum security for this data.
 			bool bForceEncryption = !string.IsNullOrEmpty(EncryptionKeyGuid);
+			bool RehydrateAssets = false;
 			string UnrealPakResponseFileName = CombinePaths(CmdEnv.LogFolder, "PakListIoStore_" + ContainerName + ".txt");
-			WritePakResponseFile(UnrealPakResponseFileName, UnrealPakResponseFile, bCompressed, CryptoSettings, bForceEncryption);
+			WritePakResponseFile(UnrealPakResponseFileName, UnrealPakResponseFile, bCompressed, RehydrateAssets, CryptoSettings, bForceEncryption);
 			CmdLine.AppendFormat(" -ResponseFile={0}", CommandUtils.MakePathSafeToUseWithCommandLine(UnrealPakResponseFileName));
 
 			if (CryptoSettings != null && CryptoSettings.bDataCryptoRequired)
@@ -2243,7 +2254,7 @@ namespace AutomationScripts
 			PakCryptoSettings.Save(CryptoKeysCacheFilename);
 
 			List<CreatePakParams> PakInputs = new List<CreatePakParams>();
-			PakInputs.Add(new CreatePakParams(SC.ShortProjectName, UnrealPakResponseFile, Params.Compressed, null));
+			PakInputs.Add(new CreatePakParams(SC.ShortProjectName, UnrealPakResponseFile, Params.Compressed, Params.RehydrateAssets, null));
 			CreatePaks(Params, SC, PakInputs, PakCryptoSettings, CryptoKeysCacheFilename);
 		}
 
@@ -2378,6 +2389,11 @@ namespace AutomationScripts
 			/// </summary>
 			public bool bCompressed;
 
+			/// <summary>
+			/// Whether to rehydrate the assets when creating the pak file or not
+			/// </summary>
+			public bool bRehydrateAssets;
+
 			/// <summary>
 			/// GUID of the encryption key for this pak file
 			/// </summary>
@@ -2389,11 +2405,12 @@ namespace AutomationScripts
 			/// <param name="PakName">Path to the base output file for this pak file</param>
 			/// <param name="UnrealPakResponseFile">Map of files within the pak file to their source file on disk</param>
 			/// <param name="bCompressed">Whether to enable compression</param>
-			public CreatePakParams(string PakName, Dictionary<string, string> UnrealPakResponseFile, bool bCompressed, string EncryptionKeyGuid)
+			public CreatePakParams(string PakName, Dictionary<string, string> UnrealPakResponseFile, bool bCompressed, bool RehydrateAssets, string EncryptionKeyGuid)
 			{
 				this.PakName = PakName;
 				this.UnrealPakResponseFile = UnrealPakResponseFile;
 				this.bCompressed = bCompressed;
+				this.bRehydrateAssets = RehydrateAssets;
 				this.EncryptionKeyGuid = EncryptionKeyGuid;
 			}
 		}
@@ -3083,6 +3100,7 @@ namespace AutomationScripts
 							OutputLocation,
 							PakFileSpecificAdditionalArgs,
 							PakParams.bCompressed,
+							PakParams.bRehydrateAssets,
 							Params.SkipEncryption ? null : CryptoSettings,
 							PatchSourceContentPath,
 							Params.SkipEncryption ? "" : PakParams.EncryptionKeyGuid));
@@ -3865,7 +3883,7 @@ namespace AutomationScripts
 					string EncryptionKeyToUse = Params.SkipEncryption ? "" : Chunk.EncryptionKeyGuid;
 					PakInputs.Add(new CreatePakParams(Chunk.ChunkName,
 						Chunk.ResponseFile.ToDictionary(entry => entry.Key, entry => entry.Value),
-						Params.Compressed || Chunk.bCompressed, EncryptionKeyToUse));
+						Params.Compressed || Chunk.bCompressed, Params.RehydrateAssets, EncryptionKeyToUse));
 				}
 			}
 
