Add missing -test= to how to run locally text

diff --git a/Engine/Source/Programs/AutomationTool/Gauntlet/Unreal/Base/Gauntlet.UnrealTestNode.cs b/Engine/Source/Programs/AutomationTool/Gauntlet/Unreal/Base/Gauntlet.UnrealTestNode.cs
--- a/Engine/Source/Programs/AutomationTool/Gauntlet/Unreal/Base/Gauntlet.UnrealTestNode.cs
+++ b/Engine/Source/Programs/AutomationTool/Gauntlet/Unreal/Base/Gauntlet.UnrealTestNode.cs
@@ -909,7 +909,7 @@ namespace Gauntlet
 
 			string StringOfCleanedArguments = string.Format("-{0}", string.Join(" -", CleanedArgs));
 			string CommandToRunLocally =
-				string.Format("RunUAT {0} {1} {2}", LaunchingBuildCommand, GetType(), StringOfCleanedArguments);
+				string.Format("RunUAT {0} -Test={1} {2}", LaunchingBuildCommand, GetType(), StringOfCleanedArguments);
 			return CommandToRunLocally;
 		}
 
