UAT - added comand line option -NoZenAutoLaunch

#rnx

diff --git a/Engine/Source/Programs/AutomationTool/AutomationUtils/ProjectParams.cs b/Engine/Source/Programs/AutomationTool/AutomationUtils/ProjectParams.cs
--- a/Engine/Source/Programs/AutomationTool/AutomationUtils/ProjectParams.cs
+++ b/Engine/Source/Programs/AutomationTool/AutomationUtils/ProjectParams.cs
@@ -731,6 +731,10 @@ namespace AutomationTool
 				this.AdditionalCookerOptions += " -ZenStore";
 			}
 			this.NoZenAutoLaunch = ParseParamValueIfNotSpecified(Command, NoZenAutoLaunch, "NoZenAutoLaunch", String.Empty);
+			if (string.IsNullOrEmpty(this.NoZenAutoLaunch) && GetParamValueIfNotSpecified(Command, null, false, "NoZenAutoLaunch"))
+			{
+				this.NoZenAutoLaunch = "127.0.0.1";
+			}
 			if (!string.IsNullOrEmpty(this.NoZenAutoLaunch) && this.Cook && !this.SkipCook)
 			{
 				this.AdditionalCookerOptions += string.Format(" -NoZenAutoLaunch={0}", this.NoZenAutoLaunch);
