- Extended the "UnstagedRuntimeDependencies" concept for pulling only minimal files needed for platform (console mainly) cooking into a directory at stage time, called "CookerSupportFiles"
- It needs to have CookerSupportFilesSubdirectory passed in to BuildCookRun to activate (giving it the name of a subdirectory under the Staged directory to put the SDK files)
- it also writes a .bat file that can set the envvars needed for the minimal SDKs to be found by Unreal

diff --git a/Engine/Source/Programs/AutomationTool/AutomationUtils/DeploymentContext.cs b/Engine/Source/Programs/AutomationTool/AutomationUtils/DeploymentContext.cs
--- a/Engine/Source/Programs/AutomationTool/AutomationUtils/DeploymentContext.cs
+++ b/Engine/Source/Programs/AutomationTool/AutomationUtils/DeploymentContext.cs
@@ -206,6 +206,11 @@ public class DeploymentContext //: ProjectParams
 	/// </summary>
 	public DirectoryReference OptionalFileInputDirectory = null;
 
+	/// <summary>
+	/// If this is specified, any files written into the receipt with the CookerSupportFiles tag will be copied into here during staging
+	/// </summary>
+	public string CookerSupportFilesSubdirectory = null;
+
 	/// <summary>
 	/// Directory name for staged projects
 	/// </summary>
@@ -394,6 +399,7 @@ public class DeploymentContext //: ProjectParams
 		DirectoryReference OptionalFileStageDirectory,
 		DirectoryReference OptionalFileInputDirectory,
 		DirectoryReference BaseArchiveDirectory,
+		string CookerSupportFilesSubdirectory,
 		Platform InSourcePlatform,
         Platform InTargetPlatform,
 		List<UnrealTargetConfiguration> InTargetConfigurations,
@@ -462,6 +468,7 @@ public class DeploymentContext //: ProjectParams
 		}
 		this.OptionalFileStageDirectory = OptionalFileStageDirectory;
 		this.OptionalFileInputDirectory = OptionalFileInputDirectory;
+		this.CookerSupportFilesSubdirectory = CookerSupportFilesSubdirectory;
 
 		if (BaseArchiveDirectory != null)
 		{
diff --git a/Engine/Source/Programs/AutomationTool/AutomationUtils/ProjectParams.cs b/Engine/Source/Programs/AutomationTool/AutomationUtils/ProjectParams.cs
--- a/Engine/Source/Programs/AutomationTool/AutomationUtils/ProjectParams.cs
+++ b/Engine/Source/Programs/AutomationTool/AutomationUtils/ProjectParams.cs
@@ -814,6 +814,7 @@ namespace AutomationTool
 			this.StageDirectoryParam = ParseParamValueIfNotSpecified(Command, StageDirectoryParam, "stagingdirectory", String.Empty, true);
 			this.OptionalFileStagingDirectory = ParseParamValueIfNotSpecified(Command, OptionalFileStagingDirectory, "optionalfilestagingdirectory", String.Empty, true);
 			this.OptionalFileInputDirectory = ParseParamValueIfNotSpecified(Command, OptionalFileInputDirectory, "optionalfileinputdirectory", String.Empty, true);
+			this.CookerSupportFilesSubdirectory = ParseParamValueIfNotSpecified(Command, CookerSupportFilesSubdirectory, "CookerSupportFilesSubdirectory", String.Empty, true);
 			this.bCodeSign = GetOptionalParamValueIfNotSpecified(Command, CodeSign, CommandUtils.IsBuildMachine, "CodeSign", "NoCodeSign");
 			this.bTreatNonShippingBinariesAsDebugFiles = GetParamValueIfNotSpecified(Command, TreatNonShippingBinariesAsDebugFiles, false, "TreatNonShippingBinariesAsDebugFiles");
 			this.bUseExtraFlavor = GetParamValueIfNotSpecified(Command, UseExtraFlavor, false, "UseExtraFlavor");
@@ -1522,6 +1523,9 @@ namespace AutomationTool
 		[Help("optionalfileinputdirectory=Path", "Directory to read the optional files from, i.e. -optionalfileinputdirectory=C:\\StageOptional")]
 		public string OptionalFileInputDirectory;
 
+		[Help("CookerSupportFilesSubdirectory=subdir", "Subdirectory under staging to copy CookerSupportFiles (as set in Build.cs files). -CookerSupportFilesSubdirectory=SDK")]
+		public string CookerSupportFilesSubdirectory;
+		
 		[Help("unrealexe=ExecutableName", "Name of the Unreal Editor executable, i.e. -unrealexe=UnrealEditor.exe")]
 		public string UnrealExe;
 
diff --git a/Engine/Source/Programs/AutomationTool/CookedEditor/CookedEditor.Automation.cs b/Engine/Source/Programs/AutomationTool/CookedEditor/CookedEditor.Automation.cs
--- a/Engine/Source/Programs/AutomationTool/CookedEditor/CookedEditor.Automation.cs
+++ b/Engine/Source/Programs/AutomationTool/CookedEditor/CookedEditor.Automation.cs
@@ -482,8 +482,8 @@ public class MakeCookedEditor : BuildCommand
 
 		StageIniPathArray(Params, SC, "EngineExtraStageFiles", Unreal.EngineDirectory, Context);
 
-		Context.FilesToUncook.Add(FileReference.Combine(Context.EngineDirectory, "Content", "EngineMaterials", "DefaultMaterial.uasset"));
-		Context.FilesToUncook.Add(FileReference.Combine(Context.EngineDirectory, "Content", "EditorLandscapeResources", "DefaultAlphaTexture.uasset"));
+		Context.FilesToUncook.Add(FileReference.Combine(Context.EngineDirectory, "Content/EngineMaterials/DefaultMaterial.uasset"));
+		Context.FilesToUncook.Add(FileReference.Combine(Context.EngineDirectory, "Content/EditorLandscapeResources/DefaultAlphaTexture.uasset"));
 	}
 
 	protected virtual void StageProjectEditorFiles(ProjectParams Params, DeploymentContext SC, ModifyStageContext Context)
@@ -1113,6 +1113,10 @@ public class MakeCookedEditor : BuildCommand
 		Params.PreModifyDeploymentContextCallback = new Action<ProjectParams, DeploymentContext>((ProjectParams P, DeploymentContext SC) => { PreModifyDeploymentContext(P, SC); });
 		Params.ModifyDeploymentContextCallback = new Action<ProjectParams, DeploymentContext>((ProjectParams P, DeploymentContext SC) => { ModifyDeploymentContext(P, SC); });
 
+		// this will make all of the files that are specified in BUild.cs files with "AdditionalPropertiesForReceipt.Add("CookerSupportFiles", ...);"  be copied into this
+		// subdirectory along with a batch file that can be used to set platform SDK environment variables during cooking
+		Params.CookerSupportFilesSubdirectory = "SDK";
+
 		ModifyParams(Params);
 
 		return Params;
diff --git a/Engine/Source/Programs/AutomationTool/Scripts/CopyBuildToStagingDirectory.Automation.cs b/Engine/Source/Programs/AutomationTool/Scripts/CopyBuildToStagingDirectory.Automation.cs
--- a/Engine/Source/Programs/AutomationTool/Scripts/CopyBuildToStagingDirectory.Automation.cs
+++ b/Engine/Source/Programs/AutomationTool/Scripts/CopyBuildToStagingDirectory.Automation.cs
@@ -547,6 +547,77 @@ namespace AutomationScripts
 			}
 		}
 
+		private static void StageCookerSupportFilesFromReceipt(DeploymentContext SC)
+		{
+			if (string.IsNullOrEmpty(SC.CookerSupportFilesSubdirectory))
+			{
+				return;
+			}
+
+			// sdk components expected to be in the form: $(SomeVar)/sub/path/to/filespec
+			// so we capture the SomeVar and sub/path/to/filespec
+			// note that filespec could be a file or a wildcard
+			Regex Var = new Regex(@"^\$\((.*)\)/(.*)$");
+
+			HashSet<string> EnvVars = new();
+			foreach (StageTarget Target in SC.StageTargets)
+			{
+				TargetReceipt Receipt = Target.Receipt;
+				foreach (ReceiptProperty Prop in Receipt.AdditionalProperties)
+				{
+					if (Prop.Name.Equals("CookerSupportFiles", StringComparison.InvariantCultureIgnoreCase))
+					{
+						// Value is a path under a envvar: $(MySdk)/Tools/Converter.exe
+						Match Match = Var.Match(Prop.Value);
+						if (Match.Success)
+						{
+							string EnvVar = Match.Groups[1].Value;
+							string SubPath = Match.Groups[2].Value.Replace("\\", "/");
+							string EnvVarValue = Environment.GetEnvironmentVariable(EnvVar);
+
+							if (EnvVarValue == null)
+							{
+								CommandUtils.LogWarning($"Using envvar that isn't set, so it can't be used to pull files from : '{EnvVar}'");
+								continue;
+							}
+							if (!Directory.Exists(EnvVarValue))
+							{
+								CommandUtils.LogWarning($"Using envvar path that doesn't exist, so it can't be used to pull files from : '{EnvVarValue}' (from env var {EnvVar})");
+								continue;
+							}
+
+							FileReference ExpandedPath = FileReference.Combine(new DirectoryReference(EnvVarValue), SubPath);
+							foreach (FileReference File in DirectoryReference.EnumerateFiles(ExpandedPath.Directory, ExpandedPath.GetFileName()))
+							{
+								string ExpandedSubPath = File.MakeRelativeTo(new DirectoryReference(EnvVarValue));
+								// insert the envvar into the path so that differnt sdks/platforms won't stomp on each other (since we don't know the platform, we can't use that)
+								ExpandedSubPath = $"{SC.CookerSupportFilesSubdirectory}/{EnvVar}/{ExpandedSubPath}";
+
+								StagedFileReference DestPath = new StagedFileReference(ExpandedSubPath);
+								CommandUtils.LogInformation($"  Staging support file '{ExpandedPath}' to '{DestPath}'");
+								SC.StageFile(StagedFileType.SystemNonUFS, File, DestPath);
+							}
+
+							EnvVars.Add(EnvVar);
+						}
+					}
+				}
+			}
+
+
+			// write a temp .bat file that will set the envvar 
+			string TempBatch = Path.Combine(Path.GetTempPath(), "AutomationTool", Path.ChangeExtension(Path.GetRandomFileName(), "bat"));
+			Directory.CreateDirectory(Path.GetDirectoryName(TempBatch));
+			// add the Envvar onto the path, like we did above
+			List<string> Lines = EnvVars.Select(x => $"set {x}=%~dp0{x}").ToList();
+			// make sure the AutoSDK var is unset, as it will mess with testing (probably not needed for final dist but can't hurt if everything else works)
+			Lines.Insert(0, "set UE_SDKS_ROOT=");
+			File.WriteAllLines(TempBatch, Lines);
+
+			// and stage it
+			SC.StageFile(StagedFileType.SystemNonUFS, new FileReference(TempBatch), new StagedFileReference($"{SC.CookerSupportFilesSubdirectory}/ActivateSupportFiles.bat"));
+		}
+
 		private static void StageAdditionalDirectoriesFromConfig(DeploymentContext SC, DirectoryReference ProjectContentRoot, StagedDirectoryReference StageContentRoot, ConfigHierarchy PlatformGameConfig, bool bUFS, string ConfigKeyName)
 		{
 			List<string> ExtraDirs;
@@ -1337,6 +1408,8 @@ namespace AutomationScripts
 					}
 					SC.MetadataDir = DirectoryReference.Combine(CookOutputDir, SC.ShortProjectName, "Metadata");
 				}
+
+				StageCookerSupportFilesFromReceipt(SC);
 			}
 
 			// Allow the calling scripts to make modifications to the deployment context before we finalize it
@@ -4526,6 +4599,7 @@ namespace AutomationScripts
 					String.IsNullOrEmpty(Params.OptionalFileStagingDirectory) ? null : new DirectoryReference(Params.OptionalFileStagingDirectory),
 					String.IsNullOrEmpty(Params.OptionalFileInputDirectory) ? null : new DirectoryReference(Params.OptionalFileInputDirectory),
 					String.IsNullOrEmpty(ArchiveDirectory) ? null : new DirectoryReference(ArchiveDirectory),
+					Params.CookerSupportFilesSubdirectory,
 					Platform.Platforms[CookedDataPlatform],
 					Platform.Platforms[StagePlatform],
 					ConfigsToProcess,
