- Made ApplePlatform.Automation to share code betwen Mac and IOS

diff --git a/Engine/Source/Programs/AutomationTool/Apple/Apple.Automation.csproj b/Engine/Source/Programs/AutomationTool/Apple/Apple.Automation.csproj
new file mode 100644
--- /dev/null
+++ b/Engine/Source/Programs/AutomationTool/Apple/Apple.Automation.csproj
@@ -0,0 +1,45 @@
+<Project Sdk="Microsoft.NET.Sdk">
+  <Import Project="..\..\Shared\UnrealEngine.csproj.props" />
+  
+  <PropertyGroup>
+    <TargetFramework>net6.0</TargetFramework>
+    <Configuration Condition=" '$(Configuration)' == '' ">Development</Configuration>
+    <OutputType>Library</OutputType>
+    <AppendTargetFrameworkToOutputPath>false</AppendTargetFrameworkToOutputPath>
+    <GenerateAssemblyInfo>false</GenerateAssemblyInfo>
+    <GenerateTargetFrameworkAttribute>false</GenerateTargetFrameworkAttribute>
+    <Configurations>Debug;Release;Development</Configurations>
+    <RootNamespace>Apple.Automation</RootNamespace>
+    <AssemblyName>Apple.Automation</AssemblyName>
+    <WarningsNotAsErrors>612,618</WarningsNotAsErrors>
+    <OutputPath>..\..\..\..\Binaries\DotNET\AutomationTool\AutomationScripts\Platforms\Apple</OutputPath>
+    <TreatWarningsAsErrors>true</TreatWarningsAsErrors>
+    <DebugType>pdbonly</DebugType>
+  </PropertyGroup>
+  <PropertyGroup Condition=" '$(Configuration)|$(Platform)' == 'Development|AnyCPU' ">
+    <DefineConstants>$(DefineConstants);TRACE</DefineConstants>
+    <Optimize>true</Optimize>
+  </PropertyGroup>
+  <PropertyGroup Condition=" '$(Configuration)|$(Platform)' == 'Release|AnyCPU' ">
+    <DefineConstants>$(DefineConstants);TRACE</DefineConstants>
+    <Optimize>true</Optimize>
+  </PropertyGroup>
+  <PropertyGroup Condition=" '$(Configuration)|$(Platform)' == 'Debug|AnyCPU' ">
+    <DefineConstants>$(DefineConstants);DEBUG;TRACE</DefineConstants>
+    <DebugSymbols>true</DebugSymbols>
+    <DebugType>full</DebugType>
+  </PropertyGroup>
+  <ItemGroup>
+    <Compile Include="..\..\Shared\MetaData.cs">
+      <Link>Properties\MetaData.cs</Link>
+    </Compile>
+  </ItemGroup>
+  <ItemGroup>
+    <ProjectReference Include="..\..\Shared\EpicGames.Core\EpicGames.Core.csproj" PrivateAssets="All"><Private>false</Private></ProjectReference>
+    <ProjectReference Include="..\..\UnrealBuildTool\UnrealBuildTool.csproj" PrivateAssets="All"><Private>false</Private></ProjectReference>
+    <ProjectReference Include="..\AutomationUtils\AutomationUtils.Automation.csproj" PrivateAssets="All"><Private>false</Private></ProjectReference>
+  </ItemGroup>
+  <ItemGroup>
+    <PackageReference Include="Microsoft.CSharp" Version="4.7.0" PrivateAssets="all" />
+  </ItemGroup>
+</Project>
\ No newline at end of file
diff --git a/Engine/Source/Programs/AutomationTool/Apple/ApplePlatform.Automation.cs b/Engine/Source/Programs/AutomationTool/Apple/ApplePlatform.Automation.cs
new file mode 100644
--- /dev/null
+++ b/Engine/Source/Programs/AutomationTool/Apple/ApplePlatform.Automation.cs
@@ -0,0 +1,66 @@
+﻿// Copyright Epic Games, Inc. All Rights Reserved.
+using System;
+using System.Linq;
+using System.Text;
+using System.IO;
+using AutomationTool;
+using UnrealBuildTool;
+using EpicGames.Core;
+using UnrealBuildBase;
+using System.Diagnostics;
+
+
+public abstract class ApplePlatform : Platform
+{
+	public ApplePlatform(UnrealTargetPlatform TargetPlatform)
+		: base(TargetPlatform)
+	{
+	}
+
+	public override bool GetSDKInstallCommand(out string Command, out string Params, ref bool bRequiresPrivilegeElevation, ref bool bCreateWindow, ITurnkeyContext TurnkeyContext)
+	{
+		Command = "";
+		Params = "";
+
+		if (UnrealBuildTool.BuildHostPlatform.Current.Platform == UnrealTargetPlatform.Mac)
+		{
+			TurnkeyContext.Log("Moving your original Xcode application from /Applications to the Trash, and unzipping the new version into /Applications");
+
+			// put current Xcode in the trash, and unzip a new one. Xcode in the dock will have to be fixed up tho!
+			Command = "osascript";
+			Params =
+				" -e \"try\"" +
+				" -e   \"tell application \\\"Finder\\\" to delete POSIX file \\\"/Applications/Xcode.app\\\"\"" +
+				" -e \"end try\"" +
+				" -e \"do shell script \\\"cd /Applications; xip --expand $(CopyOutputPath);\\\"\"" +
+				" -e \"try\"" +
+				" -e   \"do shell script \\\"xcode-select -s /Applications/Xcode.app; xcode-select --install; xcodebuild -license accept; xcodebuild -runFirstLaunch\\\" with administrator privileges\"" +
+				" -e \"end try\"";
+		}
+		else if (UnrealBuildTool.BuildHostPlatform.Current.Platform == UnrealTargetPlatform.Win64)
+		{
+
+			TurnkeyContext.Log("Uninstalling old iTunes and preparing the new one to be installed.");
+
+			Command = "$(EngineDir)/Extras/iTunes/Install_iTunes.bat";
+			Params = "$(CopyOutputPath)";
+		}
+		return true;
+	}
+
+	public override bool OnSDKInstallComplete(int ExitCode, ITurnkeyContext TurnkeyContext, DeviceInfo Device)
+	{
+		if (Device == null)
+		{
+			if (ExitCode == 0)
+			{
+				if (UnrealBuildTool.BuildHostPlatform.Current.Platform == UnrealTargetPlatform.Mac)
+				{
+					TurnkeyContext.PauseForUser("If you had Xcode in your Dock, you will need to remove it and add the new one (even though it was in the same location). macOS follows the move to the Trash for the Dock icon");
+				}
+			}
+		}
+
+		return ExitCode == 0;
+	}
+}
diff --git a/Engine/Source/Programs/AutomationTool/Apple/Properties/AssemblyInfo.cs b/Engine/Source/Programs/AutomationTool/Apple/Properties/AssemblyInfo.cs
new file mode 100644
--- /dev/null
+++ b/Engine/Source/Programs/AutomationTool/Apple/Properties/AssemblyInfo.cs
@@ -0,0 +1,12 @@
+// Copyright Epic Games, Inc. All Rights Reserved.
+
+using System.Reflection;
+using System.Runtime.CompilerServices;
+using System.Runtime.InteropServices;
+
+// General Information about an assembly is controlled through the following 
+// set of attributes. Change these attribute values to modify the information
+// associated with an assembly.
+[assembly: AssemblyTitle("Apple.Automation")]
+[assembly: AssemblyDescription("")]
+[assembly: AssemblyConfiguration("")]
diff --git a/Engine/Source/Programs/AutomationTool/IOS/IOS.Automation.csproj b/Engine/Source/Programs/AutomationTool/IOS/IOS.Automation.csproj
--- a/Engine/Source/Programs/AutomationTool/IOS/IOS.Automation.csproj
+++ b/Engine/Source/Programs/AutomationTool/IOS/IOS.Automation.csproj
@@ -54,6 +54,7 @@
     <ProjectReference Include="..\..\Shared\EpicGames.Core\EpicGames.Core.csproj" PrivateAssets="All"><Private>false</Private></ProjectReference>
     <ProjectReference Include="..\..\UnrealBuildTool\UnrealBuildTool.csproj" PrivateAssets="All"><Private>false</Private></ProjectReference>
     <ProjectReference Include="..\AutomationUtils\AutomationUtils.Automation.csproj" PrivateAssets="All"><Private>false</Private></ProjectReference>
+    <ProjectReference Include="..\Apple\Apple.Automation.csproj" />
   </ItemGroup>
   <ItemGroup>
     <PackageReference Include="Microsoft.CSharp" Version="4.7.0" PrivateAssets="all" />
diff --git a/Engine/Source/Programs/AutomationTool/IOS/IOSPlatform.Automation.cs b/Engine/Source/Programs/AutomationTool/IOS/IOSPlatform.Automation.cs
--- a/Engine/Source/Programs/AutomationTool/IOS/IOSPlatform.Automation.cs
+++ b/Engine/Source/Programs/AutomationTool/IOS/IOSPlatform.Automation.cs
@@ -153,7 +153,7 @@ class IOSClientProcess : IProcessResult
 
 }
 
-public class IOSPlatform : Platform
+public class IOSPlatform : ApplePlatform
 {
 	bool bCreatedIPA = false;
 
@@ -171,37 +171,6 @@ public class IOSPlatform : Platform
 		SDKName = (TargetPlatform == UnrealTargetPlatform.TVOS) ? "appletvos" : "iphoneos";
 	}
 
-	public override bool GetSDKInstallCommand(out string Command, out string Params, ref bool bRequiresPrivilegeElevation, ref bool bCreateWindow, ITurnkeyContext TurnkeyContext)
-	{
-		Command = "";
-		Params = "";
-
-		if (UnrealBuildTool.BuildHostPlatform.Current.Platform == UnrealTargetPlatform.Mac)
-		{
-			TurnkeyContext.Log("Moving your original Xcode application from /Applications to the Trash, and unzipping the new version into /Applications");
-
-			// put current Xcode in the trash, and unzip a new one. Xcode in the dock will have to be fixed up tho!
-			Command = "osascript";
-			Params =
-				" -e \"try\"" +
-				" -e   \"tell application \\\"Finder\\\" to delete POSIX file \\\"/Applications/Xcode.app\\\"\"" +
-				" -e \"end try\"" +
-				" -e \"do shell script \\\"cd /Applications; xip --expand $(CopyOutputPath);\\\"\"" +
-				" -e \"try\"" +
-				" -e   \"do shell script \\\"xcode-select -s /Applications/Xcode.app; xcode-select --install; xcodebuild -license accept; xcodebuild -runFirstLaunch\\\" with administrator privileges\"" +
-				" -e \"end try\"";
-		}
-		else if (UnrealBuildTool.BuildHostPlatform.Current.Platform == UnrealTargetPlatform.Win64)
-		{
-
-			TurnkeyContext.Log("Uninstalling old iTunes and preparing the new one to be installed.");
-
-			Command = "$(EngineDir)/Extras/iTunes/Install_iTunes.bat";
-			Params = "$(CopyOutputPath)";
-		}
-		return true;
-	}
-
 	public override bool GetDeviceUpdateSoftwareCommand(out string Command, out string Params, ref bool bRequiresPrivilegeElevation, ref bool bCreateWindow, ITurnkeyContext TurnkeyContext, DeviceInfo Device)
 	{
 		if (UnrealBuildTool.BuildHostPlatform.Current.Platform != UnrealTargetPlatform.Mac)
@@ -233,22 +202,6 @@ public class IOSPlatform : Platform
 		return true;
 	}
 
-	public override bool OnSDKInstallComplete(int ExitCode, ITurnkeyContext TurnkeyContext, DeviceInfo Device)
-	{
-		if (Device == null)
-		{
-			if (ExitCode == 0)
-			{
-				if (UnrealBuildTool.BuildHostPlatform.Current.Platform == UnrealTargetPlatform.Mac)
-				{
-					TurnkeyContext.PauseForUser("If you had Xcode in your Dock, you will need to remove it and add the new one (even though it was in the same location). macOS follows the move to the Trash for the Dock icon");
-				}
-			}
-		}
-
-		return ExitCode == 0;
-	}
-
 
 
 
diff --git a/Engine/Source/Programs/AutomationTool/Mac/Mac.Automation.csproj b/Engine/Source/Programs/AutomationTool/Mac/Mac.Automation.csproj
--- a/Engine/Source/Programs/AutomationTool/Mac/Mac.Automation.csproj
+++ b/Engine/Source/Programs/AutomationTool/Mac/Mac.Automation.csproj
@@ -38,6 +38,7 @@
     <ProjectReference Include="..\..\Shared\EpicGames.Core\EpicGames.Core.csproj" PrivateAssets="All"><Private>false</Private></ProjectReference>
     <ProjectReference Include="..\..\UnrealBuildTool\UnrealBuildTool.csproj" PrivateAssets="All"><Private>false</Private></ProjectReference>
     <ProjectReference Include="..\AutomationUtils\AutomationUtils.Automation.csproj" PrivateAssets="All"><Private>false</Private></ProjectReference>
+    <ProjectReference Include="..\Apple\Apple.Automation.csproj" />
   </ItemGroup>
   <ItemGroup>
     <PackageReference Include="Microsoft.CSharp" Version="4.7.0" />
diff --git a/Engine/Source/Programs/AutomationTool/Mac/MacPlatform.Automation.cs b/Engine/Source/Programs/AutomationTool/Mac/MacPlatform.Automation.cs
--- a/Engine/Source/Programs/AutomationTool/Mac/MacPlatform.Automation.cs
+++ b/Engine/Source/Programs/AutomationTool/Mac/MacPlatform.Automation.cs
@@ -10,7 +10,7 @@ using UnrealBuildTool;
 using UnrealBuildBase;
 using EpicGames.Core;
 
-public class MacPlatform : Platform
+public class MacPlatform : ApplePlatform
 {
 	/// <summary>
 	/// Default architecture to build projects for. Defaults to Intel
