Get Editor role based on Config setting (unify Editor and CookedEditor)

#rnx

diff --git a/Engine/Source/Programs/AutomationTool/Gauntlet/Unreal/Base/Gauntlet.UnrealTestConfiguration.cs b/Engine/Source/Programs/AutomationTool/Gauntlet/Unreal/Base/Gauntlet.UnrealTestConfiguration.cs
--- a/Engine/Source/Programs/AutomationTool/Gauntlet/Unreal/Base/Gauntlet.UnrealTestConfiguration.cs
+++ b/Engine/Source/Programs/AutomationTool/Gauntlet/Unreal/Base/Gauntlet.UnrealTestConfiguration.cs
@@ -714,14 +714,28 @@ namespace Gauntlet
 		/// <returns></returns>
 		public UnrealTestRole RequireRole(UnrealTargetRole InRole)
 		{
+			if(InRole.IsEditor())
+			{
+				return GetEditorRole();
+			}
 			return RequireRoles(InRole, 1).First();
 		}
 
 		public UnrealTestRole RequireRole(UnrealTargetRole InRole, UnrealTargetPlatform PlatformOverride)
 		{
+			if (InRole.IsEditor())
+			{
+				InRole = CookedEditor ? UnrealTargetRole.CookedEditor : UnrealTargetRole.Editor;
+			}
 			return RequireRoles(InRole, PlatformOverride, 1).First();
 		}
 
+		public UnrealTestRole GetEditorRole()
+		{
+			UnrealTargetRole EditorRole = CookedEditor ? UnrealTargetRole.CookedEditor : UnrealTargetRole.Editor;
+			return RequireRoles(EditorRole, 1).First();
+		}
+
 		/// <summary>
 		/// Adds 'Count' of the specified roles to this test
 		/// </summary>
