Comment out unused function for now since it is causing compilation errors/warnings.

diff --git a/Engine/Source/Programs/AutomationTool/SmartlingLocalization/SmartlingLocalizationProvider.cs b/Engine/Source/Programs/AutomationTool/SmartlingLocalization/SmartlingLocalizationProvider.cs
--- a/Engine/Source/Programs/AutomationTool/SmartlingLocalization/SmartlingLocalizationProvider.cs
+++ b/Engine/Source/Programs/AutomationTool/SmartlingLocalization/SmartlingLocalizationProvider.cs
@@ -78,10 +78,10 @@ namespace EpicGames.SmartlingLocalization
 			AuthenticationToken = null;
 		}
 
-		public async override Task InitializeProjectWithLocalizationProvider(string ProjectName, ProjectImportExportInfo ProjectImportInfo)
-		{
-			// @TODOLocalization: Throw an exeption if the config data is null/invalid 
-		}
+		// public override Task InitializeProjectWithLocalizationProvider(string ProjectName, ProjectImportExportInfo ProjectImportInfo)
+		// {
+		// 	// @TODOLocalization: Throw an exeption if the config data is null/invalid 
+		// }
 
 		public async override Task DownloadProjectFromLocalizationProvider(string ProjectName, ProjectImportExportInfo ProjectImportInfo)
 		{
@@ -369,4 +369,4 @@ namespace EpicGames.SmartlingLocalization
 		private JsonSerializerOptions JsonOptions;
 		private SmartlingAuthenticationToken AuthenticationToken;
 	}
-}
\ No newline at end of file
+}
