Fix BuildCommand run with NoProject flag

diff --git a/Engine/Source/Programs/AutomationTool/Scripts/OpenEditor.Automation.cs b/Engine/Source/Programs/AutomationTool/Scripts/OpenEditor.Automation.cs
--- a/Engine/Source/Programs/AutomationTool/Scripts/OpenEditor.Automation.cs
+++ b/Engine/Source/Programs/AutomationTool/Scripts/OpenEditor.Automation.cs
@@ -32,8 +32,11 @@ namespace AutomationTool
 			}
 		}
 
+		public bool bNoProject { get; set; }
+
 		public OpenEditor()
 		{
+			bNoProject = false;
 		}
 		
 		public override ExitCode Execute()
@@ -45,7 +48,7 @@ namespace AutomationTool
 			ProjectName = ParseParamValue("project", ProjectName);
 			IEnumerable<string> ParamList = null;
 
-			if (ParseParam("NoProject"))
+			if (bNoProject)
 			{
 				ParamList = Params;
 			}
