[Turnkey] Resolve crash while updating device flash when the following conditions occur:
* Use the interactive Turnkey menu (ex. ".uat Turnkey")
* More than one device is available
* The "default" device is disconnected
* The device selected to update is not the default

#rn Fix Turnkey Previously, if multiple devices were available and the default device was not connected a firmware flash could fail when attempting to get the FW flash version information.  Now the process will use the selected device if the default device is not connected.

diff --git a/Engine/Source/Programs/AutomationTool/Turnkey/Commands/InstallSdk.cs b/Engine/Source/Programs/AutomationTool/Turnkey/Commands/InstallSdk.cs
--- a/Engine/Source/Programs/AutomationTool/Turnkey/Commands/InstallSdk.cs
+++ b/Engine/Source/Programs/AutomationTool/Turnkey/Commands/InstallSdk.cs
@@ -138,6 +138,7 @@ namespace Turnkey.Commands
 				if (DesiredType == FileSource.SourceType.Flash)
 				{
 					InstallDevice = ChosenDevices.First();
+ 					DeviceName = InstallDevice.Name;
 				}
 
 				// skip Misc FileSources, as we dont know how they are sued
@@ -238,6 +239,11 @@ namespace Turnkey.Commands
 					if (DesiredType == FileSource.SourceType.Flash)
 					{
 						DeviceInfo Device = GetDevice(AutomationPlatform, DeviceName);
+						if (Device == null)
+						{
+							TurnkeyUtils.Log("Failed to get device (name: {0}) for platform {1}. Skipping", DeviceName == null ? "<null>/default" : DeviceName, Platform);
+							continue;
+						}
 						SDKInfo = SDK.GetAllSoftwareInfo(Device.Type, Device.SoftwareVersion);
 					}
 					else
