[Automation] fix for CheckCsprojDotNetVersion giving an error for the UnrealVS test adapater project.

changed CheckCsprojDotNetVersion to ignore all projects under the UnrealVS direction vs project containing the name UnrealVS

diff --git a/Engine/Source/Programs/AutomationTool/Scripts/CheckCsprojDotNetVersion.cs b/Engine/Source/Programs/AutomationTool/Scripts/CheckCsprojDotNetVersion.cs
--- a/Engine/Source/Programs/AutomationTool/Scripts/CheckCsprojDotNetVersion.cs
+++ b/Engine/Source/Programs/AutomationTool/Scripts/CheckCsprojDotNetVersion.cs
@@ -1,13 +1,10 @@
 // Copyright Epic Games, Inc. All Rights Reserved.
 using System;
-using System.Collections.Generic;
-using System.Text;
 using System.IO;
-using AutomationTool;
-using UnrealBuildTool;
+using System.Linq;
 using System.Text.RegularExpressions;
+using AutomationTool;
 using EpicGames.Core;
-using System.Linq;
 using UnrealBuildBase;
 
 class CheckCsprojDotNetVersion : BuildCommand
@@ -31,33 +28,33 @@ class CheckCsprojDotNetVersion : BuildCommand
 		Regex FrameworkRegex = new Regex("<TargetFrameworkVersion>v(\\d\\.\\d\\.?\\d?)<\\/TargetFrameworkVersion>");
 		Regex PossibleAppConfigRegex = new Regex("<TargetFrameworkProfile>(.+)<\\/TargetFrameworkProfile>");
 		Regex AppConfigRegex = new Regex("<supportedRuntime version=\"v(\\d\\.\\d\\.?\\d?)\" sku=\"\\.NETFramework,Version=v(\\d\\.\\d\\.?\\d?),Profile=(.+)\"\\/>");
-        Regex DotNetCoreRegex = new Regex("<TargetFramework>(netcoreapp2.0|netstandard2.0)<\\/TargetFramework>");
-        foreach (FileReference CsProj in DirectoryReference.EnumerateFiles(EngineDir, "*.csproj", SearchOption.AllDirectories))
-        {
-            if (CsProj.ContainsName("ThirdParty", EngineDir) ||
-                (CsProj.ContainsName("UE4TemplateProject", EngineDir) && CsProj.GetFileName().Equals("ProjectTemplate.csproj")) ||
-                CsProj.GetFileNameWithoutExtension().ToLower().Contains("_mono") ||
-                CsProj.GetFileNameWithoutExtension().ToLower().Contains("unrealvs") ||
+		Regex DotNetCoreRegex = new Regex("<TargetFramework>(netcoreapp2.0|netstandard2.0)<\\/TargetFramework>");
+		foreach (FileReference CsProj in DirectoryReference.EnumerateFiles(EngineDir, "*.csproj", SearchOption.AllDirectories))
+		{
+			if (CsProj.ContainsName("ThirdParty", EngineDir) ||
+				(CsProj.ContainsName("UE4TemplateProject", EngineDir) && CsProj.GetFileName().Equals("ProjectTemplate.csproj")) ||
+				CsProj.GetFileNameWithoutExtension().ToLower().Contains("_mono") ||
+				CsProj.ContainsName("UnrealVS", EngineDir) ||
 				CsProj.ContainsName("DatasmithRevitExporter", EngineDir) ||
 				CsProj.ContainsName("DatasmithNavisworksExporter", EngineDir) ||
 				CsProj.ContainsName("CSVTools", EngineDir))
 
-            {
-                continue;
-            }
+			{
+				continue;
+			}
 
-            // read in the file
-            string Contents = File.ReadAllText(CsProj.FullName);
-            Match Match = DotNetCoreRegex.Match(Contents);
-            // Check if we're a _NETCore app, ignore these.
-            if (Match.Success)
-            {
-                continue;
-            }
+			// read in the file
+			string Contents = File.ReadAllText(CsProj.FullName);
+			Match Match = DotNetCoreRegex.Match(Contents);
+			// Check if we're a _NETCore app, ignore these.
+			if (Match.Success)
+			{
+				continue;
+			}
 
 
-            Match = FrameworkRegex.Match(Contents);
-            if (Match.Success)
+			Match = FrameworkRegex.Match(Contents);
+			if (Match.Success)
 			{
 				string TargetedVersion = Match.Groups[1].Value;
 				// make sure we match, throw warning otherwise
@@ -69,7 +66,7 @@ class CheckCsprojDotNetVersion : BuildCommand
 			// if we don't have a TargetFrameworkVersion, check for the existence of TargetFrameworkProfile.
 			else
 			{
-                Match = PossibleAppConfigRegex.Match(Contents);
+				Match = PossibleAppConfigRegex.Match(Contents);
 				if (!Match.Success)
 				{
 					CommandUtils.LogInformation("No TargetFrameworkVersion or TargetFrameworkProfile found for project {0}, is it a mono project? If not, does it compile properly?", CsProj);
@@ -87,7 +84,7 @@ class CheckCsprojDotNetVersion : BuildCommand
 
 				// read in the app config
 				Contents = File.ReadAllText(AppConfigFile.FullName);
-                Match = AppConfigRegex.Match(Contents);
+				Match = AppConfigRegex.Match(Contents);
 				if (!Match.Success)
 				{
 					CommandUtils.LogInformation("Couldn't find a supportedRuntime match for the version in the app.config for project {0}.", CsProj);
