Adding workflow to quickly launch client, server and editor

diff --git a/Engine/Source/Programs/AutomationTool/Scripts/OpenEditor.Automation.cs b/Engine/Source/Programs/AutomationTool/Scripts/OpenEditor.Automation.cs
--- a/Engine/Source/Programs/AutomationTool/Scripts/OpenEditor.Automation.cs
+++ b/Engine/Source/Programs/AutomationTool/Scripts/OpenEditor.Automation.cs
@@ -14,7 +14,7 @@ namespace AutomationTool
 
 	[Help("Opens the specified project.")]
 	[Help("project=<QAGame>", "Project to open. Will search current path and paths in ueprojectdirs. If omitted will open vanilla UnrealEditor")]
-	public class OpenEditor: BuildCommand
+	public class OpenEditor : BuildCommand
 	{
 		// exposed as a property so projects can derive and set this directly
 		public string ProjectName { get; set; }
@@ -38,7 +38,7 @@ namespace AutomationTool
 		{
 			bNoProject = false;
 		}
-		
+
 		public override ExitCode Execute()
 		{
 			string EditorPath = HostPlatform.Current.GetUnrealExePath(UnrealEditorApp);
@@ -73,16 +73,19 @@ namespace AutomationTool
 
 			ParamList = new[] { EditorArgs }.Concat(ParamList);
 
-			bool bLaunched = RunUntrackedProcess(EditorPath, string.Join(" -", ParamList));
+			bool bLaunched = RunProceess.RunUntrackedProcess(EditorPath, string.Join(" -", ParamList));
 
 			return bLaunched ? ExitCode.Success : ExitCode.Error_UATLaunchFailure;
 		}
+	}
 
-		protected bool RunUntrackedProcess(string BinaryPath, string Args)
+	public class RunProceess
+	{ 
+		public static bool RunUntrackedProcess(string BinaryPath, string Args)
 		{
-			LogInformation("Running {0} {1}", BinaryPath, Args);
+			CommandUtils.LogInformation("Running {0} {1}", BinaryPath, Args);
 
-			var NewProcess = HostPlatform.Current.CreateProcess(BinaryPath);			
+			var NewProcess = HostPlatform.Current.CreateProcess(BinaryPath);
 			var Result = new ProcessResult(BinaryPath, NewProcess, false, false);
 			System.Diagnostics.Process Proc = Result.ProcessObject;
 
@@ -92,5 +95,4 @@ namespace AutomationTool
 			return Proc.Start();
 		}
 	}
-
 }
\ No newline at end of file
