Gauntlet - Clear RolesToInstalls only if it was initialized

#rnx

diff --git a/Engine/Source/Programs/AutomationTool/Gauntlet/Unreal/Base/Gauntlet.UnrealSession.cs b/Engine/Source/Programs/AutomationTool/Gauntlet/Unreal/Base/Gauntlet.UnrealSession.cs
--- a/Engine/Source/Programs/AutomationTool/Gauntlet/Unreal/Base/Gauntlet.UnrealSession.cs
+++ b/Engine/Source/Programs/AutomationTool/Gauntlet/Unreal/Base/Gauntlet.UnrealSession.cs
@@ -947,7 +947,10 @@ namespace Gauntlet
 		public void ReleaseSessionDevices()
 		{
 			UnrealDeviceReservation.ReleaseDevices();
-			RolesToInstalls.Clear();
+			if (RolesToInstalls != null)
+			{
+				RolesToInstalls.Clear();
+			}
 		}
 
 		/// <summary>
