Turnkey support for Android on Mac

#android

diff --git a/Engine/Build/Turnkey/Licensee_TurnkeyManifest.Example.xml b/Engine/Build/Turnkey/Licensee_TurnkeyManifest.Example.xml
--- a/Engine/Build/Turnkey/Licensee_TurnkeyManifest.Example.xml
+++ b/Engine/Build/Turnkey/Licensee_TurnkeyManifest.Example.xml
@@ -20,7 +20,7 @@
 		<Name>AndroidStudio</Name>
 		<Source HostPlatform="Win64">googledrive:/SdkInstallers/Android/android-studio-ide-193.6821437-windows.exe</Source>
 		<Source HostPlatform="Linux">googledrive:/SdkInstallers/Android/android-studio-ide-193.6821437-linux.tar.gz</Source>
-		<Source HostPlatform="Mac">googledrive:/SdkInstallers/Android/android-studio-ide-193.6821437-mac.zip</Source>
+		<Source HostPlatform="Mac">googledrive:/SdkInstallers/Android/android-studio-ide-193.6821437-mac.dmg</Source>
 	</FileSource>
 
 	
diff --git a/Engine/Extras/Android/SetupAndroid.command b/Engine/Extras/Android/SetupAndroid.command
--- a/Engine/Extras/Android/SetupAndroid.command
+++ b/Engine/Extras/Android/SetupAndroid.command
@@ -70,11 +70,11 @@ if [[ -z "${NDK_VERSION}" ]]; then
 fi
 
 
-STUDIO_PATH="~/Applications/Android Studio.app"
+STUDIO_PATH="$HOME/Applications/Android Studio.app"
 if [ ! -x "$STUDIO_PATH" ]; then
 	STUDIO_PATH="/Applications/Android Studio.app"
 	if [ ! -x "$STUDIO_PATH" ]; then
-		echo Android Studio not installed, please download Android Studio 3.5.3 from https://developer.android.com/studio
+		echo Android Studio not installed, please download Android Studio 4.0.2 from https://developer.android.com/studio
 		${PAUSE}
 		exit 1
 	fi
@@ -90,7 +90,7 @@ fi
 if [ ! -d "$STUDIO_SDK_PATH" ]; then
 	echo Android SDK not found at: $STUDIO_SDK_PATH
 	echo Unable to locate local Android SDK location. Did you run Android Studio after installing?
-	echo If Android Studio is installed, please run again with SDK path as parameter, otherwise download Android Studio 3.5.3 from https://developer.android.com/studio
+	echo If Android Studio is installed, please run again with SDK path as parameter, otherwise download Android Studio 4.0.2 from https://developer.android.com/studio
 	${PAUSE}
 	exit 1
 fi
diff --git a/Engine/Source/Programs/AutomationTool/Android/AndroidPlatform.Automation.cs b/Engine/Source/Programs/AutomationTool/Android/AndroidPlatform.Automation.cs
--- a/Engine/Source/Programs/AutomationTool/Android/AndroidPlatform.Automation.cs
+++ b/Engine/Source/Programs/AutomationTool/Android/AndroidPlatform.Automation.cs
@@ -220,8 +220,17 @@ public class AndroidPlatform : Platform
 		}
 		else if (OperatingSystem.IsMacOS())
 		{
-			// TODO
-			return "";
+
+			string AndroidStudioExe = "/Applications/Android Studio.app";
+			if (Directory.Exists(AndroidStudioExe))
+			{
+				return AndroidStudioExe;
+			}
+
+			string UserHome = Environment.GetEnvironmentVariable("HOME");
+			AndroidStudioExe = Path.Combine(UserHome, "Applications", "Android Studio.app");
+
+			return AndroidStudioExe;
 		}
 
 		Debug.Assert(OperatingSystem.IsWindows());
@@ -231,6 +240,7 @@ public class AndroidPlatform : Platform
 		string AndroidStudioInstallDir = RegValue == null ? DefaultAndroidStudioInstallDir : RegValue;
 		return Path.Combine(AndroidStudioInstallDir, "bin", "studio64.exe");
 	}
+
 	private static string GetSdkDir()
 	{
 		string AndroidHome = Environment.GetEnvironmentVariable("ANDROID_HOME");
@@ -248,8 +258,32 @@ public class AndroidPlatform : Platform
 		}
 		else if (OperatingSystem.IsMacOS())
 		{
-			// TODO
-			return "";
+			string BashProfilePath = Path.Combine(Environment.GetFolderPath(Environment.SpecialFolder.Personal), ".bash_profile");
+			if (!File.Exists(BashProfilePath))
+			{
+				// Try .bashrc if didn't fine .bash_profile
+				BashProfilePath = Path.Combine(Environment.GetFolderPath(Environment.SpecialFolder.Personal), ".bashrc");
+			}
+			if (File.Exists(BashProfilePath))
+			{
+				string[] BashProfileContents = File.ReadAllLines(BashProfilePath);
+
+				// Walk backwards so we keep the last export setting instead of the first
+				string SdkKey = "ANDROID_HOME";
+				for (int LineIndex = BashProfileContents.Length - 1; LineIndex >= 0; --LineIndex)
+				{
+					if (BashProfileContents[LineIndex].StartsWith("export " + SdkKey + "="))
+					{
+						string PathVar = BashProfileContents[LineIndex].Split
+('=')[1].Replace("\"", "");
+Log.TraceInformation("ANDROID_HOME = {0}", PathVar);
+						return PathVar;
+					}
+
+				}
+			}
+
+			return Environment.GetEnvironmentVariable("ANDROID_HOME");
 		}
 
 		Debug.Assert(OperatingSystem.IsWindows());
@@ -261,24 +295,22 @@ public class AndroidPlatform : Platform
 
 	public override bool UpdateHostPrerequisites(BuildCommand Command, ITurnkeyContext TurnkeyContext, bool bVerifyOnly)
 	{
-		// @todo turnkey: Handle Mac
-		if (HostPlatform.Current.HostEditorPlatform == UnrealTargetPlatform.Mac)
-		{
-			return false;
-		}
-
 		string AndroidStudioExe = GetAndroidStudioExe();
 		string SdkDir = GetSdkDir();
 
-		bool bIsInstalled = File.Exists(AndroidStudioExe) && Directory.Exists(SdkDir);
+		bool bIsMac = (HostPlatform.Current.HostEditorPlatform == UnrealTargetPlatform.Mac);
+		bool bHaveAndroidStudio = (bIsMac && Directory.Exists(AndroidStudioExe)) || 
+					(!bIsMac && FileExists(AndroidStudioExe));
+
+		bool bIsInstalled = bHaveAndroidStudio && Directory.Exists(SdkDir);
 
 		// if we are only verifying, just return the status, and if it's installed, we are done!
 		if (bVerifyOnly || bIsInstalled)
 		{
-			if (!File.Exists(AndroidStudioExe))
+			if (!bHaveAndroidStudio)
 			{
 				TurnkeyContext.ReportError("Android Studio is not installed correctly.");
-	}
+			}
 			if (!Directory.Exists(SdkDir))
 			{
 				TurnkeyContext.ReportError("Android SDK directory is not set correctly.");
@@ -286,7 +318,7 @@ public class AndroidPlatform : Platform
 			return bIsInstalled;
 		}
 
-		if (!File.Exists(AndroidStudioExe))
+		if (!bHaveAndroidStudio)
 		{
 			// get AS installer
 			string OutputPath = TurnkeyContext.RetrieveFileSource("AndroidStudio");
@@ -306,7 +338,7 @@ public class AndroidPlatform : Platform
 
 			if (OutputPath == null)
 			{
-				TurnkeyContext.PauseForUser("Unable to find Android Studio installer. Please download and install Android Studio 3.5.3 from https://developer.android.com/studio/archive before continuing.\n\nMake sure to use the Run Android Studio and complete the first-time setup!\n\nChoose all default options unless you know what you are doing.");
+				TurnkeyContext.PauseForUser("Unable to find Android Studio installer. Please download and install Android Studio 4.0.2 from https://developer.android.com/studio/archive before continuing.\n\nMake sure to use the Run Android Studio and complete the first-time setup!\n\nFollow the steps to install command-line tools (latest):\nhttps://docs.unrealengine.com/5.0/en-US/how-to-setup-up-android-sdk-and-ndk-for-your-unreal-engine-development-environment");
 			}
 			else
 			{
@@ -328,11 +360,63 @@ public class AndroidPlatform : Platform
 				}
 				else if (HostPlatform.Current.HostEditorPlatform == UnrealTargetPlatform.Mac)
 				{
-					// TODO
+
+					string UserHome = Environment.GetEnvironmentVariable("HOME");
+					string SourceApp = Path.Combine(OutputPath, "Android Studio.app");
+
+					int ExitCode = TurnkeyContext.RunExternalCommand("/usr/bin/hdiutil", "attach " + OutputPath, false, true, true);
+
+					if (ExitCode != 0)
+					{
+						TurnkeyContext.ReportError($"Android Studio installer failed. ExitCode = {ExitCode}");
+						return false;
+					}
+
+					string AndroidStudioVolume = "";
+					foreach (string Volume in Directory.GetDirectories("/Volumes"))
+
+					{
+						if (Volume.Contains("Android Studio"))
+						{
+							AndroidStudioVolume = Volume;
+							break;
+						}
+					}
+					if (AndroidStudioVolume == "")
+					{
+						TurnkeyContext.ReportError($"Android Studio installer failed. DMG did not mount");
+						return false;
+					}
+
+					string SourcePath = Path.Combine(AndroidStudioVolume, "Android Studio.app");
+					string DestPath = Path.Combine(UserHome, "Applications") + "/";
+					if (SourcePath.Contains(" "))
+					{
+						SourcePath = "\"" + SourcePath + "\"";
+					}
+					if (DestPath.Contains(" "))
+					{
+						DestPath = "\"" + DestPath + "\"";
+					}
+
+					ExitCode = TurnkeyContext.RunExternalCommand("/bin/cp", "-R " + SourcePath + " " + DestPath, false, true, true);
+
+					if (AndroidStudioVolume.Contains(" "))
+					{
+						AndroidStudioVolume = "\"" + AndroidStudioVolume + "\"";
+					}
+					int ExitCode2 = TurnkeyContext.RunExternalCommand("/usr/bin/hdiutil", "detach " + AndroidStudioVolume, false, true, true);
+
+					// give error for cp, but can ignore detach failure
+					if (ExitCode != 0)
+					{
+						TurnkeyContext.ReportError($"Android Studio installer failed. ExitCode = {ExitCode}");
+						return false;
+					}
 				}
 				else if (HostPlatform.Current.HostEditorPlatform == UnrealTargetPlatform.Win64)
 				{
-					TurnkeyContext.PauseForUser("Running the Android Studio installer, and then Android Studio for first-time setup!\n\nChoose all default options unless you know what you are doing.");
+					TurnkeyContext.PauseForUser("Running the Android Studio installer, and then Android Studio for first-time setup!\n\nFollow the steps to install command-line tools (latest):\nhttps://docs.unrealengine.com/5.0/en-US/how-to-setup-up-android-sdk-and-ndk-for-your-unreal-engine-development-environment");
 					int ExitCode = TurnkeyContext.RunExternalCommand(OutputPath, "/S", false, true, true);
 					//				Utils.RunLocalProcessAndReturnStdOut(OutputPath, "", out ExitCode, true);
 
@@ -350,23 +434,26 @@ public class AndroidPlatform : Platform
 				}
 
 				// re-query for AS location
-				TurnkeyContext.RunExternalCommand(GetAndroidStudioExe(), "", false, true, false);
+				TurnkeyContext.RunExternalCommand(GetAndroidStudioExe(), "", false, true, bIsMac);
 			}
 		}
 		else
 		{
-			TurnkeyContext.PauseForUser("The Sdk directory was not found. Running the Android Studio to perform first-time setup.\n\nChoose all default options unless you know what you are doing.");
+			TurnkeyContext.PauseForUser("The Sdk directory was not found. Running the Android Studio to perform first-time setup.\n\nFollow the steps to install command-line tools (latest):\nhttps://docs.unrealengine.com/5.0/en-US/how-to-setup-up-android-sdk-and-ndk-for-your-unreal-engine-development-environment");
 
-			TurnkeyContext.RunExternalCommand(AndroidStudioExe, "", false, true, false);
+			TurnkeyContext.RunExternalCommand(AndroidStudioExe, "", false, true, bIsMac);
 		}
 
 		// check to see if the installation worked. If so, continue on!
 		AndroidStudioExe = GetAndroidStudioExe();
 		SdkDir = GetSdkDir();
 
-		bIsInstalled = File.Exists(AndroidStudioExe) && Directory.Exists(SdkDir);
+		bHaveAndroidStudio = (bIsMac && Directory.Exists(AndroidStudioExe)) || 
+					(!bIsMac && FileExists(AndroidStudioExe));
+
+		bIsInstalled = bHaveAndroidStudio && Directory.Exists(SdkDir);
 
-		if (!File.Exists(AndroidStudioExe))
+		if (!bHaveAndroidStudio)
 		{
 			TurnkeyContext.ReportError("Android Studio is not installed correctly, after attempted installation.");
 		}
diff --git a/Engine/Source/Programs/UnrealBuildTool/Platform/Android/AndroidPlatformSDK.cs b/Engine/Source/Programs/UnrealBuildTool/Platform/Android/AndroidPlatformSDK.cs
--- a/Engine/Source/Programs/UnrealBuildTool/Platform/Android/AndroidPlatformSDK.cs
+++ b/Engine/Source/Programs/UnrealBuildTool/Platform/Android/AndroidPlatformSDK.cs
@@ -22,10 +22,51 @@ namespace UnrealBuildTool
 		{
 		}
 
-		protected override string? GetInstalledSDKVersion()
+		private string? GetNDKRoot()
 		{
 			string? NDKPath = Environment.GetEnvironmentVariable("NDKROOT");
 
+			// don't register if we don't have an NDKROOT specified
+			if (!String.IsNullOrEmpty(NDKPath))
+			{
+				return NDKPath.Replace("\"", "");
+			}
+
+			if (OperatingSystem.IsMacOS())
+			{
+				string BashProfilePath = Path.Combine(Environment.GetFolderPath(Environment.SpecialFolder.Personal), ".bash_profile");
+				if (!File.Exists(BashProfilePath))
+				{
+					// Try .bashrc if didn't fine .bash_profile
+					BashProfilePath = Path.Combine(Environment.GetFolderPath(Environment.SpecialFolder.Personal), ".bashrc");
+				}
+				if (File.Exists(BashProfilePath))
+				{
+					string[] BashProfileContents = File.ReadAllLines(BashProfilePath);
+
+					// Walk backwards so we keep the last export setting instead of the first
+					string SdkKey = "NDKROOT";
+					for (int LineIndex = BashProfileContents.Length - 1; LineIndex >= 0; --LineIndex)
+					{
+						if (BashProfileContents[LineIndex].StartsWith("export " + SdkKey + "="))
+						{
+							string PathVar = BashProfileContents[LineIndex].Split
+('=')[1].Replace("\"", "");
+Log.TraceInformation("ANDROID_HOME = {0}", PathVar);
+							return PathVar;
+						}
+
+					}
+				}
+			}
+	
+			return null;
+		}
+
+		protected override string? GetInstalledSDKVersion()
+		{
+			string? NDKPath = GetNDKRoot();
+
 			// don't register if we don't have an NDKROOT specified
 			if (String.IsNullOrEmpty(NDKPath))
 			{
