[SteamDeck] Set the IpAddr if we have a valid device

diff --git a/Engine/Source/Programs/AutomationTool/SteamDeck/SteamDeckSupport.cs b/Engine/Source/Programs/AutomationTool/SteamDeck/SteamDeckSupport.cs
--- a/Engine/Source/Programs/AutomationTool/SteamDeck/SteamDeckSupport.cs
+++ b/Engine/Source/Programs/AutomationTool/SteamDeck/SteamDeckSupport.cs
@@ -80,11 +80,10 @@ public static class SteamDeckSupport
 	 */
 	public static bool GetDeviceInfo(UnrealTargetPlatform RuntimePlatform, ProjectParams Params, out string IpAddr, out string UserName)
 	{
-		IpAddr = UserName = null;
-
 		if (Params.DeviceNames.Count != 1)
 		{
 			CommandUtils.LogWarning("SteamDeck deployment requires 1, and only, device");
+			IpAddr = UserName = null;
 			return false;
 		}
 
@@ -102,9 +101,11 @@ public static class SteamDeckSupport
 		// if the device is still null, dont assume this is a valid steam deck device
 		if (Device == null)
 		{
+			IpAddr = UserName = null;
 			return false;
 		}
 
+		IpAddr = Device.Id;
 
 		// if -deviceuser was specified, prefer that, otherwise use what's in the registry
 		UserName = Params.DeviceUsername;
