Update Smartling to download translations for files that are both in progress and published. This allows partial translations to still be used by the Engine.
#test: Ran BuildGraph locally for the "Localize UnrealEditor Release" target. Everything compiles fine but fix cannot be verified as localization data for 22.35 is not present yet.

diff --git a/Engine/Source/Programs/AutomationTool/SmartlingLocalization/SmartlingLocalizationProvider.cs b/Engine/Source/Programs/AutomationTool/SmartlingLocalization/SmartlingLocalizationProvider.cs
--- a/Engine/Source/Programs/AutomationTool/SmartlingLocalization/SmartlingLocalizationProvider.cs
+++ b/Engine/Source/Programs/AutomationTool/SmartlingLocalization/SmartlingLocalizationProvider.cs
@@ -125,8 +125,9 @@ namespace EpicGames.SmartlingLocalization
 				string SmartlingFileUri = GetSmartlingFileUri(SmartlingFilename);
 				var DownloadUriBuilder = new UriBuilder(DownloadEndpoint);
 				// Changing the retrieval type to pseudo is a good way to test downloads
-				// @TODOLocalization: Not sure if the retrieval type should be pending or published to get all of the relevant strings 
-				string DownloadQueryString = $"fileUri={SmartlingFileUri}&retrievalType=published&includeOriginalStrings=false";
+				// We use pending as the retrieval type to also download files that are in progress. We need both published and in progress files on Smartling.
+				string DownloadRetrievalType = "pending";
+				string DownloadQueryString = $"fileUri={SmartlingFileUri}&retrievalType={DownloadRetrievalType}&includeOriginalStrings=false";
 				DownloadUriBuilder.Query = DownloadQueryString;
 				var DownloadResponse = await Client.GetAsync(DownloadUriBuilder.Uri);
 				if (DownloadResponse.IsSuccessStatusCode)
