- Refactor "Default Editor" mechanism to enable specifying default Server too.

diff --git a/Engine/Source/Programs/AutomationTool/AutomationUtils/ProjectParams.cs b/Engine/Source/Programs/AutomationTool/AutomationUtils/ProjectParams.cs
--- a/Engine/Source/Programs/AutomationTool/AutomationUtils/ProjectParams.cs
+++ b/Engine/Source/Programs/AutomationTool/AutomationUtils/ProjectParams.cs
@@ -2271,29 +2271,31 @@ namespace AutomationTool
 			}
 		}
 
-		private void SelectDefaultEditorTarget(List<string> AvailableEditorTargets, ref string EditorTarget)
+		private void SelectDefaultTarget(List<string> AvailableTargets, TargetType Type, ref string Target)
 		{
-			string DefaultEditorTarget;
+			string DefaultTarget;
 
-			if (EngineConfigs[BuildHostPlatform.Current.Platform].GetString("/Script/BuildSettings.BuildSettings", "DefaultEditorTarget", out DefaultEditorTarget))
+			string ConfigKey = String.Format($"Default{Type}Target");
+
+			if (EngineConfigs[BuildHostPlatform.Current.Platform].GetString("/Script/BuildSettings.BuildSettings", ConfigKey, out DefaultTarget))
 			{
-				if (!AvailableEditorTargets.Contains(DefaultEditorTarget))
+				if (!AvailableTargets.Contains(DefaultTarget))
 				{
-					throw new AutomationException(string.Format("A default editor target '{0}' was specified in engine.ini but does not exist", DefaultEditorTarget));
+					throw new AutomationException(string.Format($"A default {Type} target '{DefaultTarget}' was specified in engine.ini but does not exist"));
 				}
 
-				EditorTarget = DefaultEditorTarget;
+				Target = DefaultTarget;
 			}
 			else
 			{
-				if (AvailableEditorTargets.Count > 1)
+				if (AvailableTargets.Count > 1)
 				{
-					throw new AutomationException("Project contains multiple editor targets but no DefaultEditorTarget is set in the [/Script/BuildSettings.BuildSettings] section of DefaultEngine.ini");
+					throw new AutomationException($"Project contains multiple {Type} targets but no {ConfigKey} is set in the [/Script/BuildSettings.BuildSettings] section of DefaultEngine.ini");
 				}
 
-				if (AvailableEditorTargets.Count > 0)
+				if (AvailableTargets.Count > 0)
 				{
-					EditorTarget = AvailableEditorTargets.First();
+					Target = AvailableTargets.First();
 				}
 			}
 		}
@@ -2410,7 +2412,7 @@ namespace AutomationTool
 				}
 
 				// Find the editor target name
-				SelectDefaultEditorTarget(TargetNamesOfType(TargetType.Editor), ref EditorTarget);
+				SelectDefaultTarget(TargetNamesOfType(TargetType.Editor), TargetType.Editor, ref EditorTarget);
 			}
 			else if (!CommandUtils.IsNullOrEmpty(Properties.Targets))
 			{
@@ -2458,11 +2460,11 @@ namespace AutomationTool
 					GameTarget = AvailableGameTargets.First();
 				}
 
-				SelectDefaultEditorTarget(AvailableEditorTargets, ref EditorTarget);
+				SelectDefaultTarget(AvailableEditorTargets, TargetType.Editor, ref EditorTarget);
 
 				if (AvailableServerTargets.Count > 0 && (DedicatedServer || Cook || CookOnTheFly)) // only if server is needed
 				{
-					ServerTarget = ChooseTarget(AvailableServerTargets, TargetType.Server);
+					SelectDefaultTarget(AvailableServerTargets, TargetType.Server, ref ServerTarget);
 				}
 			}
 			else if (!CommandUtils.IsNullOrEmpty(Properties.Programs))
