Horde: Add a temp storage test script which generates random data into a given directory.

diff --git a/Engine/Source/Programs/AutomationTool/BuildGraph/Tasks/RandomDataTask.cs b/Engine/Source/Programs/AutomationTool/BuildGraph/Tasks/RandomDataTask.cs
new file mode 100644
--- /dev/null
+++ b/Engine/Source/Programs/AutomationTool/BuildGraph/Tasks/RandomDataTask.cs
@@ -0,0 +1,130 @@
+// Copyright Epic Games, Inc. All Rights Reserved.
+
+using AutomationTool;
+using EpicGames.BuildGraph;
+using System;
+using System.Collections.Generic;
+using System.Linq;
+using System.Text;
+using System.Threading;
+using System.Threading.Tasks;
+using System.Xml;
+using EpicGames.Core;
+using UnrealBuildTool;
+using UnrealBuildBase;
+using Microsoft.Extensions.Logging;
+using System.Security.Cryptography;
+
+namespace AutomationTool.Tasks
+{
+	/// <summary>
+	/// Parameters for <see cref="RandomDataTask"/>.
+	/// </summary>
+	public class RandomDataTaskParameters
+	{
+		/// <summary>
+		/// The size of each file.
+		/// </summary>
+		[TaskParameter(Optional = true)]
+		public int Size = 1024;
+
+		/// <summary>
+		/// Number of files to write.
+		/// </summary>
+		[TaskParameter(Optional = true)]
+		public int Count = 50;
+
+		/// <summary>
+		/// Whether to generate different data for each output file.
+		/// </summary>
+		[TaskParameter(Optional = true)]
+		public bool Different = true;
+
+		/// <summary>
+		/// Output directory
+		/// </summary>
+		[TaskParameter(Optional = true)]
+		public string OutputDir;
+
+		/// <summary>
+		/// Optional filter to be applied to the list of input files.
+		/// </summary>
+		[TaskParameter(Optional = true)]
+		public string Tag;
+	}
+
+	/// <summary>
+	/// Creates files containing random data in the specified output directory. Used for generating test data for the temp storage system.
+	/// </summary>
+	[TaskElement("RandomData", typeof(RandomDataTaskParameters))]
+	public class RandomDataTask : BgTaskImpl
+	{
+		/// <summary>
+		/// Parameters for this task
+		/// </summary>
+		RandomDataTaskParameters Parameters;
+
+		/// <summary>
+		/// Constructor
+		/// </summary>
+		/// <param name="InParameters">Parameters for this task</param>
+		public RandomDataTask(RandomDataTaskParameters InParameters)
+		{
+			Parameters = InParameters;
+		}
+
+		/// <summary>
+		/// Execute the task.
+		/// </summary>
+		/// <param name="Job">Information about the current job</param>
+		/// <param name="BuildProducts">Set of build products produced by this node.</param>
+		/// <param name="TagNameToFileSet">Mapping from tag names to the set of files they include</param>
+		public override async Task ExecuteAsync(JobContext Job, HashSet<FileReference> BuildProducts, Dictionary<string, HashSet<FileReference>> TagNameToFileSet)
+		{
+			byte[] buffer = Array.Empty<byte>();
+			for (int idx = 0; idx < Parameters.Count; idx++)
+			{
+				if (idx == 0 || Parameters.Different)
+				{
+					buffer = RandomNumberGenerator.GetBytes(Parameters.Size);
+				}
+
+				FileReference file = FileReference.Combine(ResolveDirectory(Parameters.OutputDir), $"test-{Parameters.Size}-{idx}.dat");
+				await FileReference.WriteAllBytesAsync(file, buffer);
+				BuildProducts.Add(file);
+			}
+
+			// Apply the optional output tag to them
+			foreach (string TagName in FindTagNamesFromList(Parameters.Tag))
+			{
+				FindOrAddTagSet(TagNameToFileSet, TagName).UnionWith(BuildProducts);
+			}
+		}
+
+		/// <summary>
+		/// Output this task out to an XML writer.
+		/// </summary>
+		public override void Write(XmlWriter Writer)
+		{
+			Write(Writer, Parameters);
+		}
+
+		/// <summary>
+		/// Find all the tags which are used as inputs to this task
+		/// </summary>
+		/// <returns>The tag names which are read by this task</returns>
+		public override IEnumerable<string> FindConsumedTagNames()
+		{
+			yield break;
+		}
+
+		/// <summary>
+		/// Find all the tags which are modified by this task
+		/// </summary>
+		/// <returns>The tag names which are modified by this task</returns>
+		public override IEnumerable<string> FindProducedTagNames()
+		{
+			return FindTagNamesFromList(Parameters.Tag);
+		}
+	}
+}
diff --git a/Engine/Source/Programs/Horde/StorageTests.xml b/Engine/Source/Programs/Horde/StorageTests.xml
new file mode 100644
--- /dev/null
+++ b/Engine/Source/Programs/Horde/StorageTests.xml
@@ -0,0 +1,61 @@
+<?xml version='1.0' ?>
+<BuildGraph xmlns="http://www.epicgames.com/BuildGraph" xmlns:xsi="http://www.w3.org/2001/XMLSchema-instance" xsi:schemaLocation="http://www.epicgames.com/BuildGraph ../../../Engine/Build/Graph/Schema.xsd" >
+
+	<Option Name="PreflightChange" DefaultValue="" Description="Preflight changelist number"/>
+	
+	<!-- Small Chunks -->
+	<Agent Name="Agent A" Type="CompileWin64;Win64">
+		<Node Name="Create Small Chunks" Produces="#SmallChunks">
+			<StorageTest Size="1024" Count="20000" Different="true" Tag="#SmallChunks"/>
+			<Log Message="Output Files:" Files="#SmallChunks"/>
+		</Node>
+	</Agent>
+	<Agent Name="Agent B" Type="CompileWin64;Win64">
+		<Node Name="Read Small Chunks" Requires="#SmallChunks">
+			<Log Message="Output Files:" Files="#SmallChunks"/>
+		</Node>
+	</Agent>
+
+	<!-- Large Chunks -->
+	<Agent Name="Agent C" Type="CompileWin64;Win64">
+		<Node Name="Create Large Chunks" Produces="#LargeChunks">
+			<StorageTest Size="1073741824" Count="20" Different="true" Tag="#LargeChunks"/>
+			<Log Message="Output Files:" Files="#LargeChunks"/>
+		</Node>
+	</Agent>
+	<Agent Name="Agent D" Type="CompileWin64;Win64">
+		<Node Name="Read Large Chunks" Requires="#LargeChunks">
+			<Log Message="Output Files:" Files="#LargeChunks"/>
+		</Node>
+	</Agent>
+	
+	<!-- Small Identical Chunks -->
+	<Agent Name="Agent E" Type="CompileWin64;Win64">
+		<Node Name="Create Small Identical Chunks" Produces="#SmallIdenticalChunks">
+			<StorageTest Size="1024" Count="20000" Different="false" Tag="#SmallIdenticalChunks"/>
+			<Log Message="Output Files:" Files="#SmallIdenticalChunks"/>
+		</Node>
+	</Agent>
+	<Agent Name="Agent F" Type="CompileWin64;Win64">
+		<Node Name="Read Small Identical Chunks" Requires="#SmallIdenticalChunks">
+			<Log Message="Output Files:" Files="#SmallIdenticalChunks"/>
+		</Node>
+	</Agent>
+
+	<!-- Large Identical Chunks -->
+	<Agent Name="Agent G" Type="CompileWin64;Win64">
+		<Node Name="Create Large Identical Chunks" Produces="#LargeIdenticalChunks">
+			<StorageTest Size="1073741824" Count="20" Different="true" Tag="#LargeIdenticalChunks"/>
+			<Log Message="Output Files:" Files="#LargeIdenticalChunks"/>
+		</Node>
+	</Agent>
+	<Agent Name="Agent H" Type="CompileWin64;Win64">
+		<Node Name="Read Large Identical Chunks" Requires="#LargeIdenticalChunks">
+			<Log Message="Output Files:" Files="#LargeIdenticalChunks"/>
+		</Node>
+	</Agent>
+
+	<!-- Aggregate for the whole build -->
+	<Aggregate Name="Storage Tests" Requires="Read Small Chunks;Read Large Chunks;Read Small Identical Chunks;Read Large Identical Chunks"/>
+
+</BuildGraph>
