Fixed issue where the default Android architectures would prevent plugins from being built.

diff --git a/Engine/Source/Programs/AutomationTool/Scripts/BuildPluginCommand.Automation.cs b/Engine/Source/Programs/AutomationTool/Scripts/BuildPluginCommand.Automation.cs
--- a/Engine/Source/Programs/AutomationTool/Scripts/BuildPluginCommand.Automation.cs
+++ b/Engine/Source/Programs/AutomationTool/Scripts/BuildPluginCommand.Automation.cs
@@ -23,7 +23,7 @@ using UnrealBuildBase;
 [Help("Architecture_<Platform>=<Architecture[s]>", "Control architecture to compile for a platform (eg. -Architecture_Mac=arm64+x86). Default is to use UBT defaults for the platform.")]
 public sealed class BuildPlugin : BuildCommand
 {
-	const string AndroidArchitectures = "armv7+arm64";
+	const string AndroidArchitectures = "arm64";
 	string UnrealBuildToolDllRelativePath = @"Engine/Binaries/DotNET/UnrealBuildTool/UnrealBuildTool.dll";
 	FileReference UnrealBuildToolDll;
 	static private Dictionary<UnrealTargetPlatform, string> PlatformToArchitectureMap = new Dictionary<UnrealTargetPlatform, string>();
