[Core] fix for automation reporting false failures for tests that purposely cause `ensure` or `check` failures

diff --git a/Engine/Source/Programs/AutomationTool/LowLevelTests/Tests/Gauntlet.LowLevelTestNode.cs b/Engine/Source/Programs/AutomationTool/LowLevelTests/Tests/Gauntlet.LowLevelTestNode.cs
--- a/Engine/Source/Programs/AutomationTool/LowLevelTests/Tests/Gauntlet.LowLevelTestNode.cs
+++ b/Engine/Source/Programs/AutomationTool/LowLevelTests/Tests/Gauntlet.LowLevelTestNode.cs
@@ -13,7 +13,7 @@ namespace LowLevelTests
 	public class LowLevelTests : BaseTest
 	{
 		private LowLevelTestContext Context;
-		
+
 		private IAppInstance TestInstance;
 
 		private DateTime SessionStartTime = DateTime.MinValue;
@@ -131,10 +131,10 @@ namespace LowLevelTests
 				Log.Warning("No StdOut returned from low level test app.");
 			}
 			else // Save log artifact
-			{	
+			{
 				const string ClientLogFile = "ClientOutput.log";
 				string ClientOutputLog = Path.Combine(ArtifactPath, ClientLogFile);
-				
+
 				using (var ClientOutputWriter = File.CreateText(ClientOutputLog))
 				{
 					ClientOutputWriter.Write(StdOut);
@@ -158,104 +158,56 @@ namespace LowLevelTests
 				Log.Error("Failed to copy report: {0}", ex.ToString());
 			}
 
-			bool ReportResult = false;
-			if (ReportCopied)
-			{
-				string ReportContents = File.ReadAllText(ReportPath);
-				Console.WriteLine(ReportContents);
-				LowLevelTestsLogParser LowLevelTestsLogParser = new LowLevelTestsLogParser(ReportContents);
-				ReportResult = LowLevelTestsLogParser.GetCatchTestResults().Passed;
-			}
 
-			UnrealProcessResult Result = GetExitCodeAndReason(
-				InReason,
-				ReportResult,
-				StdOut,
-				out string ExitReason,
-				out int ExitCode);
-
-			if (ExitCode == 0)
-			{
-				LowLevelTestResult = TestResult.Passed;
-			}
-			else
+			string ExitReason = "";
+			if (TestInstance.WasKilled)
 			{
-				if (Result == UnrealProcessResult.TimeOut)
+				if (InReason == StopReason.MaxDuration)
 				{
 					LowLevelTestResult = TestResult.TimedOut;
+					ExitReason = "Timed Out";
 				}
 				else
 				{
 					LowLevelTestResult = TestResult.Failed;
+					ExitReason = "Process was killed by Gauntlet.";
 				}
-				Log.Info("Low level test exited with code {0} and reason: {1}", ExitCode, ExitReason);
 			}
-		}
-
-		public override void CleanupTest()
-		{
-			if (LowLevelTestsApp != null)
+			else if (TestInstance.ExitCode != 0)
 			{
-				LowLevelTestsApp.Dispose();
-				LowLevelTestsApp = null;
+				LowLevelTestResult = TestResult.Failed;
+				ExitReason = $"Process exited with exit code {TestInstance.ExitCode}";
 			}
-		}
-
-		protected virtual UnrealProcessResult GetExitCodeAndReason(StopReason InReason, bool ReportResult, string StdOut, out string ExitReason, out int ExitCode)
-		{
-			UnrealLog UnrealLog = new UnrealLogParser(StdOut).GetSummary();
-			
-			if (TestInstance.WasKilled)
+			else if (!ReportCopied)
 			{
-				if (InReason == StopReason.MaxDuration)
+				LowLevelTestResult = TestResult.Failed;
+				ExitReason = "Uabled to read test report";
+			}
+			else
+			{
+				string ReportContents = File.ReadAllText(ReportPath);
+				Log.Info(ReportContents);
+				LowLevelTestsLogParser LowLevelTestsLogParser = new LowLevelTestsLogParser(ReportContents);
+				if (LowLevelTestsLogParser.GetCatchTestResults().Passed)
 				{
-					ExitReason = "Process was killed by Gauntlet due to a timeout.";
-					ExitCode = -1;
-					return UnrealProcessResult.TimeOut;
+					LowLevelTestResult = TestResult.Passed;
+					ExitReason = "Tests passed";
 				}
 				else
 				{
-					ExitReason = "Process was killed by Gauntlet.";
-					ExitCode = 0;
-					return UnrealProcessResult.ExitOk;
+					LowLevelTestResult = TestResult.Failed;
+					ExitReason = "Tests failed";
 				}
 			}
+			Log.Info($"Low level test exited with code {TestInstance.ExitCode} and reason: {ExitReason}");
+		}
 
-			// First we check for unreal specific issues.
-			// A successful test run must be free of these.
-			if (UnrealLog.FatalError != null)
-			{
-				ExitReason = "Process encountered fatal error";
-				ExitCode = -1;
-				return UnrealProcessResult.EncounteredFatalError;
-			}
-
-			if (UnrealLog.Ensures.Count() > 0)
-			{
-				ExitReason = string.Format("Process encountered {0} Ensures", UnrealLog.Ensures.Count());
-				ExitCode = -1;
-				return UnrealProcessResult.EncounteredEnsure;
-			}
-
-			if (UnrealLog.RequestedExit)
-			{
-				ExitReason = string.Format("Exit was requested: {0}", UnrealLog.RequestedExitReason);
-				ExitCode = 0;
-				return UnrealProcessResult.ExitOk;
-			}
-
-			// Then we check for actual test results.
-			if (ReportResult)
-			{
-				ExitReason = "All Catch2 tests passed.";
-				ExitCode = 0;
-				return UnrealProcessResult.ExitOk;
-			}
-			else
+		public override void CleanupTest()
+		{
+			if (LowLevelTestsApp != null)
 			{
-				ExitReason = "Catch2 tests failed.";
-				ExitCode = -1;
-				return UnrealProcessResult.TestFailure;
+				LowLevelTestsApp.Dispose();
+				LowLevelTestsApp = null;
 			}
 		}
 	}
diff --git a/Engine/Source/Runtime/Core/Tests/Misc/AssertionMacrosTest.cpp b/Engine/Source/Runtime/Core/Tests/Misc/AssertionMacrosTest.cpp
--- a/Engine/Source/Runtime/Core/Tests/Misc/AssertionMacrosTest.cpp
+++ b/Engine/Source/Runtime/Core/Tests/Misc/AssertionMacrosTest.cpp
@@ -92,14 +92,14 @@ TEST_CASE("Core::Misc::TestEnsure", "[Core][Misc][AssertionMacros][Ensure]")
 	}
 
 	{
-		//not much of test of CHECK_ENSURE as its hard to tell if the coninue part is working
+		//not much of test of CHECK_ENSURE as its hard to tell if the continue part is working
 		int result = 0;
 		CHECK_ENSURE(result = EnsureForCheck(0));
 		REQUIRE(result == 2);
 	}
 	
 	{
-		//multiple failed ensure should continue beecause its wrapped in a scope
+		//multiple failed ensure should continue because its wrapped in a scope
 		int result = 0;
 		REQUIRE_ENSURE(result = EnsureMultipleFailures(7));
 		REQUIRE(result == 8);
@@ -114,7 +114,7 @@ TEST_CASE("Core::Misc::TestEnsure", "[Core][Misc][AssertionMacros][Ensure]")
 	}
 }
 
-DISABLED_TEST_CASE("Core::Misc::TestCheck", "[Core][Misc][AssertionMacros][Check]")
+TEST_CASE("Core::Misc::TestCheck", "[Core][Misc][AssertionMacros][Check]")
 {
 	REQUIRE_CHECK(check(1 == 2));
 	REQUIRE_CHECK_MSG("1 == 2", checkf(1 == 2, TEXT("Error Message")));
