Disabling many creative tests that don't work
Adding new ContextFunctionalTest to look at all of the previous files that have been modified and find the tests that best match the modified files

Added additional core function for getting the referencers for a set of assets up to a specified search depth.

Added the ContextFunctionalTest gauntlet job so we can run this in horde
example launching method:

diff --git a/Engine/Source/Programs/AutomationTool/Gauntlet/Unreal/Automation/UE.Automation.cs b/Engine/Source/Programs/AutomationTool/Gauntlet/Unreal/Automation/UE.Automation.cs
--- a/Engine/Source/Programs/AutomationTool/Gauntlet/Unreal/Automation/UE.Automation.cs
+++ b/Engine/Source/Programs/AutomationTool/Gauntlet/Unreal/Automation/UE.Automation.cs
@@ -846,7 +846,7 @@ namespace UE
 			MarkdownBuilder MB = new MarkdownBuilder(base.GetTestSummaryHeader());
 
 			// Everything we need is in the editor artifacts
-			var EditorRole = RoleResults.Where(R => R.Artifacts.SessionRole.RoleType == UnrealTargetRole.Editor).FirstOrDefault();
+			var EditorRole = RoleResults.Where(R => R.Artifacts != null && R.Artifacts.SessionRole.RoleType == UnrealTargetRole.Editor).FirstOrDefault();
 
 			if (EditorRole != null)
 			{
@@ -1073,7 +1073,7 @@ namespace UE
 
 			foreach (var Role in RoleResults)
 			{
-				if (Role.Artifacts.SessionRole.RoleType == UnrealTargetRole.Editor)
+				if (Role.Artifacts != null && Role.Artifacts.SessionRole.RoleType == UnrealTargetRole.Editor)
 				{
 					AutomationLogParser Parser = new AutomationLogParser(Role.LogSummary.FullLogContent);
 					AllWarnings.AddRange(
diff --git a/Engine/Source/Programs/AutomationTool/Gauntlet/Unreal/Base/Gauntlet.UnrealSession.cs b/Engine/Source/Programs/AutomationTool/Gauntlet/Unreal/Base/Gauntlet.UnrealSession.cs
--- a/Engine/Source/Programs/AutomationTool/Gauntlet/Unreal/Base/Gauntlet.UnrealSession.cs
+++ b/Engine/Source/Programs/AutomationTool/Gauntlet/Unreal/Base/Gauntlet.UnrealSession.cs
@@ -1206,83 +1206,84 @@ namespace Gauntlet
 			int DummyClientCount = 0;
 
 			List<UnrealRoleArtifacts> AllArtifacts = new List<UnrealRoleArtifacts>();
-
-			foreach (UnrealSessionInstance.RoleInstance App in TestInstance.RunningRoles)
+			if (TestInstance != null)
 			{
-				string RoleName = (App.Role.IsDummy() ? "Dummy" : "") + App.Role.RoleType.ToString();
-				string FolderName = RoleName;
+				foreach (UnrealSessionInstance.RoleInstance App in TestInstance.RunningRoles)
+				{
+					string RoleName = (App.Role.IsDummy() ? "Dummy" : "") + App.Role.RoleType.ToString();
+					string FolderName = RoleName;
 
-				int RoleCount = 1;
+					int RoleCount = 1;
 
-				if (App.Role.IsDummy())
-				{
-					DummyClientCount++;
-					RoleCount = DummyClientCount;
-				}
-				else
-				{
-					if (!RoleCounts.ContainsKey(App.Role.RoleType))
+					if (App.Role.IsDummy())
 					{
-						RoleCounts.Add(App.Role.RoleType, 1);
+						DummyClientCount++;
+						RoleCount = DummyClientCount;
 					}
 					else
 					{
-						RoleCounts[App.Role.RoleType]++;
-					}
-
-					RoleCount = RoleCounts[App.Role.RoleType];
-				}
+						if (!RoleCounts.ContainsKey(App.Role.RoleType))
+						{
+							RoleCounts.Add(App.Role.RoleType, 1);
+						}
+						else
+						{
+							RoleCounts[App.Role.RoleType]++;
+						}
 
-				
-				if (RoleCount > 1)
-				{
-					FolderName += string.Format("_{0:00}", RoleCount);
-				}
+						RoleCount = RoleCounts[App.Role.RoleType];
+					}
 
-				string DestPath = Path.Combine(OutputPath, FolderName);
 
-				if (!App.Role.IsNullRole() && !App.Role.InstallOnly )
-				{
-					Log.VeryVerbose("Calling SaveRoleArtifacts, Role: {0}  Artifact Path: {1}", App.ToString(), App.AppInstance.ArtifactPath);
-					UnrealRoleArtifacts Artifacts = null;
-					ITargetDevice device = App.AppInstance.Device;
-					IDeviceUsageReporter.RecordStart(device.Name, device.Platform, IDeviceUsageReporter.EventType.SavingArtifacts);
-					try
+					if (RoleCount > 1)
 					{
-						Artifacts = SaveRoleArtifacts(Context, App, DestPath);
+						FolderName += string.Format("_{0:00}", RoleCount);
 					}
-					catch (Exception SaveArtifactsException)
-					{
-						// Caught an exception -> report failure
-						IDeviceUsageReporter.RecordEnd(device.Name, device.Platform, IDeviceUsageReporter.EventType.SavingArtifacts, IDeviceUsageReporter.EventState.Failure);
 
-						// Retry once only, after rebooting device, if artifacts couldn't be saved.
-						if (SaveArtifactsException.Message.Contains("A retry should be performed"))
+					string DestPath = Path.Combine(OutputPath, FolderName);
+
+					if (!App.Role.IsNullRole() && !App.Role.InstallOnly)
+					{
+						Log.VeryVerbose("Calling SaveRoleArtifacts, Role: {0}  Artifact Path: {1}", App.ToString(), App.AppInstance.ArtifactPath);
+						UnrealRoleArtifacts Artifacts = null;
+						ITargetDevice device = App.AppInstance.Device;
+						IDeviceUsageReporter.RecordStart(device.Name, device.Platform, IDeviceUsageReporter.EventType.SavingArtifacts);
+						try
 						{
-							Log.Info("Rebooting device and retrying save role artifacts once.");
-							App.AppInstance.Device.Reboot();
 							Artifacts = SaveRoleArtifacts(Context, App, DestPath);
 						}
-						else
+						catch (Exception SaveArtifactsException)
 						{
-							// Pass exception to the surrounding try/catch in UnrealTestNode while preserving the original callstack
-							throw;
+							// Caught an exception -> report failure
+							IDeviceUsageReporter.RecordEnd(device.Name, device.Platform, IDeviceUsageReporter.EventType.SavingArtifacts, IDeviceUsageReporter.EventState.Failure);
+
+							// Retry once only, after rebooting device, if artifacts couldn't be saved.
+							if (SaveArtifactsException.Message.Contains("A retry should be performed"))
+							{
+								Log.Info("Rebooting device and retrying save role artifacts once.");
+								App.AppInstance.Device.Reboot();
+								Artifacts = SaveRoleArtifacts(Context, App, DestPath);
+							}
+							else
+							{
+								// Pass exception to the surrounding try/catch in UnrealTestNode while preserving the original callstack
+								throw;
+							}
 						}
-					}
-					// Did not catch -> successful reporting
-					IDeviceUsageReporter.RecordEnd(device.Name, device.Platform, IDeviceUsageReporter.EventType.SavingArtifacts, IDeviceUsageReporter.EventState.Success);
+						// Did not catch -> successful reporting
+						IDeviceUsageReporter.RecordEnd(device.Name, device.Platform, IDeviceUsageReporter.EventType.SavingArtifacts, IDeviceUsageReporter.EventState.Success);
 
-					if (Artifacts != null)
+						if (Artifacts != null)
+						{
+							AllArtifacts.Add(Artifacts);
+						}
+					}
+					else
 					{
-						AllArtifacts.Add(Artifacts);
+						Log.Verbose("Skipping SaveRoleArtifacts for Null Role: {0}", App.ToString());
 					}
 				}
-				else 
-				{
-					Log.Verbose("Skipping SaveRoleArtifacts for Null Role: {0}", App.ToString());
-				}
 			}
-
 			return AllArtifacts;
 		}
 	}
diff --git a/Engine/Source/Programs/AutomationTool/Gauntlet/Unreal/Base/Gauntlet.UnrealTestNode.cs b/Engine/Source/Programs/AutomationTool/Gauntlet/Unreal/Base/Gauntlet.UnrealTestNode.cs
--- a/Engine/Source/Programs/AutomationTool/Gauntlet/Unreal/Base/Gauntlet.UnrealTestNode.cs
+++ b/Engine/Source/Programs/AutomationTool/Gauntlet/Unreal/Base/Gauntlet.UnrealTestNode.cs
@@ -1080,19 +1080,29 @@ namespace Gauntlet
 
 			// Shutdown the instance so we can access all files, but do not null it or shutdown the UnrealApp because we still need
 			// access to these objects and their resources! Final cleanup is done in CleanupTest()
-			TestInstance.Shutdown();
+			if (TestInstance != null)
+			{
+				TestInstance.Shutdown();
+			}
 
 			try
 			{
-				Log.Info("Saving artifacts to {0}", ArtifactPath);
-				// run create dir again just in case the already made dir was cleaned up by another buildfarm job or something similar.
-				Directory.CreateDirectory(ArtifactPath);
-				Utils.SystemHelpers.MarkDirectoryForCleanup(ArtifactPath);
-
+				if (string.IsNullOrEmpty(ArtifactPath))
+				{
+					Log.Info("ArtifactsPath was null unable to save artifacts");
+				}
+				else
+				{
+					Log.Info("Saving artifacts to {0}", ArtifactPath);
+					// run create dir again just in case the already made dir was cleaned up by another buildfarm job or something similar.
+					Directory.CreateDirectory(ArtifactPath);
+					Utils.SystemHelpers.MarkDirectoryForCleanup(ArtifactPath);
+				}
 				SessionArtifacts = SaveRoleArtifacts(ArtifactPath);
 
 				// call legacy version
 				SaveArtifacts_DEPRECATED(ArtifactPath);
+				
 			}
 			catch (Exception Ex)
 			{
@@ -1518,6 +1528,10 @@ namespace Gauntlet
 		/// <returns></returns>
 		public virtual IEnumerable<UnrealRoleArtifacts> SaveRoleArtifacts(string OutputPath)
 		{
+			if (UnrealApp == null)
+			{
+				return new List<UnrealRoleArtifacts>();
+			}
 			return UnrealApp.SaveRoleArtifacts(Context, TestInstance, ArtifactPath);
 		}
 
@@ -1719,6 +1733,24 @@ namespace Gauntlet
 			return new UnrealLogParser(InArtifacts.AppInstance.StdOut).GetSummary();
 		}
 
+		protected virtual UnrealRoleResult GetRoleResultWhenNoArtifactsAreGenerated(StopReason InReason)
+		{
+			int ExitCode = 0;
+			string ExitReason;
+
+			// Give ourselves (and derived classes) a chance to analyze what happened
+			UnrealProcessResult ProcessResult = GetExitCodeAndReason(InReason, null, null, out ExitReason, out ExitCode);
+			// if the test is stopping for a reason other than completion, mark this as failing incase derived classes
+			// don't do the right thing
+			if (InReason == StopReason.MaxDuration)
+			{
+				ProcessResult = UnrealProcessResult.TimeOut;
+				ExitCode = -1;
+			}
+
+			return new UnrealRoleResult(ProcessResult, ExitCode, ExitReason, new UnrealLog(), null, new List<UnrealTestEvent>());
+		}
+
 		/// <summary>
 		/// Returns a list of all results for the roles involved in this test by calling CreateRoleResultFromArtifact for all
 		/// artifacts in the list
@@ -1727,7 +1759,14 @@ namespace Gauntlet
 		/// <returns></returns>
 		protected virtual IEnumerable<UnrealRoleResult> CreateRoleResultsFromArtifacts(StopReason InReason, IEnumerable<UnrealRoleArtifacts> InAllArtifacts)
 		{
-			return InAllArtifacts.Select(A => CreateRoleResultFromArtifact(InReason, A)).ToArray();
+			if (InAllArtifacts != null && InAllArtifacts.Count() > 0)
+			{
+				return InAllArtifacts.Select(A => CreateRoleResultFromArtifact(InReason, A)).ToArray();
+			}
+			else
+			{
+				return new List<UnrealRoleResult>() { GetRoleResultWhenNoArtifactsAreGenerated(InReason) };
+			}
 		}
 
 
@@ -1843,24 +1882,26 @@ namespace Gauntlet
 			UnrealLog LogSummary = InRoleResult.LogSummary;
 			MarkdownBuilder MB = new MarkdownBuilder();
 
-			UnrealRoleArtifacts RoleArtifacts = InRoleResult.Artifacts;
+			if (InRoleResult.Artifacts != null)
+			{
+				UnrealRoleArtifacts RoleArtifacts = InRoleResult.Artifacts;
 
-			MB.H3(string.Format("Role: {0} ({1} {2})", RoleArtifacts.SessionRole.RoleType, RoleArtifacts.SessionRole.Platform, RoleArtifacts.SessionRole.Configuration));
+				MB.H3(string.Format("Role: {0} ({1} {2})", RoleArtifacts.SessionRole.RoleType, RoleArtifacts.SessionRole.Platform, RoleArtifacts.SessionRole.Configuration));
 
-			bool HaveFatalError = LogSummary.FatalError != null;
+				bool HaveFatalError = LogSummary.FatalError != null;
 
-			// If we have no fatal error use "Error" which Horde will highlight. Otherwise use "Failed" since the fatal error will be below and highligted
-			string FailedString = HaveFatalError ? "Failed: " : "Error: ";
-			string CompletedString = "Completed: ";
+				// If we have no fatal error use "Error" which Horde will highlight. Otherwise use "Failed" since the fatal error will be below and highligted
+				string FailedString = HaveFatalError ? "Failed: " : "Error: ";
+				string CompletedString = "Completed: ";
 
-			string RoleState = InRoleResult.ExitCode != 0 && InRoleResult.LogSummary.HasAbnormalExit ? FailedString : CompletedString;
+				string RoleState = InRoleResult.ExitCode != 0 && InRoleResult.LogSummary.HasAbnormalExit ? FailedString : CompletedString;
 
-			MB.H4(string.Format("{0} {1} ({2}, ExitCode={3})", RoleState, InRoleResult.Summary, InRoleResult.ProcessResult, InRoleResult.ExitCode));
+				MB.H4(string.Format("{0} {1} ({2}, ExitCode={3})", RoleState, InRoleResult.Summary, InRoleResult.ProcessResult, InRoleResult.ExitCode));
 
-			// log command line up here for visibility
-			//MB.Paragraph(string.Format("CommandLine: {0}", RoleArtifacts.AppInstance.CommandLine));
+				// log command line up here for visibility
+				//MB.Paragraph(string.Format("CommandLine: {0}", RoleArtifacts.AppInstance.CommandLine));
 
-			MB.UnorderedList(new string[] {
+				MB.UnorderedList(new string[] {
 				string.Format("CommandLine: {0}", RoleArtifacts.AppInstance.CommandLine),
 				string.Format("Log: {0}", RoleArtifacts.LogPath),
 				string.Format("SavedDir: {0}", RoleArtifacts.ArtifactPath),
@@ -1870,112 +1911,112 @@ namespace Gauntlet
 				LogSummary.Warnings.Count() > 0 ? string.Format("Log Warnings: {0}", LogSummary.Warnings.Count()) : null
 			});
 
-			// Separate the events we want to report on
-			IEnumerable<UnrealTestEvent> Asserts = InRoleResult.Events.Where(E => E.Severity == EventSeverity.Fatal);
-			IEnumerable<UnrealTestEvent> Errors = InRoleResult.Events.Where(E => E.Severity == EventSeverity.Error);
-			IEnumerable<UnrealTestEvent> Ensures = InRoleResult.Events.Where(E => E.IsEnsure);
-			IEnumerable<UnrealTestEvent> Warnings = InRoleResult.Events.Where(E => E.Severity == EventSeverity.Warning && !E.IsEnsure);
-
-			foreach (UnrealTestEvent Event in Asserts)
-			{
-				MB.H4(string.Format("Fatal Error: {0}", Event.Summary));
+				// Separate the events we want to report on
+				IEnumerable<UnrealTestEvent> Asserts = InRoleResult.Events.Where(E => E.Severity == EventSeverity.Fatal);
+				IEnumerable<UnrealTestEvent> Errors = InRoleResult.Events.Where(E => E.Severity == EventSeverity.Error);
+				IEnumerable<UnrealTestEvent> Ensures = InRoleResult.Events.Where(E => E.IsEnsure);
+				IEnumerable<UnrealTestEvent> Warnings = InRoleResult.Events.Where(E => E.Severity == EventSeverity.Warning && !E.IsEnsure);
 
-				if (Event.Callstack.Any())
+				foreach (UnrealTestEvent Event in Asserts)
 				{
-					MB.UnorderedList(Event.Callstack.Take(MaxCallstackLines));
+					MB.H4(string.Format("Fatal Error: {0}", Event.Summary));
 
-					if (Event.Callstack.Count() > MaxCallstackLines)
+					if (Event.Callstack.Any())
 					{
-						MB.Paragraph("See log for full callstack");
+						MB.UnorderedList(Event.Callstack.Take(MaxCallstackLines));
+
+						if (Event.Callstack.Count() > MaxCallstackLines)
+						{
+							MB.Paragraph("See log for full callstack");
+						}
+					}
+					else
+					{
+						MB.Paragraph("Could not parse callstack. See log for full callstack");
 					}
 				}
-				else
-				{
-					MB.Paragraph("Could not parse callstack. See log for full callstack");
-				}
-			}
 
-			foreach (UnrealTestEvent Event in Ensures.Distinct())
-			{
-				MB.H4(string.Format("Warning: Ensure: {0}", Event.Summary));
-
-				if (Event.Callstack.Any())
+				foreach (UnrealTestEvent Event in Ensures.Distinct())
 				{
-					MB.UnorderedList(Event.Callstack.Take(MaxCallstackLines));
+					MB.H4(string.Format("Warning: Ensure: {0}", Event.Summary));
 
-					if (Event.Callstack.Count() > MaxCallstackLines)
+					if (Event.Callstack.Any())
 					{
-						MB.Paragraph("See log for full callstack");
+						MB.UnorderedList(Event.Callstack.Take(MaxCallstackLines));
+
+						if (Event.Callstack.Count() > MaxCallstackLines)
+						{
+							MB.Paragraph("See log for full callstack");
+						}
+					}
+					else
+					{
+						MB.Paragraph("Could not parse callstack. See log for full callstack");
 					}
 				}
-				else
-				{
-					MB.Paragraph("Could not parse callstack. See log for full callstack");
-				}
-			}
 
-			if (Errors.Any())
-			{
-				var ErrorList = Errors.Select(E => E.Summary).Distinct();
-				var PrintedErrorList = ErrorList;
+				if (Errors.Any())
+				{
+					var ErrorList = Errors.Select(E => E.Summary).Distinct();
+					var PrintedErrorList = ErrorList;
 
-				string TrimStatement = "";
+					string TrimStatement = "";
 
-				// too many warnings. If there was an abnormal exit show the last ones as they may be relevant
-				if (ErrorList.Count() > MaxLogLines)
-				{
-					if (LogSummary.HasAbnormalExit)
+					// too many warnings. If there was an abnormal exit show the last ones as they may be relevant
+					if (ErrorList.Count() > MaxLogLines)
 					{
-						PrintedErrorList = ErrorList.Skip(ErrorList.Count() - MaxLogLines);
-						TrimStatement = string.Format("(Last {0} of {1} errors)", MaxLogLines, ErrorList.Count());
+						if (LogSummary.HasAbnormalExit)
+						{
+							PrintedErrorList = ErrorList.Skip(ErrorList.Count() - MaxLogLines);
+							TrimStatement = string.Format("(Last {0} of {1} errors)", MaxLogLines, ErrorList.Count());
+						}
+						else
+						{
+							PrintedErrorList = ErrorList.Take(MaxLogLines);
+							TrimStatement = string.Format("(First {0} of {1} errors)", MaxLogLines, ErrorList.Count());
+						}
 					}
-					else
+
+					MB.H4("Errors:");
+					MB.UnorderedList(PrintedErrorList);
+
+					if (!string.IsNullOrEmpty(TrimStatement))
 					{
-						PrintedErrorList = ErrorList.Take(MaxLogLines);
-						TrimStatement = string.Format("(First {0} of {1} errors)", MaxLogLines, ErrorList.Count());
+						MB.Paragraph(TrimStatement);
 					}
 				}
 
-				MB.H4("Errors:");
-				MB.UnorderedList(PrintedErrorList);
-
-				if (!string.IsNullOrEmpty(TrimStatement))
+				if (Warnings.Any())
 				{
-					MB.Paragraph(TrimStatement);
-				}
-			}
+					var WarningList = Warnings.Select(E => E.Summary).Distinct();
+					var PrintedWarningList = WarningList;
 
-			if (Warnings.Any())
-			{
-				var WarningList = Warnings.Select(E => E.Summary).Distinct();
-				var PrintedWarningList = WarningList;
+					string TrimStatement = "";
 
-				string TrimStatement = "";
-
-				// too many warnings. If there was an abnormal exit show the last ones as they may be relevant
-				if (WarningList.Count() > MaxLogLines)
-				{
-					if (LogSummary.HasAbnormalExit)
-					{
-						PrintedWarningList = WarningList.Skip(WarningList.Count() - MaxLogLines);
-						TrimStatement = string.Format("(Last {0} of {1} errors)", MaxLogLines, WarningList.Count());
-					}
-					else
+					// too many warnings. If there was an abnormal exit show the last ones as they may be relevant
+					if (WarningList.Count() > MaxLogLines)
 					{
-						PrintedWarningList = WarningList.Take(MaxLogLines);
-						TrimStatement = string.Format("(First {0} of {1} errors)", MaxLogLines, WarningList.Count());
+						if (LogSummary.HasAbnormalExit)
+						{
+							PrintedWarningList = WarningList.Skip(WarningList.Count() - MaxLogLines);
+							TrimStatement = string.Format("(Last {0} of {1} errors)", MaxLogLines, WarningList.Count());
+						}
+						else
+						{
+							PrintedWarningList = WarningList.Take(MaxLogLines);
+							TrimStatement = string.Format("(First {0} of {1} errors)", MaxLogLines, WarningList.Count());
+						}
 					}
-				}
 
-				MB.H4("Warnings:");
-				MB.UnorderedList(PrintedWarningList);
+					MB.H4("Warnings:");
+					MB.UnorderedList(PrintedWarningList);
 
-				if (!string.IsNullOrEmpty(TrimStatement))
-				{
-					MB.Paragraph(TrimStatement);
+					if (!string.IsNullOrEmpty(TrimStatement))
+					{
+						MB.Paragraph(TrimStatement);
+					}
 				}
 			}
-
 			return MB.ToString();
 		}
 
@@ -2226,7 +2267,15 @@ namespace Gauntlet
 			// displayed individually by GetFormattedRoleSummary
 			foreach (var RoleResult in SortedRoles)
 			{
-				string RoleName = RoleResult.Artifacts.SessionRole.RoleType.ToString();
+				string RoleName;
+				if (RoleResult.Artifacts == null)
+				{
+					RoleName = "NoRoleArtifact";
+				}
+				else
+				{
+					RoleName = RoleResult.Artifacts.SessionRole.RoleType.ToString();
+				}
 
 				MB.Paragraph(string.Format("{0} Role: {1} ({2}, ExitCode {3})", RoleName, RoleResult.Summary, RoleResult.ProcessResult, RoleResult.ExitCode));
 
diff --git a/Engine/Source/Runtime/Engine/Classes/Engine/AssetManager.h b/Engine/Source/Runtime/Engine/Classes/Engine/AssetManager.h
--- a/Engine/Source/Runtime/Engine/Classes/Engine/AssetManager.h
+++ b/Engine/Source/Runtime/Engine/Classes/Engine/AssetManager.h
@@ -461,6 +461,9 @@ public:
 	/** Dumps out list of primary asset -> managed assets to log */
 	static void DumpReferencersForPackage(const TArray< FString >& PackageNames);
 
+	/** Finds all the referencers for a set of packages. Recursively will get all the referencers up to the max depth */
+	static void GetAllReferencersForPackage(TSet<FAssetData>& OutFoundAssets, const TArray<FName>& InPackageNames, int32 MaxDepth);
+
 	/** Starts initial load, gets called from InitializeObjectReferences */
 	virtual void StartInitialLoading();
 
diff --git a/Engine/Source/Runtime/Engine/Private/AssetManager.cpp b/Engine/Source/Runtime/Engine/Private/AssetManager.cpp
--- a/Engine/Source/Runtime/Engine/Private/AssetManager.cpp
+++ b/Engine/Source/Runtime/Engine/Private/AssetManager.cpp
@@ -3093,6 +3093,39 @@ void UAssetManager::DumpReferencersForPackage(const TArray< FString >& PackageNa
 	Manager.WriteCustomReport(FString::Printf(TEXT("ReferencersForPackage%s%s.gv"), *PackageNames[0], *FDateTime::Now().ToString()), ReportLines);
 }
 
+void UAssetManager::GetAllReferencersForPackage(TSet<FAssetData>& OutFoundAssets, const TArray<FName>& InPackageNames, int32 MaxDepth)
+{
+	ensureMsgf(MaxDepth > 0, TEXT("Max depth to search for referencers should be greater than 0"));
+
+	UAssetManager& Manager = Get();
+	IAssetRegistry& AssetRegistry = Manager.GetAssetRegistry();
+
+	TSet<FName> PackagesForNextLoop;
+	TSet<FName> CurrentPackagesToProcess;
+
+	CurrentPackagesToProcess.Append(InPackageNames);
+	for (int32 CurrentDepth = 0; CurrentDepth < MaxDepth; CurrentDepth++)
+	{
+		for (const FName& PackageName : CurrentPackagesToProcess)
+		{
+			FARFilter Filter;
+			Filter.bIncludeOnlyOnDiskAssets = true;
+			AssetRegistry.GetReferencers(PackageName, Filter.PackageNames, UE::AssetRegistry::EDependencyCategory::Package);
+
+			TArray<FAssetData> AssetReferencers;
+			AssetRegistry.GetAssets(Filter, AssetReferencers);
+			OutFoundAssets.Append(AssetReferencers);
+			for (const FAssetData& AssetData : AssetReferencers)
+			{
+				PackagesForNextLoop.Add(AssetData.PackageName);
+			}
+		}
+		CurrentPackagesToProcess.Reset();
+		CurrentPackagesToProcess.Append(PackagesForNextLoop);
+		PackagesForNextLoop.Reset();
+	}
+}
+
 FName UAssetManager::GetEncryptionKeyAssetTagName()
 {
 	static const FName NAME_EncryptionKey(TEXT("EncryptionKey"));
diff --git a/Engine/Source/Runtime/Engine/Public/Tests/AutomationCommon.h b/Engine/Source/Runtime/Engine/Public/Tests/AutomationCommon.h
--- a/Engine/Source/Runtime/Engine/Public/Tests/AutomationCommon.h
+++ b/Engine/Source/Runtime/Engine/Public/Tests/AutomationCommon.h
@@ -60,6 +60,7 @@ namespace AutomationCommon
 	 */
 	ENGINE_API SWidget* FindWidgetByTag(const FName Tag);
 
+	ENGINE_API UWorld* GetAnyGameWorld();
 #endif
 }
 
