Add null check on InLog

diff --git a/Engine/Source/Programs/AutomationTool/Gauntlet/Unreal/Automation/UE.Automation.cs b/Engine/Source/Programs/AutomationTool/Gauntlet/Unreal/Automation/UE.Automation.cs
--- a/Engine/Source/Programs/AutomationTool/Gauntlet/Unreal/Automation/UE.Automation.cs
+++ b/Engine/Source/Programs/AutomationTool/Gauntlet/Unreal/Automation/UE.Automation.cs
@@ -522,37 +522,39 @@ namespace UE
 			const float IdleTimeout = 30 * 60;
 
 			List<string> ChannelEntries = new List<string>();
+			if (TestInstance != null)
+			{
+				// We are primarily interested in what the editor is doing
+				var AppInstance = TestInstance.EditorApp;
 
-			// We are primarily interested in what the editor is doing
-			var AppInstance = TestInstance.EditorApp;
-
-			UnrealLogParser Parser = new UnrealLogParser(AppInstance.StdOut);
-			ChannelEntries.AddRange(Parser.GetEditorBusyChannels());
+				UnrealLogParser Parser = new UnrealLogParser(AppInstance.StdOut);
+				ChannelEntries.AddRange(Parser.GetEditorBusyChannels());
 
-			// Any new entries?
-			if (ChannelEntries.Count > LastAutomationEntryCount)
-			{
-				// log new entries so people have something to look at
-				ChannelEntries.Skip(LastAutomationEntryCount).ToList().ForEach(S => Log.Info("{0}", S));
-				LastAutomationEntryTime = DateTime.Now;
-				LastAutomationEntryCount = ChannelEntries.Count;
-			}
-			else
-			{
-				// Check for timeouts
-				if (LastAutomationEntryTime == DateTime.MinValue)
+				// Any new entries?
+				if (ChannelEntries.Count > LastAutomationEntryCount)
 				{
+					// log new entries so people have something to look at
+					ChannelEntries.Skip(LastAutomationEntryCount).ToList().ForEach(S => Log.Info("{0}", S));
 					LastAutomationEntryTime = DateTime.Now;
+					LastAutomationEntryCount = ChannelEntries.Count;
 				}
+				else
+				{
+					// Check for timeouts
+					if (LastAutomationEntryTime == DateTime.MinValue)
+					{
+						LastAutomationEntryTime = DateTime.Now;
+					}
 
-				double ElapsedTime = (DateTime.Now - LastAutomationEntryTime).TotalSeconds;
+					double ElapsedTime = (DateTime.Now - LastAutomationEntryTime).TotalSeconds;
 
-				// Check for timeout
-				if (ElapsedTime > IdleTimeout)
-				{
-					Log.Error("No activity observed in last {0:0.00} minutes. Aborting test", IdleTimeout / 60);
-					MarkTestComplete();
-					SetUnrealTestResult(TestResult.TimedOut);
+					// Check for timeout
+					if (ElapsedTime > IdleTimeout)
+					{
+						Log.Error("No activity observed in last {0:0.00} minutes. Aborting test", IdleTimeout / 60);
+						MarkTestComplete();
+						SetUnrealTestResult(TestResult.TimedOut);
+					}
 				}
 			}
 
@@ -750,7 +752,8 @@ namespace UE
 			UnrealProcessResult UnrealResult = base.GetExitCodeAndReason(InReason, InLog, InArtifacts, out ExitReason, out ExitCode);
 
 			// The editor is an additional arbiter of success
-			if (InArtifacts.SessionRole.RoleType == UnrealTargetRole.Editor 
+			if (InArtifacts.SessionRole.RoleType == UnrealTargetRole.Editor
+				&& InLog != null
 				&& InLog.HasAbnormalExit == false)
 			{
 				// if no fatal errors, check test results
diff --git a/Engine/Source/Programs/AutomationTool/Gauntlet/Unreal/Automation/UE.ErrorTest.cs b/Engine/Source/Programs/AutomationTool/Gauntlet/Unreal/Automation/UE.ErrorTest.cs
--- a/Engine/Source/Programs/AutomationTool/Gauntlet/Unreal/Automation/UE.ErrorTest.cs
+++ b/Engine/Source/Programs/AutomationTool/Gauntlet/Unreal/Automation/UE.ErrorTest.cs
@@ -120,88 +120,91 @@ namespace UE
 
 			TestResult FinalResult = TestResult.Invalid;
 
-			if (ErrorType == ErrorTypes.Ensure)
+			if (InLog != null)
 			{
-				// for an ensure we should have an entry and a callstack
-				int EnsureCount = InLog.Ensures.Count();
-				int CallstackLength = EnsureCount > 0 ? InLog.Ensures.First().Callstack.Length : 0;
-
-				if (EnsureCount == 0)
-				{
-					FinalResult = TestResult.Failed;
-					LocalReason = string.Format("No ensure error found for failure of type {0}", ErrorType);
-				}
-				else if (EnsureCount != 1)
-				{
-					FinalResult = TestResult.Failed;
-					LocalReason = string.Format("Incorrect ensure count found for failure of type {0}", ErrorType);
-				}
-				else if (CallstackLength == 0)
-				{
-					FinalResult = TestResult.Failed;
-					LocalReason = string.Format("No callstack error found for failure of type {0}", ErrorType);
-				}
-				else
-				{
-					FinalResult = TestResult.Passed;
-					LocalReason = string.Format("Found {0} ensures, test result = {1}", EnsureCount, FinalResult);
-				}
-			}
-			else
-			{
-				if (InLog.FatalError == null)
+				if (ErrorType == ErrorTypes.Ensure)
 				{
-					FinalResult = TestResult.Failed;
-					Log.Info("No fatal error found for failure of type {0}", ErrorType);
-				}
-				else if (InLog.FatalError.Callstack.Length == 0)
-				{
-					FinalResult = TestResult.Failed;
-					Log.Info("No callstack found for failure of type {0}", ErrorType);
+					// for an ensure we should have an entry and a callstack
+					int EnsureCount = InLog.Ensures.Count();
+					int CallstackLength = EnsureCount > 0 ? InLog.Ensures.First().Callstack.Length : 0;
+
+					if (EnsureCount == 0)
+					{
+						FinalResult = TestResult.Failed;
+						LocalReason = string.Format("No ensure error found for failure of type {0}", ErrorType);
+					}
+					else if (EnsureCount != 1)
+					{
+						FinalResult = TestResult.Failed;
+						LocalReason = string.Format("Incorrect ensure count found for failure of type {0}", ErrorType);
+					}
+					else if (CallstackLength == 0)
+					{
+						FinalResult = TestResult.Failed;
+						LocalReason = string.Format("No callstack error found for failure of type {0}", ErrorType);
+					}
+					else
+					{
+						FinalResult = TestResult.Passed;
+						LocalReason = string.Format("Found {0} ensures, test result = {1}", EnsureCount, FinalResult);
+					}
 				}
 				else
 				{
-					// all of these should contain a message and a result
-					if (ErrorType == ErrorTypes.Check)
+					if (InLog.FatalError == null)
 					{
-						if (!InLog.FatalError.Message.ToLower().Contains("assertion failed"))
-						{
-							FinalResult = TestResult.Failed;
-							LocalReason = string.Format("Unexpected assertion message");
-						}
-						else
-						{
-							FinalResult = TestResult.Passed;
-						}
-						Log.Info("Assertion message was {0}", InLog.FatalError.Message);
+						FinalResult = TestResult.Failed;
+						Log.Info("No fatal error found for failure of type {0}", ErrorType);
 					}
-					else if (ErrorType == ErrorTypes.Fatal)
+					else if (InLog.FatalError.Callstack.Length == 0)
 					{
-						if (!InLog.FatalError.Message.ToLower().Contains("fatal erro"))
-						{
-							FinalResult = TestResult.Failed;
-							LocalReason = string.Format("Unexpected Fatal Error message");
-						}
-						else
-						{
-							FinalResult = TestResult.Passed;
-						}
-
-						Log.Info("Fatal Error message was {0}", InLog.FatalError.Message);
+						FinalResult = TestResult.Failed;
+						Log.Info("No callstack found for failure of type {0}", ErrorType);
 					}
-					else if (ErrorType == ErrorTypes.GPF)
+					else
 					{
-						if (!InLog.FatalError.Message.ToLower().Contains("exception"))
+						// all of these should contain a message and a result
+						if (ErrorType == ErrorTypes.Check)
 						{
-							FinalResult = TestResult.Failed;
-							LocalReason = string.Format("Unexpected exception message");
+							if (!InLog.FatalError.Message.ToLower().Contains("assertion failed"))
+							{
+								FinalResult = TestResult.Failed;
+								LocalReason = string.Format("Unexpected assertion message");
+							}
+							else
+							{
+								FinalResult = TestResult.Passed;
+							}
+							Log.Info("Assertion message was {0}", InLog.FatalError.Message);
 						}
-						else
+						else if (ErrorType == ErrorTypes.Fatal)
 						{
-							FinalResult = TestResult.Passed;
+							if (!InLog.FatalError.Message.ToLower().Contains("fatal erro"))
+							{
+								FinalResult = TestResult.Failed;
+								LocalReason = string.Format("Unexpected Fatal Error message");
+							}
+							else
+							{
+								FinalResult = TestResult.Passed;
+							}
+
+							Log.Info("Fatal Error message was {0}", InLog.FatalError.Message);
+						}
+						else if (ErrorType == ErrorTypes.GPF)
+						{
+							if (!InLog.FatalError.Message.ToLower().Contains("exception"))
+							{
+								FinalResult = TestResult.Failed;
+								LocalReason = string.Format("Unexpected exception message");
+							}
+							else
+							{
+								FinalResult = TestResult.Passed;
+							}
+
+							Log.Info("Exception message was {0}", InLog.FatalError.Message);
 						}
-
-						Log.Info("Exception message was {0}", InLog.FatalError.Message);
 					}
 				}
 			}
diff --git a/Engine/Source/Programs/AutomationTool/Gauntlet/Unreal/Base/Gauntlet.UnrealTestNode.cs b/Engine/Source/Programs/AutomationTool/Gauntlet/Unreal/Base/Gauntlet.UnrealTestNode.cs
--- a/Engine/Source/Programs/AutomationTool/Gauntlet/Unreal/Base/Gauntlet.UnrealTestNode.cs
+++ b/Engine/Source/Programs/AutomationTool/Gauntlet/Unreal/Base/Gauntlet.UnrealTestNode.cs
@@ -1545,13 +1545,21 @@ namespace Gauntlet
 		protected virtual UnrealProcessResult GetExitCodeAndReason(StopReason InReason, UnrealLog InLog, UnrealRoleArtifacts InArtifacts, out string ExitReason, out int ExitCode)
 		{
 			// first check for fatal issues
-			if (InLog.FatalError != null)
+			if (InLog == null
+				|| InLog.FatalError != null)
 			{
 				ExitReason = "Process encountered fatal error";
 				ExitCode = -1;
 				return UnrealProcessResult.EncounteredFatalError;
 			}
 
+			if (InArtifacts == null)
+			{
+				ExitReason = "No role artifacts generated";
+				ExitCode = -1;
+				return UnrealProcessResult.TestFailure;
+			}
+
 			// Catch failed engine init. Early issues can result in the engine exiting with hard to diagnose reasons
 			if (InLog.EngineInitialized == false)
 			{
diff --git a/Engine/Source/Programs/AutomationTool/Gauntlet/Unreal/Engine/UnrealTests.cs b/Engine/Source/Programs/AutomationTool/Gauntlet/Unreal/Engine/UnrealTests.cs
--- a/Engine/Source/Programs/AutomationTool/Gauntlet/Unreal/Engine/UnrealTests.cs
+++ b/Engine/Source/Programs/AutomationTool/Gauntlet/Unreal/Engine/UnrealTests.cs
@@ -212,6 +212,7 @@ namespace Gauntlet
 
 				// The editor is an additional arbiter of success
 				if (InArtifacts.SessionRole.RoleType == UnrealTargetRole.Editor
+					&& InLog != null
 					&& InLog.HasAbnormalExit == false)
 				{
 					// if no fatal errors, check test results
