[Automation] Gives the ability to specify a "-staged" param which will append a Loose BuildFlag requirement if valid

diff --git a/Engine/Source/Programs/AutomationTool/Gauntlet/Unreal/Base/Gauntlet.UnrealSession.cs b/Engine/Source/Programs/AutomationTool/Gauntlet/Unreal/Base/Gauntlet.UnrealSession.cs
--- a/Engine/Source/Programs/AutomationTool/Gauntlet/Unreal/Base/Gauntlet.UnrealSession.cs
+++ b/Engine/Source/Programs/AutomationTool/Gauntlet/Unreal/Base/Gauntlet.UnrealSession.cs
@@ -193,6 +193,10 @@ namespace Gauntlet
 				{
 					RequiredBuildFlags |= BuildFlags.Packaged;
 				}
+				else if (Globals.Params.ParseParam("staged") && TargetBuildSupport.CanSupportBuildType(BuildFlags.Loose))
+				{
+					RequiredBuildFlags |= BuildFlags.Loose;
+				}
 			}
 
 			InstallOnly = false;
