Undo CL: 22849610
This caused issues with replay testing as null artifacts were not correctly being saved.
Adjusting the FunctionalTest so it runs then generates a blank artifact.
overriding the GetTestSummary header when we know no test will be ran.

diff --git a/Engine/Source/Programs/AutomationTool/Gauntlet/Unreal/Automation/UE.Automation.cs b/Engine/Source/Programs/AutomationTool/Gauntlet/Unreal/Automation/UE.Automation.cs
--- a/Engine/Source/Programs/AutomationTool/Gauntlet/Unreal/Automation/UE.Automation.cs
+++ b/Engine/Source/Programs/AutomationTool/Gauntlet/Unreal/Automation/UE.Automation.cs
@@ -522,39 +522,36 @@ namespace UE
 			const float IdleTimeout = 30 * 60;
 
 			List<string> ChannelEntries = new List<string>();
-			if (TestInstance != null)
-			{
-				// We are primarily interested in what the editor is doing
-				var AppInstance = TestInstance.EditorApp;
+			// We are primarily interested in what the editor is doing
+			var AppInstance = TestInstance.EditorApp;
 
-				UnrealLogParser Parser = new UnrealLogParser(AppInstance.StdOut);
-				ChannelEntries.AddRange(Parser.GetEditorBusyChannels());
+			UnrealLogParser Parser = new UnrealLogParser(AppInstance.StdOut);
+			ChannelEntries.AddRange(Parser.GetEditorBusyChannels());
 
-				// Any new entries?
-				if (ChannelEntries.Count > LastAutomationEntryCount)
+			// Any new entries?
+			if (ChannelEntries.Count > LastAutomationEntryCount)
+			{
+				// log new entries so people have something to look at
+				ChannelEntries.Skip(LastAutomationEntryCount).ToList().ForEach(S => Log.Info("{0}", S));
+				LastAutomationEntryTime = DateTime.Now;
+				LastAutomationEntryCount = ChannelEntries.Count;
+			}
+			else
+			{
+				// Check for timeouts
+				if (LastAutomationEntryTime == DateTime.MinValue)
 				{
-					// log new entries so people have something to look at
-					ChannelEntries.Skip(LastAutomationEntryCount).ToList().ForEach(S => Log.Info("{0}", S));
 					LastAutomationEntryTime = DateTime.Now;
-					LastAutomationEntryCount = ChannelEntries.Count;
 				}
-				else
-				{
-					// Check for timeouts
-					if (LastAutomationEntryTime == DateTime.MinValue)
-					{
-						LastAutomationEntryTime = DateTime.Now;
-					}
 
-					double ElapsedTime = (DateTime.Now - LastAutomationEntryTime).TotalSeconds;
+				double ElapsedTime = (DateTime.Now - LastAutomationEntryTime).TotalSeconds;
 
-					// Check for timeout
-					if (ElapsedTime > IdleTimeout)
-					{
-						Log.Error("No activity observed in last {0:0.00} minutes. Aborting test", IdleTimeout / 60);
-						MarkTestComplete();
-						SetUnrealTestResult(TestResult.TimedOut);
-					}
+				// Check for timeout
+				if (ElapsedTime > IdleTimeout)
+				{
+					Log.Error("No activity observed in last {0:0.00} minutes. Aborting test", IdleTimeout / 60);
+					MarkTestComplete();
+					SetUnrealTestResult(TestResult.TimedOut);
 				}
 			}
 
@@ -753,7 +750,6 @@ namespace UE
 
 			// The editor is an additional arbiter of success
 			if (InArtifacts.SessionRole.RoleType == UnrealTargetRole.Editor
-				&& InLog != null
 				&& InLog.HasAbnormalExit == false)
 			{
 				// if no fatal errors, check test results
@@ -849,7 +845,7 @@ namespace UE
 			MarkdownBuilder MB = new MarkdownBuilder(base.GetTestSummaryHeader());
 
 			// Everything we need is in the editor artifacts
-			var EditorRole = RoleResults.Where(R => R.Artifacts != null && R.Artifacts.SessionRole.RoleType == UnrealTargetRole.Editor).FirstOrDefault();
+			var EditorRole = RoleResults.Where(R => R.Artifacts.SessionRole.RoleType == UnrealTargetRole.Editor).FirstOrDefault();
 
 			if (EditorRole != null)
 			{
@@ -1076,7 +1072,7 @@ namespace UE
 
 			foreach (var Role in RoleResults)
 			{
-				if (Role.Artifacts != null && Role.Artifacts.SessionRole.RoleType == UnrealTargetRole.Editor)
+				if (Role.Artifacts.SessionRole.RoleType == UnrealTargetRole.Editor)
 				{
 					AutomationLogParser Parser = new AutomationLogParser(Role.LogSummary.FullLogContent);
 					AllWarnings.AddRange(
diff --git a/Engine/Source/Programs/AutomationTool/Gauntlet/Unreal/Automation/UE.ErrorTest.cs b/Engine/Source/Programs/AutomationTool/Gauntlet/Unreal/Automation/UE.ErrorTest.cs
--- a/Engine/Source/Programs/AutomationTool/Gauntlet/Unreal/Automation/UE.ErrorTest.cs
+++ b/Engine/Source/Programs/AutomationTool/Gauntlet/Unreal/Automation/UE.ErrorTest.cs
@@ -119,92 +119,88 @@ namespace UE
 			string LocalReason = "";
 
 			TestResult FinalResult = TestResult.Invalid;
+			if (ErrorType == ErrorTypes.Ensure)
+			{
+				// for an ensure we should have an entry and a callstack
+				int EnsureCount = InLog.Ensures.Count();
+				int CallstackLength = EnsureCount > 0 ? InLog.Ensures.First().Callstack.Length : 0;
 
-			if (InLog != null)
+				if (EnsureCount == 0)
+				{
+					FinalResult = TestResult.Failed;
+					LocalReason = string.Format("No ensure error found for failure of type {0}", ErrorType);
+				}
+				else if (EnsureCount != 1)
+				{
+					FinalResult = TestResult.Failed;
+					LocalReason = string.Format("Incorrect ensure count found for failure of type {0}", ErrorType);
+				}
+				else if (CallstackLength == 0)
+				{
+					FinalResult = TestResult.Failed;
+					LocalReason = string.Format("No callstack error found for failure of type {0}", ErrorType);
+				}
+				else
+				{
+					FinalResult = TestResult.Passed;
+					LocalReason = string.Format("Found {0} ensures, test result = {1}", EnsureCount, FinalResult);
+				}
+			}
+			else
 			{
-				if (ErrorType == ErrorTypes.Ensure)
+				if (InLog.FatalError == null)
 				{
-					// for an ensure we should have an entry and a callstack
-					int EnsureCount = InLog.Ensures.Count();
-					int CallstackLength = EnsureCount > 0 ? InLog.Ensures.First().Callstack.Length : 0;
-
-					if (EnsureCount == 0)
-					{
-						FinalResult = TestResult.Failed;
-						LocalReason = string.Format("No ensure error found for failure of type {0}", ErrorType);
-					}
-					else if (EnsureCount != 1)
-					{
-						FinalResult = TestResult.Failed;
-						LocalReason = string.Format("Incorrect ensure count found for failure of type {0}", ErrorType);
-					}
-					else if (CallstackLength == 0)
-					{
-						FinalResult = TestResult.Failed;
-						LocalReason = string.Format("No callstack error found for failure of type {0}", ErrorType);
-					}
-					else
-					{
-						FinalResult = TestResult.Passed;
-						LocalReason = string.Format("Found {0} ensures, test result = {1}", EnsureCount, FinalResult);
-					}
+					FinalResult = TestResult.Failed;
+					Log.Info("No fatal error found for failure of type {0}", ErrorType);
+				}
+				else if (InLog.FatalError.Callstack.Length == 0)
+				{
+					FinalResult = TestResult.Failed;
+					Log.Info("No callstack found for failure of type {0}", ErrorType);
 				}
 				else
 				{
-					if (InLog.FatalError == null)
+					// all of these should contain a message and a result
+					if (ErrorType == ErrorTypes.Check)
 					{
-						FinalResult = TestResult.Failed;
-						Log.Info("No fatal error found for failure of type {0}", ErrorType);
-					}
-					else if (InLog.FatalError.Callstack.Length == 0)
-					{
-						FinalResult = TestResult.Failed;
-						Log.Info("No callstack found for failure of type {0}", ErrorType);
+						if (!InLog.FatalError.Message.ToLower().Contains("assertion failed"))
+						{
+							FinalResult = TestResult.Failed;
+							LocalReason = string.Format("Unexpected assertion message");
+						}
+						else
+						{
+							FinalResult = TestResult.Passed;
+						}
+						Log.Info("Assertion message was {0}", InLog.FatalError.Message);
 					}
-					else
+					else if (ErrorType == ErrorTypes.Fatal)
 					{
-						// all of these should contain a message and a result
-						if (ErrorType == ErrorTypes.Check)
+						if (!InLog.FatalError.Message.ToLower().Contains("fatal erro"))
+						{
+							FinalResult = TestResult.Failed;
+							LocalReason = string.Format("Unexpected Fatal Error message");
+						}
+						else
 						{
-							if (!InLog.FatalError.Message.ToLower().Contains("assertion failed"))
-							{
-								FinalResult = TestResult.Failed;
-								LocalReason = string.Format("Unexpected assertion message");
-							}
-							else
-							{
-								FinalResult = TestResult.Passed;
-							}
-							Log.Info("Assertion message was {0}", InLog.FatalError.Message);
+							FinalResult = TestResult.Passed;
 						}
-						else if (ErrorType == ErrorTypes.Fatal)
+
+						Log.Info("Fatal Error message was {0}", InLog.FatalError.Message);
+					}
+					else if (ErrorType == ErrorTypes.GPF)
+					{
+						if (!InLog.FatalError.Message.ToLower().Contains("exception"))
 						{
-							if (!InLog.FatalError.Message.ToLower().Contains("fatal erro"))
-							{
-								FinalResult = TestResult.Failed;
-								LocalReason = string.Format("Unexpected Fatal Error message");
-							}
-							else
-							{
-								FinalResult = TestResult.Passed;
-							}
-
-							Log.Info("Fatal Error message was {0}", InLog.FatalError.Message);
+							FinalResult = TestResult.Failed;
+							LocalReason = string.Format("Unexpected exception message");
 						}
-						else if (ErrorType == ErrorTypes.GPF)
+						else
 						{
-							if (!InLog.FatalError.Message.ToLower().Contains("exception"))
-							{
-								FinalResult = TestResult.Failed;
-								LocalReason = string.Format("Unexpected exception message");
-							}
-							else
-							{
-								FinalResult = TestResult.Passed;
-							}
-
-							Log.Info("Exception message was {0}", InLog.FatalError.Message);
+							FinalResult = TestResult.Passed;
 						}
+
+						Log.Info("Exception message was {0}", InLog.FatalError.Message);
 					}
 				}
 			}
diff --git a/Engine/Source/Programs/AutomationTool/Gauntlet/Unreal/Base/Gauntlet.UnrealSession.cs b/Engine/Source/Programs/AutomationTool/Gauntlet/Unreal/Base/Gauntlet.UnrealSession.cs
--- a/Engine/Source/Programs/AutomationTool/Gauntlet/Unreal/Base/Gauntlet.UnrealSession.cs
+++ b/Engine/Source/Programs/AutomationTool/Gauntlet/Unreal/Base/Gauntlet.UnrealSession.cs
@@ -1210,84 +1210,83 @@ namespace Gauntlet
 			int DummyClientCount = 0;
 
 			List<UnrealRoleArtifacts> AllArtifacts = new List<UnrealRoleArtifacts>();
-			if (TestInstance != null)
+
+			foreach (UnrealSessionInstance.RoleInstance App in TestInstance.RunningRoles)
 			{
-				foreach (UnrealSessionInstance.RoleInstance App in TestInstance.RunningRoles)
-				{
-					string RoleName = (App.Role.IsDummy() ? "Dummy" : "") + App.Role.RoleType.ToString();
-					string FolderName = RoleName;
+				string RoleName = (App.Role.IsDummy() ? "Dummy" : "") + App.Role.RoleType.ToString();
+				string FolderName = RoleName;
 
-					int RoleCount = 1;
+				int RoleCount = 1;
 
-					if (App.Role.IsDummy())
+				if (App.Role.IsDummy())
+				{
+					DummyClientCount++;
+					RoleCount = DummyClientCount;
+				}
+				else
+				{
+					if (!RoleCounts.ContainsKey(App.Role.RoleType))
 					{
-						DummyClientCount++;
-						RoleCount = DummyClientCount;
+						RoleCounts.Add(App.Role.RoleType, 1);
 					}
 					else
 					{
-						if (!RoleCounts.ContainsKey(App.Role.RoleType))
-						{
-							RoleCounts.Add(App.Role.RoleType, 1);
-						}
-						else
-						{
-							RoleCounts[App.Role.RoleType]++;
-						}
-
-						RoleCount = RoleCounts[App.Role.RoleType];
+						RoleCounts[App.Role.RoleType]++;
 					}
 
+					RoleCount = RoleCounts[App.Role.RoleType];
+				}
 
-					if (RoleCount > 1)
-					{
-						FolderName += string.Format("_{0:00}", RoleCount);
-					}
+				
+				if (RoleCount > 1)
+				{
+					FolderName += string.Format("_{0:00}", RoleCount);
+				}
 
-					string DestPath = Path.Combine(OutputPath, FolderName);
+				string DestPath = Path.Combine(OutputPath, FolderName);
 
-					if (!App.Role.IsNullRole() && !App.Role.InstallOnly)
+				if (!App.Role.IsNullRole() && !App.Role.InstallOnly )
+				{
+					Log.VeryVerbose("Calling SaveRoleArtifacts, Role: {0}  Artifact Path: {1}", App.ToString(), App.AppInstance.ArtifactPath);
+					UnrealRoleArtifacts Artifacts = null;
+					ITargetDevice device = App.AppInstance.Device;
+					IDeviceUsageReporter.RecordStart(device.Name, device.Platform, IDeviceUsageReporter.EventType.SavingArtifacts);
+					try
 					{
-						Log.VeryVerbose("Calling SaveRoleArtifacts, Role: {0}  Artifact Path: {1}", App.ToString(), App.AppInstance.ArtifactPath);
-						UnrealRoleArtifacts Artifacts = null;
-						ITargetDevice device = App.AppInstance.Device;
-						IDeviceUsageReporter.RecordStart(device.Name, device.Platform, IDeviceUsageReporter.EventType.SavingArtifacts);
-						try
+						Artifacts = SaveRoleArtifacts(Context, App, DestPath);
+					}
+					catch (Exception SaveArtifactsException)
+					{
+						// Caught an exception -> report failure
+						IDeviceUsageReporter.RecordEnd(device.Name, device.Platform, IDeviceUsageReporter.EventType.SavingArtifacts, IDeviceUsageReporter.EventState.Failure);
+
+						// Retry once only, after rebooting device, if artifacts couldn't be saved.
+						if (SaveArtifactsException.Message.Contains("A retry should be performed"))
 						{
+							Log.Info("Rebooting device and retrying save role artifacts once.");
+							App.AppInstance.Device.Reboot();
 							Artifacts = SaveRoleArtifacts(Context, App, DestPath);
 						}
-						catch (Exception SaveArtifactsException)
-						{
-							// Caught an exception -> report failure
-							IDeviceUsageReporter.RecordEnd(device.Name, device.Platform, IDeviceUsageReporter.EventType.SavingArtifacts, IDeviceUsageReporter.EventState.Failure);
-
-							// Retry once only, after rebooting device, if artifacts couldn't be saved.
-							if (SaveArtifactsException.Message.Contains("A retry should be performed"))
-							{
-								Log.Info("Rebooting device and retrying save role artifacts once.");
-								App.AppInstance.Device.Reboot();
-								Artifacts = SaveRoleArtifacts(Context, App, DestPath);
-							}
-							else
-							{
-								// Pass exception to the surrounding try/catch in UnrealTestNode while preserving the original callstack
-								throw;
-							}
-						}
-						// Did not catch -> successful reporting
-						IDeviceUsageReporter.RecordEnd(device.Name, device.Platform, IDeviceUsageReporter.EventType.SavingArtifacts, IDeviceUsageReporter.EventState.Success);
-
-						if (Artifacts != null)
+						else
 						{
-							AllArtifacts.Add(Artifacts);
+							// Pass exception to the surrounding try/catch in UnrealTestNode while preserving the original callstack
+							throw;
 						}
 					}
-					else
+					// Did not catch -> successful reporting
+					IDeviceUsageReporter.RecordEnd(device.Name, device.Platform, IDeviceUsageReporter.EventType.SavingArtifacts, IDeviceUsageReporter.EventState.Success);
+
+					if (Artifacts != null)
 					{
-						Log.Verbose("Skipping SaveRoleArtifacts for Null Role: {0}", App.ToString());
+						AllArtifacts.Add(Artifacts);
 					}
 				}
+				else 
+				{
+					Log.Verbose("Skipping SaveRoleArtifacts for Null Role: {0}", App.ToString());
+				}
 			}
+
 			return AllArtifacts;
 		}
 	}
diff --git a/Engine/Source/Programs/AutomationTool/Gauntlet/Unreal/Base/Gauntlet.UnrealTestNode.cs b/Engine/Source/Programs/AutomationTool/Gauntlet/Unreal/Base/Gauntlet.UnrealTestNode.cs
--- a/Engine/Source/Programs/AutomationTool/Gauntlet/Unreal/Base/Gauntlet.UnrealTestNode.cs
+++ b/Engine/Source/Programs/AutomationTool/Gauntlet/Unreal/Base/Gauntlet.UnrealTestNode.cs
@@ -1080,29 +1080,19 @@ namespace Gauntlet
 
 			// Shutdown the instance so we can access all files, but do not null it or shutdown the UnrealApp because we still need
 			// access to these objects and their resources! Final cleanup is done in CleanupTest()
-			if (TestInstance != null)
-			{
-				TestInstance.Shutdown();
-			}
+			TestInstance.Shutdown();
 
 			try
 			{
-				if (string.IsNullOrEmpty(ArtifactPath))
-				{
-					Log.Info("ArtifactsPath was null unable to save artifacts");
-				}
-				else
-				{
-					Log.Info("Saving artifacts to {0}", ArtifactPath);
-					// run create dir again just in case the already made dir was cleaned up by another buildfarm job or something similar.
-					Directory.CreateDirectory(ArtifactPath);
-					Utils.SystemHelpers.MarkDirectoryForCleanup(ArtifactPath);
-				}
+				Log.Info("Saving artifacts to {0}", ArtifactPath);
+				// run create dir again just in case the already made dir was cleaned up by another buildfarm job or something similar.
+				Directory.CreateDirectory(ArtifactPath);
+				Utils.SystemHelpers.MarkDirectoryForCleanup(ArtifactPath);
+
 				SessionArtifacts = SaveRoleArtifacts(ArtifactPath);
 
 				// call legacy version
 				SaveArtifacts_DEPRECATED(ArtifactPath);
-				
 			}
 			catch (Exception Ex)
 			{
@@ -1528,10 +1518,6 @@ namespace Gauntlet
 		/// <returns></returns>
 		public virtual IEnumerable<UnrealRoleArtifacts> SaveRoleArtifacts(string OutputPath)
 		{
-			if (UnrealApp == null)
-			{
-				return new List<UnrealRoleArtifacts>();
-			}
 			return UnrealApp.SaveRoleArtifacts(Context, TestInstance, ArtifactPath);
 		}
 
@@ -1545,21 +1531,13 @@ namespace Gauntlet
 		protected virtual UnrealProcessResult GetExitCodeAndReason(StopReason InReason, UnrealLog InLog, UnrealRoleArtifacts InArtifacts, out string ExitReason, out int ExitCode)
 		{
 			// first check for fatal issues
-			if (InLog == null
-				|| InLog.FatalError != null)
+			if (InLog.FatalError != null)
 			{
 				ExitReason = "Process encountered fatal error";
 				ExitCode = -1;
 				return UnrealProcessResult.EncounteredFatalError;
 			}
 
-			if (InArtifacts == null)
-			{
-				ExitReason = "No role artifacts generated";
-				ExitCode = -1;
-				return UnrealProcessResult.TestFailure;
-			}
-
 			// Catch failed engine init. Early issues can result in the engine exiting with hard to diagnose reasons
 			if (InLog.EngineInitialized == false)
 			{
@@ -1741,24 +1719,6 @@ namespace Gauntlet
 			return new UnrealLogParser(InArtifacts.AppInstance.StdOut).GetSummary();
 		}
 
-		protected virtual UnrealRoleResult GetRoleResultWhenNoArtifactsAreGenerated(StopReason InReason)
-		{
-			int ExitCode = 0;
-			string ExitReason;
-
-			// Give ourselves (and derived classes) a chance to analyze what happened
-			UnrealProcessResult ProcessResult = GetExitCodeAndReason(InReason, null, null, out ExitReason, out ExitCode);
-			// if the test is stopping for a reason other than completion, mark this as failing incase derived classes
-			// don't do the right thing
-			if (InReason == StopReason.MaxDuration)
-			{
-				ProcessResult = UnrealProcessResult.TimeOut;
-				ExitCode = -1;
-			}
-
-			return new UnrealRoleResult(ProcessResult, ExitCode, ExitReason, new UnrealLog(), null, new List<UnrealTestEvent>());
-		}
-
 		/// <summary>
 		/// Returns a list of all results for the roles involved in this test by calling CreateRoleResultFromArtifact for all
 		/// artifacts in the list
@@ -1767,14 +1727,7 @@ namespace Gauntlet
 		/// <returns></returns>
 		protected virtual IEnumerable<UnrealRoleResult> CreateRoleResultsFromArtifacts(StopReason InReason, IEnumerable<UnrealRoleArtifacts> InAllArtifacts)
 		{
-			if (InAllArtifacts != null && InAllArtifacts.Count() > 0)
-			{
-				return InAllArtifacts.Select(A => CreateRoleResultFromArtifact(InReason, A)).ToArray();
-			}
-			else
-			{
-				return new List<UnrealRoleResult>() { GetRoleResultWhenNoArtifactsAreGenerated(InReason) };
-			}
+			return InAllArtifacts.Select(A => CreateRoleResultFromArtifact(InReason, A)).ToArray();
 		}
 
 
@@ -1890,26 +1843,24 @@ namespace Gauntlet
 			UnrealLog LogSummary = InRoleResult.LogSummary;
 			MarkdownBuilder MB = new MarkdownBuilder();
 
-			if (InRoleResult.Artifacts != null)
-			{
-				UnrealRoleArtifacts RoleArtifacts = InRoleResult.Artifacts;
+			UnrealRoleArtifacts RoleArtifacts = InRoleResult.Artifacts;
 
-				MB.H3(string.Format("Role: {0} ({1} {2})", RoleArtifacts.SessionRole.RoleType, RoleArtifacts.SessionRole.Platform, RoleArtifacts.SessionRole.Configuration));
+			MB.H3(string.Format("Role: {0} ({1} {2})", RoleArtifacts.SessionRole.RoleType, RoleArtifacts.SessionRole.Platform, RoleArtifacts.SessionRole.Configuration));
 
-				bool HaveFatalError = LogSummary.FatalError != null;
+			bool HaveFatalError = LogSummary.FatalError != null;
 
-				// If we have no fatal error use "Error" which Horde will highlight. Otherwise use "Failed" since the fatal error will be below and highligted
-				string FailedString = HaveFatalError ? "Failed: " : "Error: ";
-				string CompletedString = "Completed: ";
+			// If we have no fatal error use "Error" which Horde will highlight. Otherwise use "Failed" since the fatal error will be below and highligted
+			string FailedString = HaveFatalError ? "Failed: " : "Error: ";
+			string CompletedString = "Completed: ";
 
-				string RoleState = InRoleResult.ExitCode != 0 && InRoleResult.LogSummary.HasAbnormalExit ? FailedString : CompletedString;
+			string RoleState = InRoleResult.ExitCode != 0 && InRoleResult.LogSummary.HasAbnormalExit ? FailedString : CompletedString;
 
-				MB.H4(string.Format("{0} {1} ({2}, ExitCode={3})", RoleState, InRoleResult.Summary, InRoleResult.ProcessResult, InRoleResult.ExitCode));
+			MB.H4(string.Format("{0} {1} ({2}, ExitCode={3})", RoleState, InRoleResult.Summary, InRoleResult.ProcessResult, InRoleResult.ExitCode));
 
-				// log command line up here for visibility
-				//MB.Paragraph(string.Format("CommandLine: {0}", RoleArtifacts.AppInstance.CommandLine));
+			// log command line up here for visibility
+			//MB.Paragraph(string.Format("CommandLine: {0}", RoleArtifacts.AppInstance.CommandLine));
 
-				MB.UnorderedList(new string[] {
+			MB.UnorderedList(new string[] {
 				string.Format("CommandLine: {0}", RoleArtifacts.AppInstance.CommandLine),
 				string.Format("Log: {0}", RoleArtifacts.LogPath),
 				string.Format("SavedDir: {0}", RoleArtifacts.ArtifactPath),
@@ -1919,112 +1870,112 @@ namespace Gauntlet
 				LogSummary.Warnings.Count() > 0 ? string.Format("Log Warnings: {0}", LogSummary.Warnings.Count()) : null
 			});
 
-				// Separate the events we want to report on
-				IEnumerable<UnrealTestEvent> Asserts = InRoleResult.Events.Where(E => E.Severity == EventSeverity.Fatal);
-				IEnumerable<UnrealTestEvent> Errors = InRoleResult.Events.Where(E => E.Severity == EventSeverity.Error);
-				IEnumerable<UnrealTestEvent> Ensures = InRoleResult.Events.Where(E => E.IsEnsure);
-				IEnumerable<UnrealTestEvent> Warnings = InRoleResult.Events.Where(E => E.Severity == EventSeverity.Warning && !E.IsEnsure);
+			// Separate the events we want to report on
+			IEnumerable<UnrealTestEvent> Asserts = InRoleResult.Events.Where(E => E.Severity == EventSeverity.Fatal);
+			IEnumerable<UnrealTestEvent> Errors = InRoleResult.Events.Where(E => E.Severity == EventSeverity.Error);
+			IEnumerable<UnrealTestEvent> Ensures = InRoleResult.Events.Where(E => E.IsEnsure);
+			IEnumerable<UnrealTestEvent> Warnings = InRoleResult.Events.Where(E => E.Severity == EventSeverity.Warning && !E.IsEnsure);
 
-				foreach (UnrealTestEvent Event in Asserts)
-				{
-					MB.H4(string.Format("Fatal Error: {0}", Event.Summary));
+			foreach (UnrealTestEvent Event in Asserts)
+			{
+				MB.H4(string.Format("Fatal Error: {0}", Event.Summary));
 
-					if (Event.Callstack.Any())
-					{
-						MB.UnorderedList(Event.Callstack.Take(MaxCallstackLines));
+				if (Event.Callstack.Any())
+				{
+					MB.UnorderedList(Event.Callstack.Take(MaxCallstackLines));
 
-						if (Event.Callstack.Count() > MaxCallstackLines)
-						{
-							MB.Paragraph("See log for full callstack");
-						}
-					}
-					else
+					if (Event.Callstack.Count() > MaxCallstackLines)
 					{
-						MB.Paragraph("Could not parse callstack. See log for full callstack");
+						MB.Paragraph("See log for full callstack");
 					}
 				}
-
-				foreach (UnrealTestEvent Event in Ensures.Distinct())
+				else
 				{
-					MB.H4(string.Format("Warning: Ensure: {0}", Event.Summary));
+					MB.Paragraph("Could not parse callstack. See log for full callstack");
+				}
+			}
 
-					if (Event.Callstack.Any())
-					{
-						MB.UnorderedList(Event.Callstack.Take(MaxCallstackLines));
+			foreach (UnrealTestEvent Event in Ensures.Distinct())
+			{
+				MB.H4(string.Format("Warning: Ensure: {0}", Event.Summary));
 
-						if (Event.Callstack.Count() > MaxCallstackLines)
-						{
-							MB.Paragraph("See log for full callstack");
-						}
-					}
-					else
+				if (Event.Callstack.Any())
+				{
+					MB.UnorderedList(Event.Callstack.Take(MaxCallstackLines));
+
+					if (Event.Callstack.Count() > MaxCallstackLines)
 					{
-						MB.Paragraph("Could not parse callstack. See log for full callstack");
+						MB.Paragraph("See log for full callstack");
 					}
 				}
-
-				if (Errors.Any())
+				else
 				{
-					var ErrorList = Errors.Select(E => E.Summary).Distinct();
-					var PrintedErrorList = ErrorList;
+					MB.Paragraph("Could not parse callstack. See log for full callstack");
+				}
+			}
 
-					string TrimStatement = "";
+			if (Errors.Any())
+			{
+				var ErrorList = Errors.Select(E => E.Summary).Distinct();
+				var PrintedErrorList = ErrorList;
 
-					// too many warnings. If there was an abnormal exit show the last ones as they may be relevant
-					if (ErrorList.Count() > MaxLogLines)
+				string TrimStatement = "";
+
+				// too many warnings. If there was an abnormal exit show the last ones as they may be relevant
+				if (ErrorList.Count() > MaxLogLines)
+				{
+					if (LogSummary.HasAbnormalExit)
 					{
-						if (LogSummary.HasAbnormalExit)
-						{
-							PrintedErrorList = ErrorList.Skip(ErrorList.Count() - MaxLogLines);
-							TrimStatement = string.Format("(Last {0} of {1} errors)", MaxLogLines, ErrorList.Count());
-						}
-						else
-						{
-							PrintedErrorList = ErrorList.Take(MaxLogLines);
-							TrimStatement = string.Format("(First {0} of {1} errors)", MaxLogLines, ErrorList.Count());
-						}
+						PrintedErrorList = ErrorList.Skip(ErrorList.Count() - MaxLogLines);
+						TrimStatement = string.Format("(Last {0} of {1} errors)", MaxLogLines, ErrorList.Count());
 					}
-
-					MB.H4("Errors:");
-					MB.UnorderedList(PrintedErrorList);
-
-					if (!string.IsNullOrEmpty(TrimStatement))
+					else
 					{
-						MB.Paragraph(TrimStatement);
+						PrintedErrorList = ErrorList.Take(MaxLogLines);
+						TrimStatement = string.Format("(First {0} of {1} errors)", MaxLogLines, ErrorList.Count());
 					}
 				}
 
-				if (Warnings.Any())
+				MB.H4("Errors:");
+				MB.UnorderedList(PrintedErrorList);
+
+				if (!string.IsNullOrEmpty(TrimStatement))
 				{
-					var WarningList = Warnings.Select(E => E.Summary).Distinct();
-					var PrintedWarningList = WarningList;
+					MB.Paragraph(TrimStatement);
+				}
+			}
 
-					string TrimStatement = "";
+			if (Warnings.Any())
+			{
+				var WarningList = Warnings.Select(E => E.Summary).Distinct();
+				var PrintedWarningList = WarningList;
 
-					// too many warnings. If there was an abnormal exit show the last ones as they may be relevant
-					if (WarningList.Count() > MaxLogLines)
+				string TrimStatement = "";
+
+				// too many warnings. If there was an abnormal exit show the last ones as they may be relevant
+				if (WarningList.Count() > MaxLogLines)
+				{
+					if (LogSummary.HasAbnormalExit)
 					{
-						if (LogSummary.HasAbnormalExit)
-						{
-							PrintedWarningList = WarningList.Skip(WarningList.Count() - MaxLogLines);
-							TrimStatement = string.Format("(Last {0} of {1} errors)", MaxLogLines, WarningList.Count());
-						}
-						else
-						{
-							PrintedWarningList = WarningList.Take(MaxLogLines);
-							TrimStatement = string.Format("(First {0} of {1} errors)", MaxLogLines, WarningList.Count());
-						}
+						PrintedWarningList = WarningList.Skip(WarningList.Count() - MaxLogLines);
+						TrimStatement = string.Format("(Last {0} of {1} errors)", MaxLogLines, WarningList.Count());
 					}
-
-					MB.H4("Warnings:");
-					MB.UnorderedList(PrintedWarningList);
-
-					if (!string.IsNullOrEmpty(TrimStatement))
+					else
 					{
-						MB.Paragraph(TrimStatement);
+						PrintedWarningList = WarningList.Take(MaxLogLines);
+						TrimStatement = string.Format("(First {0} of {1} errors)", MaxLogLines, WarningList.Count());
 					}
 				}
+
+				MB.H4("Warnings:");
+				MB.UnorderedList(PrintedWarningList);
+
+				if (!string.IsNullOrEmpty(TrimStatement))
+				{
+					MB.Paragraph(TrimStatement);
+				}
 			}
+
 			return MB.ToString();
 		}
 
@@ -2275,15 +2226,7 @@ namespace Gauntlet
 			// displayed individually by GetFormattedRoleSummary
 			foreach (var RoleResult in SortedRoles)
 			{
-				string RoleName;
-				if (RoleResult.Artifacts == null)
-				{
-					RoleName = "NoRoleArtifact";
-				}
-				else
-				{
-					RoleName = RoleResult.Artifacts.SessionRole.RoleType.ToString();
-				}
+				string RoleName = RoleResult.Artifacts.SessionRole.RoleType.ToString();
 
 				MB.Paragraph(string.Format("{0} Role: {1} ({2}, ExitCode {3})", RoleName, RoleResult.Summary, RoleResult.ProcessResult, RoleResult.ExitCode));
 
diff --git a/Engine/Source/Programs/AutomationTool/Gauntlet/Unreal/Engine/UnrealTests.cs b/Engine/Source/Programs/AutomationTool/Gauntlet/Unreal/Engine/UnrealTests.cs
--- a/Engine/Source/Programs/AutomationTool/Gauntlet/Unreal/Engine/UnrealTests.cs
+++ b/Engine/Source/Programs/AutomationTool/Gauntlet/Unreal/Engine/UnrealTests.cs
@@ -212,7 +212,6 @@ namespace Gauntlet
 
 				// The editor is an additional arbiter of success
 				if (InArtifacts.SessionRole.RoleType == UnrealTargetRole.Editor
-					&& InLog != null
 					&& InLog.HasAbnormalExit == false)
 				{
 					// if no fatal errors, check test results
