[Automation] Adjusts Gauntlet's test discovery to prioritize tests containing targeted namespaces

diff --git a/Engine/Source/Programs/AutomationTool/Gauntlet/Framework/Gauntlet.Utils.cs b/Engine/Source/Programs/AutomationTool/Gauntlet/Framework/Gauntlet.Utils.cs
--- a/Engine/Source/Programs/AutomationTool/Gauntlet/Framework/Gauntlet.Utils.cs
+++ b/Engine/Source/Programs/AutomationTool/Gauntlet/Framework/Gauntlet.Utils.cs
@@ -519,10 +519,8 @@ namespace Gauntlet
 
 	namespace Utils
 	{
-		
 		public class TestConstructor
 		{
-
 			/// <summary>
 			/// Helper function that returns the type of an object based on namespaces and name
 			/// </summary>
@@ -542,11 +540,10 @@ namespace Gauntlet
 				}
 
 				Log.VeryVerbose("Will search {0} for test {1}", string.Join(" ", FullNames), TestName);
-				
-				// find all types from loaded assemblies that implement testnode
-					List < Type> CandidateTypes = new List<Type>();
 
-				foreach (Assembly Assembly in ScriptManager.AllScriptAssemblies) 
+				// find all types from loaded assemblies that implement testnode
+				List<Type> CandidateTypes = new List<Type>();
+				foreach (Assembly Assembly in ScriptManager.AllScriptAssemblies)
 				{
 					foreach (var Type in Assembly.GetTypes())
 					{
@@ -560,16 +557,15 @@ namespace Gauntlet
 				Log.VeryVerbose("Possible candidates for {0}: {1}", TestName, string.Join(" ", CandidateTypes));
 
 				// check our expanded names.. need to search in namespace order
+				IList<Type> MatchingTypes = new List<Type>();
 				foreach (string UserTypeName in FullNames)
 				{
 					// Even tho the user might have specified N1.Foo it still might be Other.N1.Foo so only
 					// compare based on the number of namespaces that were specified.
-
-
-					foreach (var Type in CandidateTypes)
+					foreach (Type Candidate in CandidateTypes)
 					{
 						string[] UserNameComponents = UserTypeName.Split('.');
-						string[] TypeNameComponents = Type.FullName.Split('.');
+						string[] TypeNameComponents = Candidate.FullName.Split('.');
 
 						int MissingUserComponents = TypeNameComponents.Length - UserNameComponents.Length;
 
@@ -578,16 +574,48 @@ namespace Gauntlet
 							TypeNameComponents = TypeNameComponents.Skip(MissingUserComponents).ToArray();
 						}
 
-						var Difference = TypeNameComponents.Except(UserNameComponents, StringComparer.OrdinalIgnoreCase);
+						IEnumerable<string> Difference = TypeNameComponents.Except(UserNameComponents, StringComparer.OrdinalIgnoreCase);
 
 						if (Difference.Count() == 0)
 						{
-							Log.VeryVerbose("Considering {0} as best match for {1}", Type, TestName);
-							return Type;
+							Log.VeryVerbose("Found match {0} for user type {1}", Candidate.FullName, UserTypeName);
+
+							// If we have any namespaces, add the candidate
+							if(Namespaces.Count() > 0)
+							{
+								MatchingTypes.Add(Candidate);
+							}
+
+							// No namespaces, just return the first type found
+							else
+							{
+								return Candidate;
+							}
 						}
 					}
 				}
 
+				// If user has specified at least 1 namespace, we prioritize types that include a provided namespace over types that do not.
+				// For example, in the case of the supplied parameters: TestName = BootTest, Namespaces = {"Game"}
+				// We prioritize "Game.BootTest" over "UE.BootTest".
+				// In the case of multiple name spaces, we default to the first test found within a namespace.
+				// For example, in the case of the supplied parameters: TestName = BootTest, Namespaces = {"Game", "UE"}
+				// We select "Game.BootTest" over "UE.BootTest".
+				// If you maintain many tests with the same base name, you should be as explicit as possible with Namespaces!
+				foreach(string Namespace in Namespaces)
+				{
+					foreach(Type Match in MatchingTypes)
+					{
+						// Split off the namespace
+						string TypeNamespace = Match.FullName.Replace(Match.Name, string.Empty);
+
+						// If it matches one of namespaces, return it
+						if (TypeNamespace.Contains(Namespace))
+						{
+							return Match;
+						}
+					}
+				}
 
 				throw new AutomationException("Unable to find type {0} in assemblies. Namespaces= {1}.", TestName, Namespaces);
 			}
