Markup P4WriteConfig with ParamHelp attributes

diff --git a/Engine/Source/Programs/AutomationTool/AutomationUtils/P4WriteConfig.cs b/Engine/Source/Programs/AutomationTool/AutomationUtils/P4WriteConfig.cs
--- a/Engine/Source/Programs/AutomationTool/AutomationUtils/P4WriteConfig.cs
+++ b/Engine/Source/Programs/AutomationTool/AutomationUtils/P4WriteConfig.cs
@@ -11,10 +11,10 @@ using UnrealBuildBase;
 namespace AutomationTool
 {
 	[Help("Auto-detects P4 settings based on the current path and creates a p4config file with the relevant settings.")]
-	[Help("setignore", "Adds a P4IGNORE to the default file (Engine/Extras/Perforce/p4ignore)")]
-	[Help("path=<path>", "Write to a path other than the current directory")]
-	[Help("p4port=<server:port>", "Optional hint/override of the server to use during lookup")]
-	[Help("p4user=<server:port>", "Optional hint/override of the username to use during lookup")]
+	[ParamHelp("SetIgnore", "Adds a P4IGNORE to the default file (Engine/Extras/Perforce/p4ignore)", ParamType = typeof(bool), Flag = "-SetIgnore")]
+	[ParamHelp("Path", "Write to a path other than the current directory")]
+	[ParamHelp("p4port=<server:port>", "Optional hint/override of the server to use during lookup")]
+	[ParamHelp("p4user=<username>", "Optional hint/override of the username to use during lookup")]
 	public class P4WriteConfig : BuildCommand
 	{
 		public override ExitCode Execute()
