Gauntlet - Fix UE error parsing from log

diff --git a/Engine/Source/Programs/AutomationTool/Gauntlet/Unreal/Utils/Gauntlet.AutomationLogParser.cs b/Engine/Source/Programs/AutomationTool/Gauntlet/Unreal/Utils/Gauntlet.AutomationLogParser.cs
--- a/Engine/Source/Programs/AutomationTool/Gauntlet/Unreal/Utils/Gauntlet.AutomationLogParser.cs
+++ b/Engine/Source/Programs/AutomationTool/Gauntlet/Unreal/Utils/Gauntlet.AutomationLogParser.cs
@@ -94,8 +94,8 @@ namespace Gauntlet
 
 			Func<string, string> SantizeLine = (L) =>
 			{
-				L = L.Replace(": Error: ", "");
-				L = L.Replace(": Warning: ", "");
+				L = L.Replace(": Error: ", ": ");
+				L = L.Replace(": Warning: ", ": ");
 				L = L.Replace("LogAutomationController: ", "");
 
 				return L;
