Add a way to save iostore containers that get modified as part of the deployment/staging process.

diff --git a/Engine/Source/Developer/LauncherServices/Private/Launcher/LauncherWorker.cpp b/Engine/Source/Developer/LauncherServices/Private/Launcher/LauncherWorker.cpp
--- a/Engine/Source/Developer/LauncherServices/Private/Launcher/LauncherWorker.cpp
+++ b/Engine/Source/Developer/LauncherServices/Private/Launcher/LauncherWorker.cpp
@@ -876,6 +876,13 @@ FString FLauncherWorker::CreateUATCommand( const ILauncherProfileRef& InProfile,
 			StageAdditionalCommandLine += TEXT(" -ReferenceContainerCryptoKeys=\"") + InProfile->GetReferenceContainerCryptoKeysFileName() + TEXT("\"") ;
 		}
 	}
+
+	if (InProfile->IsUsingIoStore() && 
+		InProfile->IsRetainStagedDirectory())
+	{
+		StageAdditionalCommandLine += TEXT(" -RetainStagedDirectory");
+	}
+
 	// stage/package/deploy
 	if (InProfile->GetDeploymentMode() != ELauncherProfileDeploymentModes::DoNotDeploy)
 	{
diff --git a/Engine/Source/Developer/LauncherServices/Private/Profiles/LauncherProfile.h b/Engine/Source/Developer/LauncherServices/Private/Profiles/LauncherProfile.h
--- a/Engine/Source/Developer/LauncherServices/Private/Profiles/LauncherProfile.h
+++ b/Engine/Source/Developer/LauncherServices/Private/Profiles/LauncherProfile.h
@@ -54,6 +54,7 @@ enum ELauncherVersion
 	LAUNCHERSERVICES_REMOVEDNUMCOOKERSTOSPAWN = 32,
 	LAUNCHERSERVICES_ADDEDORIGINALRELEASEVERSION = 33,
 	LAUNCHERSERVICES_ADDBUILDTARGETNAME = 34,
+	LAUNCHERSERVICED_ADDEDRETAINSTAGEDDIRECTORY = 35,
 	//ADD NEW STUFF HERE
 
 
@@ -577,6 +578,14 @@ public:
 	{
 		ReferenceContainerCryptoKeysFileName = InReferenceContainerCryptoKeysFileName;
 	}
+	virtual bool IsRetainStagedDirectory() const override
+	{
+		return bRetainStagedDirectory;
+	}
+	virtual void SetRetainStagedDirectory(bool bInRetainStagedDirectory) override
+	{
+		bRetainStagedDirectory = bInRetainStagedDirectory;
+	}
 
 	virtual ELauncherProfileDeploymentModes::Type GetDeploymentMode( ) const override
 	{
@@ -1093,6 +1102,11 @@ public:
 			Archive << BuildTargetName;
 		}
 
+		if (Version >= LAUNCHERSERVICED_ADDEDRETAINSTAGEDDIRECTORY)
+		{
+			Archive << bRetainStagedDirectory;
+		}		
+
 		DefaultLaunchRole->Serialize(Archive);
 
 		// serialize launch roles
@@ -1223,6 +1237,7 @@ public:
 		Writer.WriteValue("BasedOnReleaseVersionName", BasedOnReleaseVersionName);
 		Writer.WriteValue("ReferenceContainerGlobalFileName", ReferenceContainerGlobalFileName);
 		Writer.WriteValue("ReferenceContainerCryptoKeysFileName", ReferenceContainerCryptoKeysFileName);
+		Writer.WriteValue("RetainStagedDirectory", bRetainStagedDirectory);
 		Writer.WriteValue("OriginalReleaseVersionName", OriginalReleaseVersionName);
 		Writer.WriteValue("CreateDLC", CreateDLC);
 		Writer.WriteValue("DLCName", DLCName);
@@ -1563,6 +1578,12 @@ public:
 			}
 		}
 
+		if (bIsStaging)
+		{
+			// (only iostore uses this)
+			Writer.WriteValue("RetainStagedDirectory", IsRetainStagedDirectory());
+		}
+
 		/*
 		"script", ""
 		"project", ""
@@ -1893,6 +1914,10 @@ public:
 			OriginalReleaseVersionName.Empty();
 		}
 
+		if (Version >= LAUNCHERSERVICED_ADDEDRETAINSTAGEDDIRECTORY)
+		{
+			bRetainStagedDirectory = Object.GetBoolField("RetainStagedDirectory");
+		}
 
 		CreateDLC = Object.GetBoolField("CreateDLC");
 		DLCName = Object.GetStringField("DLCName");
@@ -3139,6 +3164,10 @@ private:
 	// the json file containing the keys.
 	FString ReferenceContainerCryptoKeysFileName;
 
+	// Some platforms modify iostore containers such that they can't be read as a reference container - if this is true,
+	// the deployment stage will save the staged directory beforehand..
+	bool bRetainStagedDirectory;
+
 	// create a release version of the content (this can be used to base dlc / patches from)
 	bool CreateReleaseVersion;
 
diff --git a/Engine/Source/Developer/LauncherServices/Public/ILauncherProfile.h b/Engine/Source/Developer/LauncherServices/Public/ILauncherProfile.h
--- a/Engine/Source/Developer/LauncherServices/Public/ILauncherProfile.h
+++ b/Engine/Source/Developer/LauncherServices/Public/ILauncherProfile.h
@@ -1146,6 +1146,13 @@ public:
 	virtual FString GetReferenceContainerCryptoKeysFileName() const = 0;
 	virtual void SetReferenceContainerCryptoKeysFileName(const FString& InReferenceContainerCryptoKeysFileName) = 0;
 
+	/**
+	* Whether or not to save the unaltered staged directory for platforms that modify
+	* i/o store containers for deployment. This is so that the containers can be saved and used as a reference 
+	* container later.
+	*/
+	virtual bool IsRetainStagedDirectory() const = 0;
+	virtual void SetRetainStagedDirectory(bool bInRetainStagedDirectory) = 0;
 
 	/**
 	 * Sets if we are going to generate a patch 
diff --git a/Engine/Source/Developer/ProjectLauncher/Private/Widgets/Package/SProjectLauncherPackagingSettings.cpp b/Engine/Source/Developer/ProjectLauncher/Private/Widgets/Package/SProjectLauncherPackagingSettings.cpp
--- a/Engine/Source/Developer/ProjectLauncher/Private/Widgets/Package/SProjectLauncherPackagingSettings.cpp
+++ b/Engine/Source/Developer/ProjectLauncher/Private/Widgets/Package/SProjectLauncherPackagingSettings.cpp
@@ -145,6 +145,25 @@ void SProjectLauncherPackagingSettings::Construct(const FArguments& InArgs, cons
 								]
 							]
 
+						+ SVerticalBox::Slot()
+							.AutoHeight()
+							.Padding(0.0, 4.0, 0.0, 0.0)
+							[
+								SNew(SCheckBox)
+								.IsEnabled(this, &SProjectLauncherPackagingSettings::IsEditable)
+								.IsChecked(this, &SProjectLauncherPackagingSettings::HandleRetainStagedDirectoryCheckBoxIsChecked)
+								.OnCheckStateChanged(this, &SProjectLauncherPackagingSettings::HandleRetainStagedDirectoryCheckStateChanged)
+								.Padding(FMargin(4.0f, 0.0f))
+								.ToolTipText(LOCTEXT("RetainStagedDirectoryCheckBoxTooltip", "Some platforms modify I/O store containers before deployment which prevents using them as reference block database. Set this to true to save the staged directory beforehand."))
+								.Content()
+								[
+									SNew(STextBlock)
+									.Text(LOCTEXT("RetainStagedDirectoryCheckBoxText", "Retain the staged directory if the platform will modify the I/O store containers for deployment"))
+								]
+							]
+
+
+
 						+ SVerticalBox::Slot()
 							.AutoHeight()
 							.Padding(0.0, 4.0, 0.0, 0.0)
@@ -307,6 +326,29 @@ ECheckBoxState SProjectLauncherPackagingSettings::HandleIncludePrerequisitesChec
 	return ECheckBoxState::Unchecked;
 }
 
+
+void SProjectLauncherPackagingSettings::HandleRetainStagedDirectoryCheckStateChanged(ECheckBoxState NewState)
+{
+	ILauncherProfilePtr SelectedProfile = Model->GetSelectedProfile();
+	if (SelectedProfile.IsValid())
+	{
+		SelectedProfile->SetRetainStagedDirectory(NewState == ECheckBoxState::Checked);
+	}
+}
+
+ECheckBoxState SProjectLauncherPackagingSettings::HandleRetainStagedDirectoryCheckBoxIsChecked() const
+{
+	ILauncherProfilePtr SelectedProfile = Model->GetSelectedProfile();
+	if (SelectedProfile.IsValid())
+	{
+		if (SelectedProfile->IsRetainStagedDirectory())
+		{
+			return ECheckBoxState::Checked;
+		}
+	}
+	return ECheckBoxState::Unchecked;
+}
+
 void SProjectLauncherPackagingSettings::HandleUseIoStoreCheckStateChanged(ECheckBoxState NewState)
 {
 	ILauncherProfilePtr SelectedProfile = Model->GetSelectedProfile();
diff --git a/Engine/Source/Developer/ProjectLauncher/Private/Widgets/Package/SProjectLauncherPackagingSettings.h b/Engine/Source/Developer/ProjectLauncher/Private/Widgets/Package/SProjectLauncherPackagingSettings.h
--- a/Engine/Source/Developer/ProjectLauncher/Private/Widgets/Package/SProjectLauncherPackagingSettings.h
+++ b/Engine/Source/Developer/ProjectLauncher/Private/Widgets/Package/SProjectLauncherPackagingSettings.h
@@ -53,6 +53,9 @@ private:
 	void HandleUseIoStoreCheckStateChanged(ECheckBoxState NewState);
 	ECheckBoxState HandleUseIoStoreCheckBoxIsChecked() const;
 
+	void HandleRetainStagedDirectoryCheckStateChanged(ECheckBoxState NewState);
+	ECheckBoxState HandleRetainStagedDirectoryCheckBoxIsChecked() const;
+
 	void HandleMakeBinaryConfigCheckStateChanged(ECheckBoxState NewState);
 	ECheckBoxState HandleMakeBinaryConfigCheckBoxIsChecked() const;
 
diff --git a/Engine/Source/Programs/AutomationTool/AutomationUtils/ProjectParams.cs b/Engine/Source/Programs/AutomationTool/AutomationUtils/ProjectParams.cs
--- a/Engine/Source/Programs/AutomationTool/AutomationUtils/ProjectParams.cs
+++ b/Engine/Source/Programs/AutomationTool/AutomationUtils/ProjectParams.cs
@@ -375,6 +375,7 @@ namespace AutomationTool
 			this.ServerDevice = InParams.ServerDevice;
             this.NullRHI = InParams.NullRHI;
 			this.WriteBackMetadataToAssetRegistry = InParams.WriteBackMetadataToAssetRegistry;
+			this.RetainStagedDirectory = InParams.RetainStagedDirectory;
             this.FakeClient = InParams.FakeClient;
             this.EditorTest = InParams.EditorTest;
             this.RunAutomationTests = InParams.RunAutomationTests;
@@ -578,7 +579,8 @@ namespace AutomationTool
 			string SessionLabel = null,
 			ParamList<string> InMapsToRebuildLightMaps = null,
             ParamList<string> InMapsToRebuildHLOD = null,
-            ParamList<string> TitleID = null
+            ParamList<string> TitleID = null,
+            bool? RetainStagedDirectory = null
 			)
 		{
 			//
@@ -944,6 +946,7 @@ namespace AutomationTool
 			this.UbtArgs = ParseParamValueIfNotSpecified(Command, UbtArgs, "ubtargs", String.Empty);
 			this.AdditionalPackageOptions = ParseParamValueIfNotSpecified(Command, AdditionalPackageOptions, "AdditionalPackageOptions", String.Empty);
 			this.WriteBackMetadataToAssetRegistry = ParseParamValueIfNotSpecified(Command, WriteBackMetadataToAssetRegistry, "WriteBackMetadataToAssetRegistry", String.Empty);
+			this.RetainStagedDirectory = GetParamValueIfNotSpecified(Command, RetainStagedDirectory, this.RetainStagedDirectory, "RetainStagedDirectory");
 
 			// -trace can be used with or without a value
 			if (Trace != null || GetParamValueIfNotSpecified(Command, null, false, "trace"))
@@ -2131,6 +2134,9 @@ namespace AutomationTool
         [Help("WriteBackMetadataToAssetRegistry", "Passthru to iostore staging, see IoStoreUtilities.cpp")]
         public string WriteBackMetadataToAssetRegistry;
 
+        [Help("RetainStagedDirectory", "If set, retain the staged directory for platforms that modify the I/O store containers for deployment. This is necessary for using the reference container for patch preventing on such platforms.")]
+        public bool RetainStagedDirectory;
+
         /// <summary>
         /// Run:adds ?fake to the server URL
         /// </summary>
@@ -3138,6 +3144,7 @@ namespace AutomationTool
 				CommandUtils.LogLog("ForcePackageData={0}", ForcePackageData);
 				CommandUtils.LogLog("NullRHI={0}", NullRHI);
 				CommandUtils.LogLog("WriteBackMetadataToAssetRegistry={0}", WriteBackMetadataToAssetRegistry);
+				CommandUtils.LogLog("RetainStagedDirectory={0}", RetainStagedDirectory);
 				CommandUtils.LogLog("FakeClient={0}", FakeClient);
                 CommandUtils.LogLog("EditorTest={0}", EditorTest);
                 CommandUtils.LogLog("RunAutomationTests={0}", RunAutomationTests); 
