Fix unable to find an explicitly requested Gauntlet test not part of the preferred namespaces

diff --git a/Engine/Source/Programs/AutomationTool/Gauntlet/Framework/Gauntlet.Utils.cs b/Engine/Source/Programs/AutomationTool/Gauntlet/Framework/Gauntlet.Utils.cs
--- a/Engine/Source/Programs/AutomationTool/Gauntlet/Framework/Gauntlet.Utils.cs
+++ b/Engine/Source/Programs/AutomationTool/Gauntlet/Framework/Gauntlet.Utils.cs
@@ -549,6 +549,11 @@ namespace Gauntlet
 					{
 						if (typeof(ITestNode).IsAssignableFrom(Type))
 						{
+							// If there is an Exact match, just take it
+							if (Type.FullName == TestName)
+							{
+								return Type;
+							}
 							CandidateTypes.Add(Type);
 						}
 					}
@@ -617,6 +622,12 @@ namespace Gauntlet
 					}
 				}
 
+				// If nothing found from the prefered namespaces, then return the first item with less amount of steps.
+				if(MatchingTypes.Count() > 0)
+				{
+					return MatchingTypes.OrderBy(M => M.FullName.Split('.').Length).First();
+				}
+
 				throw new AutomationException("Unable to find type {0} in assemblies. Namespaces= {1}.", TestName, Namespaces);
 			}
 
