- Changed a warning to log about missing VIsualStudio assuiming on compilation ability - it is fine for no VS if compilation is not needed (like with running precompiled UAT under Wine)

diff --git a/Engine/Source/Programs/AutomationTool/AutomationUtils/CommandEnvironment.cs b/Engine/Source/Programs/AutomationTool/AutomationUtils/CommandEnvironment.cs
--- a/Engine/Source/Programs/AutomationTool/AutomationUtils/CommandEnvironment.cs
+++ b/Engine/Source/Programs/AutomationTool/AutomationUtils/CommandEnvironment.cs
@@ -209,8 +209,8 @@ namespace AutomationTool
 			}
 			catch (Exception Ex)
 			{
-				Log.WriteLine(LogEventType.Warning, Ex.Message);
-				Log.WriteLine(LogEventType.Warning, "Assuming no compilation capability for NET Framework projects.");
+				Log.WriteLine(LogEventType.Log, Ex.Message);
+				Log.WriteLine(LogEventType.Log, "Assuming no compilation capability for NET Framework projects.");
 				HasCapabilityToCompile = false;
 				FrameworkMsbuildPath = "";
 			}
