Modify Callstack matching regex to elide module address if present.

diff --git a/Engine/Source/Programs/AutomationTool/Gauntlet/Unreal/Utils/Gauntlet.UnrealLogParser.cs b/Engine/Source/Programs/AutomationTool/Gauntlet/Unreal/Utils/Gauntlet.UnrealLogParser.cs
--- a/Engine/Source/Programs/AutomationTool/Gauntlet/Unreal/Utils/Gauntlet.UnrealLogParser.cs
+++ b/Engine/Source/Programs/AutomationTool/Gauntlet/Unreal/Utils/Gauntlet.UnrealLogParser.cs
@@ -743,9 +743,10 @@ namespace Gauntlet
 						// Must have [Callstack] 0x00123456
 						// The module name is optional, must start with whitespace, and continues until next whites[ace
 						// filename is optional, must be in [file]
-						// Note - Ubreal callstacks are always meant to omit all three with placeholders for missing values, but
+						// module address is optional, must be in quotes() after address
+						// Note - Unreal callstacks are always meant to omit all three with placeholders for missing values, but
 						// we'll assume that may not happen...
-						Match CSMatch = Regex.Match(Line, @"(0[xX][0-9A-f]{8,16})(?:\s+([^\s]+))?(?:\s+\[(.*?)\])?", RegexOptions.IgnoreCase);
+						Match CSMatch = Regex.Match(Line, @"(0[xX][0-9A-f]{8,16})(?:\s+\(0[xX][0-9A-f]{8,16}\))?(?:\s+([^\s]+))?(?:\s+\[(.*?)\])?", RegexOptions.IgnoreCase);
 
 						if (CSMatch.Success)
 						{
