Add missing Insight events during EditorPerf tests

#rnx

diff --git a/Engine/Source/Programs/AutomationTool/Gauntlet/Unreal/Editor/EditorTests.cs b/Engine/Source/Programs/AutomationTool/Gauntlet/Unreal/Editor/EditorTests.cs
--- a/Engine/Source/Programs/AutomationTool/Gauntlet/Unreal/Editor/EditorTests.cs
+++ b/Engine/Source/Programs/AutomationTool/Gauntlet/Unreal/Editor/EditorTests.cs
@@ -96,7 +96,8 @@ namespace UnrealEditor
 			{
 				AppConfig.CommandLineParams.Add(string.Format("tracefile={0}", TraceFile));
 				AppConfig.CommandLineParams.Add("tracefiletrunc"); // replace existing
-				AppConfig.CommandLineParams.Add("trace=cpu,counters,bookmarks");
+				AppConfig.CommandLineParams.Add("trace=cpu,counters,stats,gpu,frame,bookmark,log,loadtime,savetime,assetloadtime");
+				AppConfig.CommandLineParams.Add("statnamedevents");
 			}
 
 			if (SuppressLogWarnings)
