Run p4 add with -f to re-add files with p4 wildcards in their name

#rnx

diff --git a/Engine/Source/Programs/AutomationTool/Scripts/FixPerforceCase.cs b/Engine/Source/Programs/AutomationTool/Scripts/FixPerforceCase.cs
--- a/Engine/Source/Programs/AutomationTool/Scripts/FixPerforceCase.cs
+++ b/Engine/Source/Programs/AutomationTool/Scripts/FixPerforceCase.cs
@@ -111,7 +111,7 @@ namespace AutomationTool
 			int AddChangeNumber = P4.CreateChange(Description: String.Format("Fixing case of {0} (2/2){1}", FileSpec, BoilerplateText));
 			foreach (string NewFile in TargetFiles)
 			{
-				P4.LogP4("", String.Format("add -c {0} {1}", AddChangeNumber, NewFile));
+				P4.LogP4("", String.Format("add -f -c {0} {1}", AddChangeNumber, NewFile));
 			}
 			P4.Submit(AddChangeNumber);
 		}
