Horde: Add version to Gauntlet horde reports, add some basic test information to simple reports for new backend
#rnx

diff --git a/Engine/Source/Programs/AutomationTool/Gauntlet/Framework/Utils/Gauntlet.HordeReport.cs b/Engine/Source/Programs/AutomationTool/Gauntlet/Framework/Utils/Gauntlet.HordeReport.cs
--- a/Engine/Source/Programs/AutomationTool/Gauntlet/Framework/Utils/Gauntlet.HordeReport.cs
+++ b/Engine/Source/Programs/AutomationTool/Gauntlet/Framework/Utils/Gauntlet.HordeReport.cs
@@ -85,9 +85,15 @@ namespace Gauntlet
 
 		public abstract class BaseHordeReport : BaseTestReport
 		{
+
+			/// <summary>
+			/// Horde report version
+			/// </summary>
+			public int Version { get; set; } = 1;
+
 			protected string OutputArtifactPath;
 			protected HashSet<string> ArtifactProcessedHashes;
-			protected Dictionary<string, object> ExtraReports;
+			protected Dictionary<string, object> ExtraReports;			
 
 			/// <summary>
 			/// Attach Artifact to the Test Report
@@ -1067,26 +1073,61 @@ namespace Gauntlet
 		/// </summary>
 		public class SimpleTestReport : BaseHordeReport
 		{
+			public class TestRole
+			{
+				public string Type { get; set; }
+				public string Platform { get; set; }
+				public string Configuration { get; set; }
+
+				public TestRole(UnrealTestRole Role, UnrealTestRoleContext Context)
+				{
+					Type = Role.Type.ToString();
+					Platform = Context.Platform.ToString();
+					Configuration = Context.Configuration.ToString();
+				}
+			}
+
 			public override string Type
 			{
 				get { return "Simple Report"; }
 			}
+
 			public SimpleTestReport() : base()
 			{
-				Logs = new List<String>();
-				Errors = new List<String>();
-				Warnings = new List<String>();
+
+			}
+
+			public SimpleTestReport(Gauntlet.TestResult TestResult, UnrealTestContext Context, UnrealTestConfiguration Configuration) : base()
+			{
+
+				BuildChangeList = Context.BuildInfo.Changelist;
+				this.TestResult = TestResult.ToString();
+
+				// populate roles
+				UnrealTestRole MainRole = Configuration.GetMainRequiredRole();
+				this.MainRole = new TestRole(MainRole, Context.GetRoleContext(MainRole.Type));
+
+				IEnumerable<UnrealTestRole> Roles = Configuration.RequiredRoles.Values.SelectMany(V => V);
+				foreach (UnrealTestRole Role in Roles)
+				{
+					this.Roles.Add(new TestRole(Role, Context.GetRoleContext(Role.Type)));
+				}
 			}
 
+			public string TestName { get; set; }
 			public string Description { get; set; }
 			public string ReportCreatedOn { get; set; }
 			public float TotalDurationSeconds { get; set; }
 			public bool HasSucceeded { get; set; }
 			public string Status { get; set; }
 			public string URLLink { get; set; }
-			public List<String> Logs { get; set; }
-			public List<String> Errors { get; set; }
-			public List<String> Warnings { get; set; }
+			public int BuildChangeList { get; set; }
+			public TestRole MainRole { get; set; }
+			public List<TestRole> Roles { get; set; } = new List<TestRole>();
+			public string TestResult { get; set; }
+			public List<String> Logs { get; set; } = new List<String>();
+			public List<String> Errors { get; set; } = new List<String>();
+			public List<String> Warnings { get; set; } = new List<String>();
 
 			public override void AddEvent(EventType Type, string Message, object Context = null)
 			{
@@ -1121,7 +1162,7 @@ namespace Gauntlet
 			public class DataItem
 			{
 				public string Key { get; set; }
-				public object Data { get; set; }
+				public object Data { get; set; }				
 			}
 			public TestDataCollection()
 			{
diff --git a/Engine/Source/Programs/AutomationTool/Gauntlet/Unreal/Base/Gauntlet.UnrealTestNode.cs b/Engine/Source/Programs/AutomationTool/Gauntlet/Unreal/Base/Gauntlet.UnrealTestNode.cs
--- a/Engine/Source/Programs/AutomationTool/Gauntlet/Unreal/Base/Gauntlet.UnrealTestNode.cs
+++ b/Engine/Source/Programs/AutomationTool/Gauntlet/Unreal/Base/Gauntlet.UnrealTestNode.cs
@@ -1325,7 +1325,15 @@ namespace Gauntlet
 			{
 				GetCachedConfiguration().HordeTestDataKey = Name + " " + Context.ToString();
 			}
-			HordeReport.SimpleTestReport HordeTestReport = new HordeReport.SimpleTestReport();
+
+			string TestName = Name;
+			if (TestName.Split('.').Length > 1)
+			{
+				TestName = TestName.Split('.').Last();
+			}
+
+			HordeReport.SimpleTestReport HordeTestReport = new HordeReport.SimpleTestReport(UnrealTestResult, Context, GetCachedConfiguration());
+			HordeTestReport.TestName = TestName;
 			HordeTestReport.ReportCreatedOn = DateTime.Now.ToString();
 			HordeTestReport.TotalDurationSeconds = (float) (DateTime.Now - SessionStartTime).TotalSeconds;
 			HordeTestReport.Description = Context.ToString();
@@ -1447,23 +1455,31 @@ namespace Gauntlet
 			}
 		}
 
+		void SubmitToHorde(ITestReport Report)
+		{
+			if (!GetCachedConfiguration().WriteTestResultsForHorde)
+			{
+				return;
+			}
+
+			// write test data collection for Horde
+			string HordeTestDataFilePath = Path.Combine(
+				string.IsNullOrEmpty(GetCachedConfiguration().HordeTestDataPath) ? HordeReport.DefaultTestDataDir : GetCachedConfiguration().HordeTestDataPath,
+				FileUtils.SanitizeFilename(Name) + ".TestData.json"
+			);
+			HordeReport.TestDataCollection HordeTestDataCollection = new HordeReport.TestDataCollection();
+			HordeTestDataCollection.AddNewTestReport(Report, GetCachedConfiguration().HordeTestDataKey);
+			HordeTestDataCollection.WriteToJson(HordeTestDataFilePath, true);
+		}
+
 		/// <summary>
 		/// Optional function that is called on test completion and gives an opportunity to submit a report to a Dashboard
 		/// </summary>
 		/// <param name="Report"></param>
 		public virtual void SubmitToDashboard(ITestReport Report)
 		{
-			if (GetCachedConfiguration().WriteTestResultsForHorde)
-			{
-				// write test data collection for Horde
-				string HordeTestDataFilePath = Path.Combine(
-					string.IsNullOrEmpty(GetCachedConfiguration().HordeTestDataPath) ? HordeReport.DefaultTestDataDir : GetCachedConfiguration().HordeTestDataPath,
-					FileUtils.SanitizeFilename(Name) + ".TestData.json"
-				);
-				HordeReport.TestDataCollection HordeTestDataCollection = new HordeReport.TestDataCollection();
-				HordeTestDataCollection.AddNewTestReport(Report, GetCachedConfiguration().HordeTestDataKey);
-				HordeTestDataCollection.WriteToJson(HordeTestDataFilePath, true);
-			}
+			SubmitToHorde(Report);
+
 			if (!string.IsNullOrEmpty(GetCachedConfiguration().PublishTelemetryTo) && Report is ITelemetryReport Telemetry)
 			{
 				IEnumerable<TelemetryData> DataRows = Telemetry.GetAllTelemetryData();
