Avoid using a time based query when validating log

#rnx

diff --git a/Engine/Source/Programs/AutomationTool/Gauntlet/Unreal/Automation/UE.BootTest.cs b/Engine/Source/Programs/AutomationTool/Gauntlet/Unreal/Automation/UE.BootTest.cs
--- a/Engine/Source/Programs/AutomationTool/Gauntlet/Unreal/Automation/UE.BootTest.cs
+++ b/Engine/Source/Programs/AutomationTool/Gauntlet/Unreal/Automation/UE.BootTest.cs
@@ -155,39 +155,21 @@ namespace UE
 				else
 				{
 					// find a logfile or something that indicates the process ran successsfully
-					bool MissingFiles = false;
+					bool MissingLogs = false;
 
 					foreach (var RoleResult in InResults)
 					{
-						DirectoryInfo RoleDir = new DirectoryInfo(RoleResult.Artifacts.ArtifactPath);
-
-						IEnumerable<FileInfo> ArtifactFiles = RoleDir.EnumerateFiles("*.*", SearchOption.AllDirectories);
-
-						// user may not have cleared paths between runs, so throw away anything that's older than 2m
-						ArtifactFiles = ArtifactFiles.Where(F => (DateTime.Now - F.LastWriteTime).TotalMinutes < 2);
-
-						if (ArtifactFiles.Any() == false)
-						{
-							MissingFiles = true;
-							ReportError("No artifact files found for {0}. Were they not retrieved from the device?", RoleResult.Artifacts.SessionRole);
-						}
-
-						IEnumerable<FileInfo> LogFiles = ArtifactFiles.Where(F => F.Extension.Equals(".log", StringComparison.OrdinalIgnoreCase));
-
-						if (LogFiles.Any() == false)
+						if (!File.Exists(RoleResult.Artifacts.LogPath))
 						{
-							MissingFiles = true;
+							MissingLogs = true;
 							ReportError("No log files found for {0}. Were they not retrieved from the device?", RoleResult.Artifacts.SessionRole);
 						}
 					}
 
-					if (MissingFiles)
+					if (!MissingLogs)
 					{
-						SetUnrealTestResult(TestResult.Failed);
-						ReportError("One or more roles did not generated any artifacts");
+						Log.Info("Found valid log artifacts for test");
 					}
-
-					Log.Info("Found valid artifacts for test");
 				}
 			}
 
