Typo fixup

diff --git a/Engine/Source/Programs/AutomationTool/Gauntlet/Unreal/Base/Gauntlet.UnrealSession.cs b/Engine/Source/Programs/AutomationTool/Gauntlet/Unreal/Base/Gauntlet.UnrealSession.cs
--- a/Engine/Source/Programs/AutomationTool/Gauntlet/Unreal/Base/Gauntlet.UnrealSession.cs
+++ b/Engine/Source/Programs/AutomationTool/Gauntlet/Unreal/Base/Gauntlet.UnrealSession.cs
@@ -272,7 +272,7 @@ namespace Gauntlet
 		}
 
 		/// <summary>
-		/// Helper for accessing server process. May return null if no serer is invlved
+		/// Helper for accessing server process. May return null if no server is involved
 		/// </summary>
 		public IAppInstance ServerApp
 		{
@@ -283,7 +283,7 @@ namespace Gauntlet
 		}
 
 		/// <summary>
-		/// Helper for accessing server process. May return null if no serer is invlved
+		/// Helper for accessing editor process. May return null if no editor is involved
 		/// </summary>
 		public IAppInstance EditorApp
 		{
