Fixed crash during run command when deploy to a Linux server. Fix verified using Lyra Game using ushell to build, stage, and deploy the server.

diff --git a/Engine/Source/Programs/AutomationTool/Linux/LinuxPlatform.Automation.cs b/Engine/Source/Programs/AutomationTool/Linux/LinuxPlatform.Automation.cs
--- a/Engine/Source/Programs/AutomationTool/Linux/LinuxPlatform.Automation.cs
+++ b/Engine/Source/Programs/AutomationTool/Linux/LinuxPlatform.Automation.cs
@@ -14,6 +14,9 @@ using UnrealBuildBase;
 
 public abstract class BaseLinuxPlatform : Platform
 {
+	// Matches strings of the form "DeviceName@IP Address" such as "[REDACTED_EMAIL]"
+	static Regex DeviceRegex = new Regex(@"\w.+@([A-Za-z0-9\.\-]+)[\+]?");
+
 	static string PScpPath = CombinePaths(Unreal.RootDirectory.FullName, "\\Engine\\Extras\\ThirdPartyNotUE\\putty\\PSCP.EXE");
 	static string PlinkPath = CombinePaths(Unreal.RootDirectory.FullName, "\\Engine\\Extras\\ThirdPartyNotUE\\putty\\PLINK.EXE");
 	static string LaunchOnHelperShellScriptName = "LaunchOnHelper.sh";
@@ -322,6 +325,12 @@ chmod +x {0}
 	{
 		if ((ProjParams.Deploy || ProjParams.Run) && BuildHostPlatform.Current.Platform != UnrealTargetPlatform.Linux)
 		{
+ 			Match ServerDeviceMatch = DeviceRegex.Match(ProjParams.ServerDevice);
+			if (ServerDeviceMatch.Success)
+			{
+				ProjParams.ServerDeviceAddress = ServerDeviceMatch.Groups[1].Value;
+			}
+
 			// we don't need username/password if we are only targeting steamdeck devices
 			bool bNeedsUsernameAndPassword = false;
 			foreach (string DeviceId in ProjParams.DeviceNames)
diff --git a/Engine/Source/Programs/AutomationTool/Scripts/RunProjectCommand.Automation.cs b/Engine/Source/Programs/AutomationTool/Scripts/RunProjectCommand.Automation.cs
--- a/Engine/Source/Programs/AutomationTool/Scripts/RunProjectCommand.Automation.cs
+++ b/Engine/Source/Programs/AutomationTool/Scripts/RunProjectCommand.Automation.cs
@@ -862,13 +862,13 @@ namespace AutomationScripts
 				ServerApp = Exes[0].FullName;
 			}
 			var Args = bCooked ? "" : (SC.ProjectArgForCommandLines + " ");
-			Console.WriteLine(Params.ServerDeviceAddress);
+			Console.WriteLine("Running dedicated server on device with address: " + Params.ServerDeviceAddress);
 			TargetPlatformDescriptor ServerPlatformDesc = ServerParams.ServerTargetPlatforms[0];
 			if (ServerParams.Cook && ServerPlatformDesc.Type == UnrealTargetPlatform.Linux && !String.IsNullOrEmpty(ServerParams.ServerDeviceAddress))
 			{
 				ServerApp = @"C:\Windows\system32\cmd.exe";
 
-				string plinkPath = CombinePaths(Environment.GetEnvironmentVariable("LINUX_ROOT"), "bin/PLINK.exe ");
+				string plinkPath = CombinePaths(Unreal.EngineDirectory.FullName, "Extras", "ThirdPartyNotUE", "putty", "PLINK.exe");
 				string exePath = CombinePaths(SC.ShortProjectName, "Binaries", ServerPlatformDesc.Type.ToString(), SC.ShortProjectName + "Server");
 				if (ServerParams.ServerConfigsToBuild[0] != UnrealTargetConfiguration.Development)
 				{
