Use Structured Logging in Gauntlet

diff --git a/Engine/Source/Programs/AutomationTool/AutomationUtils/Matchers/GauntletEventMatcher.cs b/Engine/Source/Programs/AutomationTool/AutomationUtils/Matchers/GauntletEventMatcher.cs
--- a/Engine/Source/Programs/AutomationTool/AutomationUtils/Matchers/GauntletEventMatcher.cs
+++ b/Engine/Source/Programs/AutomationTool/AutomationUtils/Matchers/GauntletEventMatcher.cs
@@ -13,13 +13,12 @@ namespace AutomationUtils.Matchers
 	/// </summary>
 	class GauntletEventMatcher : ILogEventMatcher
 	{
-		static readonly Regex s_engineTestPattern = new Regex(@"^(?<indent>\s*)Error: EngineTest.RunTests Group:(?<group>[^\s]+) \(");
+		static readonly Regex s_engineTestPattern = new Regex(@"^(?<indent>\s*)## Error: UE\.(?<target>\w+?)Automation\(RunTest=(?<group>[^\s]+)\) \(");
 
 		static readonly Regex s_mdHeadingPattern = new Regex(@"^\s*#{1,3} ");
-		static readonly Regex s_mdTestNamePattern = new Regex(@"^\s*#####\s+(?<friendly_name>.*):\s*(?<name>\S+)\s*");
 		static readonly Regex s_mdTestListPattern = new Regex(@"^\s*### The following tests failed:");
-
-		static readonly Regex s_screenshotPattern = new Regex(@"Error: Screenshot '(?<screenshot>[^']+)' test failed");
+		static readonly Regex s_mdTestNamePattern = new Regex(@"^\s*#####\s+Error: Test '(?<name>\S+)'");
+		static readonly Regex s_mdTestFullNamePattern = new Regex(@"^\s*FullName:\s+(?<fullname>\S+)");
 
 		static readonly Regex s_genericPattern = new Regex("\\[ERROR\\] (.*)$");
 
@@ -55,9 +54,15 @@ namespace AutomationUtils.Matchers
 							builder.AddProperty("group", group);//.Annotate().AddProperty("group", Group);
 
 							builder.Annotate(testNameMatch.Groups["name"]);
-							builder.Annotate(testNameMatch.Groups["friendly_name"]);
 
-							eventId = KnownLogEvents.Gauntlet_UnitTest;//							ChildEvents.Add(ChildEventBuilder);
+							builder.MoveNext();
+							Match? testFullNameMatch;
+							if (builder.IsNextLineAligned() && builder.Current.TryMatch(s_mdTestFullNamePattern, out testFullNameMatch))
+							{
+								builder.Annotate(testFullNameMatch.Groups["fullname"]);
+							}
+
+							eventId = KnownLogEvents.Gauntlet_UnrealEngineTestEvent;//							ChildEvents.Add(ChildEventBuilder);
 						}
 					}
 					else
@@ -72,13 +77,6 @@ namespace AutomationUtils.Matchers
 				return builder.ToMatch(LogEventPriority.High, LogLevel.Error, eventId);
 			}
 
-			if (cursor.TryMatch(s_screenshotPattern, out match))
-			{
-				LogEventBuilder builder = new LogEventBuilder(cursor);
-				builder.Annotate(match.Groups["screenshot"], LogEventMarkup.ScreenshotTest);//.MarkAsScreenshotTest();
-				return builder.ToMatch(LogEventPriority.High, LogLevel.Error, KnownLogEvents.Gauntlet_ScreenshotTest);
-			}
-
 			if (cursor.TryMatch(s_genericPattern, out _))
 			{
 				LogEventBuilder builder = new LogEventBuilder(cursor);
diff --git a/Engine/Source/Programs/AutomationTool/Gauntlet/Framework/Devices/Gauntlet.DevicePool.cs b/Engine/Source/Programs/AutomationTool/Gauntlet/Framework/Devices/Gauntlet.DevicePool.cs
--- a/Engine/Source/Programs/AutomationTool/Gauntlet/Framework/Devices/Gauntlet.DevicePool.cs
+++ b/Engine/Source/Programs/AutomationTool/Gauntlet/Framework/Devices/Gauntlet.DevicePool.cs
@@ -10,6 +10,7 @@ using System.Text.RegularExpressions;
 using System.Linq;
 using System.Text.Json;
 using System.Text.Json.Serialization;
+using EpicGames.Core;
 
 namespace Gauntlet
 {
@@ -344,7 +345,7 @@ namespace Gauntlet
 
 						if (Failures >= 0.33)
 						{
-							FailedProvisions.ForEach(D => Log.Warning("Device {0} could not be added", D));
+							FailedProvisions.ForEach(D => Log.Warning(KnownLogEvents.Gauntlet_DeviceEvent, "Device {Name} could not be added", D));
 						}
 					}
 
@@ -354,7 +355,7 @@ namespace Gauntlet
 					// warn if anything was left registered
 					ReservedDevices.ForEach(D =>
 					{
-						Log.Warning("Device {0} was not unregistered prior to DevicePool.Dispose!. Forcibly disposing.", D.Name);
+						Log.Warning(KnownLogEvents.Gauntlet_DeviceEvent, "Device {Name} was not unregistered prior to DevicePool.Dispose!. Forcibly disposing.", D.Name);
 						D.Dispose();
 					});
 
@@ -499,7 +500,7 @@ namespace Gauntlet
 		{
 			if (MaxDeviceErrorReports == 0)
 			{
-				Log.Warning("Maximum device errors reported to backend, {0} : {1} ignored", ServiceDeviceName, ErrorMessage);
+				Log.Warning(KnownLogEvents.Gauntlet_DeviceEvent, "Maximum device errors reported to backend, {Name} : {Message} ignored", ServiceDeviceName, ErrorMessage);
 				return;
 			}
 
@@ -529,7 +530,7 @@ namespace Gauntlet
 				catch (Exception Ex)
 				{
 					// warn and use a different directory
-					Log.Warning("Failed to delete {0}. {1}", ClientTempDir, Ex.Message);
+					Log.Warning(KnownLogEvents.Gauntlet_DeviceEvent, "Failed to delete {Folder}. {Message}", ClientTempDir, Ex.Message);
 					ClientTempDir = Path.Combine(PlatformCache, string.Format("{0}_{1}", InDeviceDefiniton.Name, CleanAttempts++));
 				}
 			}
@@ -639,7 +640,7 @@ namespace Gauntlet
                 if (!DeviceMap.ContainsKey(Entry.Key.Platform.Value))
 				{
 					// if an unsupported device, we can't reserve it
-					Log.Info("Unable to reserve service device of type: {0}", Entry.Key);
+					Log.Info(KnownLogEvents.Gauntlet_DeviceEvent, "Unable to reserve service device of type: {Type}", Entry.Key);
 					return false;
 				}
 
@@ -648,7 +649,7 @@ namespace Gauntlet
 					// if specific device model, we can't currently reserve it from (legacy) service
 					if (DeviceURL.ToLower().Contains("deviceservice.epicgames.net"))
 					{
-						Log.Info("Unable to reserve service device of model: {0} on legacy service", Entry.Key.Model);
+						Log.Info(KnownLogEvents.Gauntlet_DeviceEvent, "Unable to reserve service device of model: {Model} on legacy service", Entry.Key.Model);
 						return false;
 					}
 									
@@ -680,7 +681,7 @@ namespace Gauntlet
 				}
 				catch (Exception Ex)
 				{
-					Log.Info("Unable to make device registration: {0}", Ex);
+					Log.Info(KnownLogEvents.Gauntlet_DeviceEvent, "Unable to make device registration: {Exception}", Ex);
 					return false;
 				}
 
@@ -727,7 +728,7 @@ namespace Gauntlet
 							DeviceReservation.Dispose();
 						}
 
-						Log.Info("Unable to make device registration: device registration failed for {0}:{1}", Def.Platform, Def.Name);
+						Log.Info(KnownLogEvents.Gauntlet_DeviceEvent, "Unable to make device registration: device registration failed for {Platform}:{Name}", Def.Platform, Def.Name);
 						return false;
 					}
 					else
@@ -739,14 +740,14 @@ namespace Gauntlet
 					ServiceReservations[TargetDevice] = DeviceReservation;
 				}
 
-				Log.Info("Successfully reserved service devices");
-				Devices.ForEach(Device => Log.Verbose("    Device: {0}", Device));
+				Log.Info(KnownLogEvents.Gauntlet_DeviceEvent, "Successfully reserved service devices");
+				Devices.ForEach(Device => Log.Verbose(KnownLogEvents.Gauntlet_DeviceEvent, "    Device: {Name}", Device));
 
 				return true;
 			}
 
-			Log.Info("Unable to reserve service devices:");
-			Devices.ForEach(Device => Log.Info("    Device: {0}", Device));
+			Log.Info(KnownLogEvents.Gauntlet_DeviceEvent, "Unable to reserve service devices:");
+			Devices.ForEach(Device => Log.Info(KnownLogEvents.Gauntlet_DeviceEvent, "    Device: {Name}", Device));
 			return false;
 		}
 
@@ -769,7 +770,7 @@ namespace Gauntlet
 				// Did they specify a file?
 				if (PossibleFileName && File.Exists(InputReference))
 				{
-					Gauntlet.Log.Info("Adding devices from {0}", InputReference);
+					Log.Info(KnownLogEvents.Gauntlet_DeviceEvent, "Adding devices from {Reference}", InputReference);
 					List<DeviceDefinition> DeviceDefinitions = JsonSerializer.Deserialize<List < DeviceDefinition >>(
 						File.ReadAllText(InputReference),
 						new JsonSerializerOptions { PropertyNameCaseInsensitive = true }
@@ -777,7 +778,7 @@ namespace Gauntlet
 
 					foreach (DeviceDefinition Def in DeviceDefinitions)
 					{
-						Gauntlet.Log.Info("Adding {0}", Def);
+						Log.Info(KnownLogEvents.Gauntlet_DeviceEvent, "Adding {Name}", Def);
 
 						// use Legacy field if it exists
 						if (Def.Platform == null)
@@ -791,7 +792,7 @@ namespace Gauntlet
 							// check whether disabled
 							if (String.Compare(Def.Available, "disabled", true) == 0)
 							{
-								Gauntlet.Log.Info("Skipping {0} due to being disabled", Def.Name);
+								Log.Info(KnownLogEvents.Gauntlet_DeviceEvent, "Skipping {Name} due to being disabled", Def.Name);
 								continue;
 							}
 
@@ -821,13 +822,13 @@ namespace Gauntlet
 
 									if (DateTime.Now < From || DateTime.Now > To)
 									{
-										Gauntlet.Log.Info("Skipping {0} due to availability constraint {1}", Def.Name, Def.Available);
+										Log.Info(KnownLogEvents.Gauntlet_DeviceEvent, "Skipping {Name} due to availability constraint {Constraint}", Def.Name, Def.Available);
 										continue;
 									}
 								}
 								else
 								{
-									Gauntlet.Log.Warning("Failed to parse availability {0} for {1}", Def.Available, Def.Name);
+									Log.Warning(KnownLogEvents.Gauntlet_DeviceEvent, "Failed to parse availability {Constraint} for {Name}", Def.Available, Def.Name);
 								}
 							}
 						}
@@ -887,7 +888,7 @@ namespace Gauntlet
 
 							} 
 
-							Log.Info("Added device {0}:{1} to pool", DevicePlatform, DeviceAddress);
+							Log.Info(KnownLogEvents.Gauntlet_DeviceEvent, "Added device {Platform}:{Address} to pool", DevicePlatform, DeviceAddress);
 							DeviceDefinition Def = new DeviceDefinition();
 							Def.Address = DeviceAddress;
 							Def.Name = DeviceAddress;
@@ -972,12 +973,12 @@ namespace Gauntlet
 
 				if (NewDevice.IsAvailable == false)
 				{
-					Log.Info("Assigned device {0} reports unavailable. Requesting a forced disconnect", NewDevice);
+					Log.Info(KnownLogEvents.Gauntlet_DeviceEvent, "Assigned device {Name} reports unavailable. Requesting a forced disconnect", NewDevice);
 					NewDevice.Disconnect(true);
 
 					if  (NewDevice.IsAvailable == false)
 					{
-						Log.Info("Assigned device {0} still  unavailable. Requesting a reboot", NewDevice);
+						Log.Info(KnownLogEvents.Gauntlet_DeviceEvent, "Assigned device {Name} still  unavailable. Requesting a reboot", NewDevice);
 						NewDevice.Reboot();
 					}
 				}
@@ -992,7 +993,7 @@ namespace Gauntlet
 			}
 			catch (Exception Ex)
 			{
-				Log.Info("Failed to create device {0}. {1}", Def.ToString(), Ex.Message);
+				Log.Info(KnownLogEvents.Gauntlet_DeviceEvent, "Failed to create device {Name}. {Message}", Def.ToString(), Ex.Message);
 			}
 
 			return NewDevice;
@@ -1080,7 +1081,7 @@ namespace Gauntlet
 				TotalDeviceTypes[Constraint] = DevicePool.Instance.GetTotalDeviceCount(Constraint, Validate);
 
 
-				Log.Verbose("{0}: {1} devices required. Total:{2}, Available:{3}",
+				Log.Verbose(KnownLogEvents.Gauntlet_DeviceEvent, "{Constraint}: {Platform} devices required. Total:{Total}, Available:{Available}",
 					Constraint, PlatformRequirement.Value,
 					TotalDeviceTypes[PlatformRequirement.Key], AvailableDeviceTypes[PlatformRequirement.Key]);
 			}
@@ -1100,7 +1101,7 @@ namespace Gauntlet
 
 				Devices.ToList().ForEach(Platform => DeviceCounts[Platform] = RequiredDevices[Platform]);
 
-				Log.Info("Requesting devices from service at {0}", DeviceURL);
+				Log.Info(KnownLogEvents.Gauntlet_DeviceEvent, "Requesting devices from service at {URL}", DeviceURL);
 
 				// Acquire necessary devices from service
 				if (!ReserveDevicesFromService(DeviceURL, DeviceCounts))
@@ -1114,7 +1115,7 @@ namespace Gauntlet
 				if (TooFewTotalDevices.Count() > 0)
 				{
 					var MissingDeviceStrings = TooFewTotalDevices.Select(D => string.Format("Not enough devices of type {0} exist for test. ({1} required, {2} available)", D, RequiredDevices[D], AvailableDeviceTypes[D]));
-					Log.Error("{0}", string.Join("\n", MissingDeviceStrings));
+					Log.Error(KnownLogEvents.Gauntlet_DeviceEvent, string.Join("\n", MissingDeviceStrings));
 					throw new AutomationException("Not enough devices available");
 				}
 
@@ -1122,7 +1123,7 @@ namespace Gauntlet
 				if (TooFewCurrentDevices.Count() > 0)
 				{
 					var MissingDeviceStrings = TooFewCurrentDevices.Select(D => string.Format("Not enough devices of type {0} available for test. ({1} required, {2} available)", D, RequiredDevices[D], AvailableDeviceTypes[D]));
-					Log.Verbose("{0}", string.Join("\n", MissingDeviceStrings));
+					Log.Verbose(KnownLogEvents.Gauntlet_DeviceEvent, string.Join("\n", MissingDeviceStrings));
 					return false;
 				}
 			}
@@ -1145,13 +1146,13 @@ namespace Gauntlet
 			{
 				List<ITargetDevice> Selection = new List<ITargetDevice>();
 				
-				Log.Verbose($"Enumerating devices for constraint {Constraint}");
+				Log.Verbose(KnownLogEvents.Gauntlet_DeviceEvent, $"Enumerating devices for constraint {Constraint}");
 
-				Log.Verbose($"   Available devices:");
-				AvailableDevices.ForEach(D => Log.Verbose($"      {D.Platform}:{D.Name}"));
+				Log.Verbose(KnownLogEvents.Gauntlet_DeviceEvent, $"   Available devices:");
+				AvailableDevices.ForEach(D => Log.Verbose(KnownLogEvents.Gauntlet_DeviceEvent, $"      {D.Platform}:{D.Name}"));
 
-				Log.Verbose($"   Unprovisioned devices:");
-				UnprovisionedDevices.ForEach(D => Log.Verbose($"      {D}"));
+				Log.Verbose(KnownLogEvents.Gauntlet_DeviceEvent, $"   Unprovisioned devices:");
+				UnprovisionedDevices.ForEach(D => Log.Verbose(KnownLogEvents.Gauntlet_DeviceEvent, $"      {D}"));
 				
 				// randomize the order of all devices that are of this platform
 				var MatchingProvisionedDevices = AvailableDevices.Where(D => Constraint.Check(D)).ToList();
@@ -1169,7 +1170,7 @@ namespace Gauntlet
 
 					while (NextDevice != null && ContinuePredicate)
 					{
-						Log.Verbose("Checking {0} against predicate", NextDevice.Name);
+						Log.Verbose(KnownLogEvents.Gauntlet_DeviceEvent, "Checking {Name} against predicate", NextDevice.Name);
 						MatchingProvisionedDevices.Remove(NextDevice);
 						ContinuePredicate = Predicate(NextDevice);
 
@@ -1185,7 +1186,7 @@ namespace Gauntlet
 
 						if (NextDeviceDef != null)
 						{
-							Log.Verbose("Provisioning device {0} for the pool", NextDeviceDef.Name);
+							Log.Verbose(KnownLogEvents.Gauntlet_DeviceEvent, "Provisioning device {Name} for the pool", NextDeviceDef.Name);
 
 							// try to create a device. This can fail, but if so we'll just end up back here
 							// on the next iteration
@@ -1197,11 +1198,11 @@ namespace Gauntlet
 							if (NewDevice != null)
 							{
 								MatchingProvisionedDevices.Add(NewDevice);
-								Log.Verbose("Added device {0} to pool", NewDevice.Name);
+								Log.Verbose(KnownLogEvents.Gauntlet_DeviceEvent, "Added device {Name} to pool", NewDevice.Name);
 							}
 							else
 							{
-								Log.Info("Failed to provision {0}", NextDeviceDef.Name);
+								Log.Info(KnownLogEvents.Gauntlet_DeviceEvent, "Failed to provision {Name}", NextDeviceDef.Name);
 								// track this
 								if (FailedProvisions.Contains(NextDeviceDef) == false)
 								{
@@ -1211,7 +1212,7 @@ namespace Gauntlet
 						}
 						else
 						{
-							Log.Info("Pool ran out of devices of type {0}!", Constraint);
+							Log.Info(KnownLogEvents.Gauntlet_DeviceEvent, "Pool ran out of devices of type {Constraint}!", Constraint);
 							OutOfDevices = true;
 						}
 					}
diff --git a/Engine/Source/Programs/AutomationTool/Gauntlet/Framework/Gauntlet.TestExecutor.cs b/Engine/Source/Programs/AutomationTool/Gauntlet/Framework/Gauntlet.TestExecutor.cs
--- a/Engine/Source/Programs/AutomationTool/Gauntlet/Framework/Gauntlet.TestExecutor.cs
+++ b/Engine/Source/Programs/AutomationTool/Gauntlet/Framework/Gauntlet.TestExecutor.cs
@@ -7,6 +7,7 @@ using System.Linq;
 using System.Threading;
 using System.Text.RegularExpressions;
 using System.Drawing;
+using EpicGames.Core;
 
 namespace Gauntlet
 {
@@ -120,7 +121,7 @@ namespace Gauntlet
 		{
 			Options = InOptions;
 
-			Log.Info("Preparing to start {0} automation test(s)", RequiredTests.Count());
+			Log.Info("Preparing to start {Number} automation test(s)", RequiredTests.Count());
 
 			// install a cancel handler so we can stop parallel-for gracefully
 			Action CancelHandler = delegate ()
@@ -173,7 +174,7 @@ namespace Gauntlet
 
 				DateTime StartPassTime = DateTime.Now;
 
-				Log.Info("Starting test iteration {0} of {1}", CurrentTestPass + 1, Options.TestIterations);
+				Log.Info("Starting test iteration {Index} of {Total}", CurrentTestPass + 1, Options.TestIterations);
 
 				// Tests that we want to run
 				List<TestExecutionInfo> PendingTests = RequiredTests.Select(N => new TestExecutionInfo(N)).ToList();
@@ -230,7 +231,7 @@ namespace Gauntlet
 							}
 							catch (System.Exception ex)
 							{
-								Log.Error("Test {0} threw an exception during ready check. Ex: {1}", Node, ex);
+								Log.Error(KnownLogEvents.Gauntlet_TestEvent, "Test {Name} threw an exception during ready check. Ex: {Exception}", Node, ex);
 								Node.AddTestEvent(new UnrealTestEvent(EventSeverity.Error, "Test Failed to Start", new List<string> {ex.Message}));
 								PendingTests[i] = null;
 								NodeInfo.TimeSetupBegan = NodeInfo.TimeSetupEnded = NodeInfo.TimeTestEnded = DateTime.Now;
@@ -265,7 +266,7 @@ namespace Gauntlet
 									double TimeWaiting = (DateTime.Now - NodeInfo.FirstReadyCheckTime).TotalSeconds;
 									if (TimeWaiting >= Options.Wait)
 									{
-										Log.Warning("Test {0} has been waiting to run resource-free for {1:00} seconds. Removing from wait list", Node, TimeWaiting);
+										Log.Warning(KnownLogEvents.Gauntlet_DeviceEvent, "Test {Name} has been waiting to run resource-free for {Time:00} seconds. Removing from wait list", Node, TimeWaiting);
 										Node.AddTestEvent(new UnrealTestEvent(EventSeverity.Error, "Insufficient devices found", new List<string> {string.Format("Test {0} was unable to find enough devices after trying for {1:00} seconds.", Node, TimeWaiting), "This is not a test-related failure."}));
 										PendingTests[i] = null;
 										NodeInfo.TimeSetupBegan = NodeInfo.TimeSetupEnded = NodeInfo.TimeTestEnded = DateTime.Now;
@@ -281,7 +282,7 @@ namespace Gauntlet
 
 						if (TestToStart != null)
 						{
-							Log.Info("Test {0} is ready to run", TestToStart);
+							Log.Info("Test {Name} is ready to run", TestToStart);
 
 							PendingTests.Remove(TestToStart);
 							StartingTests.Add(TestToStart);
@@ -301,7 +302,7 @@ namespace Gauntlet
 									{
 										TestToStart.TimeSetupEnded = TestToStart.TimeTestEnded = DateTime.Now;
 										CompletedTests.Add(TestToStart);
-										Log.Error("Test {0} failed to start", TestToStart);
+										Log.Error(KnownLogEvents.Gauntlet_TestEvent, "Test {Name} failed to start", TestToStart);
 									}
 									else
 									{
@@ -341,7 +342,7 @@ namespace Gauntlet
 
 							if ((SecondsRunning % 60) == 0)
 							{
-								Log.Verbose("Test {0} is still running. {1:00} seconds elapsed, will timeout in {2:00} seconds",
+								Log.Verbose("Test {Name} is still running. {Elapsed:00} seconds elapsed, will timeout in {Max:00} seconds",
 									TestInfo,
 									RunningTime.TotalSeconds,
 									TestInfo.TestNode.MaxDuration - RunningTime.TotalSeconds);
@@ -363,7 +364,7 @@ namespace Gauntlet
 					if ((DateTime.Now - LastStatusUpdateTime).TotalSeconds >= StatusUpdatePeriod)
 					{
 						LastStatusUpdateTime = DateTime.Now;
-						Log.Info("Status: Completed:{0}, Running:{1}, Starting: {2}, Waiting:{3}",
+						Log.Info("Status: Completed:{Completed}, Running:{Running}, Starting: {Starting}, Waiting:{Waiting}",
 							CompletedTests.Count(), RunningTests.Count(), StartingTests.Count(), PendingTests.Count());
 					}
 
@@ -390,7 +391,7 @@ namespace Gauntlet
 
 						if (Elapsed >= 5)
 						{
-							Log.Error("Giving up waiting for tests after {0:00} seconds", Elapsed);
+							Log.Error(KnownLogEvents.Gauntlet_TestEvent, "Giving up waiting for tests after {Elapsed:00} seconds", Elapsed);
 							break;
 						}
 					}
@@ -410,20 +411,20 @@ namespace Gauntlet
 
 					foreach (TestExecutionInfo TestInfo in StartingTests)
 					{
-						Log.Info("Forcing pending test {0} to run CleanupTest", TestInfo.TestNode.Name);
+						Log.Info("Forcing pending test {Name} to run CleanupTest", TestInfo.TestNode.Name);
 						TestInfo.TestNode.CleanupTest();
 						CompletedTests.Add(TestInfo);
 					}
 
 					foreach (TestExecutionInfo TestInfo in RunningTests)
 					{
-						Log.Info("Ticking test {0} to cancel", TestInfo.TestNode.Name);
+						Log.Info("Ticking test {Name} to cancel", TestInfo.TestNode.Name);
 						TestResult Res = TickTest(TestInfo);
 						CompletedTests.Add(TestInfo);
 
 						if (Res != TestResult.Failed)
 						{
-							Log.Warning("Ticking of cancelled test {0} returnd {1}", TestInfo.TestNode.Name, Res);
+							Log.Warning(KnownLogEvents.Gauntlet_TestEvent, "Ticking of cancelled test {Name} returned {Result}", TestInfo.TestNode.Name, Res);
 						}
 					}
 				}
@@ -475,14 +476,14 @@ namespace Gauntlet
 						Log.Info(Msg);
 
 						// log test timing to info
-						Log.Info(string.Format("Test Time: {0:mm\\:ss} (Waited:{1:mm\\:ss}, Setup:{2:mm\\:ss})", TestDuration, TimeWaiting, SetupTime));
+						Log.Info("Test Time: {Duration:mm\\:ss} (Waited:{Waited:mm\\:ss}, Setup:{Setup:mm\\:ss})", TestDuration, TimeWaiting, SetupTime);
 
 					});
 
 					if (Options.Parallel > 1)
 					{
-						Log.Info("MaxParallelTasks: {0}", MaxParallelTasks);
-						Log.Info("MaxStartingTasks: {0}", MaxStartingTasks);
+						Log.Info("MaxParallelTasks: {Count}", MaxParallelTasks);
+						Log.Info("MaxStartingTasks: {Count}", MaxStartingTasks);
 					}
 
 					// report all tests
@@ -531,7 +532,7 @@ namespace Gauntlet
 
 				MB.HorizontalLine();
 
-				Log.Info("{0}", MB.ToString());
+				Log.Info(MB.ToString());
 			}			
 
 			IsRunning = false;
@@ -552,7 +553,7 @@ namespace Gauntlet
 		{
 			string Name = TestInfo.TestNode.Name;
 
-			Log.Info("Starting Test {0}", TestInfo);
+			Log.Info("Starting Test {Name}", TestInfo);
 
 			try
 			{
@@ -560,13 +561,13 @@ namespace Gauntlet
 				if (TestInfo.TestNode.StartTest(Pass, NumPasses))
 				{
 					TestInfo.TimeSetupEnded = DateTime.Now;
-					Log.Info("Launched test {0} at {1}", Name, TestInfo.TimeSetupEnded.ToString("h:mm:ss"));
+					Log.Info("Launched test {Name} at {Time}", Name, TestInfo.TimeSetupEnded.ToString("h:mm:ss"));
 					return true;
 				}
 			}
 			catch (Exception Ex)
 			{
-				Log.Error("Test {0} threw an exception during launch. Skipping test. Ex: {1}\n{2}", Name, Ex.Message, Ex.StackTrace);
+				Log.Error(KnownLogEvents.Gauntlet_TestEvent, "Test {Name} threw an exception during launch. Skipping test. Ex: {Exception}\n{Callstack}", Name, Ex.Message, Ex.StackTrace);
 			}			
 
 			return false;			
@@ -579,41 +580,54 @@ namespace Gauntlet
 		/// <returns></returns>
 		void ReportTestSummary(TestExecutionInfo TestInfo)
 		{
-			string Summary = TestInfo.TestNode.GetTestSummary();
-
 			Log.SuspendSanitization();
 
+			string Summary = TestInfo.TestNode.GetTestSummary();
+
 			// Show summary
-			Summary.Split('\n').ToList().ForEach(L => Log.Info("  " + L));
+			Summary.Split('\n').ToList().ForEach(L => {
+				if (L.Contains("Error: "))
+				{
+					Log.Error(KnownLogEvents.Gauntlet_TestEvent, " " + L);
+				}
+				else if (L.Contains("Warning: "))
+				{
+					Log.Warning(KnownLogEvents.Gauntlet_TestEvent, " " + L);
+				}
+				else
+				{
+					Log.Info(" " + L);
+				}
+			});
 
 			Log.ResumeSanitization();
 
 			// list warnings/errors if the test wants that
 			if (TestInfo.TestNode.LogWarningsAndErrorsAfterSummary)
 			{
-				TestInfo.TestNode.GetErrors().ToList().ForEach(E => Log.Error("{0}", E));
+				TestInfo.TestNode.GetErrors().ToList().ForEach(E => Log.Error(KnownLogEvents.Gauntlet_TestEvent, E));
 
-				TestInfo.TestNode.GetWarnings().ToList().ForEach(E => Log.Warning("{0}", E));
+				TestInfo.TestNode.GetWarnings().ToList().ForEach(E => Log.Warning(KnownLogEvents.Gauntlet_TestEvent, E));
 			}
 
 			// display the final result
 			if (TestInfo.FinalResult != TestResult.Passed)
 			{
-				Log.Info("{0} result={1}", TestInfo, TestInfo.FinalResult);
+				Log.Info("{TestInfo} result={Result}", TestInfo, TestInfo.FinalResult);
 				if (string.IsNullOrEmpty(TestInfo.CancellationReason) == false)
 				{
-					Log.Info("\tReason: {0}", TestInfo.CancellationReason);
+					Log.Info("\tReason: {Reason}", TestInfo.CancellationReason);
 				}				
 			}
 			else
 			{
 				if (TestInfo.TestNode.GetWarnings().Any())
 				{
-					Log.Info("{0} result={1} with warnings", TestInfo, TestInfo.FinalResult);
+					Log.Info("{TestInfo} result={Result} with warnings", TestInfo, TestInfo.FinalResult);
 				}
 				else
 				{
-					Log.Info("{0} result={1}", TestInfo, TestInfo.FinalResult);
+					Log.Info("{TestInfo} result={Result}", TestInfo, TestInfo.FinalResult);
 				}
 			}
 
@@ -673,7 +687,7 @@ namespace Gauntlet
 			// only show pass info for multiple passes
 			if (NumPasses > 1)
 			{
-				Log.Info("Completed test pass {0} of {1}.", CurrentPass, NumPasses);
+				Log.Info("Completed test pass {Index} of {Total}.", CurrentPass, NumPasses);
 			}
 
 			// only show count of passed/failed etc for multiple test
@@ -722,13 +736,13 @@ namespace Gauntlet
 				{
 					TestInfo.CancellationReason = string.Format("Terminating Test {0} due to maximum duration of {1} seconds. ", TestInfo.TestNode, TestInfo.TestNode.MaxDuration);
 					TestInfo.FinalResult = TestResult.TimedOut;
-					Log.Error("{0}", TestInfo.CancellationReason);
+					Log.Error(KnownLogEvents.Gauntlet_TestEvent, TestInfo.CancellationReason);
 				}
 				else if (TestInfo.TestNode.MaxDurationReachedResult == EMaxDurationReachedResult.Success)
 				{
 					TestInfo.FinalResult = TestResult.Passed;
 					TestIsRunning = false;
-					Log.Info(string.Format("Test {0} successfully reached maximum duration of {1} seconds. ", TestInfo.TestNode, TestInfo.TestNode.MaxDuration));
+					Log.Info("Test {Name} successfully reached maximum duration of {Time} seconds. ", TestInfo.TestNode, TestInfo.TestNode.MaxDuration);
 				}
 			}
 
@@ -736,7 +750,7 @@ namespace Gauntlet
 			{
 				TestInfo.CancellationReason = string.Format("Cancelling Test {0} on request", TestInfo.TestNode);
 				TestInfo.FinalResult = TestResult.Cancelled;
-				Log.Info("{0}", TestInfo.CancellationReason);
+				Log.Info(TestInfo.CancellationReason);
 			}
 
 			if (!string.IsNullOrEmpty(TestInfo.CancellationReason))
@@ -754,7 +768,7 @@ namespace Gauntlet
 					// artifcat links
 					Log.Info("*");
 					Log.Info("****************************************************************");
-					Log.Info("Finished Test: {0} in {1:mm\\:ss}", TestInfo, DateTime.Now - TestInfo.TimeSetupEnded);
+					Log.Info("Finished Test: {Name} in {Time:mm\\:ss}", TestInfo, DateTime.Now - TestInfo.TimeSetupEnded);
 
 					// Tell the test it's done. If it still thinks its running it was cancelled
 					TestInfo.TestNode.StopTest(TestIsRunning ? StopReason.MaxDuration : StopReason.Completed);
@@ -766,7 +780,7 @@ namespace Gauntlet
 					bool bCanFinalizeTest = true;
 					if (TestInfo.FinalResult == TestResult.WantRetry)
 					{
-						Log.Info("{0} requested retry. Cleaning up old test and relaunching", TestInfo);
+						Log.Info("{Name} requested retry. Cleaning up old test and relaunching", TestInfo);
 
 						DateTime OriginalStartTime = TestInfo.TimeSetupEnded;
 
@@ -781,14 +795,14 @@ namespace Gauntlet
 						else
 						{
 							TestInfo.CancellationReason = "Failed to restart during retry.";
-							Log.Error(TestInfo.CancellationReason);
+							Log.Error(KnownLogEvents.Gauntlet_TestEvent, TestInfo.CancellationReason);
 							TestInfo.FinalResult = TestResult.Failed;
 						}
 					}
 
 					if (bCanFinalizeTest)
 					{
-						Log.Info("{0} result={1}", TestInfo, TestInfo.FinalResult);
+						Log.Info("{Name} result={Result}", TestInfo, TestInfo.FinalResult);
 
 						if (!Options.DeferReports)
 						{
@@ -802,7 +816,7 @@ namespace Gauntlet
 						}
 						catch (System.Exception ex)
 						{
-							Log.Error("Test {0} threw an exception while cleaning up. Ex: {1}", TestInfo.TestNode.Name, ex.Message);
+							Log.Error(KnownLogEvents.Gauntlet_TestEvent, "Test {Name} threw an exception while cleaning up. Ex: {Exception}", TestInfo.TestNode.Name, ex.Message);
 						}
 					}
 
@@ -813,12 +827,12 @@ namespace Gauntlet
 				{
 					if (TestIsRunning)
 					{
-						Log.Warning("Cancelled Test {0} threw an exception while stopping. Ex: {1}\n{2}", 
+						Log.Warning(KnownLogEvents.Gauntlet_TestEvent, "Cancelled Test {Name} threw an exception while stopping. Ex: {Eception}\n{Callstack}", 
 							TestInfo.TestNode.Name, ex.Message, ex.StackTrace);
 					}
 					else
 					{
-						Log.Error("Test {0} threw an exception while stopping. Ex: {1}\n{2}",
+						Log.Error(KnownLogEvents.Gauntlet_TestEvent, "Test {Name} threw an exception while stopping. Ex: {Exception}\n{Callstack}",
 							TestInfo.TestNode.Name, ex.Message, ex.StackTrace);
 					}
 
@@ -835,7 +849,7 @@ namespace Gauntlet
 		/// </summary>
 		bool WaitForTests()
 		{
-			Log.Info("Waiting for {0} tests to complete", RunningTests.Count);
+			Log.Info("Waiting for {Count} tests to complete", RunningTests.Count);
 
 			DateTime LastUpdateMsg = DateTime.Now;
 
@@ -859,7 +873,7 @@ namespace Gauntlet
 
 						if ((DateTime.Now - LastUpdateMsg).TotalSeconds > 60.0f)
 						{
-							Log.Verbose("Test {0} is still running. {1:00} seconds elapsed, will timeout in {2:00} seconds",
+							Log.Verbose("Test {Name} is still running. {Elapsed:00} seconds elapsed, will timeout in {Max:00} seconds",
 								Process.TestNode.Name,
 								RunningTime.TotalSeconds,
 								Process.TestNode.MaxDuration - RunningTime.TotalSeconds);
@@ -873,7 +887,7 @@ namespace Gauntlet
 						{
 							AllTestsPassed = false;
 						}
-						Log.Info("Test {0} Result: {1}", Process.TestNode.Name, Result);
+						Log.Info("Test {Name} Result: {Result}", Process.TestNode.Name, Result);
 					}
 				}
 
diff --git a/Engine/Source/Programs/AutomationTool/Gauntlet/Framework/Gauntlet.Utils.cs b/Engine/Source/Programs/AutomationTool/Gauntlet/Framework/Gauntlet.Utils.cs
--- a/Engine/Source/Programs/AutomationTool/Gauntlet/Framework/Gauntlet.Utils.cs
+++ b/Engine/Source/Programs/AutomationTool/Gauntlet/Framework/Gauntlet.Utils.cs
@@ -14,6 +14,8 @@ using UnrealBuildTool;
 using EpicGames.Core;
 using System.Security.Cryptography;
 using System.Text;
+using Microsoft.Extensions.Logging;
+using Logging = Microsoft.Extensions.Logging;
 
 namespace Gauntlet
 {
@@ -204,6 +206,7 @@ namespace Gauntlet
 
 		static int SanitizationSuspendCount = 0;
 
+		static ILogger Logger = EpicGames.Core.Log.Logger;
 
 		public static void AddCallback(Action<string> CB)
 		{
@@ -264,7 +267,7 @@ namespace Gauntlet
 			{
 				if (CommandUtils.IsBuildMachine)
 				{
-					OutputMessage("<-- Suspend Log Parsing -->");
+					OutputMessage("<-- Suspend Log Parsing -->", KnownLogEvents.Gauntlet);
 				}
 			}
 		}
@@ -275,7 +278,7 @@ namespace Gauntlet
 			{
 				if (CommandUtils.IsBuildMachine)
 				{
-					OutputMessage("<-- Resume Log Parsing -->");
+					OutputMessage("<-- Resume Log Parsing -->", KnownLogEvents.Gauntlet);
 				}
 			}
 		}
@@ -312,15 +315,16 @@ namespace Gauntlet
 		}
 
 		/// <summary>
-		/// Outputs the message to the console with an optional prefix and sanitization. Sanitizing
-		/// allows errors and exceptions to be passed through to logs without triggering CIS warnings
-		/// about out log
+		/// Outputs the message to the console with an optional log level and sanitization.
+		/// Sanitizing allows errors and exceptions to be passed through to logs without triggering CIS warnings about out log
 		/// </summary>
 		/// <param name="Message"></param>
-		/// <param name="Prefix"></param>
+		/// <param name="EventId"></param>
+		/// <param name="Level"></param>
 		/// <param name="Sanitize"></param>
+		/// <param name="Args"></param>
 		/// <returns></returns>
-		static private void OutputMessage(string Message, string Prefix="", bool Sanitize=true)
+		static private void OutputMessage(string Message, Logging.EventId EventId, Logging.LogLevel Level = Logging.LogLevel.Information, bool Sanitize=true, params object[] Args)
 		{
 			// EC detects error statements in the log as a failure. Need to investigate best way of 
 			// reporting errors, but not errors we've handled from tools.
@@ -347,45 +351,58 @@ namespace Gauntlet
 				Message = "[" + ElapsedTime.ToString() + "] " + Message;
 			}
 
-			if (string.IsNullOrEmpty(Prefix) == false)
-			{
-				Message = Prefix + ": " + Message;
-			}
-
-			// TODO - Remove all Gauntlet logging and switch to UBT log?
-			CommandUtils.LogInformation(Message);
+			Logger.Log(Level, EventId, Message, Args);
 
-			try
+			if (LogFile != null || Callbacks != null)
 			{
-				if (LogFile != null)
+				try
 				{
-					LogFile.WriteLine(Message);
-				}
+					if (Args.Length > 0)
+					{
+						// Adjust Message and Arguments for string.Format
+						int Index = 0;
+						string SequenceMessage = Regex.Replace(Message, @"\{[^{}:]+(:[^{}]+)?\}", M => "{" + Index++.ToString() + M.Groups[1].Value + "}", RegexOptions.IgnoreCase);
+						if (Index > Args.Length)
+						{
+							Args = Args.Concat(Enumerable.Repeat("null", Index - Args.Length)).ToArray();
+						}
+						else if (Index < Args.Length)
+						{
+							Args = Args.Take(Index).ToArray();
+						}
+						Message = string.Format(SequenceMessage, Args);
+					}
 
-				if (Callbacks != null)
+					if (LogFile != null)
+					{
+						LogFile.WriteLine(Message);
+					}
+
+					if (Callbacks != null)
+					{
+						Callbacks.ForEach(A => A(Message));
+					}
+				}
+				catch (Exception Ex)
 				{
-					Callbacks.ForEach(A => A(Message));
+					Logger.LogWarning(KnownLogEvents.Gauntlet, "Exception logging '{Message}'. {Exception}", Message, Ex.ToString());
 				}
 			}
-			catch (Exception Ex)
-			{
-				CommandUtils.LogWarning("Exception logging '{0}'. {1}", Message, Ex.ToString());
-			}			
 		}	
 
 		static public void Verbose(string Format, params object[] Args)
 		{
 			if (IsVerbose)
 			{
-				Verbose(string.Format(Format, Args));
+				Verbose(KnownLogEvents.Gauntlet, Format, Args);
 			}
 		}
 
-		static public void Verbose(string Message)
+		static public void Verbose(Logging.EventId EventId, string Format, params object[] Args)
 		{
 			if (IsVerbose)
 			{
-				OutputMessage(Message);
+				OutputMessage(Format, EventId, Args: Args);
 			}
 		}
 
@@ -393,44 +410,44 @@ namespace Gauntlet
 		{
 			if (IsVeryVerbose)
 			{
-				VeryVerbose(string.Format(Format, Args));
+				VeryVerbose(KnownLogEvents.Gauntlet, Format, Args);
 			}
 		}
 
-		static public void VeryVerbose(string Message)
+		static public void VeryVerbose(Logging.EventId EventId, string Format, params object[] Args)
 		{
 			if (IsVeryVerbose)
 			{
-				OutputMessage(Message);
+				OutputMessage(Format, EventId, Args: Args);
 			}
 		}
 
 		static public void Info(string Format, params object[] Args)
 		{
-			Info(string.Format(Format, Args));
+			Info(KnownLogEvents.Gauntlet, Format, Args);
 		}
 
-		static public void Info(string Message)
+		static public void Info(Logging.EventId EventId, string Format, params object[] Args)
 		{
-			OutputMessage(Message);
+			OutputMessage(Format, EventId, Args: Args);
 		}
 
 		static public void Warning(string Format, params object[] Args)
 		{
-			Warning(string.Format(Format, Args));
+			Warning(KnownLogEvents.Gauntlet, Format, Args);
 		}
 
-		static public void Warning(string Message)
+		static public void Warning(Logging.EventId EventId, string Format, params object[] Args)
 		{
-			OutputMessage(Message, "Warning");
+			OutputMessage(Format, EventId, Logging.LogLevel.Warning, Args: Args);
 		}
 		static public void Error(string Format, params object[] Args)
 		{
-			Error(string.Format(Format, Args));
+			Error(KnownLogEvents.Gauntlet, Format, Args);
 		}
-		static public void Error(string Message)
-		{		
-			OutputMessage(Message, "Error", false);
+		static public void Error(Logging.EventId EventId, string Format, params object[] Args)
+		{
+			OutputMessage(Format, EventId, Logging.LogLevel.Error, false, Args);
 		}
 	}
 
diff --git a/Engine/Source/Programs/AutomationTool/Gauntlet/Platform/Android/Gauntlet.AndroidBuildSource.cs b/Engine/Source/Programs/AutomationTool/Gauntlet/Platform/Android/Gauntlet.AndroidBuildSource.cs
--- a/Engine/Source/Programs/AutomationTool/Gauntlet/Platform/Android/Gauntlet.AndroidBuildSource.cs
+++ b/Engine/Source/Programs/AutomationTool/Gauntlet/Platform/Android/Gauntlet.AndroidBuildSource.cs
@@ -8,6 +8,7 @@ using UnrealBuildTool;
 using System.Threading;
 using System.Text.RegularExpressions;
 using System.Linq;
+using EpicGames.Core;
 
 namespace Gauntlet
 {
@@ -104,7 +105,7 @@ namespace Gauntlet
 				var SourceApkMatch = Regex.Match(BatContents, @" install\s+(.+\.apk)");
 				if ( SourceApkMatch.Groups.Count <= 0)
 				{
-					Log.Warning("Could not parse install command from {0}", Fi.FullName);
+					Log.Warning(KnownLogEvents.Gauntlet_BuildDropEvent, "Could not parse install command from {File}", Fi.FullName);
 					continue;
 				}
 				string SourceApkPath = Path.Combine(AbsPath,SourceApkMatch.Groups[1].ToString());
@@ -132,19 +133,19 @@ namespace Gauntlet
 	
 				if (string.IsNullOrEmpty(SourceApkPath))
 				{
-					Log.Warning("No APK found for build at {0}", Fi.FullName);
+					Log.Warning(KnownLogEvents.Gauntlet_BuildDropEvent, "No APK found for build at {File}", Fi.FullName);
 					continue;
 				}
 
 				if (!File.Exists(SourceApkPath))
 				{
-					Log.Warning("Resolved APK name but it doesn't exist {0}", SourceApkPath);
+					Log.Warning(KnownLogEvents.Gauntlet_BuildDropEvent, "Resolved APK name but it doesn't exist {File}", SourceApkPath);
 					continue;
 				}
 
 				if (string.IsNullOrEmpty(AndroidPackageName))
 				{
-					Log.Warning("No product name found for build at {0}", Fi.FullName);
+					Log.Warning(KnownLogEvents.Gauntlet_BuildDropEvent, "No product name found for build at {File}", Fi.FullName);
 					continue;
 				}
 
diff --git a/Engine/Source/Programs/AutomationTool/Gauntlet/Platform/Android/Gauntlet.TargetDeviceAndroid.cs b/Engine/Source/Programs/AutomationTool/Gauntlet/Platform/Android/Gauntlet.TargetDeviceAndroid.cs
--- a/Engine/Source/Programs/AutomationTool/Gauntlet/Platform/Android/Gauntlet.TargetDeviceAndroid.cs
+++ b/Engine/Source/Programs/AutomationTool/Gauntlet/Platform/Android/Gauntlet.TargetDeviceAndroid.cs
@@ -161,7 +161,7 @@ namespace Gauntlet
 
 			if (Install.AndroidDevice != null && Install.AndroidDevice.Disposed)
 			{
-				Log.Warning("Attempting to cache log using disposed Android device");
+				Log.Warning(KnownLogEvents.Gauntlet_DeviceEvent, "Attempting to cache log using disposed Android device");
 				return;
 			}
 
@@ -229,7 +229,7 @@ namespace Gauntlet
 				}
 				catch
 				{
-					Log.Warning("Failed to remove old cache folder {0}", Install.AndroidDevice.LocalCachePath);
+					Log.Warning(KnownLogEvents.Gauntlet_DeviceEvent, "Failed to remove old cache folder {Folder}", Install.AndroidDevice.LocalCachePath);
 				}
 			}
 
@@ -240,7 +240,7 @@ namespace Gauntlet
 			}
 			catch (Exception Ex)
 			{
-				Log.Warning("Exception marking directory for cleanup {0}", Ex.Message);
+				Log.Warning(KnownLogEvents.Gauntlet_DeviceEvent, "Exception marking directory for cleanup {Exception}", Ex.Message);
 			}			
 
 			string LocalSaved = Path.Combine(Install.AndroidDevice.LocalCachePath, "Saved");
@@ -257,7 +257,7 @@ namespace Gauntlet
 
 			if (PullCmd.ExitCode != 0)
 			{
-				Log.Warning("Failed to retrieve artifacts. {0}", PullCmd.Output);
+				Log.Warning(KnownLogEvents.Gauntlet_DeviceEvent, "Failed to retrieve artifacts. {Output}", PullCmd.Output);
 			}
 			else
 			{
@@ -416,7 +416,7 @@ namespace Gauntlet
 
 				if (AllDevices[DeviceName] == false)
 				{
-					Log.Warning("Device {0} is connected but we are not authorized", DeviceName);
+					Log.Warning(KnownLogEvents.Gauntlet_DeviceEvent, "Device {Name} is connected but we are not authorized", DeviceName);
 					return false;
 				}
 
@@ -568,7 +568,7 @@ namespace Gauntlet
 				}
 				catch (Exception Ex)
 				{
-					Log.Warning("TargetDeviceAndroid.Dispose() threw: {0}", Ex.Message);
+					Log.Warning(KnownLogEvents.Gauntlet_DeviceEvent, "TargetDeviceAndroid.Dispose() threw: {Exception}", Ex.Message);
 				}
 				finally
 				{
@@ -840,7 +840,7 @@ namespace Gauntlet
 
 				if (AdbResult.ExitCode != 0)
 				{
-					Log.Warning("Failed to write dependency file {0}", DepFile);
+					Log.Warning(KnownLogEvents.Gauntlet_DeviceEvent, "Failed to write dependency file {File}", DepFile);
 				}
 			}
 
@@ -963,7 +963,7 @@ namespace Gauntlet
 
                     else
                     {
-                        Log.Warning("File to copy {0} not found", FileToCopy);
+                        Log.Warning(KnownLogEvents.Gauntlet_DeviceEvent, "File to copy {File} not found", FileToCopy);
                     }
                 }
             }
diff --git a/Engine/Source/Programs/AutomationTool/Gauntlet/Platform/IOS/Gauntlet.IOSBuildSource.cs b/Engine/Source/Programs/AutomationTool/Gauntlet/Platform/IOS/Gauntlet.IOSBuildSource.cs
--- a/Engine/Source/Programs/AutomationTool/Gauntlet/Platform/IOS/Gauntlet.IOSBuildSource.cs
+++ b/Engine/Source/Programs/AutomationTool/Gauntlet/Platform/IOS/Gauntlet.IOSBuildSource.cs
@@ -8,6 +8,7 @@ using UnrealBuildTool;
 using System.Threading;
 using System.Text.RegularExpressions;
 using System.Linq;
+using EpicGames.Core;
 
 namespace Gauntlet
 {
@@ -79,7 +80,7 @@ namespace Gauntlet
 					{
 						if (!File.Exists(ShouldExist) && !Directory.Exists(ShouldExist))
 						{
-							Log.Error(String.Format("unzip encountered an error or warning procesing IPA, possibly due to Zip64 issue, {0} missing", ShouldExist));
+							Log.Error(KnownLogEvents.Gauntlet_BuildDropEvent, "unzip encountered an error or warning procesing IPA, possibly due to Zip64 issue, {File} missing", ShouldExist);
 							return false;
 						}
 					}
diff --git a/Engine/Source/Programs/AutomationTool/Gauntlet/Platform/IOS/Gauntlet.TargetDeviceIOS.cs b/Engine/Source/Programs/AutomationTool/Gauntlet/Platform/IOS/Gauntlet.TargetDeviceIOS.cs
--- a/Engine/Source/Programs/AutomationTool/Gauntlet/Platform/IOS/Gauntlet.TargetDeviceIOS.cs
+++ b/Engine/Source/Programs/AutomationTool/Gauntlet/Platform/IOS/Gauntlet.TargetDeviceIOS.cs
@@ -10,6 +10,7 @@ using AutomationTool;
 using UnrealBuildTool;
 using System.Text;
 using System.Text.RegularExpressions;
+using EpicGames.Core;
 
 /*
 
@@ -69,7 +70,7 @@ namespace Gauntlet
 
 			if (DownloadCmd.ExitCode != 0)
 			{
-				Log.Warning("Failed to retrieve artifacts. {0}", DownloadCmd.Output);
+				Log.Warning(KnownLogEvents.Gauntlet_DeviceEvent, "Failed to retrieve artifacts. {Output}", DownloadCmd.Output);
 			}
 			
 		}
@@ -238,7 +239,7 @@ namespace Gauntlet
 				}
 				catch (Exception Ex)
 				{
-					Log.Warning("TargetDeviceIOS.Dispose() threw: {0}", Ex.Message);
+					Log.Warning(KnownLogEvents.Gauntlet_DeviceEvent, "TargetDeviceIOS.Dispose() threw: {Exception}", Ex.Message);
 				}
 				finally
 				{
@@ -295,7 +296,7 @@ namespace Gauntlet
 			// Give ios-deploy a chance to throw out any errors...
 			if (Result.HasExited)
 			{
-				Log.Warning("ios-deploy exited early: " + Result.Output);
+				Log.Warning(KnownLogEvents.Gauntlet_DeviceEvent, "ios-deploy exited early: " + Result.Output);
 				throw new DeviceException("Failed to launch on {0}. {1}", Name, Result.Output);
 			}
 
@@ -325,7 +326,7 @@ namespace Gauntlet
 
 				if (Result.ExitCode != 0)
 				{
-					Log.Warning("Failed to clean artifacts from device");
+					Log.Warning(KnownLogEvents.Gauntlet_DeviceEvent, "Failed to clean artifacts from device");
 					return false;
 				}
 
@@ -449,7 +450,7 @@ namespace Gauntlet
 			}
 			else
 			{
-				Log.Warning("No symbols found for local build at {0}, removing cached app symbols", LocalSymbolsDir);
+				Log.Warning(KnownLogEvents.Gauntlet_DeviceEvent, "No symbols found for local build at {Directory}, removing cached app symbols", LocalSymbolsDir);
 			}
 
 			// resign application
@@ -566,7 +567,7 @@ namespace Gauntlet
 						}
 						else
 						{
-							Log.Warning("File to copy {0} not found", FileToCopy);
+							Log.Warning(KnownLogEvents.Gauntlet_DeviceEvent, "File to copy {File} not found", FileToCopy);
 						}
 					}
 				}
@@ -660,7 +661,7 @@ namespace Gauntlet
 
 			if (!rebooted) 
 			{
-				Log.Warning("Failed to reboot iOS device {0}, device didn't come back after restart", DeviceName);
+				Log.Warning(KnownLogEvents.Gauntlet_DeviceEvent, "Failed to reboot iOS device {Name}, device didn't come back after restart", DeviceName);
 			}
 
 			return true; 
@@ -1017,7 +1018,7 @@ namespace Gauntlet
 			} 
 			catch (Exception Ex)
 			{
-				Log.Warning("Exception parsing LLDB callstack {0}", Ex.Message);
+				Log.Warning(KnownLogEvents.Gauntlet_DeviceEvent, "Exception parsing LLDB callstack {Exception}", Ex.Message);
 			}
 
 			return null;
@@ -1155,14 +1156,14 @@ namespace Gauntlet
 
 			if (Threads.Count(T => T.Current == true) > 1)
 			{
-				Log.Warning("LLDB debug parsed more than one current thread");
+				Log.Warning(KnownLogEvents.Gauntlet_DeviceEvent, "LLDB debug parsed more than one current thread");
 			}
 
 			Thread = Threads.FirstOrDefault(T => T.Current == true);
 
 			if (Threads.Count > 0 && Thread == null)
 			{
-				Log.Warning("Unable to parse full crash callstack");
+				Log.Warning(KnownLogEvents.Gauntlet_DeviceEvent, "Unable to parse full crash callstack");
 			}
 
 
diff --git a/Engine/Source/Programs/AutomationTool/Gauntlet/Platform/Linux/Gauntlet.TargetDeviceLinux.cs b/Engine/Source/Programs/AutomationTool/Gauntlet/Platform/Linux/Gauntlet.TargetDeviceLinux.cs
--- a/Engine/Source/Programs/AutomationTool/Gauntlet/Platform/Linux/Gauntlet.TargetDeviceLinux.cs
+++ b/Engine/Source/Programs/AutomationTool/Gauntlet/Platform/Linux/Gauntlet.TargetDeviceLinux.cs
@@ -8,6 +8,7 @@ using System.Threading;
 using AutomationTool;
 using UnrealBuildTool;
 using System.Text.RegularExpressions;
+using EpicGames.Core;
 
 namespace Gauntlet
 {
@@ -84,7 +85,7 @@ namespace Gauntlet
 				}
 				catch (Exception Ex)
 				{
-					Log.Warning("Failed to delete {0}. {1}", ArtifactPath, Ex.Message);
+					Log.Warning(KnownLogEvents.Gauntlet_DeviceEvent, "Failed to delete {File}. {Exception}", ArtifactPath, Ex.Message);
 				}
 			}
 		}
@@ -279,7 +280,7 @@ namespace Gauntlet
 					}
 					else
 					{
-						Log.Warning("File to copy {0} not found", FileToCopy);
+						Log.Warning(KnownLogEvents.Gauntlet_DeviceEvent, "File to copy {File} not found", FileToCopy);
 					}
 				}
 			}
diff --git a/Engine/Source/Programs/AutomationTool/Gauntlet/Platform/Mac/Gauntlet.TargetDeviceMac.cs b/Engine/Source/Programs/AutomationTool/Gauntlet/Platform/Mac/Gauntlet.TargetDeviceMac.cs
--- a/Engine/Source/Programs/AutomationTool/Gauntlet/Platform/Mac/Gauntlet.TargetDeviceMac.cs
+++ b/Engine/Source/Programs/AutomationTool/Gauntlet/Platform/Mac/Gauntlet.TargetDeviceMac.cs
@@ -67,7 +67,7 @@ namespace Gauntlet
 				}
 				catch (Exception Ex)
 				{
-					Log.Warning("Failed to delete {0}. {1}", ArtifactPath, Ex.Message);
+					Log.Warning(KnownLogEvents.Gauntlet_DeviceEvent, "Failed to delete {File}. {Exception}", ArtifactPath, Ex.Message);
 				}
 			}
 		}
diff --git a/Engine/Source/Programs/AutomationTool/Gauntlet/Platform/Windows/Gauntlet.TargetDeviceWindows.cs b/Engine/Source/Programs/AutomationTool/Gauntlet/Platform/Windows/Gauntlet.TargetDeviceWindows.cs
--- a/Engine/Source/Programs/AutomationTool/Gauntlet/Platform/Windows/Gauntlet.TargetDeviceWindows.cs
+++ b/Engine/Source/Programs/AutomationTool/Gauntlet/Platform/Windows/Gauntlet.TargetDeviceWindows.cs
@@ -8,6 +8,7 @@ using System.Threading;
 using AutomationTool;
 using UnrealBuildTool;
 using System.Text.RegularExpressions;
+using EpicGames.Core;
 
 namespace Gauntlet
 {
@@ -90,7 +91,7 @@ namespace Gauntlet
 					{
 						string Output = LogReader.ReadToEnd();
 
-						if (Output != null)
+						if (!string.IsNullOrEmpty(Output))
 						{
 							ProcessLogOutput += Output;
 						}
@@ -186,7 +187,7 @@ namespace Gauntlet
 				}
 				catch (Exception Ex)
 				{
-					Log.Warning("Failed to delete {0}. {1}", ArtifactPath, Ex.Message);
+					Log.Warning(KnownLogEvents.Gauntlet_DeviceEvent, "Failed to delete {File}. {Exception}", ArtifactPath, Ex.Message);
 				}
 			}
 		}
@@ -336,7 +337,7 @@ namespace Gauntlet
 				catch (Exception Ex)
 				{
 					//throw new AutomationException("Unable to delete existing log file {0} {1}", ProcessLogFile, Ex.Message);
-					Log.Warning("Unable to delete existing log file {0}. {1}", ProcessLogFile, Ex.Message);
+					Log.Warning(KnownLogEvents.Gauntlet_DeviceEvent, "Unable to delete existing log file {File}. {Exception}", ProcessLogFile, Ex.Message);
 				}
 
 				Log.Verbose("\t{0}", CmdLine);
@@ -380,7 +381,7 @@ namespace Gauntlet
 					}
 					else
 					{
-						Log.Warning("File to copy {0} not found", FileToCopy);
+						Log.Warning(KnownLogEvents.Gauntlet_DeviceEvent, "File to copy {File} not found", FileToCopy);
 					}
 				}
 			}
diff --git a/Engine/Source/Programs/AutomationTool/Gauntlet/Unreal/Automation/UE.Automation.cs b/Engine/Source/Programs/AutomationTool/Gauntlet/Unreal/Automation/UE.Automation.cs
--- a/Engine/Source/Programs/AutomationTool/Gauntlet/Unreal/Automation/UE.Automation.cs
+++ b/Engine/Source/Programs/AutomationTool/Gauntlet/Unreal/Automation/UE.Automation.cs
@@ -9,6 +9,8 @@ using System.Linq;
 using System.Text;
 using System.Threading.Tasks;
 using UnrealBuildTool;
+using EpicGames.Core;
+using Log = Gauntlet.Log;
 
 namespace UE
 {
@@ -532,7 +534,7 @@ namespace UE
 			if (ChannelEntries.Count > LastAutomationEntryCount)
 			{
 				// log new entries so people have something to look at
-				ChannelEntries.Skip(LastAutomationEntryCount).ToList().ForEach(S => Log.Info("{0}", S));
+				ChannelEntries.Skip(LastAutomationEntryCount).ToList().ForEach(S => Log.Info(S));
 				LastAutomationEntryTime = DateTime.Now;
 				LastAutomationEntryCount = ChannelEntries.Count;
 			}
@@ -549,7 +551,7 @@ namespace UE
 				// Check for timeout
 				if (ElapsedTime > IdleTimeout)
 				{
-					Log.Error("No activity observed in last {0:0.00} minutes. Aborting test", IdleTimeout / 60);
+					Log.Warning(KnownLogEvents.Gauntlet_TestEvent, "No activity observed in last {Time:0.00} minutes. Aborting test", IdleTimeout / 60);
 					MarkTestComplete();
 					SetUnrealTestResult(TestResult.TimedOut);
 				}
@@ -577,7 +579,7 @@ namespace UE
 							}
 							catch (Exception Ex)
 							{
-								Log.Warning("Failed to load Json report. {0}", Ex);
+								Log.Warning(KnownLogEvents.Gauntlet_TestEvent, "Failed to load Json report. {Exception}", Ex);
 							}
 						}
 					}
@@ -630,7 +632,7 @@ namespace UE
 			if (JsonTestPassResults.InProcess > 0)
 			{
 				// The test pass did not run completely
-				Log.Verbose("Found in-process tests: {0}", JsonTestPassResults.InProcess);
+				Log.Verbose("Found in-process tests: {Count}", JsonTestPassResults.InProcess);
 				// Get any critical error and push it to json report and resave it.
 				if (RoleResults != null)
 				{
@@ -682,7 +684,7 @@ namespace UE
 			if (JsonTestPassResults.NotRun > 0)
 			{
 				// The test pass did not run at all
-				Log.Verbose("Found not-run tests: {0}", JsonTestPassResults.NotRun);
+				Log.Verbose("Found not-run tests: {Count}", JsonTestPassResults.NotRun);
 				if (GetConfiguration().ResumeOnCriticalFailure && !HasTimeout)
 				{
 					// Reschedule test to resume from last 'in-process' test.
@@ -709,7 +711,7 @@ namespace UE
 					}
 					else
 					{
-						Log.Error("Reach maximum of retries({0}) to resume on critical failure!", Retries);
+						Log.Warning(KnownLogEvents.Gauntlet_TestEvent, "Reach maximum of retries({Count}) to resume on critical failure!", Retries);
 						// Adding a note to the report about why the not-run are not going to be run
 						string Message = string.Format("Session reached maximum of retries({0}) to resume on critical failure!", Retries);
 						if (!string.IsNullOrEmpty(LastTestWithCriticalFailure))
@@ -825,7 +827,7 @@ namespace UE
 					}
 					else
 					{
-						Log.Warning("Publishing Telemetry is requested but '{0}' does not support telemetry input.", Report.GetType().FullName);
+						Log.Warning("Publishing Telemetry is requested but '{Config}' does not support telemetry input.", Report.GetType().FullName);
 					}
 				}
 			}
@@ -838,11 +840,11 @@ namespace UE
 		/// a link to the reports
 		/// </summary>
 		/// <returns></returns>
-		protected override string GetTestSummaryHeader()
+		protected override void LogTestSummaryHeader()
 		{
-			const int kMaxErrorsOrWarningsToDisplay = 5;
+			const int kMaxErrorsOrWarningsToDisplay = 10;
 
-			MarkdownBuilder MB = new MarkdownBuilder(base.GetTestSummaryHeader());
+			base.LogTestSummaryHeader();
 
 			// Everything we need is in the editor artifacts
 			var EditorRole = RoleResults.Where(R => R.Artifacts.SessionRole.RoleType == UnrealTargetRole.Editor).FirstOrDefault();
@@ -857,13 +859,11 @@ namespace UE
 				IEnumerable<UnrealAutomatedTestResult> FailedTests = AllTests.Where(T => T.IsComplete && T.HasFailed);
 				IEnumerable<UnrealAutomatedTestResult> TestsWithWarnings = AllTests.Where(T => T.HasSucceeded && T.HasWarnings);
 
-				Func<string, string> ErrorLineForHorde = (L) => string.Format("Err: {0}", L);
-				Func<string, string> WarningLineForHorde = (L) => string.Format("Warn: {0}", L);
 				Func<IEnumerable<string>, IEnumerable<string>> CapErrorOrWarningList = (E) =>
 				{
 					if (E.Count() > kMaxErrorsOrWarningsToDisplay)
 					{
-						E = E.Skip(E.Count() - kMaxErrorsOrWarningsToDisplay);
+						E = E.Take(kMaxErrorsOrWarningsToDisplay);
 					}
 					return E;
 				};
@@ -872,27 +872,37 @@ namespace UE
 				{
 					if (AllTests.Count() == 0)
 					{
-						MB.H3("Error: No tests were executed.");
+						Log.Error(KnownLogEvents.Gauntlet_UnrealEngineTestEvent, " * No tests were executed.");
+						Log.Info("");
 					}
 					else if (IncompleteTests.Count() > 0)
 					{
 						var InProcess = IncompleteTests.Where(T => T.State == TestStateType.InProcess);
 						if (InProcess.Any())
 						{
-							MB.H3("The following test(s) were incomplete:");
+							Log.Info(" ### The following Engine test(s) were incomplete:");
 							foreach (UnrealAutomatedTestResult Result in InProcess)
 							{
-								MB.H4(string.Format("Error: Test '{0}' did not complete", Result.TestDisplayName));
-								MB.Paragraph("FullName: " + Result.FullTestPath);
-								MB.UnorderedList(CapErrorOrWarningList(Result.ErrorEvents.Select(E => ErrorLineForHorde(E.Message)).Distinct()));
-								MB.UnorderedList(CapErrorOrWarningList(Result.WarningEvents.Select(E => WarningLineForHorde(E.Message)).Distinct()));
+								Log.Error(KnownLogEvents.Gauntlet_UnrealEngineTestEvent, " * Test '{Name}' did not complete.", Result.FullTestPath);
+								var Errors = CapErrorOrWarningList(Result.ErrorEvents.Select(E => E.Message).Distinct());
+								foreach (var Error in Errors)
+								{
+									Log.Error(KnownLogEvents.Gauntlet_UnrealEngineTestEvent, "    " + Error);
+								}
+								var Warnings = CapErrorOrWarningList(Result.WarningEvents.Select(E => E.Message).Distinct());
+								foreach (var Warning in Warnings)
+								{
+									Log.Warning(KnownLogEvents.Gauntlet_UnrealEngineTestEvent, "    " + Warning);
+								}
+								Log.Info("");
 							}
 						}
 						var NotRun = IncompleteTests.Where(T => T.State == TestStateType.NotRun);
 						if (NotRun.Any())
 						{
-							MB.H3("The following test(s) were not run:");
-							MB.UnorderedList(NotRun.Select(T => T.FullTestPath));
+							Log.Info(" ### The following Engine test(s) were not run:");
+							Log.Info(string.Join("\n", NotRun.Select(T => " * "+T.FullTestPath)));
+							Log.Info("");
 						}
 					}
 				}
@@ -900,18 +910,26 @@ namespace UE
 				{
 					if (AllTests.Count() == 0)
 					{
-						MB.H3("Error: No tests were executed.");
-
-						IEnumerable<string> WarningsAndErrors = Events.Where(E => E.IsError || E.IsWarning).Select(E => E.IsError? ErrorLineForHorde(E.Message) : WarningLineForHorde(E.Message)).Distinct();
+						Log.Error(KnownLogEvents.Gauntlet_TestEvent, " * No tests were executed.");
 
-						if (WarningsAndErrors.Any())
+						IEnumerable<string> Errors = Events.Where(E => E.IsError).Select(E => E.Message).Distinct();
+						IEnumerable<string> Warnings = Events.Where(E => E.IsWarning).Select(E => E.Message).Distinct();
+						if (Errors.Any() || Warnings.Any())
 						{
-							WarningsAndErrors = CapErrorOrWarningList(WarningsAndErrors);
-							MB.UnorderedList(WarningsAndErrors);
+							Log.Info("   See log above for details.");
+							foreach (var Error in CapErrorOrWarningList(Errors))
+							{
+								Log.Error(KnownLogEvents.Gauntlet_UnrealEngineTestEvent, "    " + Error);
+							}
+							foreach (var Warning in CapErrorOrWarningList(Warnings))
+							{
+								Log.Warning(KnownLogEvents.Gauntlet_UnrealEngineTestEvent, "    " + Warning);
+							}
+							Log.Info("");
 						}
 						else
 						{
-							MB.Paragraph("Unknown failure.");
+							Log.Info("   Unknown failure.\n");
 						}
 					}
 					else
@@ -919,53 +937,56 @@ namespace UE
 						// Now list the tests that failed
 						if (FailedTests.Count() > 0)
 						{
-							MB.H3("The following test(s) failed:");
+							Log.Info(" ### The following Engine test(s) failed:");
 
 							foreach (UnrealAutomatedTestResult Result in FailedTests)
 							{
+								Log.Error(KnownLogEvents.Gauntlet_UnrealEngineTestEvent, " * Test '{Name}' failed.", Result.FullTestPath);
 								// only show the last N items
-								IEnumerable<string> Events = Result.ErrorEvents.Select(E => ErrorLineForHorde(E.Message)).Distinct();
-
+								IEnumerable<string> Events = Result.ErrorEvents.Select(E => E.Message).Distinct();
 								Events = CapErrorOrWarningList(Events);
-
-								MB.H4(string.Format("Error: Test '{0}' failed", Result.TestDisplayName));
-								MB.Paragraph("FullName: " + Result.FullTestPath);
-								MB.UnorderedList(Events);
+								foreach (var Event in Events)
+								{
+									Log.Error(KnownLogEvents.Gauntlet_UnrealEngineTestEvent, "    " + Event);
+								}
+								Log.Info("");
 							}
 						}
 
 						if (TestsWithWarnings.Count() > 0)
 						{
-							MB.H3("The following test(s) completed with warnings:");
+							Log.Info(" ### The following Engine test(s) completed with warnings:");
 
 							foreach (UnrealAutomatedTestResult Result in TestsWithWarnings)
 							{
+								Log.Warning(KnownLogEvents.Gauntlet_UnrealEngineTestEvent, " * Test '{Name}' completed with warnings.", Result.FullTestPath);
 								// only show the last N items
-								IEnumerable<string> WarningEvents = Result.WarningEvents.Select(E => WarningLineForHorde(E.Message)).Distinct();
-
+								IEnumerable<string> WarningEvents = Result.WarningEvents.Select(E => E.Message).Distinct();
 								WarningEvents = CapErrorOrWarningList(WarningEvents);
-
-								MB.H4(string.Format("Warning: Test '{0}' completed with warnings", Result.TestDisplayName));
-								MB.Paragraph("FullName: " + Result.FullTestPath);
-								MB.UnorderedList(WarningEvents);
+								foreach (var Event in WarningEvents)
+								{
+									Log.Warning(KnownLogEvents.Gauntlet_UnrealEngineTestEvent, "    " + Event);
+								}
+								Log.Info("");
 							}
 						}
 
 						if (IncompleteTests.Count() > 0)
 						{
-							MB.H3("The following test(s) timed out or did not run:");
+							Log.Info(" ### The following Engine test(s) timed out or did not run:");
 
 							foreach (UnrealAutomatedTestResult Result in IncompleteTests)
 							{
+								Log.Error(KnownLogEvents.Gauntlet_UnrealEngineTestEvent, " * Test '{Name}' did not run or complete.", Result.FullTestPath);
 								// only show the last N items
-								IEnumerable<string> ErrorAndWarningEvents = Result.WarningAndErrorEvents.Select(E => E.IsError? ErrorLineForHorde(E.Message) : WarningLineForHorde(E.Message)).Distinct();
-
+								IEnumerable<string> ErrorAndWarningEvents = Result.WarningAndErrorEvents.Select(E => E.IsError? E.Message : E.Message).Distinct();
 								ErrorAndWarningEvents = CapErrorOrWarningList(ErrorAndWarningEvents);
-
-								MB.H4(string.Format("Error: Test '{0}' did not run or complete", Result.TestDisplayName));
-								MB.Paragraph("FullName: " + Result.FullTestPath);
-								MB.UnorderedList(ErrorAndWarningEvents);
+								foreach (var Event in ErrorAndWarningEvents)
+								{
+									Log.Error(KnownLogEvents.Gauntlet_UnrealEngineTestEvent, "    " + Event);
+								}
 							}
+							Log.Info("");
 						}
 
 						// show a brief summary at the end where it's most visible
@@ -974,31 +995,32 @@ namespace UE
 						int PassedTests = AllTests.Count() - (FailedTests.Count() + IncompleteTests.Count());
 						int TestsPassedWithoutWarnings = PassedTests - TestsWithWarnings.Count();
 
-						TestSummary.Add(string.Format("{0} Test(s) Requested", AllTests.Count()));
+						TestSummary.Add(string.Format(" * {0} Test(s) Requested", AllTests.Count()));
 
 						// Print out a summary of each category of result
 						if (TestsPassedWithoutWarnings > 0)
 						{
-							TestSummary.Add(string.Format("{0} Test(s) Passed", TestsPassedWithoutWarnings));
+							TestSummary.Add(string.Format(" * {0} Test(s) Passed", TestsPassedWithoutWarnings));
 						}
 
 						if (TestsWithWarnings.Count() > 0)
 						{
-							TestSummary.Add(string.Format("{0} Test(s) Passed with warnings", TestsWithWarnings.Count()));
+							TestSummary.Add(string.Format(" * {0} Test(s) Passed with warnings", TestsWithWarnings.Count()));
 						}
 
 						if (FailedTests.Count() > 0)
 						{
-							TestSummary.Add(string.Format("{0} Test(s) Failed", FailedTests.Count()));
+							TestSummary.Add(string.Format(" * {0} Test(s) Failed", FailedTests.Count()));
 						}
 
 						if (IncompleteTests.Count() > 0)
 						{
-							TestSummary.Add(string.Format("{0} Test(s) didn't complete", IncompleteTests.Count()));
+							TestSummary.Add(string.Format(" * {0} Test(s) didn't complete", IncompleteTests.Count()));
 						}
 
-						MB.H3("Summary");
-						MB.UnorderedList(TestSummary);
+						Log.Info(" ### Summary");
+						Log.Info(string.Join("\n", TestSummary));
+						Log.Info("");
 					}
 				}
 
@@ -1013,22 +1035,21 @@ namespace UE
 					}
 					if (!string.IsNullOrEmpty(AutomationReportPath) || !string.IsNullOrEmpty(AutomationReportURL))
 					{
-						MB.H3("Links");
+						Log.Info(" ### Links");
 
 						if (string.IsNullOrEmpty(AutomationReportURL) == false)
 						{
-							MB.Paragraph(string.Format("View results here: {0}", AutomationReportURL));
+							Log.Info("  View results here: {URL}", AutomationReportURL);
 						}
 
 						if (string.IsNullOrEmpty(AutomationReportPath) == false)
 						{
-							MB.Paragraph(string.Format("Open results in UnrealEd from {0}", AutomationReportPath));
+							Log.Info("  Open results in UnrealEd from {Path}", AutomationReportPath);
 						}
+						Log.Info("");
 					}
 				}
 			}
-
-			return MB.ToString();
 		}
 
 		/// <summary>
diff --git a/Engine/Source/Programs/AutomationTool/Gauntlet/Unreal/Base/Gauntlet.UnrealSession.cs b/Engine/Source/Programs/AutomationTool/Gauntlet/Unreal/Base/Gauntlet.UnrealSession.cs
--- a/Engine/Source/Programs/AutomationTool/Gauntlet/Unreal/Base/Gauntlet.UnrealSession.cs
+++ b/Engine/Source/Programs/AutomationTool/Gauntlet/Unreal/Base/Gauntlet.UnrealSession.cs
@@ -13,6 +13,7 @@ using System.Drawing;
 using System.Linq;
 using System.Collections;
 using System.Diagnostics;
+using EpicGames.Core;
 
 namespace Gauntlet
 {
@@ -519,7 +520,7 @@ namespace Gauntlet
 
 			if (!CheckRolesArePossible(ref ValidationIssues))
 			{
-				ValidationIssues.ForEach(S => Log.Error("{0}", S));
+				ValidationIssues.ForEach(S => Log.Error(KnownLogEvents.Gauntlet_BuildDropEvent, S));
 				throw new AutomationException("One or more issues occurred when validating build {0} against requested roles", InSource.BuildName);
 			}
 
@@ -778,7 +779,7 @@ namespace Gauntlet
 							string ErrorMessage = string.Format("Failed to install app onto device {0} for role {1}. {2}. Will retry with new device", Device, Role, Ex);
 							if (ErrorMessage.Contains("not enough space"))
 							{
-								Log.Warning(ErrorMessage);
+								Log.Error(KnownLogEvents.Gauntlet_DeviceEvent, ErrorMessage);
 								if (Device.Platform == BuildHostPlatform.Current.Platform)
 								{
 									// If on desktop platform, we are not retrying.
@@ -788,7 +789,7 @@ namespace Gauntlet
 							}
 							else
 							{
-								Log.Info(ErrorMessage);
+								Log.Warning(KnownLogEvents.Gauntlet_DeviceEvent, ErrorMessage);
 							}
 							UnrealDeviceReservation.MarkProblemDevice(Device);
 							InstallSuccess = false;
@@ -867,13 +868,13 @@ namespace Gauntlet
 						catch (DeviceException Ex)
 						{
 							// shutdown all 
-							Log.Warning("Device {0} threw an exception during launch. \nException={1}", CurrentInstall.Device, Ex.Message);
+							Log.Warning(KnownLogEvents.Gauntlet_DeviceEvent, "Device {Name} threw an exception during launch. \nException={Exception}", CurrentInstall.Device, Ex.Message);
 							Success = false;
 						}
 
 						if (Success == false)
 						{
-							Log.Warning("Failed to start build on {0}. Marking as problem device and retrying with new set", CurrentInstall.Device);
+							Log.Warning(KnownLogEvents.Gauntlet_DeviceEvent, "Failed to start build on {Name}. Marking as problem device and retrying with new set", CurrentInstall.Device);
 
 							// terminate anything that's running
 							foreach (UnrealSessionInstance.RoleInstance RunningRole in RunningRoles)
diff --git a/Engine/Source/Programs/AutomationTool/Gauntlet/Unreal/Base/Gauntlet.UnrealTestNode.cs b/Engine/Source/Programs/AutomationTool/Gauntlet/Unreal/Base/Gauntlet.UnrealTestNode.cs
--- a/Engine/Source/Programs/AutomationTool/Gauntlet/Unreal/Base/Gauntlet.UnrealTestNode.cs
+++ b/Engine/Source/Programs/AutomationTool/Gauntlet/Unreal/Base/Gauntlet.UnrealTestNode.cs
@@ -11,6 +11,7 @@ using System.Drawing;
 using System.Linq;
 using System.Text;
 using System.Security.Cryptography;
+using EpicGames.Core;
 
 namespace Gauntlet
 {
@@ -137,7 +138,7 @@ namespace Gauntlet
 		{
 			Message = string.Format(Message, Args);
 			Events.Add(new UnrealAutomationEvent(EventType.Error, Message));
-			if (!LogWarningsAndErrorsAfterSummary) { Log.Error(Message); }
+			if (!LogWarningsAndErrorsAfterSummary) { Log.Error(KnownLogEvents.Gauntlet_TestEvent, Message); }
 			if (GetTestStatus() == TestStatus.Complete && GetTestResult() == TestResult.Passed)
 			{
 				SetUnrealTestResult(TestResult.Failed);
@@ -151,8 +152,8 @@ namespace Gauntlet
 		public virtual void ReportWarning(string Message, params object[] Args)
 		{
 			Message = string.Format(Message, Args);
-			if (!LogWarningsAndErrorsAfterSummary) { Events.Add(new UnrealAutomationEvent(EventType.Warning, Message)); }
-			Log.Warning(Message);
+			Events.Add(new UnrealAutomationEvent(EventType.Warning, Message));
+			if (!LogWarningsAndErrorsAfterSummary) { Log.Warning(KnownLogEvents.Gauntlet_TestEvent, Message); }
 		}
 
 		// Begin UnrealTestNode properties and members
@@ -520,7 +521,7 @@ namespace Gauntlet
 			if (TimeOfFirstMissingProcess == DateTime.MinValue)
 			{
 				TimeOfFirstMissingProcess = DateTime.Now;
-				Log.Verbose("Role {0} exited. Waiting {1} seconds for others to exit", MissingRoles.First().ToString(), TimeToWaitForProcesses);
+				Log.Verbose("Role {Role} exited. Waiting {Duration} seconds for others to exit", MissingRoles.First().ToString(), TimeToWaitForProcesses);
 			}
 
 			if ((DateTime.Now - TimeOfFirstMissingProcess).TotalSeconds < TimeToWaitForProcesses)
@@ -529,7 +530,7 @@ namespace Gauntlet
 				return true;
 			}
 
-			Log.Info("Ending {0} due to exit of Role {1}. {2} processes still running", Name, MissingRoles.First().ToString(), TestInstance.RunningRoles.Count());
+			Log.Info("Ending {Name} due to exit of Role {Role}. {Count} processes still running", Name, MissingRoles.First().ToString(), TestInstance.RunningRoles.Count());
 
 			// Done!
 			return false;
@@ -609,7 +610,7 @@ namespace Gauntlet
 					SessionRole.CommandLineParams = TestRole.CommandLineParams;
  					SessionRole.RoleModifier = TestRole.RoleType;
 					SessionRole.Constraint = UseContextConstraint ? Context.Constraint : new UnrealDeviceTargetConstraint(SessionPlatform);
-					Log.Verbose("Created SessionRole {0} from RoleContext {1} (RoleType={2})", SessionRole, RoleContext, TypesToRoles.Key);
+					Log.Verbose("Created SessionRole {Role} from RoleContext {Context} (RoleType={Type})", SessionRole, RoleContext, TypesToRoles.Key);
 
 					// TODO - this can all / mostly go into UnrealTestConfiguration.ApplyToConfig
 
@@ -724,7 +725,7 @@ namespace Gauntlet
 
 				if (ArtifactPathIsTaken)
 				{
-					Log.Info("Directory already exists at {0}", PotentialPath);
+					Log.Info("Directory already exists at {Path}", PotentialPath);
 					ArtifactNumericPostfix++;
 				}
 				else
@@ -780,7 +781,7 @@ namespace Gauntlet
 
 			// We need to create this directory at the start of the test rather than the end of the test - we are running into instances where multiple A/B tests
 			// on the same build are seeing the directory as non-existent and thinking it is safe to write to.
-			Log.Info("UnrealTestNode.StartTest Calling CreateDirectory for artifacts at {0}", ArtifactPath);
+			Log.Info("UnrealTestNode.StartTest Calling CreateDirectory for artifacts at {Path}", ArtifactPath);
 			Directory.CreateDirectory(ArtifactPath);
 
 			// Launch the test
@@ -959,7 +960,7 @@ namespace Gauntlet
 			LastActiveHeartbeatTime = DateTime.MinValue;
 
 			ReserveArtifactPath();
-			Log.Info("UnrealTestNode.ReStartTest Calling CreateDirectory for artifacts at {0}", ArtifactPath);
+			Log.Info("UnrealTestNode.ReStartTest Calling CreateDirectory for artifacts at {Path}", ArtifactPath);
 			Directory.CreateDirectory(ArtifactPath);
 
 			// Relaunch the test
@@ -1084,7 +1085,7 @@ namespace Gauntlet
 
 			try
 			{
-				Log.Info("Saving artifacts to {0}", ArtifactPath);
+				Log.Info("Saving artifacts to {Path}", ArtifactPath);
 				// run create dir again just in case the already made dir was cleaned up by another buildfarm job or something similar.
 				Directory.CreateDirectory(ArtifactPath);
 				Utils.SystemHelpers.MarkDirectoryForCleanup(ArtifactPath);
@@ -1096,7 +1097,7 @@ namespace Gauntlet
 			}
 			catch (Exception Ex)
 			{
-				Log.Warning("Failed to save artifacts. {0}", Ex);
+				Log.Warning("Failed to save artifacts. {Exception}", Ex);
 			}
 
 			if (CurrentPass + 1 == NumPasses)
@@ -1108,7 +1109,7 @@ namespace Gauntlet
 				}
 				catch (Exception Ex)
 				{
-					Log.Warning("Failed to release devices. {0}", Ex);
+					Log.Warning(KnownLogEvents.Gauntlet_DeviceEvent, "Failed to release devices. {Exception}", Ex);
 				}
 			}
 
@@ -1137,7 +1138,7 @@ namespace Gauntlet
 			{
 				CreateReportFailed = true;
 				Message = Ex.Message;				
-				Log.Warning("Failed to save completion report. {0}", Ex);
+				Log.Warning("Failed to save completion report. {Exception}", Ex);
 			}
 
 			if (CreateReportFailed)
@@ -1148,7 +1149,7 @@ namespace Gauntlet
 				}
 				catch (Exception Ex)
 				{
-					Log.Warning("Failed to handle completion report failure. {0}", Ex);
+					Log.Warning("Failed to handle completion report failure. {Exception}", Ex);
 				}
 			}
 
@@ -1162,7 +1163,7 @@ namespace Gauntlet
 			}
 			catch (Exception Ex)
 			{
-				Log.Warning("Failed to submit results to dashboard. {0}", Ex);
+				Log.Warning("Failed to submit results to dashboard. {Exception}", Ex);
 			}
 		}
 
@@ -1176,7 +1177,7 @@ namespace Gauntlet
 				Message = string.Format("See Gauntlet.log for details");
 			}
 
-			Log.Warning("CreateReport Failed: {0}", Message);
+			Log.Warning("CreateReport Failed: {Failure}", Message);
 		}
 
 		/// <summary>
@@ -1390,7 +1391,7 @@ namespace Gauntlet
 			if (File.Exists(JsonReportPath))
 			{
 				string HordeArtifactPath = GetCachedConfiguration().HordeArtifactPath;
-				Log.Verbose("Reading json Unreal Automated test report from {0}", JsonReportPath);
+				Log.Verbose("Reading json Unreal Automated test report from {Path}", JsonReportPath);
 				UnrealAutomatedTestPassResults JsonTestPassResults = UnrealAutomatedTestPassResults.LoadFromJson(JsonReportPath);
 				var MainRole = GetCachedConfiguration().GetMainRequiredRole();
 				string HordeTestDataKey = string.IsNullOrEmpty(GetCachedConfiguration().HordeTestDataKey) ? Name + " " + Context.ToString() : GetCachedConfiguration().HordeTestDataKey;
@@ -1428,7 +1429,7 @@ namespace Gauntlet
 			}
 			else
 			{
-				Log.Warning("Could not find Unreal Automated test report at {0}. Reverting to base report.", JsonReportPath);
+				Log.Warning("Could not find Unreal Automated test report at {FilePath}. Reverting to base report.", JsonReportPath);
 				return CreateSimpleReportForHorde(GetTestResult());
 			}
 		}
@@ -1490,7 +1491,7 @@ namespace Gauntlet
 					{
 						DBConfig.LoadConfig(GetCachedConfiguration().DatabaseConfigPath);
 						IDatabaseDriver<TelemetryData> DB = DBConfig.GetDriver();
-						Log.Verbose("Submitting telemetry data to {0}", DB.ToString());
+						Log.Verbose("Submitting telemetry data to {Database}", DB.ToString());
 
 						UnrealTelemetryContext TestContext = new UnrealTelemetryContext();
 						TestContext.SetProperty("ProjectName", Context.BuildInfo.ProjectName);
@@ -1505,7 +1506,7 @@ namespace Gauntlet
 					}
 					else
 					{
-						Log.Warning("Got telemetry data, but database configuration is unknown '{0}'.", GetCachedConfiguration().PublishTelemetryTo);
+						Log.Warning("Got telemetry data, but database configuration is unknown '{Config}'.", GetCachedConfiguration().PublishTelemetryTo);
 					}
 				}
 			}
@@ -1611,7 +1612,7 @@ namespace Gauntlet
 			{
 				if (InLog.HasTestExitCode == false)
 				{
-					Log.Verbose("Role {0} had 0 exit code but used Gauntlet and no TestExitCode was found. Assuming failure", InArtifacts.SessionRole.RoleType);
+					Log.Verbose("Role {Role} had 0 exit code but used Gauntlet and no TestExitCode was found. Assuming failure", InArtifacts.SessionRole.RoleType);
 					ExitReason = "Process terminated prematurely! No test result from Gauntlet controller";
 					ExitCode = -1;
 					return UnrealProcessResult.TestFailure;
@@ -1765,7 +1766,7 @@ namespace Gauntlet
 				double SecondsSinceSessionStart = DateTime.Now.Subtract(SessionStartTime).TotalSeconds;
 				if (SecondsSinceSessionStart > HeartbeatOptions.TimeoutBeforeFirstActiveHeartbeat)
 				{
-					Log.Error("{0} seconds have passed without detecting the first active Gauntlet heartbeat.", HeartbeatOptions.TimeoutBeforeFirstActiveHeartbeat);
+					Log.Error(KnownLogEvents.Gauntlet_TestEvent, "{Time} seconds have passed without detecting the first active Gauntlet heartbeat.", HeartbeatOptions.TimeoutBeforeFirstActiveHeartbeat);
 					MarkTestComplete();
 					SetUnrealTestResult(TestResult.TimedOut);
 				}
@@ -1777,7 +1778,7 @@ namespace Gauntlet
 				double SecondsSinceLastActiveHeartbeat = DateTime.Now.Subtract(LastActiveHeartbeatTime).TotalSeconds;
 				if (SecondsSinceLastActiveHeartbeat > HeartbeatOptions.TimeoutBetweenActiveHeartbeats)
 				{
-					Log.Error("{0} seconds have passed without detecting any active Gauntlet heartbeats.", HeartbeatOptions.TimeoutBetweenActiveHeartbeats);
+					Log.Error(KnownLogEvents.Gauntlet_TestEvent, "{Time} seconds have passed without detecting any active Gauntlet heartbeats.", HeartbeatOptions.TimeoutBetweenActiveHeartbeats);
 					MarkTestComplete();
 					SetUnrealTestResult(TestResult.TimedOut);
 				}
@@ -1789,7 +1790,7 @@ namespace Gauntlet
 				double SecondsSinceLastHeartbeat = DateTime.Now.Subtract(LastHeartbeatTime).TotalSeconds;
 				if (SecondsSinceLastHeartbeat > HeartbeatOptions.TimeoutBetweenAnyHeartbeats)
 				{
-					Log.Error("{0} seconds have passed without detecting any Gauntlet heartbeats.", HeartbeatOptions.TimeoutBetweenAnyHeartbeats);
+					Log.Error(KnownLogEvents.Gauntlet_TestEvent, "{Time} seconds have passed without detecting any Gauntlet heartbeats.", HeartbeatOptions.TimeoutBetweenAnyHeartbeats);
 					MarkTestComplete();
 					SetUnrealTestResult(TestResult.TimedOut);
 				}
@@ -1846,45 +1847,56 @@ namespace Gauntlet
 		}
 
 		/// <summary>
-		/// Returns a formatted summary of the role that's suitable for displaying
+		/// Deprecated
 		/// </summary>
 		/// <param name="InRoleResult"></param>
 		/// <returns></returns>
 		protected virtual string GetFormattedRoleSummary(UnrealRoleResult InRoleResult)
+		{
+			return "";
+		}
+
+		/// <summary>
+		/// Log a formatted summary of the role that's suitable for displaying
+		/// </summary>
+		/// <param name="InRoleResult"></param>
+		/// <returns></returns>
+		protected virtual void LogRoleSummary(UnrealRoleResult InRoleResult)
 		{
 
 			const int MaxLogLines = 10;
 			const int MaxCallstackLines = 20;
 
 			UnrealLog LogSummary = InRoleResult.LogSummary;
-			MarkdownBuilder MB = new MarkdownBuilder();
 
 			UnrealRoleArtifacts RoleArtifacts = InRoleResult.Artifacts;
 
-			MB.H3(string.Format("Role: {0} ({1} {2})", RoleArtifacts.SessionRole.RoleType, RoleArtifacts.SessionRole.Platform, RoleArtifacts.SessionRole.Configuration));
-
-			bool HaveFatalError = LogSummary.FatalError != null;
-
-			// If we have no fatal error use "Error" which Horde will highlight. Otherwise use "Failed" since the fatal error will be below and highligted
-			string FailedString = HaveFatalError ? "Failed: " : "Error: ";
-			string CompletedString = "Completed: ";
+			Log.Info(" #### Role: {0} ({1} {2})", RoleArtifacts.SessionRole.RoleType, RoleArtifacts.SessionRole.Platform, RoleArtifacts.SessionRole.Configuration);
 
-			string RoleState = InRoleResult.ExitCode != 0 && InRoleResult.LogSummary.HasAbnormalExit ? FailedString : CompletedString;
-
-			MB.H4(string.Format("{0} {1} ({2}, ExitCode={3})", RoleState, InRoleResult.Summary, InRoleResult.ProcessResult, InRoleResult.ExitCode));
+			bool HasFailed = InRoleResult.ExitCode != 0 && InRoleResult.LogSummary.HasAbnormalExit;
+			string RoleState = HasFailed ? "failed:" : "completed:";
+			string StatusMessage = string.Format(" #### {0} {1} {2} ({3}, ExitCode={4})", RoleArtifacts.SessionRole.RoleType, RoleState, InRoleResult.Summary, InRoleResult.ProcessResult, InRoleResult.ExitCode);
+			if(HasFailed)
+			{
+				Log.Error(KnownLogEvents.Gauntlet_TestEvent, StatusMessage);
+			}
+			else
+			{
+				Log.Info(StatusMessage);
+			}
 
 			// log command line up here for visibility
-			//MB.Paragraph(string.Format("CommandLine: {0}", RoleArtifacts.AppInstance.CommandLine));
-
-			MB.UnorderedList(new string[] {
-				string.Format("CommandLine: {0}", RoleArtifacts.AppInstance.CommandLine),
-				string.Format("Log: {0}", RoleArtifacts.LogPath),
-				string.Format("SavedDir: {0}", RoleArtifacts.ArtifactPath),
-				LogSummary.FatalError != null ? "Fatal Errors: 1" : null,
-				LogSummary.Ensures.Count() > 0 ? string.Format("Ensures: {0}", LogSummary.Ensures.Count()) : null,
-				LogSummary.Errors.Count() > 0 ? string.Format("Log Errors: {0}", LogSummary.Errors.Count()) : null,
-				LogSummary.Warnings.Count() > 0 ? string.Format("Log Warnings: {0}", LogSummary.Warnings.Count()) : null
-			});
+
+			Log.Info( string.Join("\n", new string[] {
+				string.Format(" * CommandLine: {0}", RoleArtifacts.AppInstance.CommandLine),
+				string.Format(" * Log: {0}", RoleArtifacts.LogPath),
+				string.Format(" * SavedDir: {0}", RoleArtifacts.ArtifactPath),
+				LogSummary.FatalError != null ? " * Fatal Errors: 1" : null,
+				LogSummary.Ensures.Count() > 0 ? string.Format(" * Ensures: {0}", LogSummary.Ensures.Count()) : null,
+				LogSummary.Errors.Count() > 0 ? string.Format(" * Log Errors: {0}", LogSummary.Errors.Count()) : null,
+				LogSummary.Warnings.Count() > 0 ? string.Format(" * Log Warnings: {0}", LogSummary.Warnings.Count()) : null
+			}.Where(L => !string.IsNullOrEmpty(L))));
+			Log.Info("");
 
 			// Separate the events we want to report on
 			IEnumerable<UnrealTestEvent> Asserts = InRoleResult.Events.Where(E => E.Severity == EventSeverity.Fatal);
@@ -1894,40 +1906,41 @@ namespace Gauntlet
 
 			foreach (UnrealTestEvent Event in Asserts)
 			{
-				MB.H4(string.Format("Fatal Error: {0}", Event.Summary));
-
+				List<string> Callstack = new List<string>();
 				if (Event.Callstack.Any())
 				{
-					MB.UnorderedList(Event.Callstack.Take(MaxCallstackLines));
-
+					Callstack = Event.Callstack.Take(MaxCallstackLines).ToList();
 					if (Event.Callstack.Count() > MaxCallstackLines)
 					{
-						MB.Paragraph("See log for full callstack");
+						Callstack.Add("See log for full callstack");
 					}
 				}
 				else
 				{
-					MB.Paragraph("Could not parse callstack. See log for full callstack");
+					Callstack.Add("Could not parse callstack. See log for full callstack");
 				}
+				Log.Error(KnownLogEvents.Gauntlet_TestEvent, " * Fatal Error: {Summary}\n{Callstack}", Event.Summary, string.Join("\n", Callstack.Select(C=>"    "+C)));
+				Log.Info("");
 			}
 
 			foreach (UnrealTestEvent Event in Ensures.Distinct())
 			{
-				MB.H4(string.Format("Warning: Ensure: {0}", Event.Summary));
-
+				List<string> Callstack = new List<string>();
 				if (Event.Callstack.Any())
 				{
-					MB.UnorderedList(Event.Callstack.Take(MaxCallstackLines));
+					Callstack = Event.Callstack.Take(MaxCallstackLines).ToList();
 
 					if (Event.Callstack.Count() > MaxCallstackLines)
 					{
-						MB.Paragraph("See log for full callstack");
+						Callstack.Add("See log for full callstack");
 					}
 				}
 				else
 				{
-					MB.Paragraph("Could not parse callstack. See log for full callstack");
+					Callstack.Add("Could not parse callstack. See log for full callstack");
 				}
+				Log.Warning(KnownLogEvents.Gauntlet_TestEvent, " * Ensure: {Summary}\n{Callstack}", Event.Summary, string.Join("\n", Callstack.Select(C => "    " + C)));
+				Log.Info("");
 			}
 
 			if (Errors.Any())
@@ -1943,22 +1956,24 @@ namespace Gauntlet
 					if (LogSummary.HasAbnormalExit)
 					{
 						PrintedErrorList = ErrorList.Skip(ErrorList.Count() - MaxLogLines);
-						TrimStatement = string.Format("(Last {0} of {1} errors)", MaxLogLines, ErrorList.Count());
+						TrimStatement = string.Format("  (Last {0} of {1} errors)", MaxLogLines, ErrorList.Count());
 					}
 					else
 					{
 						PrintedErrorList = ErrorList.Take(MaxLogLines);
-						TrimStatement = string.Format("(First {0} of {1} errors)", MaxLogLines, ErrorList.Count());
+						TrimStatement = string.Format("  (First {0} of {1} errors)", MaxLogLines, ErrorList.Count());
 					}
 				}
 
-				MB.H4("Errors:");
-				MB.UnorderedList(PrintedErrorList);
-
+				foreach (var Error in PrintedErrorList)
+				{
+					Log.Error(KnownLogEvents.Gauntlet_TestEvent, " * {Message}", Error);
+				}
 				if (!string.IsNullOrEmpty(TrimStatement))
 				{
-					MB.Paragraph(TrimStatement);
+					Log.Info(KnownLogEvents.Gauntlet_TestEvent, TrimStatement);
 				}
+				Log.Info("");
 			}
 
 			if (Warnings.Any())
@@ -1974,25 +1989,25 @@ namespace Gauntlet
 					if (LogSummary.HasAbnormalExit)
 					{
 						PrintedWarningList = WarningList.Skip(WarningList.Count() - MaxLogLines);
-						TrimStatement = string.Format("(Last {0} of {1} errors)", MaxLogLines, WarningList.Count());
+						TrimStatement = string.Format(" (Last {0} of {1} warnings)", MaxLogLines, WarningList.Count());
 					}
 					else
 					{
 						PrintedWarningList = WarningList.Take(MaxLogLines);
-						TrimStatement = string.Format("(First {0} of {1} errors)", MaxLogLines, WarningList.Count());
+						TrimStatement = string.Format(" (First {0} of {1} warnings)", MaxLogLines, WarningList.Count());
 					}
 				}
 
-				MB.H4("Warnings:");
-				MB.UnorderedList(PrintedWarningList);
-
+				foreach (var Warning in PrintedWarningList)
+				{
+					Log.Warning(KnownLogEvents.Gauntlet_TestEvent, " * {Message}", Warning);
+				}
 				if (!string.IsNullOrEmpty(TrimStatement))
 				{
-					MB.Paragraph(TrimStatement);
+					Log.Info(KnownLogEvents.Gauntlet_TestEvent, TrimStatement);
 				}
+				Log.Info("");
 			}
-
-			return MB.ToString();
 		}
 
 		/// <summary>
@@ -2004,8 +2019,6 @@ namespace Gauntlet
 
 			AddProcessResultEventsFromTestNode();
 			
-
-			const int MaxCallstackLines = 20;
 			MarkdownBuilder MB = new MarkdownBuilder();
 
 			if (TestNodeEvents.Count == 0)
@@ -2023,43 +2036,18 @@ namespace Gauntlet
 
 			foreach (UnrealTestEvent Event in Fatals)
 			{
-				MB.H4(string.Format("Fatal Error: {0}", Event.Summary));
-
-				if (Event.Callstack.Any())
-				{
-					MB.UnorderedList(Event.Callstack.Take(MaxCallstackLines));
-
-					if (Event.Callstack.Count() > MaxCallstackLines)
-					{
-						MB.Paragraph("See log for full callstack");
-					}
-				}
-				else
-				{
-					MB.Paragraph("Could not parse callstack. See log for full callstack");
-				}
+				MB.H4(Event.Summary);
+				MB.UnorderedList(Event.Details);
+				MB.NewLine();
 			}
 
 			foreach (UnrealTestEvent Event in Ensures.Distinct())
 			{
-				MB.H4(string.Format("Warning: Ensure: {0}", Event.Summary));
-
-				if (Event.Callstack.Any())
-				{
-					MB.UnorderedList(Event.Callstack.Take(MaxCallstackLines));
-
-					if (Event.Callstack.Count() > MaxCallstackLines)
-					{
-						MB.Paragraph("See log for full callstack");
-					}
-				}
-				else
-				{
-					MB.Paragraph("Could not parse callstack. See log for full callstack");
-				}
+				MB.H4(Event.Summary);
+				MB.UnorderedList(Event.Details);
+				MB.NewLine();
 			}
 
-
 			if (Errors.Any())
 			{
 				MB.H3("Errors:");
@@ -2200,7 +2188,7 @@ namespace Gauntlet
 			{
 				foreach (var Role in FailedRoles)
 				{
-					Log.Info("Failing test because {0} exited with {1}. ({2})", Role.Artifacts.SessionRole, Role.ExitCode, Role.Summary);
+					Log.Info("Failing test because {Role} exited with {ExitCode}. ({Result})", Role.Artifacts.SessionRole, Role.ExitCode, Role.Summary);
 				}
 				ExitCode = FailedRoles.FirstOrDefault().ExitCode;
 			}
@@ -2214,28 +2202,43 @@ namespace Gauntlet
 			return ExitCode == 0 ? TestResult.Passed : TestResult.Failed;
 		}
 
+
 		/// <summary>
-		/// Return the header for the test summary. The header is the first block of text and will be
-		/// followed by the summary of each individual role in the test
+		/// Deprecated
 		/// </summary>
 		/// <returns></returns>
 		protected virtual string GetTestSummaryHeader()
+		{
+			return "";
+		}
+
+		/// <summary>
+		/// Log header for the test summary. The header is the first block of text and will be
+		/// followed by the summary of each individual role in the test
+		/// </summary>
+		/// <returns></returns>
+		protected virtual void LogTestSummaryHeader()
 		{
 			int FatalErrors = 0;
 			int Ensures = 0;
 			int Errors = 0;
 			int Warnings = 0;
 
-			MarkdownBuilder MB = new MarkdownBuilder();
-
 			bool TestFailed = GetTestResult() != TestResult.Passed;
 
 			// Good/Bad news upfront
 			string Prefix = TestFailed ? "Error: " : "";
 			string WarningStatement = (HasWarnings && !TestFailed)  ? " With Warnings" : "";
-			string ResultString = string.Format("{0}{1} {2}{3} ***", Prefix, this.Name, GetTestResult(), WarningStatement);
-			MB.H2(ResultString);
-						
+			string ResultString = string.Format(" ### {0}{1} {2}{3} ***", Prefix, this.Name, GetTestResult(), WarningStatement);
+			if(TestFailed)
+			{
+				Log.Error(KnownLogEvents.Gauntlet_TestEvent, ResultString);
+			}
+			else
+			{
+				Log.Info(ResultString);
+			}
+
 			IEnumerable<UnrealRoleResult> SortedRoles = RoleResults.OrderBy(R => R.ProcessResult == UnrealProcessResult.ExitOk);
 
 			// create a quick summary of total failures, ensures, errors, etc. Don't write out errors etc for roles, those will be
@@ -2244,7 +2247,7 @@ namespace Gauntlet
 			{
 				string RoleName = RoleResult.Artifacts.SessionRole.RoleType.ToString();
 
-				MB.Paragraph(string.Format("{0} Role: {1} ({2}, ExitCode {3})", RoleName, RoleResult.Summary, RoleResult.ProcessResult, RoleResult.ExitCode));
+				Log.Info(string.Format("\n {0} Role: {1} ({2}, ExitCode {3})\n", RoleName, RoleResult.Summary, RoleResult.ProcessResult, RoleResult.ExitCode));
 
 				FatalErrors += RoleResult.LogSummary.FatalError != null ? 1 : 0;
 				Ensures += RoleResult.LogSummary.Ensures.Count();
@@ -2252,26 +2255,21 @@ namespace Gauntlet
 				Warnings += RoleResult.LogSummary.Warnings.Count();
 			}
 
-			MB.UnorderedList(new string[] {
+			Log.Info(string.Join("\n", new string[] {
 				string.Format("Context: {0}", Context.ToString()),
 				FatalErrors > 0 ? string.Format("FatalErrors: {0}", FatalErrors) : null,
 				Ensures > 0 ? string.Format("Ensures: {0}", Ensures) : null,
 				Errors > 0 ? string.Format("Log Errors: {0}", Errors) : null,
 				Warnings > 0 ? string.Format("Log Warnings: {0}", Warnings) : null,
 				string.Format("Result: {0}", GetTestResult())
-			});
+			}.Where(L=>!string.IsNullOrEmpty(L)).Select(L=>" * "+L)));
 
 			if (TestFailed && GetRolesThatFailed().Where(R => R.ProcessResult == UnrealProcessResult.Unknown).Any())
 			{
-				MB.Paragraph(string.Format("Error: {0} failed due to undiagnosed reasons", this.Name));
-				MB.Paragraph("See 'Role Report' below for more details on each role");
+				Log.Info("");
+				Log.Error(KnownLogEvents.Gauntlet_TestEvent, " * {Name} failed due to undiagnosed reasons", this.Name);
 			}
-			else
-			{
-				MB.Paragraph("See 'Role Report' below for more details on each role");
-			}			
-
-			return MB.ToString();
+			Log.Info("\n See 'Role Report' below for more details on each role\n");
 		}
 
 		/// <summary>
@@ -2282,46 +2280,60 @@ namespace Gauntlet
 		{
 			MarkdownBuilder ReportBuilder = new MarkdownBuilder();
 
-			// Handle case where there aren't any session artifacts, for example with device starvation
-			if (SessionArtifacts == null)
-			{
-				ReportBuilder.H1(string.Format("Gauntlet summary events:"));
-				ReportBuilder.HorizontalLine();
-				ReportBuilder.Append(CreateFormattedEventListFromTestNode());
-				return ReportBuilder.ToString();
-			}
-			// Sort roles so problem ones are at the bottom, just above the summary
-			var SortedRoles = RoleResults.OrderBy(R =>
+			// Handle case where there are session artifacts
+			if (SessionArtifacts != null)
 			{
-				int Score = 0;
-
-				if (R.ProcessResult != UnrealProcessResult.ExitOk)
+				// Sort roles so problem ones are at the bottom, just above the summary
+				var SortedRoles = RoleResults.OrderBy(R =>
 				{
-					Score += 1000000;
-				}
+					int Score = 0;
 
-				Score += R.LogSummary.Errors.Count() * 10;
-				Score += R.LogSummary.Warnings.Count();
+					if (R.ProcessResult != UnrealProcessResult.ExitOk)
+					{
+						Score += 1000000;
+					}
 
-				return Score;
-			});		
+					Score += R.LogSummary.Errors.Count() * 10;
+					Score += R.LogSummary.Warnings.Count();
 
-			// add Summary
-			ReportBuilder.HorizontalLine();
-			ReportBuilder.H1("Test Summary: " + this.Name);
-			ReportBuilder.HorizontalLine();
-			ReportBuilder.Append(GetTestSummaryHeader());
+					return Score;
+				});
 
-			ReportBuilder.HorizontalLine();
-			ReportBuilder.H1(string.Format("Role(s): {0}", this.Name));
-			ReportBuilder.HorizontalLine();
+				// add Summary
+				string HorizontalLine = " " + new string('-', 80);
+				Log.Info(HorizontalLine);
+				Log.Info(" # Test Summary: " + this.Name);
+				Log.Info(HorizontalLine);
+				string HeaderSummary = GetTestSummaryHeader();
+				if (!string.IsNullOrEmpty(HeaderSummary))
+				{
+					Log.Info(HeaderSummary);
+				}
+				else
+				{
+					LogTestSummaryHeader();
+				}
 
-			// Add a summary of each 
-			foreach (var Role in SortedRoles)
-			{
-				string Summary = GetFormattedRoleSummary(Role);
-				ReportBuilder.Append(Summary);
+				Log.Info(HorizontalLine);
+				Log.Info(" # Role(s): {Name}", this.Name);
+				Log.Info(HorizontalLine);
+
+				// Add a summary of each 
+				foreach (var Role in SortedRoles)
+				{
+					string RoleSummary = GetFormattedRoleSummary(Role);
+					if (!string.IsNullOrEmpty(RoleSummary))
+					{
+						Log.Info(RoleSummary);
+					}
+					else
+					{
+						LogRoleSummary(Role);
+					}
+				}
 			}
+
+			ReportBuilder.HorizontalLine();
 			ReportBuilder.H1(string.Format("Gauntlet summary events:"));
 			ReportBuilder.HorizontalLine();
 			ReportBuilder.Append(CreateFormattedEventListFromTestNode());
diff --git a/Engine/Source/Programs/AutomationTool/Gauntlet/Unreal/Engine/UnrealTests.cs b/Engine/Source/Programs/AutomationTool/Gauntlet/Unreal/Engine/UnrealTests.cs
--- a/Engine/Source/Programs/AutomationTool/Gauntlet/Unreal/Engine/UnrealTests.cs
+++ b/Engine/Source/Programs/AutomationTool/Gauntlet/Unreal/Engine/UnrealTests.cs
@@ -7,6 +7,7 @@ using System.Linq;
 using AutomationTool;
 using UnrealBuildTool;
 using UnrealBuildBase;
+using EpicGames.Core;
 
 namespace Gauntlet
 {
@@ -163,7 +164,7 @@ namespace Gauntlet
 					// Check for timeout
 					if (ElapsedTime > IdleTimeout)
 					{
-						Log.Error("No activity observed in last {0:0.00} minutes. Aborting test", IdleTimeout / 60);
+						Log.Error(KnownLogEvents.Gauntlet_TestEvent, "No activity observed in last {0:0.00} minutes. Aborting test", IdleTimeout / 60);
 						MarkTestComplete();
 						SetUnrealTestResult(TestResult.TimedOut);
 					}
@@ -247,9 +248,9 @@ namespace Gauntlet
 			/// Override the summary report so we can insert a link to our report and the failed tests
 			/// </summary>
 			/// <returns></returns>
-			protected override string GetTestSummaryHeader()
+			protected override void LogTestSummaryHeader()
 			{
-				MarkdownBuilder MB = new MarkdownBuilder(base.GetTestSummaryHeader());
+				base.LogTestSummaryHeader();
 
 				var EditorRole = RoleResults.Where(R => R.Artifacts.SessionRole.RoleType == UnrealTargetRole.Editor).FirstOrDefault();
 
@@ -263,32 +264,38 @@ namespace Gauntlet
 
 					if (AllTests.Count() == 0)
 					{
-						MB.Paragraph("No tests were executed!");
+						Log.Error(KnownLogEvents.Gauntlet_UnrealEngineTestEvent, " * No tests were executed!");
 					}
 					else
 					{
 						int PassedTests = AllTests.Count() - (FailedTests.Count() + IncompleteTests.Count());
-						MB.Paragraph(string.Format("{0} of {1} tests passed", PassedTests, AllTests.Count()));
+						Log.Info(" * {0} of {1} tests passed", PassedTests, AllTests.Count());
 
 						if (FailedTests.Count() > 0)
 						{
-							MB.H3("The following tests failed:");
+							Log.Info(" ### The following tests failed:");
 
 							foreach (UnrealAutomatedTestResult Result in FailedTests)
 							{
-								MB.H4(string.Format("{0}: {1}", Result.TestDisplayName, Result.FullTestPath));
-								MB.UnorderedList(Result.ErrorEvents.Select(E => string.Format("Error: {0}",E.Message)));
+								Log.Error(KnownLogEvents.Gauntlet_UnrealEngineTestEvent, " * {0} failed", Result.FullTestPath);
+								foreach (var Event in Result.ErrorEvents)
+								{
+									Log.Error(KnownLogEvents.Gauntlet_UnrealEngineTestEvent, "    "+Event.Message);
+								}
 							}
 						}
 
 						if (IncompleteTests.Count() > 0)
 						{
-							MB.H3("The following tests timed out:");
+							Log.Info(" ### The following tests timed out:");
 
 							foreach (UnrealAutomatedTestResult Result in IncompleteTests)
 							{
-								MB.H4(string.Format("{0}: {1}", Result.TestDisplayName, Result.FullTestPath));
-								MB.UnorderedList(Result.ErrorEvents.Select(E => string.Format("Error: {0}", E.Message)));
+								Log.Error(KnownLogEvents.Gauntlet_UnrealEngineTestEvent, " * {0} timed out", Result.FullTestPath);
+								foreach (var Event in Result.ErrorEvents)
+								{
+									Log.Error(KnownLogEvents.Gauntlet_UnrealEngineTestEvent, "    " + Event.Message);
+								}
 							}
 						}
 
@@ -297,24 +304,22 @@ namespace Gauntlet
 
 						if (!string.IsNullOrEmpty(ReportLink) || !string.IsNullOrEmpty(ReportPath))
 						{
-							MB.H3("Links");
+							Log.Info(" ### Links");
 
 							if (string.IsNullOrEmpty(ReportLink) == false)
 							{
-								MB.Paragraph(string.Format("View results here: {0}", ReportLink));
+								Log.Info("  View results here: {URL}", ReportLink);
 							}
 
 							if (string.IsNullOrEmpty(ReportPath) == false)
 							{
-								MB.Paragraph(string.Format("Open results in UnrealEd from {0}", ReportPath));
+								Log.Info("  Open results in UnrealEd from {Path}", ReportPath);
 							}
 						}
 
 
 					}
 				}
-
-				return MB.ToString();
 			}
 		}
 
diff --git a/Engine/Source/Programs/Shared/EpicGames.Core/KnownLogEvents.cs b/Engine/Source/Programs/Shared/EpicGames.Core/KnownLogEvents.cs
--- a/Engine/Source/Programs/Shared/EpicGames.Core/KnownLogEvents.cs
+++ b/Engine/Source/Programs/Shared/EpicGames.Core/KnownLogEvents.cs
@@ -146,14 +146,24 @@ namespace EpicGames.Core
 		public static EventId Gauntlet { get; } = new EventId(600);
 
 		/// <summary>
-		/// Error message from Gauntlet engine tests
+		/// Error message from Gauntlet test events
 		/// </summary>
-		public static EventId Gauntlet_UnitTest { get; } = new EventId(601);
+		public static EventId Gauntlet_TestEvent { get; } = new EventId(601);
 
 		/// <summary>
-		/// Error message from Gauntlet screenshot tests
+		/// Error message from Gauntlet device events
 		/// </summary>
-		public static EventId Gauntlet_ScreenshotTest { get; } = new EventId(602);
+		public static EventId Gauntlet_DeviceEvent { get; } = new EventId(602);
+
+		/// <summary>
+		/// Error message from Gauntlet Unreal Engine test events
+		/// </summary>
+		public static EventId Gauntlet_UnrealEngineTestEvent { get; } = new EventId(603);
+
+		/// <summary>
+		/// Error message from Gauntlet build drop events
+		/// </summary>
+		public static EventId Gauntlet_BuildDropEvent { get; } = new EventId(604);
 
 		/// <summary>
 		/// A systemic event, relating to the health of the farm
