Build: Symbol Server: Added support for source server (source code indexing) for consoles.

Source code indexing allows Visual Studio to automatically fetch the right source code files from Perforce when debugging builds/crashdumps.

Previously, we had this functionality available only on Windows in the form of SrcSrv build task (build graph). This task takes generated pdbs and embeds information about the source code used into them. However, this functionality doesn't translate well into similar features on other platforms.

Resolved by extending SymStore task to handle source indexing in additional to uploading symbols to the symbol store. SrcSrv remains functional, but it still works only on Windows. For this reason, SymStore becomes the preferred solution as it's more general.

Fixed a few additional problems, both in the task itself and our game's build graph.

diff --git a/Engine/Source/Programs/AutomationTool/AutomationUtils/Platform.cs b/Engine/Source/Programs/AutomationTool/AutomationUtils/Platform.cs
--- a/Engine/Source/Programs/AutomationTool/AutomationUtils/Platform.cs
+++ b/Engine/Source/Programs/AutomationTool/AutomationUtils/Platform.cs
@@ -747,7 +747,9 @@ namespace AutomationTool
 			}
 		}
 
-		public virtual bool PublishSymbols(DirectoryReference SymbolStoreDirectory, List<FileReference> Files, string Product, string BuildVersion = null)
+		public virtual bool PublishSymbols(DirectoryReference SymbolStoreDirectory, List<FileReference> Files,
+			bool bIndexSources, List<FileReference> SourceFiles,
+			string Product, string Branch, int Change, string BuildVersion = null)
 		{
 			CommandUtils.LogWarning("PublishSymbols() has not been implemented for {0}", PlatformType.ToString());
 			return false;
@@ -772,6 +774,16 @@ namespace AutomationTool
 			get { return false; }
 		}
 
+		/// <summary>
+		/// When true, callers of PublishSymbols() must provide an explicit list of source files to create the index from.
+		/// Some platforms discover the source files via other means, so it is possible to turn this step of the process
+		/// off, since it can be slow.
+		/// </summary>
+		public virtual bool SymbolServerSourceIndexingRequiresListOfSourceFiles
+		{
+			get { return true; }
+		}
+
 		/// <summary>
 		/// If true, indicates the platform's symbol server directory must be locked for
 		/// exclusive access before any operation is performed on it. Platforms may override
diff --git a/Engine/Source/Programs/AutomationTool/BuildGraph/Tasks/SymStoreTask.cs b/Engine/Source/Programs/AutomationTool/BuildGraph/Tasks/SymStoreTask.cs
--- a/Engine/Source/Programs/AutomationTool/BuildGraph/Tasks/SymStoreTask.cs
+++ b/Engine/Source/Programs/AutomationTool/BuildGraph/Tasks/SymStoreTask.cs
@@ -43,11 +43,42 @@ namespace AutomationTool.Tasks
         [TaskParameter]
         public string Product;
 
+		/// <summary>
+		/// Name of the Branch to base all the depot source files from.
+		/// Used when IndexSources is true (may be used only on some platforms).
+		/// </summary>
+		[TaskParameter(Optional = true)]
+		public string Branch;
+
+		/// <summary>
+		/// Changelist to which all the depot source files have been synced to.
+		/// Used when IndexSources is true (may be used only on some platforms).
+		/// </summary>
+		[TaskParameter(Optional = true)]
+		public int Change;
+
 		/// <summary>
 		/// BuildVersion associated with these symbols. Used for clean-up in AgeStore by matching this version against a directory name in a build share.
 		/// </summary>
 		[TaskParameter(Optional = true)]
 		public string BuildVersion;
+
+		/// <summary>
+		/// Whether to include the source code index in the uploaded symbols.
+		/// When enabled, the task will generate data required by a source server (only some platforms and source control servers are supported).
+		/// The source server allows debuggers to automatically fetch the matching source code when debbugging builds or analyzing dumps.
+		/// </summary>
+		[TaskParameter(Optional = true)]
+		public bool IndexSources = false;
+
+		/// <summary>
+		/// Filter for the depot source files that are to be indexed.
+		/// It's a semicolon-separated list of perforce filter e.g. Engine/....cpp;Engine/....h.
+		/// It may also be a name of a previously defined tag e.g. "#SourceFiles
+		/// Used when IndexSources is true (may be used only on some platforms).
+		/// </summary>
+		[TaskParameter(Optional = true)]
+		public string SourceFiles;
 	}
 
     /// <summary>
@@ -86,9 +117,20 @@ namespace AutomationTool.Tasks
 
 			// Take the lock before accessing the symbol server, if required by the platform
 			Platform TargetPlatform = Platform.GetPlatform(Parameters.Platform);
+
+			List<FileReference> SourceFiles = new List<FileReference>();
+
+			if (Parameters.IndexSources && TargetPlatform.SymbolServerSourceIndexingRequiresListOfSourceFiles)
+			{
+				CommandUtils.LogInformation("Discovering source code files...");
+
+				SourceFiles = ResolveFilespec(Unreal.RootDirectory, Parameters.SourceFiles, TagNameToFileSet).ToList();
+			}
+
 			CommandUtils.OptionallyTakeLock(TargetPlatform.SymbolServerRequiresLock, StoreDir, TimeSpan.FromMinutes(60), () =>
 			{
-				if (!TargetPlatform.PublishSymbols(StoreDir, Files, Parameters.Product, Parameters.BuildVersion))
+				if (!TargetPlatform.PublishSymbols(StoreDir, Files, Parameters.IndexSources, SourceFiles,
+					Parameters.Product, Parameters.Branch, Parameters.Change, Parameters.BuildVersion))
 				{
 					throw new AutomationException("Failure publishing symbol files.");
 				}
diff --git a/Engine/Source/Programs/AutomationTool/Win/Tasks/SrcSrvTask.cs b/Engine/Source/Programs/AutomationTool/Win/Tasks/SrcSrvTask.cs
--- a/Engine/Source/Programs/AutomationTool/Win/Tasks/SrcSrvTask.cs
+++ b/Engine/Source/Programs/AutomationTool/Win/Tasks/SrcSrvTask.cs
@@ -1,22 +1,14 @@
 // Copyright Epic Games, Inc. All Rights Reserved.
 
-using System;
 using System.Collections.Generic;
 using System.Linq;
-using System.Text;
-using System.Threading.Tasks;
 using AutomationTool;
 using UnrealBuildTool;
-using System.IO;
-using System.Threading;
 using System.Xml;
-using System.Diagnostics;
 
-using Action = System.Action;
 using EpicGames.Core;
 using UnrealBuildBase;
 using EpicGames.BuildGraph;
-using AutomationTool.Tasks;
 using System.Runtime.Versioning;
 
 namespace Win.Automation
@@ -53,6 +45,11 @@ namespace Win.Automation
 
 	/// <summary>
 	/// Task which strips symbols from a set of files
+	/// Note that this task only supports source indexing for Windows-like platforms.
+	/// SymStore can be considered a generalization of this task because in addition to uploading symbols
+	/// it can also index sources and supports consoles as well (any missing platform that supports
+	/// source indexing in general may add support for it by extending PublishSymbols).
+	/// Check SymStoreTaskParameters.IndexSources/SourceFiles/Branch/Change
 	/// </summary>
 	[TaskElement("SrcSrv", typeof(SrcSrvTaskParameters))]
 	public class SrcSrvTask : CustomTask
@@ -88,90 +85,13 @@ namespace Win.Automation
 		[SupportedOSPlatform("windows")]
 		internal static void Execute(FileReference[] BinaryFiles, FileReference[] SourceFiles, string Branch, int Change)
 		{
-			FileReference[] PdbFiles = BinaryFiles.Where(x => x.HasExtension(".pdb")).ToArray();
+			IEnumerable<FileReference> PdbFiles = BinaryFiles.Where(x => x.HasExtension(".pdb"));
 
-			// Get the PDBSTR.EXE path, using the latest SDK version we can find.
-			FileReference PdbStrExe = GetPdbStrExe();
+			Win64Platform WindowsPlatform = Platform.GetPlatform(UnrealTargetPlatform.Win64) as Win64Platform;
 
-			// Get the path to the generated SRCSRV.INI file
-			FileReference SrcSrvIni = FileReference.Combine(Unreal.RootDirectory, "Engine", "Intermediate", "SrcSrv.ini");
-			DirectoryReference.CreateDirectory(SrcSrvIni.Directory);
-
-			// Generate the SRCSRV.INI file
-			using (StreamWriter Writer = new StreamWriter(SrcSrvIni.FullName))
-			{
-				Writer.WriteLine("SRCSRV: ini------------------------------------------------");
-				Writer.WriteLine("VERSION=1");
-				Writer.WriteLine("VERCTRL=Perforce");
-				Writer.WriteLine("SRCSRV: variables------------------------------------------");
-				Writer.WriteLine("SRCSRVTRG=%sdtrg%");
-				Writer.WriteLine("SRCSRVCMD=%sdcmd%");
-				Writer.WriteLine("SDCMD=p4.exe print -o %srcsrvtrg% \"{0}/%var2%@{1}\"", Branch.TrimEnd('/'), Change);
-				Writer.WriteLine("SDTRG=%targ%\\{0}\\{1}\\%fnbksl%(%var2%)", Branch.Replace('/', '+'), Change);
-				Writer.WriteLine("SRCSRV: source files ---------------------------------------");
-				foreach (FileReference SourceFile in SourceFiles)
-				{
-					string RelativeSourceFile = SourceFile.MakeRelativeTo(Unreal.RootDirectory);
-					Writer.WriteLine("{0}*{1}", SourceFile.FullName, RelativeSourceFile.Replace('\\', '/'));
-				}
-				Writer.WriteLine("SRCSRV: end------------------------------------------------");
-			}
-
-            // Execute PDBSTR on the PDB files in parallel.
-            Parallel.For(0, PdbFiles.Length, (Idx, State) => { ExecuteTool(PdbStrExe, PdbFiles[Idx], SrcSrvIni, State); });
+			WindowsPlatform.AddSourceIndexToSymbols(PdbFiles, SourceFiles, Branch, Change);
 		}
 
-        /// <summary>
-        /// Executes the PdbStr tool.
-        /// </summary>
-        /// <param name="PdbStrExe">Path to PdbStr.exe</param>
-        /// <param name="PdbFile">The PDB file to embed source information for</param>
-        /// <param name="SrcSrvIni">Ini file containing settings to embed</param>
-        /// <param name="State">The current loop state</param>
-        /// <returns>True if the tool executed succesfully</returns>
-        static void ExecuteTool(FileReference PdbStrExe, FileReference PdbFile, FileReference SrcSrvIni, ParallelLoopState State)
-        {
-            using (Process Process = new Process())
-            {
-                List<string> Messages = new List<string>();
-                Messages.Add(String.Format("Writing source server data: {0}", PdbFile));
-
-                DataReceivedEventHandler OutputHandler = (s, e) => { if (e.Data != null) { Messages.Add(e.Data); } };
-                Process.StartInfo.FileName = PdbStrExe.FullName;
-                Process.StartInfo.Arguments = String.Format("-w -p:\"{0}\" -i:\"{1}\" -s:srcsrv", PdbFile.FullName, SrcSrvIni.FullName);
-                Process.StartInfo.UseShellExecute = false;
-                Process.StartInfo.RedirectStandardOutput = true;
-                Process.StartInfo.RedirectStandardError = true;
-                Process.StartInfo.RedirectStandardInput = false;
-                Process.StartInfo.CreateNoWindow = true;
-                Process.OutputDataReceived += OutputHandler;
-                Process.ErrorDataReceived += OutputHandler;
-                Process.Start();
-                Process.BeginOutputReadLine();
-                Process.BeginErrorReadLine();
-                Process.WaitForExit();
-
-                if(Process.ExitCode != 0)
-                {
-                    Messages.Add(String.Format("Exit code: {0}", Process.ExitCode));
-                    State.Break();
-                }
-
-                lock (State)
-                {
-                    foreach (string Message in Messages)
-                    {
-                        CommandUtils.LogInformation(Message);
-                    }
-                }
-
-				if(Process.ExitCode != 0)
-				{
-					throw new AutomationException("Failed to embed source server data for {0}", PdbFile);
-				}
-            }
-        }
-
         /// <summary>
         /// Output this task out to an XML writer.
         /// </summary>
@@ -204,31 +124,6 @@ namespace Win.Automation
 		{
 			yield break;
 		}
-
-		/// <summary>
-		/// Try to get the PDBSTR.EXE path from the Windows SDK
-		/// </summary>
-		/// <returns>Path to PDBSTR.EXE</returns>
-		[SupportedOSPlatform("windows")]
-		public static FileReference GetPdbStrExe()
-		{
-			List<KeyValuePair<string, DirectoryReference>> WindowsSdkDirs = WindowsExports.GetWindowsSdkDirs();
-			foreach (DirectoryReference WindowsSdkDir in WindowsSdkDirs.Select(x => x.Value))
-			{
-				FileReference CheckPdbStrExe64 = FileReference.Combine(WindowsSdkDir, "Debuggers", "x64", "SrcSrv", "PdbStr.exe");
-				if (FileReference.Exists(CheckPdbStrExe64))
-				{
-					return CheckPdbStrExe64;
-				}
-
-				FileReference CheckPdbStrExe32 = FileReference.Combine(WindowsSdkDir, "Debuggers", "x86", "SrcSrv", "PdbStr.exe");
-				if (FileReference.Exists(CheckPdbStrExe32))
-				{
-					return CheckPdbStrExe32;
-				}
-			}
-			throw new AutomationException("Unable to find a Windows SDK installation containing PDBSTR.EXE");
-		}
 	}
 
 	public static partial class BgStateExtensions
diff --git a/Engine/Source/Programs/AutomationTool/Win/WinPlatform.Automation.cs b/Engine/Source/Programs/AutomationTool/Win/WinPlatform.Automation.cs
--- a/Engine/Source/Programs/AutomationTool/Win/WinPlatform.Automation.cs
+++ b/Engine/Source/Programs/AutomationTool/Win/WinPlatform.Automation.cs
@@ -12,7 +12,8 @@ using System.Diagnostics;
 using EpicGames.Core;
 using UnrealBuildBase;
 using System.Runtime.Versioning;
-
+using System.Threading.Tasks;
+using System.Reflection.PortableExecutable;
 
 public class Win64Platform : Platform
 {
@@ -461,10 +462,14 @@ public class Win64Platform : Platform
 	}
 
 	[SupportedOSPlatform("windows")]
-	public override bool PublishSymbols(DirectoryReference SymbolStoreDirectory, List<FileReference> Files, string Product, string BuildVersion = null)
+	public override bool PublishSymbols(DirectoryReference SymbolStoreDirectory, List<FileReference> Files,
+			bool bIndexSources, List<FileReference> SourceFiles,
+			string Product, string Branch, int Change, string BuildVersion = null)
     {
-        // Get the SYMSTORE.EXE path, using the latest SDK version we can find.
-        FileReference SymStoreExe = GetSymStoreExe();
+		LogInformation($"Publishing symbols to \"{SymbolStoreDirectory}\" (source indexing: {bIndexSources})");
+
+		// Get the SYMSTORE.EXE path, using the latest SDK version we can find.
+		FileReference SymStoreExe = GetSymStoreExe();
 
 		List<FileReference> FilesToAdd = Files.Where(x => x.HasExtension(".pdb") || x.HasExtension(".exe") || x.HasExtension(".dll")).ToList();
 		if(FilesToAdd.Count > 0)
@@ -478,10 +483,26 @@ public class Win64Platform : Platform
 				DirectoryReference.CreateDirectory(TempSymStoreDir);
 			}
 
+			DirectoryReference TempSymStoreIndexedDir = DirectoryReference.Combine(Unreal.RootDirectory, "Engine", "Intermediate", "SymStoreIndexed");
+
 			string TempFileName = Path.GetTempFileName();
 			try
 			{
-				File.WriteAllLines(TempFileName, FilesToAdd.Select(x => x.FullName), Encoding.ASCII);
+				IEnumerable<FileReference> SymbolsToIndex = Enumerable.Empty<FileReference>();
+				IEnumerable<FileReference> SymbolsAfterIndexing = Enumerable.Empty<FileReference>();
+
+				if (bIndexSources)
+				{
+					// Skip read-only PDBs as we won't be able to add data to them. They are likely to be symbols for third-party libraries checked into the source control.
+					SymbolsToIndex = FilesToAdd.Where(x => x.HasExtension(".pdb") && !(new FileInfo(x.FullName).IsReadOnly));
+
+					// We can't write to the original symbol files because they may have been generated by a different build step,
+					// and clobbering such build products is considered an error.
+					// For this reason, we copy the symbol files and modify the copied.
+					SymbolsAfterIndexing = CopySymbolsWithSourceIndexing(TempSymStoreIndexedDir, SymbolsToIndex, SourceFiles, Branch, Change);
+				}
+
+				File.WriteAllLines(TempFileName, FilesToAdd.Except(SymbolsToIndex).Union(SymbolsAfterIndexing).Select(x => x.FullName), Encoding.ASCII);
 
 				// Copy everything to the temp symstore
 				ProcessStartInfo StartInfo = new ProcessStartInfo();
@@ -497,6 +518,7 @@ public class Win64Platform : Platform
 			finally
 			{
 				File.Delete(TempFileName);
+				CommandUtils.DeleteDirectory(TempSymStoreIndexedDir);
 			}
 			DateTime CompressDone = DateTime.Now;
 			LogInformation("Took {0}s to compress the symbol files to temp path {1}", (CompressDone - Start).TotalSeconds, TempSymStoreDir);
@@ -582,6 +604,43 @@ public class Win64Platform : Platform
 		return true;
     }
 
+	[SupportedOSPlatform("windows")]
+	private IEnumerable<FileReference> CopySymbolsWithSourceIndexing(DirectoryReference TargetDirectory, IEnumerable<FileReference> SymbolFiles,
+	IEnumerable<FileReference> SourceFiles, string Branch, int Change)
+	{
+		if (DirectoryReference.Exists(TargetDirectory))
+		{
+			CommandUtils.DeleteDirectory(TargetDirectory);
+			DirectoryReference.CreateDirectory(TargetDirectory);
+		}
+
+		List<FileReference> CopiedSymbolFiles = new List<FileReference>();
+
+		// Copy all the symbol files to the given directory.
+		foreach (FileReference File in SymbolFiles)
+		{
+			string RelativePath = File.MakeRelativeTo(Unreal.RootDirectory);
+
+			FileReference DestinationFile = FileReference.Combine(TargetDirectory, RelativePath);
+
+			try
+			{
+				CommandUtils.CopyFile(File.FullName, DestinationFile.FullName);
+			}
+			catch (Exception Ex)
+			{
+				throw new AutomationException("Couldn't copy the pdb file to the temp directory for indexing! Reason: {0}", Ex.ToString());
+			}
+
+			CopiedSymbolFiles.Add(DestinationFile);
+		}
+
+		// Index all symbol files in one go (indexing source code from the source control is shared work between all pdbs).
+		AddSourceIndexToSymbols(CopiedSymbolFiles, SourceFiles, Branch, Change);
+
+		return CopiedSymbolFiles;
+	}
+
 	bool IsSymbolFile(FileReference File)
 	{
 		if (File.HasExtension(".dll") || File.HasExtension(".exe") || File.HasExtension(".pdb"))
@@ -595,6 +654,144 @@ public class Win64Platform : Platform
 		return false;
 	}
 
+	/// <summary>
+	/// 
+	/// </summary>
+	/// <param name="PdbFiles"></param>
+	/// <param name="SourceFiles"></param>
+	/// <param name="Branch"></param>
+	/// <param name="Change"></param>
+	[SupportedOSPlatform("windows")]
+	public void AddSourceIndexToSymbols(IEnumerable<FileReference> PdbFiles, IEnumerable<FileReference> SourceFiles, string Branch, int Change)
+	{
+		LogInformation("Adding source control information to PDB files...");
+
+		// Get the PDBSTR.EXE path, using the latest SDK version we can find.
+		FileReference PdbStrExe = GetPdbStrExe();
+
+		// Get the path to the generated SRCSRV.INI file
+		FileReference SrcSrvIni = FileReference.Combine(Unreal.RootDirectory, "Engine", "Intermediate", "SrcSrv.ini");
+		DirectoryReference.CreateDirectory(SrcSrvIni.Directory);
+
+		// Generate the SRCSRV.INI file
+		using (StreamWriter Writer = new StreamWriter(SrcSrvIni.FullName))
+		{
+			Writer.WriteLine("SRCSRV: ini------------------------------------------------");
+			Writer.WriteLine("VERSION=1");
+			Writer.WriteLine("VERCTRL=Perforce");
+			Writer.WriteLine("SRCSRV: variables------------------------------------------");
+			Writer.WriteLine("SRCSRVTRG=%sdtrg%");
+			Writer.WriteLine("SRCSRVCMD=%sdcmd%");
+			Writer.WriteLine("SDCMD=p4.exe print -o %srcsrvtrg% \"{0}/%var2%@{1}\"", Branch.TrimEnd('/'), Change);
+			Writer.WriteLine("SDTRG=%targ%\\{0}\\{1}\\%fnbksl%(%var2%)", Branch.Replace('/', '+'), Change);
+			Writer.WriteLine("SRCSRV: source files ---------------------------------------");
+			foreach (FileReference SourceFile in SourceFiles)
+			{
+				string RelativeSourceFile = SourceFile.MakeRelativeTo(Unreal.RootDirectory);
+				Writer.WriteLine("{0}*{1}", SourceFile.FullName, RelativeSourceFile.Replace('\\', '/'));
+			}
+			Writer.WriteLine("SRCSRV: end------------------------------------------------");
+		}
+
+		// Execute PDBSTR on the PDB files in parallel.
+		Parallel.ForEach(PdbFiles, (PdbFile, State) => { ExecutePdbStrTool(PdbStrExe, PdbFile, SrcSrvIni, State); });
+	}
+
+	/// <summary>
+	/// Executes the PdbStr tool.
+	/// </summary>
+	/// <param name="PdbStrExe">Path to PdbStr.exe</param>
+	/// <param name="PdbFile">The PDB file to embed source information for</param>
+	/// <param name="SrcSrvIni">Ini file containing settings to embed</param>
+	/// <param name="State">The current loop state</param>
+	/// <returns>True if the tool executed successfully</returns>
+	static void ExecutePdbStrTool(FileReference PdbStrExe, FileReference PdbFile, FileReference SrcSrvIni, ParallelLoopState State)
+	{
+		FileInfo PdbInfo = new FileInfo(PdbFile.FullName);
+		FileInfo IniInfo = new FileInfo(SrcSrvIni.FullName);
+
+		const long Size4GB = 4L * 1024 * 1024 * 1024;
+
+		// PdbStr.exe tool can't handle pdb's larger than 4GB (linked with /PDBPAGESIZE:8192 switch).
+		// It's a known limitation and it may happen for large game projects, especially in development configurations.
+		// We choose to emit a warning for this known issue rather than failing the whole build with an exception, if we run into this case.
+		if (PdbInfo.Length + IniInfo.Length >= Size4GB)
+		{
+			lock (State)
+			{
+				CommandUtils.LogWarning("Skipping indexing of {0} because it would make the file larger than 4GB (not supported by pdbstr.exe tool).", PdbFile.GetFileName());
+			}
+
+			return;
+		}
+
+		using (Process Process = new Process())
+		{
+			List<string> Messages = new List<string>();
+
+			Messages.Add(String.Format("Writing source server data: {0}", PdbFile));
+
+			DataReceivedEventHandler OutputHandler = (s, e) => { if (e.Data != null) { Messages.Add(e.Data); } };
+			Process.StartInfo.FileName = PdbStrExe.FullName;
+			Process.StartInfo.Arguments = String.Format("-w -p:\"{0}\" -i:\"{1}\" -s:srcsrv", PdbFile.FullName, SrcSrvIni.FullName);
+			Process.StartInfo.UseShellExecute = false;
+			Process.StartInfo.RedirectStandardOutput = true;
+			Process.StartInfo.RedirectStandardError = true;
+			Process.StartInfo.RedirectStandardInput = false;
+			Process.StartInfo.CreateNoWindow = true;
+			Process.OutputDataReceived += OutputHandler;
+			Process.ErrorDataReceived += OutputHandler;
+			Process.Start();
+			Process.BeginOutputReadLine();
+			Process.BeginErrorReadLine();
+			Process.WaitForExit();
+
+			if (Process.ExitCode != 0)
+			{
+				Messages.Add(String.Format("Exit code: {0}", Process.ExitCode));
+				State.Break();
+			}
+
+			lock (State)
+			{
+				foreach (string Message in Messages)
+				{
+					CommandUtils.LogInformation(Message);
+				}
+			}
+
+			if (Process.ExitCode != 0)
+			{
+				throw new AutomationException("Failed to embed source server data for {0}", PdbFile);
+			}
+		}
+	}
+
+	/// <summary>
+	/// Try to get the PDBSTR.EXE path from the Windows SDK
+	/// </summary>
+	/// <returns>Path to PDBSTR.EXE</returns>
+	[SupportedOSPlatform("windows")]
+	static FileReference GetPdbStrExe()
+	{
+		List<KeyValuePair<string, DirectoryReference>> WindowsSdkDirs = WindowsExports.GetWindowsSdkDirs();
+		foreach (DirectoryReference WindowsSdkDir in WindowsSdkDirs.Select(x => x.Value))
+		{
+			FileReference CheckPdbStrExe64 = FileReference.Combine(WindowsSdkDir, "Debuggers", "x64", "SrcSrv", "PdbStr.exe");
+			if (FileReference.Exists(CheckPdbStrExe64))
+			{
+				return CheckPdbStrExe64;
+			}
+
+			FileReference CheckPdbStrExe32 = FileReference.Combine(WindowsSdkDir, "Debuggers", "x86", "SrcSrv", "PdbStr.exe");
+			if (FileReference.Exists(CheckPdbStrExe32))
+			{
+				return CheckPdbStrExe32;
+			}
+		}
+		throw new AutomationException("Unable to find a Windows SDK installation containing PDBSTR.EXE");
+	}
+
 	public override string[] SymbolServerDirectoryStructure
     {
         get
