Fixes crash when PLINK or PSCP are inside of a folder path which contains a space

diff --git a/Engine/Source/Programs/AutomationTool/Linux/LinuxPlatform.Automation.cs b/Engine/Source/Programs/AutomationTool/Linux/LinuxPlatform.Automation.cs
--- a/Engine/Source/Programs/AutomationTool/Linux/LinuxPlatform.Automation.cs
+++ b/Engine/Source/Programs/AutomationTool/Linux/LinuxPlatform.Automation.cs
@@ -17,8 +17,8 @@ public abstract class BaseLinuxPlatform : Platform
 	// Matches strings of the form "DeviceName@IP Address" such as "[REDACTED_EMAIL]"
 	static Regex DeviceRegex = new Regex(@"\w.+@([A-Za-z0-9\.\-]+)[\+]?");
 
-	static string PScpPath = CombinePaths(Unreal.RootDirectory.FullName, "\\Engine\\Extras\\ThirdPartyNotUE\\putty\\PSCP.EXE");
-	static string PlinkPath = CombinePaths(Unreal.RootDirectory.FullName, "\\Engine\\Extras\\ThirdPartyNotUE\\putty\\PLINK.EXE");
+	static string PScpPath  = MakePathSafeToUseWithCommandLine(CombinePaths(Unreal.RootDirectory.FullName, "Engine", "Extras", "ThirdPartyNotUE", "putty", "PSCP.EXE"));
+	static string PlinkPath = MakePathSafeToUseWithCommandLine(CombinePaths(Unreal.RootDirectory.FullName, "Engine", "Extras", "ThirdPartyNotUE", "putty", "PLINK.EXE"));
 	static string LaunchOnHelperShellScriptName = "LaunchOnHelper.sh";
 
 	public BaseLinuxPlatform(UnrealTargetPlatform P)
diff --git a/Engine/Source/Programs/AutomationTool/Scripts/RunProjectCommand.Automation.cs b/Engine/Source/Programs/AutomationTool/Scripts/RunProjectCommand.Automation.cs
--- a/Engine/Source/Programs/AutomationTool/Scripts/RunProjectCommand.Automation.cs
+++ b/Engine/Source/Programs/AutomationTool/Scripts/RunProjectCommand.Automation.cs
@@ -868,8 +868,8 @@ namespace AutomationScripts
 			{
 				ServerApp = @"C:\Windows\system32\cmd.exe";
 
-				string plinkPath = CombinePaths(Unreal.EngineDirectory.FullName, "Extras", "ThirdPartyNotUE", "putty", "PLINK.exe");
-				string exePath = CombinePaths(SC.ShortProjectName, "Binaries", ServerPlatformDesc.Type.ToString(), SC.ShortProjectName + "Server");
+				string plinkPath = MakePathSafeToUseWithCommandLine(CombinePaths(Unreal.RootDirectory.FullName, "Engine", "Extras", "ThirdPartyNotUE", "putty", "PLINK.exe"));
+				string exePath = MakePathSafeToUseWithCommandLine(CombinePaths(SC.ShortProjectName, "Binaries", ServerPlatformDesc.Type.ToString(), SC.ShortProjectName + "Server"));
 				if (ServerParams.ServerConfigsToBuild[0] != UnrealTargetConfiguration.Development)
 				{
 					exePath += "-" + ServerPlatformDesc.Type.ToString() + "-" + ServerParams.ServerConfigsToBuild[0].ToString();
