do not fail a -noclient -server build if there are multiple game build targets.

 - this matches the logic when running UAT with -target=myserver etc.

diff --git a/Engine/Source/Programs/AutomationTool/AutomationUtils/ProjectParams.cs b/Engine/Source/Programs/AutomationTool/AutomationUtils/ProjectParams.cs
--- a/Engine/Source/Programs/AutomationTool/AutomationUtils/ProjectParams.cs
+++ b/Engine/Source/Programs/AutomationTool/AutomationUtils/ProjectParams.cs
@@ -2435,7 +2435,7 @@ namespace AutomationTool
 					GameTarget = ChooseTarget(AvailableClientTargets, TargetType.Client);
 					ProjectType = TargetType.Client;
 				}
-				else if (AvailableGameTargets.Count > 0)
+				else if (AvailableGameTargets.Count > 0 && !NoClient)
 				{
 					if (AvailableGameTargets.Count > 1)
 					{
