Always take the last member of the string for build command
+ add more arguments to filter out

#rnx

diff --git a/Engine/Source/Programs/AutomationTool/Gauntlet/Framework/Gauntlet.TestExecutor.cs b/Engine/Source/Programs/AutomationTool/Gauntlet/Framework/Gauntlet.TestExecutor.cs
--- a/Engine/Source/Programs/AutomationTool/Gauntlet/Framework/Gauntlet.TestExecutor.cs
+++ b/Engine/Source/Programs/AutomationTool/Gauntlet/Framework/Gauntlet.TestExecutor.cs
@@ -105,7 +105,7 @@ namespace Gauntlet
 		/// </summary>
 		public TestExecutor(string InBuildCommand)
 		{
-			BuildCommandThatLaunchedExecutor = InBuildCommand;
+			BuildCommandThatLaunchedExecutor = InBuildCommand.Split(".").Last();
 			RunningTests = new List<TestExecutionInfo>();
 		}
 
diff --git a/Engine/Source/Programs/AutomationTool/Gauntlet/Unreal/Base/Gauntlet.UnrealTestNode.cs b/Engine/Source/Programs/AutomationTool/Gauntlet/Unreal/Base/Gauntlet.UnrealTestNode.cs
--- a/Engine/Source/Programs/AutomationTool/Gauntlet/Unreal/Base/Gauntlet.UnrealTestNode.cs
+++ b/Engine/Source/Programs/AutomationTool/Gauntlet/Unreal/Base/Gauntlet.UnrealTestNode.cs
@@ -864,7 +864,7 @@ namespace Gauntlet
 		public override string GetRunLocalCommand(string LaunchingBuildCommand)
 		{
 			string[] ArgsToNotDisplay =
-			{
+			new List<string>{
 					"test",
 					"tests",
 					"tempdir",
@@ -875,11 +875,20 @@ namespace Gauntlet
 					"RecordDeviceUsage",
 					"uploadreport",
 					"skipdashboardsubmit",
-			};
+					"ReportURL",
+					"ReportExportPath",
+					"WriteTestResultsForHorde",
+					"HordeTestDataKey",
+					"PublishTelemetryTo",
+					"ECBranch",
+					"ECChangelist",
+					"PreFlightChange",
+			}.Select(I => I.ToLower()).ToArray();
 
 			bool ShouldArgBeDisplayed(string InArg)
 			{
-				return !ArgsToNotDisplay.Any(str => InArg.StartsWith(str));
+				InArg = InArg.Split("=", 2).First().ToLower();
+				return !ArgsToNotDisplay.Any(str => InArg == str);
 			}
 			
 			string WrapParameterInQuotes(string InString)
