Build: Symbol Server: Downgraded failure to add source indexing to a pdb from a breaking exception to a log message (Windows-based platforms).

PdbStr.exe tool that is used to add source information to .pdb can't handle large (> 4 BG) symbol files and there isn't much we can do about it. Some of our builds produce such symbols which leads to whole build tasks being failed (and symbols not being stored at all).

I tried to predice such cases by checking the size of .pdb before trying to invoke PdbStr.exe and it works in some cases, but apparently it's not precise enough because we have cases in our build farm which pass my check but fail later on when PdbStr.exe is started.

#rnx

diff --git a/Engine/Source/Programs/AutomationTool/Win/WinPlatform.Automation.cs b/Engine/Source/Programs/AutomationTool/Win/WinPlatform.Automation.cs
--- a/Engine/Source/Programs/AutomationTool/Win/WinPlatform.Automation.cs
+++ b/Engine/Source/Programs/AutomationTool/Win/WinPlatform.Automation.cs
@@ -748,8 +748,7 @@ public class Win64Platform : Platform
 
 			if (Process.ExitCode != 0)
 			{
-				Messages.Add(String.Format("Exit code: {0}", Process.ExitCode));
-				State.Break();
+				Messages.Add($"Failed to embed source server data for {PdbFile} (exit code: {Process.ExitCode})");
 			}
 
 			lock (State)
@@ -759,11 +758,6 @@ public class Win64Platform : Platform
 					CommandUtils.LogInformation(Message);
 				}
 			}
-
-			if (Process.ExitCode != 0)
-			{
-				throw new AutomationException("Failed to embed source server data for {0}", PdbFile);
-			}
 		}
 	}
 
