Update Smartling localization to use Smartling namespaces to reduce duplicated strings and leverage the string sharing feature for Smartling.

diff --git a/Engine/Source/Programs/AutomationTool/SmartlingLocalization/SmartlingLocalizationProvider.cs b/Engine/Source/Programs/AutomationTool/SmartlingLocalization/SmartlingLocalizationProvider.cs
--- a/Engine/Source/Programs/AutomationTool/SmartlingLocalization/SmartlingLocalizationProvider.cs
+++ b/Engine/Source/Programs/AutomationTool/SmartlingLocalization/SmartlingLocalizationProvider.cs
@@ -256,6 +256,10 @@ namespace EpicGames.SmartlingLocalization
 					UploadMultipartFormDataContent.Add(UploadFileStreamContent, "file", SmartlingFileUri);	
 					UploadMultipartFormDataContent.Add(new StringContent(SmartlingFileUri, Encoding.UTF8, "application/json"), "fileUri");
 					UploadMultipartFormDataContent.Add(new StringContent("gettext", Encoding.UTF8, "application/json"), "fileType");
+					// We introduce a Smartling namespace to leverage Smartling's string sharing feature
+					// Following Smartling best practices, we make the Smartling namespace the same as the full file path of the file.
+					// https://help.smartling.com/hc/en-us/articles/360008143833-String-Sharing-and-Namespaces-via-Smartling-API
+					UploadMultipartFormDataContent.Add(new StringContent(SmartlingFilename, Encoding.UTF8, "application/json"), "smartling.namespace");
 					var UploadResponse = await Client.PostAsync(UploadEndpoint, UploadMultipartFormDataContent);
 					if (UploadResponse.IsSuccessStatusCode)
 					{
