Log Suspend/Resume log parsing through the legacy logger

diff --git a/Engine/Source/Programs/AutomationTool/Gauntlet/Framework/Gauntlet.Utils.cs b/Engine/Source/Programs/AutomationTool/Gauntlet/Framework/Gauntlet.Utils.cs
--- a/Engine/Source/Programs/AutomationTool/Gauntlet/Framework/Gauntlet.Utils.cs
+++ b/Engine/Source/Programs/AutomationTool/Gauntlet/Framework/Gauntlet.Utils.cs
@@ -267,7 +267,7 @@ namespace Gauntlet
 			{
 				if (CommandUtils.IsBuildMachine)
 				{
-					OutputMessage("<-- Suspend Log Parsing -->", KnownLogEvents.Gauntlet);
+					CommandUtils.LogInformation("<-- Suspend Log Parsing -->");
 				}
 			}
 		}
@@ -278,7 +278,7 @@ namespace Gauntlet
 			{
 				if (CommandUtils.IsBuildMachine)
 				{
-					OutputMessage("<-- Resume Log Parsing -->", KnownLogEvents.Gauntlet);
+					CommandUtils.LogInformation("<-- Resume Log Parsing -->");
 				}
 			}
 		}
