- Removed special case Android logic, because it was adding /arm64/ into the path that was being looked for with InstalledBuilds. This change was tested in 5.1 stream with Android plugins

diff --git a/Engine/Source/Programs/AutomationTool/Scripts/BuildPluginCommand.Automation.cs b/Engine/Source/Programs/AutomationTool/Scripts/BuildPluginCommand.Automation.cs
--- a/Engine/Source/Programs/AutomationTool/Scripts/BuildPluginCommand.Automation.cs
+++ b/Engine/Source/Programs/AutomationTool/Scripts/BuildPluginCommand.Automation.cs
@@ -314,10 +314,6 @@ public sealed class BuildPlugin : BuildCommand
 				ManifestFileNames.Add(ManifestFileName);
 				
 				string Arguments = String.Format("-plugin={0} -iwyu -noubtmakefiles -manifest={1} -nohotreload", CommandUtils.MakePathSafeToUseWithCommandLine(HostProjectPluginFile.FullName), CommandUtils.MakePathSafeToUseWithCommandLine(ManifestFileName.FullName));
-				if (Platform == UnrealTargetPlatform.Android)
-				{
-					Arguments += String.Format(" -architectures={0}", AndroidArchitectures);
-				}
 
 				if (!String.IsNullOrEmpty(InAdditionalArgs))
 				{
@@ -367,10 +363,6 @@ public sealed class BuildPlugin : BuildCommand
 				{
 					Arguments += String.Format(" -architecture={0}", SpecifiedArchitecture);
 				}
-				else if (Platform == UnrealTargetPlatform.Android)
-				{
-					Arguments += String.Format(" -architecture={0}", AndroidArchitectures);
-				}
 
 				if (!String.IsNullOrEmpty(InAdditionalArgs))
 				{
