Build: Symbol Server: Downgraded a warning about a PDB being larger than 4GB to an information log.

This message should really be a warning, but given we have no way of fixing it (it's a limitation of pdbstr.exe tool from Windows SDK) there's not much use of having it as it just prevents our build machines from being green.

#rnx

diff --git a/Engine/Source/Programs/AutomationTool/Win/WinPlatform.Automation.cs b/Engine/Source/Programs/AutomationTool/Win/WinPlatform.Automation.cs
--- a/Engine/Source/Programs/AutomationTool/Win/WinPlatform.Automation.cs
+++ b/Engine/Source/Programs/AutomationTool/Win/WinPlatform.Automation.cs
@@ -719,7 +719,7 @@ public class Win64Platform : Platform
 		{
 			lock (State)
 			{
-				CommandUtils.LogWarning("Skipping indexing of {0} because it would make the file larger than 4GB (not supported by pdbstr.exe tool).", PdbFile.GetFileName());
+				CommandUtils.LogInformation("Skipping indexing of {0} because it would make the file larger than 4GB (not supported by pdbstr.exe tool).", PdbFile.GetFileName());
 			}
 
 			return;
