- Added a new callback to staging that allow for finalization of staged files
- Added a default implementation of GetCookPlatform, which allows BuildCookRun to run without needing any platform-specific automation DLLs (cook and staging can happen, but running/packaging would still need the platform automation dlls)

diff --git a/Engine/Source/Programs/AutomationTool/AutomationUtils/Platform.cs b/Engine/Source/Programs/AutomationTool/AutomationUtils/Platform.cs
--- a/Engine/Source/Programs/AutomationTool/AutomationUtils/Platform.cs
+++ b/Engine/Source/Programs/AutomationTool/AutomationUtils/Platform.cs
@@ -472,7 +472,11 @@ namespace AutomationTool
 		/// <returns>Cook platform string.</returns>
 		public virtual string GetCookPlatform(bool bDedicatedServer, bool bIsClientOnly)
 		{
-			throw new AutomationException("{0} does not yet implement GetCookPlatform.", PlatformType);
+			// this should get all cases, but a platform can override if needed
+
+			string Suffix = bIsClientOnly ? "Client" : bDedicatedServer ? "Server" : "";
+			string PlatformName = GetGenericPlatformName(TargetPlatformType);
+			return $"{PlatformName}{Suffix}";
 		}
 
 		/// <summary>
diff --git a/Engine/Source/Programs/AutomationTool/AutomationUtils/ProjectParams.cs b/Engine/Source/Programs/AutomationTool/AutomationUtils/ProjectParams.cs
--- a/Engine/Source/Programs/AutomationTool/AutomationUtils/ProjectParams.cs
+++ b/Engine/Source/Programs/AutomationTool/AutomationUtils/ProjectParams.cs
@@ -1965,15 +1965,20 @@ namespace AutomationTool
 		public Action<ProjectParams, DeploymentContext> PreModifyDeploymentContextCallback = null;
 
 		/// <summary>
-		/// Stage: Optional callback that a build script can use to modify a deployment context before it is applied
+		/// Stage: Optional callback that a build script can use to modify a deployment context before it is applied (and before it is finalized)
 		/// </summary>
 		public Action<ProjectParams, DeploymentContext> ModifyDeploymentContextCallback = null;
 
-        /// <summary>
-        /// On Windows, adds an executable to the root of the staging directory which checks for prerequisites being 
+		/// <summary>
+		/// Stage: Optional callback that a build script can use to finalize a deployment context before it is applied
+		/// </summary>
+		public Action<ProjectParams, DeploymentContext> FinalizeDeploymentContextCallback = null;
+
+		/// <summary>
+		/// On Windows, adds an executable to the root of the staging directory which checks for prerequisites being 
 		/// installed and launches the game with a path to the .uproject file.
 		/// </summary>
-        public bool NoBootstrapExe { get; set; }
+		public bool NoBootstrapExe { get; set; }
 
 		/// <summary>
 		/// By default we don't code sign unless it is required or requested
diff --git a/Engine/Source/Programs/AutomationTool/CookedEditor/CookedEditor.Automation.cs b/Engine/Source/Programs/AutomationTool/CookedEditor/CookedEditor.Automation.cs
--- a/Engine/Source/Programs/AutomationTool/CookedEditor/CookedEditor.Automation.cs
+++ b/Engine/Source/Programs/AutomationTool/CookedEditor/CookedEditor.Automation.cs
@@ -743,6 +743,10 @@ public class MakeCookedEditor : BuildCommand
 		CookedEditorStageDirectory = SC.StageDirectory;
 	}
 
+	protected virtual void FinalizeDeploymentContext(ProjectParams Params, DeploymentContext SC)
+	{
+	}
+
 	protected virtual void SetupDLCMode(FileReference ProjectFile, out string DLCName, out string ReleaseVersion, out TargetType Type)
 	{
 		bool bBuildAgainstRelease = ConfigHelper.GetBool("bBuildAgainstRelease");
@@ -915,7 +919,7 @@ public class MakeCookedEditor : BuildCommand
 					throw new AutomationException($"Unable to stage directories with \"Dest\" setting for CookedEditor: '{Entry}'");
 				}
 
-				// now stage it to a different location a specified in the params
+				// now stage it to a different location as specified in the params
 				StagedFileType FileType = (FileList == Context.NonUFSFilesToStage) ? StagedFileType.NonUFS : StagedFileType.UFS;
 				SC.StageFile(FileType, SourceFile, new StagedFileReference(Props["Dest"]));
 				continue;
@@ -1110,8 +1114,9 @@ public class MakeCookedEditor : BuildCommand
 		}
 
 		// set up override functions
-		Params.PreModifyDeploymentContextCallback = new Action<ProjectParams, DeploymentContext>((ProjectParams P, DeploymentContext SC) => { PreModifyDeploymentContext(P, SC); });
-		Params.ModifyDeploymentContextCallback = new Action<ProjectParams, DeploymentContext>((ProjectParams P, DeploymentContext SC) => { ModifyDeploymentContext(P, SC); });
+		Params.PreModifyDeploymentContextCallback = (P, SC) => PreModifyDeploymentContext(P, SC);
+		Params.ModifyDeploymentContextCallback = (P, SC) => ModifyDeploymentContext(P, SC);
+		Params.FinalizeDeploymentContextCallback = (P, SC) => FinalizeDeploymentContext(P, SC);
 
 		// this will make all of the files that are specified in BUild.cs files with "AdditionalPropertiesForReceipt.Add("CookerSupportFiles", ...);"  be copied into this
 		// subdirectory along with a batch file that can be used to set platform SDK environment variables during cooking
diff --git a/Engine/Source/Programs/AutomationTool/Scripts/CopyBuildToStagingDirectory.Automation.cs b/Engine/Source/Programs/AutomationTool/Scripts/CopyBuildToStagingDirectory.Automation.cs
--- a/Engine/Source/Programs/AutomationTool/Scripts/CopyBuildToStagingDirectory.Automation.cs
+++ b/Engine/Source/Programs/AutomationTool/Scripts/CopyBuildToStagingDirectory.Automation.cs
@@ -1,4 +1,4 @@
-﻿// Copyright Epic Games, Inc. All Rights Reserved.
+// Copyright Epic Games, Inc. All Rights Reserved.
 
 using AutomationTool;
 using System;
@@ -1518,6 +1518,13 @@ namespace AutomationScripts
 			{
 				SC.FilesToStage.NonUFSFiles[Pair.Key] = Pair.Value;
 			}
+
+			// Allow the calling scripts to make modifications to the deployment context before we finalize it
+			if (Params.FinalizeDeploymentContextCallback != null)
+			{
+				Params.FinalizeDeploymentContextCallback(Params, SC);
+			}
+
 		}
 
 		/// <summary>
