Update Smartling integration to filter out UE FText placeholders.
The placeholders we use for FText are of the form {0} for positional argumens and {MyVariableName} as named arguments.
Smartling will flag those as warnings as potential placeholders and need to be manually reviewed. This feature will fix all the warnings and have all files to be uploaded have the directive applied to them.
#test Ran the BuildGraph command locally to make sure that UAt is able to compile properly

diff --git a/Engine/Source/Programs/AutomationTool/SmartlingLocalization/SmartlingLocalizationProvider.cs b/Engine/Source/Programs/AutomationTool/SmartlingLocalization/SmartlingLocalizationProvider.cs
--- a/Engine/Source/Programs/AutomationTool/SmartlingLocalization/SmartlingLocalizationProvider.cs
+++ b/Engine/Source/Programs/AutomationTool/SmartlingLocalization/SmartlingLocalizationProvider.cs
@@ -260,6 +260,10 @@ namespace EpicGames.SmartlingLocalization
 					// Following Smartling best practices, we make the Smartling namespace the same as the full file path of the file.
 					// https://help.smartling.com/hc/en-us/articles/360008143833-String-Sharing-and-Namespaces-via-Smartling-API
 					UploadMultipartFormDataContent.Add(new StringContent(SmartlingFilename, Encoding.UTF8, "application/json"), "smartling.namespace");
+					// all placeholder values in the source files will raise warnings in Smartling which need to be manually reviewed.
+					// This allows the warnings to be resolved. The regex designates anything within {} to be a placeholder value.
+					// This accounts for the common FText formatted string placeholders of {0} or {MyPlaceholderVariable}
+					UploadMultipartFormDataContent.Add(new StringContent("\\{([^}]+)\\}", Encoding.UTF8, "application/json"), "smartling.placeholder_format_custom");
 					var UploadResponse = await Client.PostAsync(UploadEndpoint, UploadMultipartFormDataContent);
 					if (UploadResponse.IsSuccessStatusCode)
 					{
