Automation: Added an option to CommandUtils.Run to correctly deal with tools writing wide-char text to the standard output.

While the command prompt always shows such text correctly, the issue becomes apparent when redirecting the output to a file as text would appear as:

T h i s   i s   s o m e   u n i c o d e   o u t p u t

#rnx

diff --git a/Engine/Source/Programs/AutomationTool/AutomationUtils/ProcessUtils.cs b/Engine/Source/Programs/AutomationTool/AutomationUtils/ProcessUtils.cs
--- a/Engine/Source/Programs/AutomationTool/AutomationUtils/ProcessUtils.cs
+++ b/Engine/Source/Programs/AutomationTool/AutomationUtils/ProcessUtils.cs
@@ -740,6 +740,10 @@ namespace AutomationTool
 			/// </summary>
 			NoStdOutRedirect = 1 << 3,
 			NoLoggingOfRunCommand = 1 << 4,
+
+			/// <summary>
+			/// Output of the spawned process is expected to be encoded as UTF-8.
+			/// </summary>
 			UTF8Output = 1 << 5,
 
 			/// When specified with AllowSpew, the output will be TraceEventType.Verbose instead of TraceEventType.Information
@@ -759,6 +763,11 @@ namespace AutomationTool
 			/// </summary>
 			NoStdOutCapture = 1 << 9,
 
+			/// <summary>
+			/// Output of the spawned process is expected to be encoded as UTF-16 (System.Text.UnicodeEncoding, in .NET terminology)
+			/// </summary>
+			UTF16Output = 1 << 10,
+
 			Default = AllowSpew | AppMustExist,
 		}
 
@@ -892,6 +901,10 @@ namespace AutomationTool
 				{
 					Proc.StartInfo.StandardOutputEncoding = new System.Text.UTF8Encoding(false, false);
 				}
+				else if ((Options & ERunOptions.UTF16Output) == ERunOptions.UTF16Output)
+				{
+					Proc.StartInfo.StandardOutputEncoding = new System.Text.UnicodeEncoding(false, false, false);
+				}
 				Proc.Start();
 
 				if (bRedirectStdOut)
