Actually make the Namespace list as a string for log output

diff --git a/Engine/Source/Programs/AutomationTool/Gauntlet/Framework/Gauntlet.Utils.cs b/Engine/Source/Programs/AutomationTool/Gauntlet/Framework/Gauntlet.Utils.cs
--- a/Engine/Source/Programs/AutomationTool/Gauntlet/Framework/Gauntlet.Utils.cs
+++ b/Engine/Source/Programs/AutomationTool/Gauntlet/Framework/Gauntlet.Utils.cs
@@ -645,7 +645,7 @@ namespace Gauntlet
 					return MatchingTypes.OrderBy(M => M.FullName.Split('.').Length).First();
 				}
 
-				throw new AutomationException("Unable to find type {0} in assemblies. Namespaces= {1}.", TestName, Namespaces);
+				throw new AutomationException("Unable to find type {0} in assemblies. Namespaces={1}.", TestName, string.Join(", ", Namespaces));
 			}
 
 
