Improve Android device handling:
* Suspend log parsing when power on or off device
* Retry on disconnection when copy via adb push fails for disconnection
* Let DoAutoRenew() trap exceptions when renewing and deleting device

diff --git a/Engine/Source/Programs/AutomationTool/AutomationUtils/DeviceReservation.cs b/Engine/Source/Programs/AutomationTool/AutomationUtils/DeviceReservation.cs
--- a/Engine/Source/Programs/AutomationTool/AutomationUtils/DeviceReservation.cs
+++ b/Engine/Source/Programs/AutomationTool/AutomationUtils/DeviceReservation.cs
@@ -357,27 +357,13 @@ namespace AutomationTool.DeviceReservation
 
 		public void Renew(Uri BaseUri, TimeSpan NewDuration)
 		{
-			try
-			{
-				Utils.InvokeAPI(BaseUri.AppendPath("api/v1/reservations/" + Guid.ToString()), "PUT", NewDuration);
-			}
-			catch (Exception ex)
-			{
-				throw new AutomationException(ex, "Failed to renew device reservation.");
-			}
+			Utils.InvokeAPI(BaseUri.AppendPath("api/v1/reservations/" + Guid.ToString()), "PUT", NewDuration);
 		}
 
 		public void Delete(Uri BaseUri)
 		{
-			try
-			{
-				Utils.InvokeAPI(BaseUri.AppendPath("api/v1/reservations/" + Guid.ToString()), "DELETE");
-				Console.WriteLine("Successfully deleted device reservation \"{0}\".", Guid);
-			}
-			catch (Exception ex)
-			{
-				Utils.Log(string.Format("Failed to delete device reservation: {0}", ex.Message));
-			}
+			Utils.InvokeAPI(BaseUri.AppendPath("api/v1/reservations/" + Guid.ToString()), "DELETE");
+			Console.WriteLine("Successfully deleted device reservation \"{0}\".", Guid);
 		}
 
 		static public void ReportDeviceError(string InBaseUri, string DeviceName, string Error)
diff --git a/Engine/Source/Programs/AutomationTool/Gauntlet/Platform/Android/Gauntlet.TargetDeviceAndroid.cs b/Engine/Source/Programs/AutomationTool/Gauntlet/Platform/Android/Gauntlet.TargetDeviceAndroid.cs
--- a/Engine/Source/Programs/AutomationTool/Gauntlet/Platform/Android/Gauntlet.TargetDeviceAndroid.cs
+++ b/Engine/Source/Programs/AutomationTool/Gauntlet/Platform/Android/Gauntlet.TargetDeviceAndroid.cs
@@ -414,9 +414,9 @@ namespace Gauntlet
 					return false;
 				}
 
-				if (AllDevices[DeviceName] == false)
+				if (AllDevices[DeviceName] != "device")
 				{
-					Log.Warning(KnownLogEvents.Gauntlet_DeviceEvent, "Device {Name} is connected but we are not authorized", DeviceName);
+					Log.Warning(KnownLogEvents.Gauntlet_DeviceEvent, "Device {Name} is '{State}'", DeviceName, AllDevices[DeviceName]);
 					return false;
 				}
 
@@ -467,7 +467,7 @@ namespace Gauntlet
 				{
 					if (GetAllConnectedDevices().Count > 0)
 					{
-						throw new AutomationException("No default device available. One or more devices are connected but unauthorized. See 'adb devices'");
+						throw new AutomationException("No default device available. One or more devices are connected but unauthorized or offline. See 'adb devices'");
 					}
 					else
 					{
@@ -537,10 +537,10 @@ namespace Gauntlet
 				throw new AutomationException("Failed to find new device {0} in connection list", DeviceName);
 			}
 
-			if (ConnectedDevices[DeviceName] == false)
+			if (ConnectedDevices[DeviceName] != "device")
 			{
 				Dispose();
-				throw new AutomationException("Device {0} is connected but this PC is not authorized.", DeviceName);
+				throw new AutomationException("Device {0} is '{1}'.", DeviceName, ConnectedDevices[DeviceName]);
 			}
 		}
 
@@ -603,7 +603,7 @@ namespace Gauntlet
 		/// Returns a list of locally connected devices (e.g. 'adb devices'). 
 		/// </summary>
 		/// <returns></returns>
-		static private Dictionary<string, bool> GetAllConnectedDevices()
+		static private Dictionary<string, string> GetAllConnectedDevices()
 		{
            var Result = RunAdbGlobalCommand("devices");
 
@@ -611,7 +611,7 @@ namespace Gauntlet
 
             var DeviceList = DeviceMatches.Cast<Match>().ToDictionary(
                 M => M.Groups[1].ToString(),
-                M => !M.Groups[2].ToString().ToLower().Contains("unauthorized")
+                M => M.Groups[2].ToString().ToLower()
             );
 
             return DeviceList;
@@ -620,7 +620,7 @@ namespace Gauntlet
 		static private IEnumerable<string> GetAllAvailableDevices()
 		{
 			var AllDevices = GetAllConnectedDevices();
-			return AllDevices.Keys.Where(D => AllDevices[D] == true);
+			return AllDevices.Keys.Where(D => AllDevices[D] == "device");
 		}
 
 		static public ITargetDevice[] GetDefaultDevices()
@@ -657,8 +657,7 @@ namespace Gauntlet
 			get
 			{
 				string CommandLine = "shell dumpsys power";
-				IProcessResult OnAndUnlockedQuery = RunAdbDeviceCommand(CommandLine);
-
+				IProcessResult OnAndUnlockedQuery = RunAdbDeviceCommand(CommandLine, bPauseErrorParsing: true);
 				return OnAndUnlockedQuery.Output.Contains("mHoldingDisplaySuspendBlocker=true")
 					&& OnAndUnlockedQuery.Output.Contains("mHoldingWakeLockSuspendBlocker=true");
 			}
@@ -668,16 +667,15 @@ namespace Gauntlet
 		{
 			Log.Verbose("{0}: Powering on", ToString());
 			string CommandLine = "shell \"input keyevent KEYCODE_WAKEUP && input keyevent KEYCODE_MENU\"";
-			RunAdbDeviceCommand(CommandLine);
-			return true;
+			IProcessResult PowerOnQuery = RunAdbDeviceCommand(CommandLine, bPauseErrorParsing: true);
+			return !PowerOnQuery.Output.Contains("error");
 		}
 		public bool PowerOff()
 		{
 			Log.Verbose("{0}: Powering off", ToString());
-
 			string CommandLine = "shell \"input keyevent KEYCODE_SLEEP\"";
-			RunAdbDeviceCommand(CommandLine);
-			return true;
+			IProcessResult PowerOffQuery = RunAdbDeviceCommand(CommandLine, bPauseErrorParsing: true);
+			return !PowerOffQuery.Output.Contains("error");
 		}
 
 		public bool Reboot()
@@ -821,7 +819,29 @@ namespace Gauntlet
 
 					if (AdbResult.ExitCode != 0)
 					{
-						throw new AutomationException("Failed to push {0} to device. Error {1}", SourcePath, AdbResult.Output);
+						if ((AdbResult.Output.Contains("couldn't read from device") || AdbResult.Output.Contains("offline")) && !IsConnected)
+						{
+							Log.Info("Lost connection with device '{Name}'.", Name);
+							// Disconnection occurred. Let's retry a second time before given up
+							// Try to reconnect
+							if (Connect())
+							{
+								Log.Info("Retrying to copy via adb push...");
+								AdbResult = RunAdbDeviceCommand(AdbCommand);
+								if (AdbResult.ExitCode != 0)
+								{
+									throw new AutomationException("Failed to push {0} to device. Error {1}", SourcePath, AdbResult.Output);
+								}
+							}
+							else
+							{
+								throw new AutomationException("Failed to reconnect {0}", Name);
+							}
+						}
+						else
+						{
+							throw new AutomationException("Failed to push {0} to device. Error {1}", SourcePath, AdbResult.Output);
+						}
 					}
 
 					// Now pull info about the file which we'll write as a dep
