Move ResumeOnCriticalFailure option to UE.AutomationTestConfig

diff --git a/Engine/Source/Programs/AutomationTool/Gauntlet/Unreal/Automation/UE.Automation.cs b/Engine/Source/Programs/AutomationTool/Gauntlet/Unreal/Automation/UE.Automation.cs
--- a/Engine/Source/Programs/AutomationTool/Gauntlet/Unreal/Automation/UE.Automation.cs
+++ b/Engine/Source/Programs/AutomationTool/Gauntlet/Unreal/Automation/UE.Automation.cs
@@ -112,6 +112,12 @@ namespace UE
 		[AutoParam]
 		public virtual bool SimpleHordeReport { get; set; } = false;
 
+		/// <summary>
+		/// Resume test run on critical failure through pass retry
+		/// </summary>
+		[AutoParam]
+		public bool ResumeOnCriticalFailure = false;
+
 		/// <summary>
 		/// Validate DDC during tests
 		/// </summary>
@@ -685,7 +691,7 @@ namespace UE
 			{
 				// The test pass did not run at all
 				Log.Verbose("Found not-run tests: {Count}", JsonTestPassResults.NotRun);
-				if (GetConfiguration().ResumeOnCriticalFailure && !HasTimeout)
+				if ((GetConfiguration() is AutomationTestConfig Config) && Config.ResumeOnCriticalFailure && !HasTimeout)
 				{
 					// Reschedule test to resume from last 'in-process' test.
 					if (SetToRetryIfPossible())
diff --git a/Engine/Source/Programs/AutomationTool/Gauntlet/Unreal/Base/Gauntlet.UnrealTestConfiguration.cs b/Engine/Source/Programs/AutomationTool/Gauntlet/Unreal/Base/Gauntlet.UnrealTestConfiguration.cs
--- a/Engine/Source/Programs/AutomationTool/Gauntlet/Unreal/Base/Gauntlet.UnrealTestConfiguration.cs
+++ b/Engine/Source/Programs/AutomationTool/Gauntlet/Unreal/Base/Gauntlet.UnrealTestConfiguration.cs
@@ -553,12 +553,6 @@ namespace Gauntlet
 		[AutoParam(600.0f)]
 		public float MaxDuration { get; set; }
 
-		/// <summary>
-		/// Resume test run on critical failure through pass retry
-		/// </summary>
-		[AutoParam]
-		public bool ResumeOnCriticalFailure = false;
-
 		/// <summary>
 		/// Max number of retries in case of critical failure
 		/// </summary>
