Allow running Unreal tests from installed build
* Introduced the command line argument -editordir to specify the alternate location of the editor build.

diff --git a/Engine/Source/Programs/AutomationTool/Gauntlet/Unreal/Base/Gauntlet.UnrealTestContext.cs b/Engine/Source/Programs/AutomationTool/Gauntlet/Unreal/Base/Gauntlet.UnrealTestContext.cs
--- a/Engine/Source/Programs/AutomationTool/Gauntlet/Unreal/Base/Gauntlet.UnrealTestContext.cs
+++ b/Engine/Source/Programs/AutomationTool/Gauntlet/Unreal/Base/Gauntlet.UnrealTestContext.cs
@@ -63,6 +63,12 @@ namespace Gauntlet
 		[AutoParamWithNames("Build", "Builds")]
 		public string Build = "";
 
+		/// <summary>
+		/// Location of the editor build
+		/// </summary>
+		[AutoParam]
+		public string EditorDir = "";
+
 		/// <summary>
 		/// Does this project use 'Game' or Client/Server?
 		/// </summary>
diff --git a/Engine/Source/Programs/AutomationTool/Gauntlet/Unreal/BuildSource/Gauntlet.UnrealBuildSource.cs b/Engine/Source/Programs/AutomationTool/Gauntlet/Unreal/BuildSource/Gauntlet.UnrealBuildSource.cs
--- a/Engine/Source/Programs/AutomationTool/Gauntlet/Unreal/BuildSource/Gauntlet.UnrealBuildSource.cs
+++ b/Engine/Source/Programs/AutomationTool/Gauntlet/Unreal/BuildSource/Gauntlet.UnrealBuildSource.cs
@@ -641,7 +641,7 @@ namespace Gauntlet
 			{
 				FileSystemReference EditorExe = null;
 
-				if (BuildName == "Editor")
+				if (BuildName == "Editor" && string.IsNullOrEmpty(Globals.Params.ParseValue("EditorDir", null)))
 				{
 					try
 					{
diff --git a/Engine/Source/Programs/AutomationTool/Gauntlet/Unreal/RunUnreal.cs b/Engine/Source/Programs/AutomationTool/Gauntlet/Unreal/RunUnreal.cs
--- a/Engine/Source/Programs/AutomationTool/Gauntlet/Unreal/RunUnreal.cs
+++ b/Engine/Source/Programs/AutomationTool/Gauntlet/Unreal/RunUnreal.cs
@@ -154,8 +154,8 @@ namespace Gauntlet
 			UnrealTargetPlatform DefaultPlatform = BuildHostPlatform.Current.Platform;
 			UnrealTargetConfiguration DefaultConfiguration = UnrealTargetConfiguration.Development;
 
-			DirectoryReference UnrealPath = new DirectoryReference(Environment.CurrentDirectory);
-					
+			DirectoryReference UnrealPath = new DirectoryReference(!string.IsNullOrEmpty(ContextOptions.EditorDir) ? ContextOptions.EditorDir : Environment.CurrentDirectory);
+
 			// todo, pass this in as a BuildSource and remove the ContextOption params specific to finding builds
 			UnrealBuildSource BuildInfo = (UnrealBuildSource)Activator.CreateInstance(ContextOptions.BuildSourceType, new object[] { ContextOptions.Project, ContextOptions.ProjectPath, UnrealPath, ContextOptions.UsesSharedBuildType, ContextOptions.Build, ContextOptions.SearchPaths });
 
