Gauntlet - Fix incomple and crashed test not being log in summary report

#rnx

diff --git a/Engine/Source/Programs/AutomationTool/Gauntlet/Unreal/Automation/UE.Automation.cs b/Engine/Source/Programs/AutomationTool/Gauntlet/Unreal/Automation/UE.Automation.cs
--- a/Engine/Source/Programs/AutomationTool/Gauntlet/Unreal/Automation/UE.Automation.cs
+++ b/Engine/Source/Programs/AutomationTool/Gauntlet/Unreal/Automation/UE.Automation.cs
@@ -873,7 +873,7 @@ namespace UE
 					}
 					return E;
 				};
-				// If there were abnormal exits then look only at the incomplete tests to avoid confusing things.
+				// If there were abnormal exits then look only at the failed and incomplete tests only to avoid confusing things.
 				if (GetRolesThatExitedAbnormally().Any())
 				{
 					if (AllTests.Count() == 0)
@@ -881,15 +881,16 @@ namespace UE
 						Log.Error(KnownLogEvents.Gauntlet_UnrealEngineTestEvent, " * No tests were executed.");
 						Log.Info("");
 					}
-					else if (IncompleteTests.Count() > 0)
+					else if (FailedTests.Count() > 0 || IncompleteTests.Count() > 0)
 					{
-						var InProcess = IncompleteTests.Where(T => T.State == TestStateType.InProcess);
-						if (InProcess.Any())
+						var Failures = FailedTests.Concat(IncompleteTests.Where(T => T.State == TestStateType.InProcess));
+						if (Failures.Any())
 						{
-							Log.Info(" ### The following Engine test(s) were incomplete:");
-							foreach (UnrealAutomatedTestResult Result in InProcess)
+							Log.Info(" ### The following Engine test(s) were incomplete or failed:");
+							foreach (UnrealAutomatedTestResult Result in Failures)
 							{
-								Log.Error(KnownLogEvents.Gauntlet_UnrealEngineTestEvent, " * Test '{Name}' did not complete.", Result.FullTestPath);
+								string Message = !Result.IsComplete ? " * Test '{Name}' did not complete." : " * Test '{Name}' failed.";
+								Log.Error(KnownLogEvents.Gauntlet_UnrealEngineTestEvent, Message, Result.FullTestPath);
 								var Errors = CapErrorOrWarningList(Result.ErrorEvents.Select(E => E.Message).Distinct());
 								foreach (var Error in Errors)
 								{
