PR #8047: Fix missing quotes when Gauntlet run on iOS with automation test command (Contributed by [REDACTED_PERSON])

diff --git a/Engine/Source/Programs/AutomationTool/Gauntlet/Platform/IOS/Gauntlet.TargetDeviceIOS.cs b/Engine/Source/Programs/AutomationTool/Gauntlet/Platform/IOS/Gauntlet.TargetDeviceIOS.cs
--- a/Engine/Source/Programs/AutomationTool/Gauntlet/Platform/IOS/Gauntlet.TargetDeviceIOS.cs
+++ b/Engine/Source/Programs/AutomationTool/Gauntlet/Platform/IOS/Gauntlet.TargetDeviceIOS.cs
@@ -269,7 +269,7 @@ namespace Gauntlet
 				throw new DeviceException("AppInstance is of incorrect type!");
 			}
 
-			string CommandLine = IOSApp.CommandLine.Replace("\"", "\\\"");
+			string CommandLine = IOSApp.CommandLine.Replace("\"", "\\\\\"");
 			
 			Log.Info("Launching {0} on {1}", App.Name, ToString());
 			Log.Verbose("\t{0}", CommandLine);
