Adding a simple replay-based PGO Gauntlet task.

```RunUnreal -Test=BasicReplayPGONode -project=MyProject -platform=Win64 -configuration=Shipping -build=MyProject\Saved\StagedBuilds\Win64 -LocalReplay=MyProject\Replays\PGO.replay -ProfileOutputDirectory=MyProject\Platforms\Windows\Build\PGO -MaxDuration=7200```

Note: not supported on all platforms yet.

diff --git a/Engine/Source/Programs/AutomationTool/Gauntlet/Unreal/Base/Gauntlet.UnrealPGONode.cs b/Engine/Source/Programs/AutomationTool/Gauntlet/Unreal/Base/Gauntlet.UnrealPGONode.cs
--- a/Engine/Source/Programs/AutomationTool/Gauntlet/Unreal/Base/Gauntlet.UnrealPGONode.cs
+++ b/Engine/Source/Programs/AutomationTool/Gauntlet/Unreal/Base/Gauntlet.UnrealPGONode.cs
@@ -190,6 +190,60 @@ namespace Gauntlet
 		}
 	}
 
+
+	public class BasicReplayPGOConfig : PGOConfig
+	{
+		[AutoParam("")]
+		public string LocalReplay;
+
+		[AutoParam("")]
+		public string AdditionalCommandLine;
+
+		public override void ParametersWereApplied(string[] Params)
+		{
+			base.ParametersWereApplied(Params);
+			
+			// Check that replay file exists if specified
+			if (string.IsNullOrEmpty(LocalReplay) || !File.Exists(LocalReplay))
+			{
+				throw new AutomationException("LocalReplay option must be specified with absolute path to existing replay file");
+			}			
+		}
+	}
+
+	/// <summary>
+	/// PGO Profiling Node that runs a replay & gathers the results. Requires an existing build compiled with  -PGOProfile -Configuration=Shipping
+	/// e.g. RunUnreal -Test=BasicReplayPGONode -project=MyProject -platform=Win64 -configuration=Shipping -build=MyProject\Saved\StagedBuilds\Windows -LocalReplay=MyProject\Replays\PGO.replay -ProfileOutputDirectory=MyProject\Platforms\Windows\Build\PGO -MaxDuration=7200
+	/// </summary>
+	public class BasicReplayPGONode : PGONode<BasicReplayPGOConfig>
+	{
+		public BasicReplayPGONode(UnrealTestContext InContext) : base(InContext)
+		{
+		}
+
+		public override BasicReplayPGOConfig GetConfiguration()
+		{
+			if (CachedConfig != null)
+			{
+				return CachedConfig;
+			}
+			var Config = base.GetConfiguration();
+
+			// create a client that runs the given replay & then exits when it's finished
+			UnrealTestRole ClientRole = Config.RequireRole(UnrealTargetRole.Client);	
+			ClientRole.FilesToCopy.Add(new UnrealFileToCopy(Config.LocalReplay, EIntendedBaseCopyDirectory.Demos, Path.GetFileName(Config.LocalReplay)));
+			
+			string ReplayName = Path.GetFileNameWithoutExtension(Config.LocalReplay);
+			ClientRole.CommandLine += $" -ExecCmds=\"DEMOPLAY {ReplayName}\"";
+			ClientRole.CommandLine += $" -Deterministic -ExitAfterReplay -NoCheckpointHangDetector";
+			ClientRole.CommandLine += $" {Config.AdditionalCommandLine} ";
+
+			return Config;
+		}
+	}
+
+
+
 	internal interface IPGOPlatform
 	{
 		void ApplyConfiguration(PGOConfig Config);
