UE.PLMTest waits for engine init before starting.

#rnx

diff --git a/Engine/Source/Programs/AutomationTool/Gauntlet/Unreal/Automation/UE.PLMTest.cs b/Engine/Source/Programs/AutomationTool/Gauntlet/Unreal/Automation/UE.PLMTest.cs
--- a/Engine/Source/Programs/AutomationTool/Gauntlet/Unreal/Automation/UE.PLMTest.cs
+++ b/Engine/Source/Programs/AutomationTool/Gauntlet/Unreal/Automation/UE.PLMTest.cs
@@ -13,9 +13,9 @@ namespace UE
 	/// </summary>
 	public class PLMTestConfiguration : UnrealTestConfiguration
 	{
-		/// How long to wait before starting
+		/// How long to wait before starting after the engine has been initialized
 		[AutoParam]
-		public int InitialWait = 45;
+		public int InitialWait = 0;
 
 		/// How long to wait between loops
 		[AutoParam]
@@ -215,6 +215,7 @@ namespace UE
 
 		Action[] Actions;
 		int CurrentActionIndex;
+		bool bHasEngineInit;
 
 		/// <summary>
 		/// Default constructor
@@ -264,9 +265,18 @@ namespace UE
 			}
 
 			CurrentActionIndex = 0;
+			bHasEngineInit = false;
 			return true;
 		}
 
+		/// <summary>
+		/// String that we search for to be considered "Booted"
+		/// </summary>
+		/// <returns></returns>
+		protected virtual string GetCompletionString()
+		{
+			return "Engine is initialized. Leaving FEngineLoop::Init()";
+		}
 
 		/// <summary>
 		/// Called periodically while the test is running to allow code to monitor health.
@@ -276,7 +286,18 @@ namespace UE
 			// run the base class tick;
 			base.TickTest();
 
-			if (CurrentActionIndex >= Actions.Length)
+			if (!bHasEngineInit)
+			{
+				// see if the engine has been initialized yet
+				IAppInstance RunningInstance = this.TestInstance.RunningRoles.First().AppInstance;
+				UnrealLogParser LogParser = new UnrealLogParser(RunningInstance.StdOut);
+				if (LogParser.Content.IndexOf( GetCompletionString(), StringComparison.OrdinalIgnoreCase) > 0 )
+				{
+					Log.Info("Found '{0}'. Starting PLM tests...", GetCompletionString() );
+					bHasEngineInit = true;
+				}
+			}
+			else if (CurrentActionIndex >= Actions.Length)
 			{
 				Log.Info("All actions completed. Ending test");
 				MarkTestComplete();
