When signing with sha1 sign with /as to append certs and not override existing certs

#rnx

diff --git a/Engine/Source/Programs/AutomationTool/AutomationUtils/CommandUtils.cs b/Engine/Source/Programs/AutomationTool/AutomationUtils/CommandUtils.cs
--- a/Engine/Source/Programs/AutomationTool/AutomationUtils/CommandUtils.cs
+++ b/Engine/Source/Programs/AutomationTool/AutomationUtils/CommandUtils.cs
@@ -2894,7 +2894,7 @@ namespace AutomationTool
 					StringBuilder CommandLine = new StringBuilder();
 					if(SignatureType == SignatureType.SHA1)
 					{
-						CommandLine.AppendFormat("sign{0} /a /n \"{1}\" /t {2} /v", SpecificStoreArg, SigningIdentity, TimestampServersSHA1[NumAttempts % TimestampServersSHA1.Length]);
+						CommandLine.AppendFormat("sign{0} /a /as /n \"{1}\" /tr {2} /v", SpecificStoreArg, SigningIdentity, TimestampServersSHA1[NumAttempts % TimestampServersSHA1.Length]);
 					}
 					else if(SignatureType == SignatureType.SHA256)
 					{
