Basic replay PGO node improvements:
 - Use -replay="xxx" command line instead of -execcmds="demoplay xxx"
 - Added a hook to allow replays to be set up in a platform-specific way, if the FilesToCopy / -replay method isn't suitable for that platform.

#rnx

diff --git a/Engine/Source/Programs/AutomationTool/Gauntlet/Unreal/Base/Gauntlet.UnrealPGONode.cs b/Engine/Source/Programs/AutomationTool/Gauntlet/Unreal/Base/Gauntlet.UnrealPGONode.cs
--- a/Engine/Source/Programs/AutomationTool/Gauntlet/Unreal/Base/Gauntlet.UnrealPGONode.cs
+++ b/Engine/Source/Programs/AutomationTool/Gauntlet/Unreal/Base/Gauntlet.UnrealPGONode.cs
@@ -43,7 +43,7 @@ namespace Gauntlet
 	public class PGONode<TConfigClass> : UnrealTestNode<TConfigClass> where TConfigClass : PGOConfig, new()
 	{
 		protected string LocalOutputDirectory;
-		private IPGOPlatform PGOPlatform;
+		internal IPGOPlatform PGOPlatform;
 		bool ProcessPGODataFailed = false;
 
 		public PGONode(UnrealTestContext InContext) : base(InContext)
@@ -191,6 +191,12 @@ namespace Gauntlet
 	}
 
 
+	public interface IPGOCustomReplayPlatform
+	{
+		void ConfigureClientForReplay( PGOConfig Config, UnrealTestRole ClientRole, string LocalReplayFileName );
+	}
+
+
 	public class BasicReplayPGOConfig : PGOConfig
 	{
 		[AutoParam("")]
@@ -230,13 +236,22 @@ namespace Gauntlet
 			var Config = base.GetConfiguration();
 
 			// create a client that runs the given replay & then exits when it's finished
-			UnrealTestRole ClientRole = Config.RequireRole(UnrealTargetRole.Client);	
-			ClientRole.FilesToCopy.Add(new UnrealFileToCopy(Config.LocalReplay, EIntendedBaseCopyDirectory.Demos, Path.GetFileName(Config.LocalReplay)));
-			
-			string ReplayName = Path.GetFileNameWithoutExtension(Config.LocalReplay);
-			ClientRole.CommandLine += $" -ExecCmds=\"DEMOPLAY {ReplayName}\"";
+			UnrealTestRole ClientRole = Config.RequireRole(UnrealTargetRole.Client);
 			ClientRole.CommandLine += $" -Deterministic -ExitAfterReplay -NoCheckpointHangDetector";
-			ClientRole.CommandLine += $" {Config.AdditionalCommandLine} ";
+			ClientRole.CommandLine += $" {Config.AdditionalCommandLine}";
+
+			if (PGOPlatform is IPGOCustomReplayPlatform PGOCustomReplayPlatform)
+			{
+				// the platform needs bespoke replay handling
+				PGOCustomReplayPlatform.ConfigureClientForReplay(Config, ClientRole, Config.LocalReplay);
+			}
+			else
+			{
+				// normal replay handling - copy to the demos folder and then launch from command line
+				string ReplayName = Path.GetFileNameWithoutExtension(Config.LocalReplay);
+				ClientRole.FilesToCopy.Add(new UnrealFileToCopy(Config.LocalReplay, EIntendedBaseCopyDirectory.Demos, Path.GetFileName(Config.LocalReplay)));
+				ClientRole.CommandLine += $" -Replay=\"{ReplayName}\" ";
+			}
 
 			return Config;
 		}
