- More Architecture cleanup
- Added UAT params to pass architecture by target type (-editorarchitecture, -clientarchitecture, -serverarchitecture, -programarchitecture). Keeping -specifiedarchitecture for compat, and it will set all the other 4 if they are not specified
- Changed ProjectGenerator to find Program .uproject files (logic was looking for .Target files under the .uproject's directory, which is not the case for Programs)
- Added two params to GetProjectArchitectures:
    - bGetAllSupported - used to get all supported architectures, not just the ones currently being built for (needed for XcodeProject to know what architectures to allow in the dropdown)
    - bIsDistributionMode - used to get architectures to compile when making distribution builds (probably want to compile all architectures for distro's)
- Removed the MacExports.TargetsAlowedForAppleSilicon, updated XcodeProject code to use the new GetProjectArchitectures()

diff --git a/Engine/Source/Programs/AutomationTool/AutomationUtils/ProjectParams.cs b/Engine/Source/Programs/AutomationTool/AutomationUtils/ProjectParams.cs
--- a/Engine/Source/Programs/AutomationTool/AutomationUtils/ProjectParams.cs
+++ b/Engine/Source/Programs/AutomationTool/AutomationUtils/ProjectParams.cs
@@ -561,6 +561,10 @@ namespace AutomationTool
 			bool? Prebuilt = null,
             int? RunTimeoutSeconds = null,
 			string SpecifiedArchitecture = null,
+			string ServerArchitecture = null,
+			string EditorArchitecture = null,
+			string ClientArchitecture = null,
+			string ProgramArchitecture = null,
 			string UbtArgs = null,
 			string AdditionalPackageOptions = null,
 			bool? IterativeDeploy = null,
@@ -943,6 +947,12 @@ namespace AutomationTool
 			this.DevicePassword = ParseParamValueIfNotSpecified(Command, DevicePassword, "devicepass", String.Empty);
 			this.CrashReporter = GetParamValueIfNotSpecified(Command, CrashReporter, this.CrashReporter, "crashreporter");
 			this.SpecifiedArchitecture = ParseParamValueIfNotSpecified(Command, SpecifiedArchitecture, "specifiedarchitecture", String.Empty);
+			// if SpecifiedArchitecture is used, then set them all to it, and then allow comandline to override specific ones
+			this.ServerArchitecture = this.EditorArchitecture = this.ClientArchitecture = this.ProgramArchitecture = this.SpecifiedArchitecture;
+			this.ServerArchitecture = ParseParamValueIfNotSpecified(Command, ServerArchitecture, "serverarchitecture", this.ServerArchitecture);
+			this.EditorArchitecture = ParseParamValueIfNotSpecified(Command, EditorArchitecture, "editorarchitecture", this.EditorArchitecture);
+			this.ClientArchitecture = ParseParamValueIfNotSpecified(Command, ClientArchitecture, "clientarchitecture", this.ClientArchitecture);
+			this.ProgramArchitecture = ParseParamValueIfNotSpecified(Command, ProgramArchitecture, "programarchitecture", this.ProgramArchitecture);
 			this.UbtArgs = ParseParamValueIfNotSpecified(Command, UbtArgs, "ubtargs", String.Empty);
 			this.AdditionalPackageOptions = ParseParamValueIfNotSpecified(Command, AdditionalPackageOptions, "AdditionalPackageOptions", String.Empty);
 			this.WriteBackMetadataToAssetRegistry = ParseParamValueIfNotSpecified(Command, WriteBackMetadataToAssetRegistry, "WriteBackMetadataToAssetRegistry", String.Empty);
@@ -2221,9 +2231,21 @@ namespace AutomationTool
         [Help("RunTimeoutSeconds", "timeout to wait after we lunch the game")]
         public int RunTimeoutSeconds;
 
-		[Help("SpecifiedArchitecture", "Determine a specific Minimum OS")]
+		[Help("SpecifiedArchitecture", "Architecture to use for building any executables (see EditorArchitecture, etc for specific target type control)")]
 		public string SpecifiedArchitecture;
 
+		[Help("EditorArchitecture", "Architecture to use for building editor executables")]
+		public string EditorArchitecture;
+
+		[Help("ServerArchitecture", "Architecture to use for building server executables")]
+		public string ServerArchitecture;
+
+		[Help("ClientArchitecture", "Architecture to use for building client/game executables")]
+		public string ClientArchitecture;
+
+		[Help("ProgramArchitecture", "Architecture to use for building program executables")]
+		public string ProgramArchitecture;
+
 		[Help("UbtArgs", "extra options to pass to ubt")]
 		public string UbtArgs;
 
diff --git a/Engine/Source/Programs/AutomationTool/Mac/MacPlatform.Automation.cs b/Engine/Source/Programs/AutomationTool/Mac/MacPlatform.Automation.cs
--- a/Engine/Source/Programs/AutomationTool/Mac/MacPlatform.Automation.cs
+++ b/Engine/Source/Programs/AutomationTool/Mac/MacPlatform.Automation.cs
@@ -88,17 +88,6 @@ public class MacPlatform : ApplePlatform
 		return "MacEditor";
 	}
 
-	/// <summary>
-	/// Returns true if UAT can build this target for all Mac architectures
-	/// </summary>
-	/// <param name="InTarget"></param>
-	/// <param name="InParams"></param>
-	/// <returns></returns>
-	protected bool CanBuildTargetForAllArchitectures(UnrealBuild.BuildTarget InTarget, ProjectParams InParams)
-	{
-		return MacExports.TargetsAllowedForAppleSilicon.Contains(InTarget.TargetName, StringComparer.OrdinalIgnoreCase);
-	}
-
 	/// <summary>
 	/// Override PreBuildAgenda so we can control the architecture that targets are built for based on
 	/// project settings and the current user environment
@@ -110,75 +99,14 @@ public class MacPlatform : ApplePlatform
 	{
 		base.PreBuildAgenda(Build, Agenda, Params);
 
-		string LocalArchitecture = MacExports.HostArchitecture;
-
-		bool ProjectIsUniversal = ProjectTargetArchitectures.Count() > 1;
-
-		// Go through the agenda for all targets and set the architecture appropriately
+		// Go through the agenda for all targets and set the architecture if needed
 		foreach (UnrealBuild.BuildTarget Target in Agenda.Targets)
 		{
-			bool IsTarget = Params.ClientCookedTargets.Contains(Target.TargetName) || Params.ServerCookedTargets.Contains(Target.TargetName);
-
-			// Default to Intel. 
-			string UBTArchitectureParam = MacExports.IntelArchitecture;
-
-			// If an architecture was specified, use that
-			if (!string.IsNullOrEmpty(Params.SpecifiedArchitecture))
-			{
-				UBTArchitectureParam = Params.SpecifiedArchitecture;
-				Log.TraceInformation("Building {0} as {1} due to -specifiedarchitecture", Target.TargetName, UBTArchitectureParam);
-			}
-			else
-			{
-				// Targets are easy. 
-				// - If the project is set to Intel/Apple we do that
-				// - If it's universal we build the local architecture unless distributing
-				// - SpecifiedArchitecture overrides these
-				if (IsTarget)
-				{
-					// If the project isn't marked as universal built what it's set to
-					if (!ProjectIsUniversal)
-					{
-						UBTArchitectureParam = ProjectTargetArchitectures.First();
-						Log.TraceInformation("Building {0} as {1}", Target.TargetName, UBTArchitectureParam);
-					}
-					else
-					{
-						// if it is universal, build everything for distribution or just the local architecture otherwise
-						if (Params.Distribution)
-						{
-							UBTArchitectureParam = string.Join("+", ProjectTargetArchitectures);
-							Log.TraceInformation("Building {0} as {1} for distribution", Target.TargetName, UBTArchitectureParam);
-						}
-						else
-						{
-							UBTArchitectureParam = LocalArchitecture;
-							Log.TraceInformation("Building {0} as {1} for local non-distribution", Target.TargetName, UBTArchitectureParam);
-						}
-					}
-				}
-				else
-				{
-					// We build tools for the local architecture if possible
-					if (CanBuildTargetForAllArchitectures(Target, Params) || LocalArchitecture == MacExports.IntelArchitecture)
-					{
-						UBTArchitectureParam = LocalArchitecture;
-						Log.TraceInformation("Building {0} as {1} for host", Target.TargetName, UBTArchitectureParam);
-					}
-					else if (MacExports.IsRunningOnAppleArchitecture)
-					{
-						// Tell them why to avoid confusion
-						Log.TraceInformation("Building {0} as {1} (arm64 not currently supported)", Target.TargetName, UBTArchitectureParam);
-						UBTArchitectureParam = MacExports.IntelArchitecture;
-					}
-				}
-			}
-			
-			// TODO - This needs to be handled in a more graceful way, however for now
-			// when in an installedbuild just leave it to the default arch and don't add anything extra
-			if (!Unreal.IsEngineInstalled())
+			// if building for Distribution, and no arch is already specified, then get the distro architectures and use that for this build
+			if (Params.Distribution && !Target.UBTArgs.ToLower().Contains("-architecture="))
 			{
-				Target.UBTArgs += string.Format(" -architecture={0}", UBTArchitectureParam);
+				IEnumerable<string> DistroArches = PlatformExports.GetProjectArchitectures(UnrealTargetPlatform.Mac, Params.RawProjectPath, Target.TargetName, false, true);
+				Target.UBTArgs += " -architecture=" + string.Join("+", DistroArches);
 			}
 		}
 	}
diff --git a/Engine/Source/Programs/AutomationTool/Scripts/BuildProjectCommand.Automation.cs b/Engine/Source/Programs/AutomationTool/Scripts/BuildProjectCommand.Automation.cs
--- a/Engine/Source/Programs/AutomationTool/Scripts/BuildProjectCommand.Automation.cs
+++ b/Engine/Source/Programs/AutomationTool/Scripts/BuildProjectCommand.Automation.cs
@@ -76,6 +76,13 @@ namespace AutomationScripts
 			var Agenda = new UnrealBuild.BuildAgenda();
 			var CrashReportPlatforms = new HashSet<UnrealTargetPlatform>();
 
+
+
+			string ServerArchitecture = Params.ServerArchitecture == "" ? "" : $" -architecture={Params.ServerArchitecture}";
+			string EditorArchitecture = Params.EditorArchitecture == "" ? "" : $" -architecture={Params.EditorArchitecture}";
+			string ClientArchitecture = Params.ClientArchitecture == "" ? "" : $" -architecture={Params.ClientArchitecture}";
+			string ProgramArchitecture = Params.ProgramArchitecture == "" ? "" : $" -architecture={Params.ProgramArchitecture}";
+
 			// Setup editor targets
 			if (Params.HasEditorTargets && (!Params.SkipBuildEditor) && (TargetMask & ProjectBuildTargets.Editor) == ProjectBuildTargets.Editor)
 			{
@@ -83,23 +90,23 @@ namespace AutomationScripts
 				UnrealTargetPlatform EditorPlatform = HostPlatform.Current.HostEditorPlatform;
 				const UnrealTargetConfiguration EditorConfiguration = UnrealTargetConfiguration.Development;
 
-				Agenda.AddTargets(Params.EditorTargets.ToArray(), EditorPlatform, EditorConfiguration, Params.CodeBasedUprojectPath);
+				Agenda.AddTargets(Params.EditorTargets.ToArray(), EditorPlatform, EditorConfiguration, Params.CodeBasedUprojectPath, InAddArgs: EditorArchitecture);
 
 				if (!Unreal.IsEngineInstalled())
 				{
 					CrashReportPlatforms.Add(EditorPlatform);
 					if (Params.EditorTargets.Contains("UnrealHeaderTool") == false)
 					{
-						Agenda.AddTargets(new string[] { "UnrealHeaderTool" }, EditorPlatform, EditorConfiguration);
-						Agenda.AddTargets(new string[] { "UnrealHeaderTool" }, EditorPlatform, EditorConfiguration, Params.CodeBasedUprojectPath);
+						Agenda.AddTargets(new string[] { "UnrealHeaderTool" }, EditorPlatform, EditorConfiguration, InAddArgs:ProgramArchitecture);
+						Agenda.AddTargets(new string[] { "UnrealHeaderTool" }, EditorPlatform, EditorConfiguration, Params.CodeBasedUprojectPath, InAddArgs: ProgramArchitecture);
 					}
 					if (Params.EditorTargets.Contains("ShaderCompileWorker") == false)
 					{
-						Agenda.AddTargets(new string[] { "ShaderCompileWorker" }, EditorPlatform, EditorConfiguration);
+						Agenda.AddTargets(new string[] { "ShaderCompileWorker" }, EditorPlatform, EditorConfiguration, InAddArgs: ProgramArchitecture);
 					}
 					if (Params.FileServer && Params.EditorTargets.Contains("UnrealFileServer") == false)
 					{
-						Agenda.AddTargets(new string[] { "UnrealFileServer" }, EditorPlatform, EditorConfiguration);
+						Agenda.AddTargets(new string[] { "UnrealFileServer" }, EditorPlatform, EditorConfiguration, InAddArgs: ProgramArchitecture);
 					}
 				}
 			}
@@ -109,7 +116,7 @@ namespace AutomationScripts
 			{
 				if (!Params.HasEditorTargets || Params.EditorTargets.Contains("UnrealPak") == false)
 				{
-					Agenda.AddTargets(new string[] { "UnrealPak" }, HostPlatform.Current.HostEditorPlatform, UnrealTargetConfiguration.Development, Params.CodeBasedUprojectPath);
+					Agenda.AddTargets(new string[] { "UnrealPak" }, HostPlatform.Current.HostEditorPlatform, UnrealTargetConfiguration.Development, Params.CodeBasedUprojectPath, InAddArgs: ProgramArchitecture);
 				}
 			}
 
@@ -157,7 +164,7 @@ namespace AutomationScripts
 					{
 						UnrealTargetPlatform CrashReportPlatform = Platform.GetPlatform(ClientPlatformType).CrashReportPlatform ?? ClientPlatformType;
 						CrashReportPlatforms.Add(CrashReportPlatform);
-						Agenda.AddTargets(Params.ClientCookedTargets.ToArray(), ClientPlatformType, BuildConfig, Params.CodeBasedUprojectPath, InAddArgs: " -remoteini=\"" + Params.RawProjectPath.Directory.FullName + "\"" + AdditionalArgs);
+						Agenda.AddTargets(Params.ClientCookedTargets.ToArray(), ClientPlatformType, BuildConfig, Params.CodeBasedUprojectPath, InAddArgs: " -remoteini=\"" + Params.RawProjectPath.Directory.FullName + "\"" + AdditionalArgs + ClientArchitecture);
 					}
 				}
 			}
@@ -171,7 +178,7 @@ namespace AutomationScripts
 					{
 						UnrealTargetPlatform CrashReportPlatform = Platform.GetPlatform(ServerPlatformType).CrashReportPlatform ?? ServerPlatformType;
 						CrashReportPlatforms.Add(CrashReportPlatform);
-						Agenda.AddTargets(Params.ServerCookedTargets.ToArray(), ServerPlatformType, BuildConfig, Params.CodeBasedUprojectPath, InAddArgs: " -remoteini=\"" + Params.RawProjectPath.Directory.FullName + "\"" + AdditionalArgs);
+						Agenda.AddTargets(Params.ServerCookedTargets.ToArray(), ServerPlatformType, BuildConfig, Params.CodeBasedUprojectPath, InAddArgs: " -remoteini=\"" + Params.RawProjectPath.Directory.FullName + "\"" + AdditionalArgs + ServerArchitecture);
 					}
 				}
 			}
@@ -182,7 +189,7 @@ namespace AutomationScripts
 				{
 					if (Params.ClientTargetPlatforms.Contains(new TargetPlatformDescriptor(BootstrapPackagedGamePlatformType)))
 					{
-						Agenda.AddTarget("BootstrapPackagedGame", BootstrapPackagedGamePlatformType, UnrealBuildTool.UnrealTargetConfiguration.Shipping);
+						Agenda.AddTarget("BootstrapPackagedGame", BootstrapPackagedGamePlatformType, UnrealBuildTool.UnrealTargetConfiguration.Shipping, InAddArgs: ClientArchitecture);
 					}
 				}
 			}
@@ -192,7 +199,7 @@ namespace AutomationScripts
 				{
 					if (PlatformSupportsCrashReporter(CrashReportPlatform))
 					{
-						Agenda.AddTarget("CrashReportClient", CrashReportPlatform, UnrealTargetConfiguration.Shipping, InAddArgs: " -remoteini=\"" + Params.RawProjectPath.Directory.FullName + "\"");
+						Agenda.AddTarget("CrashReportClient", CrashReportPlatform, UnrealTargetConfiguration.Shipping, InAddArgs: " -remoteini=\"" + Params.RawProjectPath.Directory.FullName + "\"" + ProgramArchitecture);
 					}
 				}
 			}
@@ -204,7 +211,7 @@ namespace AutomationScripts
 				{
 					foreach (var ClientPlatformType in UniquePlatformTypes)
 					{
-						Agenda.AddTargets(Params.ProgramTargets.ToArray(), ClientPlatformType, BuildConfig, Params.CodeBasedUprojectPath);
+						Agenda.AddTargets(Params.ProgramTargets.ToArray(), ClientPlatformType, BuildConfig, Params.CodeBasedUprojectPath, ProgramArchitecture);
 					}
 				}
 			}
diff --git a/Engine/Source/Programs/UnrealBuildTool/Configuration/TargetDescriptor.cs b/Engine/Source/Programs/UnrealBuildTool/Configuration/TargetDescriptor.cs
--- a/Engine/Source/Programs/UnrealBuildTool/Configuration/TargetDescriptor.cs
+++ b/Engine/Source/Programs/UnrealBuildTool/Configuration/TargetDescriptor.cs
@@ -420,7 +420,7 @@ namespace UnrealBuildTool
 						if (ParamArchitectures.Count() == 0)
 						{
 							// ask the platform what achitectures it wants for this project
-							Architectures = BuildPlatform.GetProjectArchitectures(TargetProjectFile, TargetName);
+							Architectures = BuildPlatform.GetProjectArchitectures(TargetProjectFile, TargetName, false, false);
 						}
 						else
 						{
diff --git a/Engine/Source/Programs/UnrealBuildTool/Configuration/UEBuildPlatform.cs b/Engine/Source/Programs/UnrealBuildTool/Configuration/UEBuildPlatform.cs
--- a/Engine/Source/Programs/UnrealBuildTool/Configuration/UEBuildPlatform.cs
+++ b/Engine/Source/Programs/UnrealBuildTool/Configuration/UEBuildPlatform.cs
@@ -254,7 +254,9 @@ namespace UnrealBuildTool
 		/// </summary>
 		/// <param name="ProjectFile">Optional project to read settings from </param>
 		/// <param name="TargetName">Optional name of project, can be useful for programs, etc that have no projectfile</param>
-		public virtual IEnumerable<string> GetProjectArchitectures(FileReference? ProjectFile, string? TargetName=null)
+		/// <param name="bGetAllSupported">If true, return all supported architectures for this target</param>
+		/// <param name="bIsDistributionMode">If true, return architectures when packaging in distribution mode</param>
+		public virtual IEnumerable<string> GetProjectArchitectures(FileReference? ProjectFile, string? TargetName, bool bGetAllSupported, bool bIsDistributionMode)
 		{
 			// by default, use an empty architecture (which is really just a modifer to the platform for some paths/names)
 			return new string[] { GetDefaultArchitecture(ProjectFile) };
diff --git a/Engine/Source/Programs/UnrealBuildTool/Platform/Mac/MacExports.cs b/Engine/Source/Programs/UnrealBuildTool/Platform/Mac/MacExports.cs
--- a/Engine/Source/Programs/UnrealBuildTool/Platform/Mac/MacExports.cs
+++ b/Engine/Source/Programs/UnrealBuildTool/Platform/Mac/MacExports.cs
@@ -49,30 +49,6 @@ namespace UnrealBuildTool
 			}
 		}
 
-		/// <summary>
-		/// Returns the current list of default targets that are known to build for Apple Silicon. Individual projects will
-		/// use the Project setting.
-		/// </summary>
-		/// <returns></returns>
-		static public IEnumerable<string> TargetsAllowedForAppleSilicon
-		{
-			get
-			{
-				return new[] { 
-					"BenchmarkTool",
-					"BlankProgram",
-					"BuildPatchTool",
-					"ShaderCompileWorker",
-					"UnrealClient",
-					"UnrealGame", 
-					"UnrealServer",
-					"UnrealHeaderTool", 
-					"UnrealPak",
-					"SlateUGS",
-				};
-			}
-		}
-
 		/// <summary>
 		/// Cached result for AppleArch check
 		/// </summary>
diff --git a/Engine/Source/Programs/UnrealBuildTool/Platform/Mac/UEBuildMac.cs b/Engine/Source/Programs/UnrealBuildTool/Platform/Mac/UEBuildMac.cs
--- a/Engine/Source/Programs/UnrealBuildTool/Platform/Mac/UEBuildMac.cs
+++ b/Engine/Source/Programs/UnrealBuildTool/Platform/Mac/UEBuildMac.cs
@@ -267,17 +267,8 @@ namespace UnrealBuildTool
 			return MacExports.DefaultArchitecture;
 		}
 
-		private List<string>? CachedProjectArches = null;
-		private FileReference? CachedArchesProject = null;
-		private string? CachedArchesTarget = null;
-		public override IEnumerable<string> GetProjectArchitectures(FileReference? ProjectFile, string? TargetName)
+		public override IEnumerable<string> GetProjectArchitectures(FileReference? ProjectFile, string? TargetName, bool bGetAllSupported, bool bIsDistributionMode)
 		{
-			// use the cached one if it's valid
-			if (CachedProjectArches != null && ProjectFile == CachedArchesProject && TargetName == CachedArchesTarget)
-			{
-				return CachedProjectArches;
-			}
-
 			bool bIsEditor = false;
 			//// get project ini from ProjetFile, or if null, then try to get it from the target rules
 			//DirectoryReference? ProjectDir = null;
@@ -294,6 +285,12 @@ namespace UnrealBuildTool
 				}
 				TargetRules? Rules = RulesAsm.CreateTargetRules(TargetName, UnrealTargetPlatform.Mac, UnrealTargetConfiguration.Development, "", ProjectFile, null, Log.Logger);
 				bIsEditor = Rules.Type == TargetType.Editor;
+				// the projectfile passed in may be a game's uproject file that we are compiling a program in the context of, 
+				// but we still want the settings for the program
+				if (Rules.Type == TargetType.Program)
+				{
+					ProjectFile = Rules.ProjectFile;
+				}
 			}
 
 			ConfigHierarchy EngineIni = ConfigCache.ReadHierarchy(ConfigHierarchyType.Engine, ProjectFile?.Directory, UnrealTargetPlatform.Mac);
@@ -319,7 +316,9 @@ namespace UnrealBuildTool
 
 			// choose a supported architecture(s) based on desired type
 			List<string> Architectures = new();
-			if (DefaultArchitecture.Equals("all", StringComparison.InvariantCultureIgnoreCase))
+
+			// return all supported if getting supported, compiling for distribution, or we want active, and "all" is selected
+			if (bGetAllSupported || bIsDistributionMode || DefaultArchitecture.Equals("all", StringComparison.InvariantCultureIgnoreCase))
 			{
 				if (bSupportsArm64)
 				{
@@ -362,12 +361,7 @@ namespace UnrealBuildTool
 				throw new BuildException($"Unknown {DefaultKey} value found ('{DefaultArchitecture}') in .ini");
 			}
 
-
-			CachedProjectArches = Architectures;
-			CachedArchesProject = ProjectFile;
-			CachedArchesTarget = TargetName;
-
-			return CachedProjectArches;
+			return Architectures;
 		}
 
 		/// <summary>
diff --git a/Engine/Source/Programs/UnrealBuildTool/ProjectFiles/ProjectFileGenerator.cs b/Engine/Source/Programs/UnrealBuildTool/ProjectFiles/ProjectFileGenerator.cs
--- a/Engine/Source/Programs/UnrealBuildTool/ProjectFiles/ProjectFileGenerator.cs
+++ b/Engine/Source/Programs/UnrealBuildTool/ProjectFiles/ProjectFileGenerator.cs
@@ -869,8 +869,8 @@ namespace UnrealBuildTool
 			// Sort the targets by name. When we have multiple targets of a given type for a project, we'll use the order to determine which goes in the primary project file (so that client names with a suffix will go into their own project).
 			AllTargetFiles = AllTargetFiles.OrderBy(x => x.FullName, StringComparer.OrdinalIgnoreCase).ToList();
 
-			// Remove any game projects that don't have a target
-			AllGameProjects.RemoveAll(x => !AllTargetFiles.Any(y => y.IsUnderDirectory(x.Directory)));
+			// Remove any game projects that don't have a target (or are under a directory "Programs" since they don't have Targets under the .uproject)
+			AllGameProjects.RemoveAll(x => !x.ContainsName("Programs", 0) && !AllTargetFiles.Any(y => y.IsUnderDirectory(x.Directory)));
 
 			Dictionary<FileReference, List<DirectoryReference>> AdditionalSearchPaths = new Dictionary<FileReference, List<DirectoryReference>>();
 
@@ -2381,6 +2381,13 @@ namespace UnrealBuildTool
 					RulesAssembly RulesAssembly;
 
 					FileReference? CheckProjectFile = AllGames.FirstOrDefault(x => TargetFilePath.IsUnderDirectory(x.Directory));
+					// if it wasn't found, it may be a program, and the target file isn't under the uproject, so look up by name
+					if (CheckProjectFile == null && TargetFilePath.ContainsName("Programs", 0))
+					{
+						string BaseName = TargetFilePath.GetFileNameWithoutAnyExtensions();
+						CheckProjectFile = AllGames.FirstOrDefault(x => x.GetFileNameWithoutAnyExtensions().Equals(BaseName, StringComparison.InvariantCultureIgnoreCase));
+					}
+
 					if (CheckProjectFile == null)
 					{
 						RulesAssembly = RulesCompiler.CreateEngineRulesAssembly(false, false, false, Logger);
diff --git a/Engine/Source/Programs/UnrealBuildTool/ProjectFiles/Xcode/XcodeProject.cs b/Engine/Source/Programs/UnrealBuildTool/ProjectFiles/Xcode/XcodeProject.cs
--- a/Engine/Source/Programs/UnrealBuildTool/ProjectFiles/Xcode/XcodeProject.cs
+++ b/Engine/Source/Programs/UnrealBuildTool/ProjectFiles/Xcode/XcodeProject.cs
@@ -639,54 +639,7 @@ namespace UnrealBuildTool.XcodeProjectXcconfig
 			// First time seeing this target?
 			if (!CachedMacProjectArchitectures.ContainsKey(TargetName))
 			{
-				// Default to Intel
-				IEnumerable<string> TargetArchitectures = new[] { MacExports.IntelArchitecture };
-
-				// These targets are known to work so are allow-listed
-				bool IsAllowed = MacExports.TargetsAllowedForAppleSilicon.Contains(TargetName, StringComparer.OrdinalIgnoreCase);
-
-				// determine the target architectures based on what's allowed/denied
-				if (IsAllowed)
-				{
-					TargetArchitectures = AllArchitectures;
-				}
-				else
-				{
-					// if this is an unspecified tool/program, default to Intel for installed builds because we know all of that works. 
-					if (Config.ProjectTarget!.TargetRules!.Type == TargetType.Program)
-					{
-						// For misc tools we default to Intel for installed builds because we know all of that works. 
-						TargetArchitectures = Unreal.IsEngineInstalled() ? new[] { MacExports.IntelArchitecture } : AllArchitectures;
-					}
-					else
-					{
-						// For project targets we default to Intel then check the project settings. Note the editor target will have
-						// been denied above already.
-						TargetArchitectures = new[] { MacExports.IntelArchitecture };
-
-						ConfigHierarchy EngineIni = ConfigCache.ReadHierarchy(ConfigHierarchyType.Engine, InProjectFile?.Directory, UnrealTargetPlatform.Mac);
-						string TargetArchitecture;
-						string Key = Config.ProjectTarget!.TargetRules!.Type == TargetType.Editor ? "EditorTargetArchitecture" : "TargetArchitecture";
-						if (EngineIni.GetString("/Script/MacTargetPlatform.MacTargetSettings", Key, out TargetArchitecture))
-						{
-							if (TargetArchitecture.Contains("Universal", StringComparison.OrdinalIgnoreCase))
-							{
-								TargetArchitectures = AllArchitectures;
-							}
-							else if (TargetArchitecture.Contains("Intel", StringComparison.OrdinalIgnoreCase))
-							{
-								TargetArchitectures = new[] { MacExports.IntelArchitecture };
-							}
-							else if (TargetArchitecture.Contains("Apple", StringComparison.OrdinalIgnoreCase))
-							{
-								TargetArchitectures = new[] { MacExports.AppleArchitecture };
-							}
-						}
-					}
-				}
-
-				// Cache this so we don't need to keep checking this file
-				CachedMacProjectArchitectures.Add(TargetName, TargetArchitectures);
+				CachedMacProjectArchitectures[TargetName] = PlatformExports.GetProjectArchitectures(UnrealTargetPlatform.Mac, InProjectFile, TargetName, bGetAllSupported: true);
 			}
 
 			return CachedMacProjectArchitectures[TargetName];
diff --git a/Engine/Source/Programs/UnrealBuildTool/ProjectFiles/Xcode/XcodeProjectLegacy.cs b/Engine/Source/Programs/UnrealBuildTool/ProjectFiles/Xcode/XcodeProjectLegacy.cs
--- a/Engine/Source/Programs/UnrealBuildTool/ProjectFiles/Xcode/XcodeProjectLegacy.cs
+++ b/Engine/Source/Programs/UnrealBuildTool/ProjectFiles/Xcode/XcodeProjectLegacy.cs
@@ -960,54 +960,7 @@ namespace UnrealBuildTool.XcodeProjectLegacy
 			// First time seeing this target?
 			if (!CachedMacProjectArcitectures.ContainsKey(TargetName))
 			{
-				// Default to Intel
-				IEnumerable<string> TargetArchitectures = new[] { MacExports.IntelArchitecture };
-
-				// These targets are known to work so are allow-listed
-				bool IsAllowed = MacExports.TargetsAllowedForAppleSilicon.Contains(TargetName, StringComparer.OrdinalIgnoreCase);
-
-				// determine the target architectures based on what's allowed/denied
-				if (IsAllowed)
-				{
-					TargetArchitectures = AllArchitectures;
-				}
-				else
-				{
-					// if this is an unspecified tool/program, default to Intel for installed builds because we know all of that works. 
-					if (Config.ProjectTarget!.TargetRules!.Type == TargetType.Program)
-					{
-						// For misc tools we default to Intel for installed builds because we know all of that works. 
-						TargetArchitectures = Unreal.IsEngineInstalled() ? new[] { MacExports.IntelArchitecture } : AllArchitectures;
-					}
-					else
-					{
-						// For project targets we default to Intel then check the project settings. Note the editor target will have
-						// been denied above already.
-						TargetArchitectures = new[] { MacExports.IntelArchitecture };
-
-						ConfigHierarchy EngineIni = ConfigCache.ReadHierarchy(ConfigHierarchyType.Engine, InProjectFile?.Directory, UnrealTargetPlatform.Mac);
-						string TargetArchitecture;
-						string Key = Config.ProjectTarget!.TargetRules!.Type == TargetType.Editor ? "EditorTargetArchitecture" : "TargetArchitecture";
-						if (EngineIni.GetString("/Script/MacTargetPlatform.MacTargetSettings", Key, out TargetArchitecture))
-						{
-							if (TargetArchitecture.Contains("Universal", StringComparison.OrdinalIgnoreCase))
-							{
-								TargetArchitectures = AllArchitectures;
-							}
-							else if (TargetArchitecture.Contains("Intel", StringComparison.OrdinalIgnoreCase))
-							{
-								TargetArchitectures = new[] { MacExports.IntelArchitecture };
-							}
-							else if (TargetArchitecture.Contains("Apple", StringComparison.OrdinalIgnoreCase))
-							{
-								TargetArchitectures = new[] { MacExports.AppleArchitecture };
-							}
-						}
-					}
-				}
-
-				// Cache this so we don't need to keep checking this file
-				CachedMacProjectArcitectures.Add(TargetName, TargetArchitectures);
+				CachedMacProjectArcitectures[TargetName] = PlatformExports.GetProjectArchitectures(UnrealTargetPlatform.Mac, InProjectFile, TargetName, bGetAllSupported:true);
 			}
 
 			return CachedMacProjectArcitectures[TargetName];
diff --git a/Engine/Source/Programs/UnrealBuildTool/System/PlatformExports.cs b/Engine/Source/Programs/UnrealBuildTool/System/PlatformExports.cs
--- a/Engine/Source/Programs/UnrealBuildTool/System/PlatformExports.cs
+++ b/Engine/Source/Programs/UnrealBuildTool/System/PlatformExports.cs
@@ -47,17 +47,20 @@ namespace UnrealBuildTool
 		}
 
 		/// <summary>
-		/// Gets the default architecture for a given platform
+		/// Gets the architecture(s) the platform wants to currently compile
 		/// </summary>
 		/// <param name="Platform">The platform to get the default architecture for</param>
 		/// <param name="ProjectFile">Project file to read settings from</param>
 		/// <param name="TargetName">Target name - most Programs don't have a project file, but they will have a Target name that can be used to determine architecture</param>
-		/// <returns>The default architecture</returns>
-		public static IEnumerable<string> GetProjectArchitectures(UnrealTargetPlatform Platform, FileReference? ProjectFile, string? TargetName=null)
+		/// <param name="bGetAllSupported">If true, return all supported architectures for this target</param>
+		/// <param name="bIsDistributionMode">If true, return architectures when packaging in distribution mode</param>
+		/// <returns>The architecture to build</returns>
+		public static IEnumerable<string> GetProjectArchitectures(UnrealTargetPlatform Platform, FileReference? ProjectFile, string? TargetName, bool bGetAllSupported=false, bool bIsDistributionMode=false)
 		{
-			return UEBuildPlatform.GetBuildPlatform(Platform).GetProjectArchitectures(ProjectFile, TargetName);
+			return UEBuildPlatform.GetBuildPlatform(Platform).GetProjectArchitectures(ProjectFile, TargetName, bGetAllSupported, bIsDistributionMode);
 		}
 
+
 		/// <summary>
 		/// Checks whether the given project has a default build configuration
 		/// </summary>
