Fix for Android architecture for building plugins with binary builds

#android

diff --git a/Engine/Source/Programs/AutomationTool/Scripts/BuildPluginCommand.Automation.cs b/Engine/Source/Programs/AutomationTool/Scripts/BuildPluginCommand.Automation.cs
--- a/Engine/Source/Programs/AutomationTool/Scripts/BuildPluginCommand.Automation.cs
+++ b/Engine/Source/Programs/AutomationTool/Scripts/BuildPluginCommand.Automation.cs
@@ -363,6 +363,10 @@ public sealed class BuildPlugin : BuildCommand
 				{
 					Arguments += String.Format(" -architecture={0}", SpecifiedArchitecture);
 				}
+				else if (Platform == UnrealTargetPlatform.Android)
+				{
+					Arguments += String.Format(" -architectures={0}", AndroidArchitectures);
+				}
 
 				if (!String.IsNullOrEmpty(InAdditionalArgs))
 				{
