Build Graph reports which previous steps created any modified build products.

#rnx

diff --git a/Engine/Source/Programs/AutomationTool/BuildGraph/BuildGraph.cs b/Engine/Source/Programs/AutomationTool/BuildGraph/BuildGraph.cs
--- a/Engine/Source/Programs/AutomationTool/BuildGraph/BuildGraph.cs
+++ b/Engine/Source/Programs/AutomationTool/BuildGraph/BuildGraph.cs
@@ -1027,6 +1027,15 @@ namespace AutomationTool
 				string Message;
 				if(!ModifiedFiles.ContainsKey(File.RelativePath) && !File.Compare(Unreal.RootDirectory, out Message))
 				{
+					// look up the previous nodes to help with error diagnosis
+					List<string> PreviousNodeNames = InputManifests.Where( x => x.Value.Files.Contains(File) )
+						.Select( x => x.Key.NodeName )
+						.ToList();
+					if (PreviousNodeNames.Count > 0)
+					{
+						Message += $" (previous {(PreviousNodeNames.Count==1?"step":"steps")}: {string.Join(" + ", PreviousNodeNames)})";
+					}
+
 					ModifiedFiles.Add(File.RelativePath, Message);
 				}
 			}
