Remove unused code that could cause error on Mac

#android

diff --git a/Engine/Source/Programs/AutomationTool/Android/AndroidFileClient.cs b/Engine/Source/Programs/AutomationTool/Android/AndroidFileClient.cs
--- a/Engine/Source/Programs/AutomationTool/Android/AndroidFileClient.cs
+++ b/Engine/Source/Programs/AutomationTool/Android/AndroidFileClient.cs
@@ -123,7 +123,6 @@ namespace AutomationTool
 				// create TCP/IP socket
 				try
 				{
-					IPHostEntry ipHost = Dns.GetHostEntry(Dns.GetHostName());
 					IPAddress ipAddr = IPAddress.Parse("127.0.0.1");
 					IPEndPoint localEndPoint = new IPEndPoint(ipAddr, Port);
 
@@ -1225,7 +1224,6 @@ namespace AutomationTool
 			// create TCP/IP socket
 			try
 			{
-				IPHostEntry ipHost = Dns.GetHostEntry(Dns.GetHostName());
 				IPAddress ipAddr = IPAddress.Parse(Address);
 				IPEndPoint localEndPoint = new IPEndPoint(ipAddr, UsePort);
 
diff --git a/Engine/Source/Programs/UnrealAndroidFileTool/Program.cs b/Engine/Source/Programs/UnrealAndroidFileTool/Program.cs
--- a/Engine/Source/Programs/UnrealAndroidFileTool/Program.cs
+++ b/Engine/Source/Programs/UnrealAndroidFileTool/Program.cs
@@ -16,7 +16,7 @@ namespace UnrealAndroidFileTool
 {
 	class Program
 	{
-		private static string AppVersion = "1.0.0";
+		private static string AppVersion = "1.0.1";
 
 		static void ShowHelp(string command)
 		{
