Don't use /as for msi files

#rnx

diff --git a/Engine/Source/Programs/AutomationTool/AutomationUtils/CommandUtils.cs b/Engine/Source/Programs/AutomationTool/AutomationUtils/CommandUtils.cs
--- a/Engine/Source/Programs/AutomationTool/AutomationUtils/CommandUtils.cs
+++ b/Engine/Source/Programs/AutomationTool/AutomationUtils/CommandUtils.cs
@@ -2869,17 +2869,19 @@ namespace AutomationTool
 		};
 
 		[SupportedOSPlatform("windows")]
-		public static void Sign(FileReference File, SignatureType SignatureType)
+		public static void Sign(FileReference File, SignatureType SignatureType, bool AllowMultipleSignatures = true)
 		{
 			List<FileReference> Files = new List<FileReference> { File };
-			Sign(Files, SignatureType);
+			Sign(Files, SignatureType, AllowMultipleSignatures);
 		}
 
 		[SupportedOSPlatform("windows")]
-		public static void Sign(List<FileReference> Files, SignatureType SignatureType)
+		public static void Sign(List<FileReference> Files, SignatureType SignatureType, bool AllowMultipleSignatures = true)
 		{
 			string SignToolPath = GetSignToolPath();
 			string SpecificStoreArg = bUseMachineStoreForCertificates ? " /sm" : "";
+			string MultipleSignatureArg = AllowMultipleSignatures ? " /as" : "";
+			string SHA1TimestampArg = AllowMultipleSignatures ? " /tr" : " /t";
 
 			for(int FileIdx = 0; FileIdx < Files.Count; )
 			{
@@ -2894,7 +2896,7 @@ namespace AutomationTool
 					StringBuilder CommandLine = new StringBuilder();
 					if(SignatureType == SignatureType.SHA1)
 					{
-						CommandLine.AppendFormat("sign{0} /a /as /n \"{1}\" /tr {2} /v", SpecificStoreArg, SigningIdentity, TimestampServersSHA1[NumAttempts % TimestampServersSHA1.Length]);
+						CommandLine.AppendFormat("sign{0} /a /n \"{1}\" {2} {3} /v {4}", SpecificStoreArg, SigningIdentity, SHA1TimestampArg, TimestampServersSHA1[NumAttempts % TimestampServersSHA1.Length], MultipleSignatureArg);
 					}
 					else if(SignatureType == SignatureType.SHA256)
 					{
@@ -3208,7 +3210,7 @@ namespace AutomationTool
 					FinalFiles.Add(new FileReference(TargetFileInfo));
 				}
 			}
-			CodeSignWindows.Sign(FinalFiles, CodeSignWindows.SignatureType.SHA1);
+			CodeSignWindows.Sign(FinalFiles, CodeSignWindows.SignatureType.SHA1, !FinalFiles.Any(x => x.HasExtension(".msi"))); // By default we append signatures, but disable if there are MSI files in the list
 			CodeSignWindows.Sign(FinalFiles.Where(x => !x.HasExtension(".msi")).ToList(), CodeSignWindows.SignatureType.SHA256); // MSI files can only have one signature; prefer SHA1 for compatibility
 		}
 	}
