BuildGraph: Add a better error when an exception is thrown parsing a task argument.

diff --git a/Engine/Source/Programs/AutomationTool/BuildGraph/BgNodeExecutor.cs b/Engine/Source/Programs/AutomationTool/BuildGraph/BgNodeExecutor.cs
--- a/Engine/Source/Programs/AutomationTool/BuildGraph/BgNodeExecutor.cs
+++ b/Engine/Source/Programs/AutomationTool/BuildGraph/BgNodeExecutor.cs
@@ -221,30 +221,38 @@ namespace AutomationTool
 				}
 
 				// If it's a collection type, split it into separate values
-				if (Parameter.CollectionType == null)
+				try
 				{
-					// Parse it and assign it to the parameters object
-					object? FieldValue = ParseValue(Value, Parameter.ValueType);
-					Parameter.FieldInfo.SetValue(ParametersObject, FieldValue);
-				}
-				else
-				{
-					// Get the collection, or create one if necessary
-					object? CollectionValue = Parameter.FieldInfo.GetValue(ParametersObject);
-					if (CollectionValue == null)
+					if (Parameter.CollectionType == null)
 					{
-						CollectionValue = Activator.CreateInstance(Parameter.FieldInfo.FieldType)!;
-						Parameter.FieldInfo.SetValue(ParametersObject, CollectionValue);
+						// Parse it and assign it to the parameters object
+						object? FieldValue = ParseValue(Value, Parameter.ValueType);
+						Parameter.FieldInfo.SetValue(ParametersObject, FieldValue);
 					}
-
-					// Parse the values and add them to the collection
-					List<string> ValueStrings = BgTaskImpl.SplitDelimitedList(Value);
-					foreach (string ValueString in ValueStrings)
+					else
 					{
-						object ElementValue = ParseValue(ValueString, Parameter.ValueType)!;
-						Parameter.CollectionType.InvokeMember("Add", BindingFlags.InvokeMethod | BindingFlags.Instance | BindingFlags.Public, null, CollectionValue, new object[] { ElementValue });
+						// Get the collection, or create one if necessary
+						object? CollectionValue = Parameter.FieldInfo.GetValue(ParametersObject);
+						if (CollectionValue == null)
+						{
+							CollectionValue = Activator.CreateInstance(Parameter.FieldInfo.FieldType)!;
+							Parameter.FieldInfo.SetValue(ParametersObject, CollectionValue);
+						}
+
+						// Parse the values and add them to the collection
+						List<string> ValueStrings = BgTaskImpl.SplitDelimitedList(Value);
+						foreach (string ValueString in ValueStrings)
+						{
+							object ElementValue = ParseValue(ValueString, Parameter.ValueType)!;
+							Parameter.CollectionType.InvokeMember("Add", BindingFlags.InvokeMethod | BindingFlags.Instance | BindingFlags.Public, null, CollectionValue, new object[] { ElementValue });
+						}
 					}
 				}
+				catch (Exception ex)
+				{
+					Logger.LogScriptError(TaskInfo.Location, "Unable to parse argument {Name} from {Value}", Name, Value);
+					Logger.LogDebug(ex, "Exception while parsing argument {Name}", Name);
+				}
 			}
 
 			// Construct the task
