Gauntlet - Enable Lumen ray tracing when enabling raytracing

diff --git a/Engine/Source/Programs/AutomationTool/Gauntlet/Unreal/Automation/UE.Automation.cs b/Engine/Source/Programs/AutomationTool/Gauntlet/Unreal/Automation/UE.Automation.cs
--- a/Engine/Source/Programs/AutomationTool/Gauntlet/Unreal/Automation/UE.Automation.cs
+++ b/Engine/Source/Programs/AutomationTool/Gauntlet/Unreal/Automation/UE.Automation.cs
@@ -246,7 +246,7 @@ namespace UE
 
 				if (RayTracing)
 				{
-					AppConfig.CommandLine += " -dpcvars=r.RayTracing=1,r.SkinCache.CompileShaders=1,AutomationAllowFrameTraceCapture=0";
+					AppConfig.CommandLine += " -dpcvars=r.RayTracing=1,r.SkinCache.CompileShaders=1,r.Lumen.HardwareRayTracing=1,AutomationAllowFrameTraceCapture=0";
 				}
 				else
 				{
