UAT: Include deleted files when searching for last code changelist.

diff --git a/Engine/Source/Programs/AutomationTool/AutomationUtils/P4Environment.cs b/Engine/Source/Programs/AutomationTool/AutomationUtils/P4Environment.cs
--- a/Engine/Source/Programs/AutomationTool/AutomationUtils/P4Environment.cs
+++ b/Engine/Source/Programs/AutomationTool/AutomationUtils/P4Environment.cs
@@ -211,7 +211,7 @@ namespace AutomationTool
 			if(String.IsNullOrEmpty(CodeChangelistString) && CommandUtils.P4CLRequired)
 			{
 				P4Connection Connection = new P4Connection(User, Client, ServerAndPort);
-				CodeChangelistString = DetectCurrentCodeCL(Connection, ClientRoot);
+				CodeChangelistString = DetectCurrentCodeCL(Connection, ClientRoot, Changelist);
 				CommandUtils.SetEnvVar(EnvVarNames.CodeChangelist, CodeChangelistString);
 			}
 			if(!String.IsNullOrEmpty(CodeChangelistString))
@@ -451,7 +451,7 @@ namespace AutomationTool
 		/// </summary>
 		/// <param name="ClientRootPath">Workspace path.</param>
 		/// <returns>Changelist number as a string.</returns>
-		private static string DetectCurrentCodeCL(P4Connection Connection, string ClientRootPath)
+		private static string DetectCurrentCodeCL(P4Connection Connection, string ClientRootPath, int Changelist)
 		{
 			CommandUtils.LogVerbose("uebp_CodeCL not set, detecting last code CL...");
 
@@ -460,7 +460,7 @@ namespace AutomationTool
 
 			foreach (string CodeExtension in EpicGames.Perforce.PerforceUtils.CodeExtensions)
 			{
-				P4Cmd.AppendFormat(" \"{0}/...{1}#have\"", CommandUtils.CombinePaths(PathSeparator.Depot, ClientRootPath), CodeExtension);
+				P4Cmd.AppendFormat(" \"{0}/...{1}@<={2}\"", CommandUtils.CombinePaths(PathSeparator.Depot, ClientRootPath), CodeExtension, Changelist);
 			}
 
 			IProcessResult P4Result = Connection.P4(P4Cmd.ToString(), AllowSpew: false);
