Horde: Add an event matcher for "Assertion failed" messages from the engine.

diff --git a/Engine/Source/Programs/AutomationTool/AutomationUtils/Matchers/CrashEventMatcher.cs b/Engine/Source/Programs/AutomationTool/AutomationUtils/Matchers/CrashEventMatcher.cs
--- a/Engine/Source/Programs/AutomationTool/AutomationUtils/Matchers/CrashEventMatcher.cs
+++ b/Engine/Source/Programs/AutomationTool/AutomationUtils/Matchers/CrashEventMatcher.cs
@@ -17,6 +17,8 @@ namespace AutomationUtils.Matchers
 
 		static readonly Regex s_appErrorPattern = new Regex(@"^\s*[A-Za-z]+: Error: appError called: ");
 
+		static readonly Regex s_assertionFailedPattern = new Regex(@"^Assertion failed: ");
+
 		public LogEventMatch? Match(ILogCursor cursor)
 		{
 			if (cursor.Contains("begin: stack for UAT"))
@@ -35,6 +37,11 @@ namespace AutomationUtils.Matchers
 				LogEventBuilder builder = new LogEventBuilder(cursor);
 				return builder.ToMatch(LogEventPriority.Normal, LogLevel.Error, KnownLogEvents.Engine_AppError);
 			}
+			if (cursor.IsMatch(s_assertionFailedPattern))
+			{
+				LogEventBuilder builder = new LogEventBuilder(cursor);
+				return builder.ToMatch(LogEventPriority.Normal, LogLevel.Error, KnownLogEvents.Engine_AssertionFailed);
+			}
 			if (cursor.Contains("AutomationTool: Stack:"))
 			{
 				LogEventBuilder builder = new LogEventBuilder(cursor);
diff --git a/Engine/Source/Programs/Horde/Horde.Agent.Tests/LogParserTests.cs b/Engine/Source/Programs/Horde/Horde.Agent.Tests/LogParserTests.cs
--- a/Engine/Source/Programs/Horde/Horde.Agent.Tests/LogParserTests.cs
+++ b/Engine/Source/Programs/Horde/Horde.Agent.Tests/LogParserTests.cs
@@ -175,6 +175,18 @@ namespace Horde.Agent.Tests
 			CheckEventGroup(logEvents, 0, 1, LogLevel.Error, KnownLogEvents.AutomationTool_CrashExitCode);
 		}
 
+		[TestMethod]
+		public void AssertionFailedEventMatcher()
+		{
+			string[] lines =
+			{
+				@"Assertion failed: !NumUsed.GetValue() [File:Runtime/Core/Public/Containers/LockFreeFixedSizeAllocator.h] [Line: 201]"
+			};
+
+			List<LogEvent> logEvents = Parse(lines);
+			CheckEventGroup(logEvents, 0, 1, LogLevel.Error, KnownLogEvents.Engine_AssertionFailed);
+		}
+
 		[TestMethod]
 		public void SymbolStripSpuriousEventMatcher()
 		{
diff --git a/Engine/Source/Programs/Shared/EpicGames.Core/KnownLogEvents.cs b/Engine/Source/Programs/Shared/EpicGames.Core/KnownLogEvents.cs
--- a/Engine/Source/Programs/Shared/EpicGames.Core/KnownLogEvents.cs
+++ b/Engine/Source/Programs/Shared/EpicGames.Core/KnownLogEvents.cs
@@ -85,6 +85,11 @@ namespace EpicGames.Core
 		/// </summary>
 		public static EventId Engine_AppError { get; } = new EventId(305);
 
+		/// <summary>
+		/// Engine: Assertion failed
+		/// </summary>
+		public static EventId Engine_AssertionFailed { get; } = new EventId(306);
+
 		/// <summary>
 		/// Engine: Shader compiler output
 		/// </summary>
