- Added MultiShelve UAT script that can shelve multiple changelists into a single shelf
- Added a Reopen command to P4Utils.cs
- Refactored Edit and Reopen (that take file list) into a BatchCommand function to split up huge file lists across multiple commands to respect a max commandline length
- Increased max comandline length from 1024 to the ProcessStartInfo max of 32699

diff --git a/Engine/Source/Programs/AutomationTool/AutomationUtils/P4Utils.cs b/Engine/Source/Programs/AutomationTool/AutomationUtils/P4Utils.cs
--- a/Engine/Source/Programs/AutomationTool/AutomationUtils/P4Utils.cs
+++ b/Engine/Source/Programs/AutomationTool/AutomationUtils/P4Utils.cs
@@ -2069,23 +2069,35 @@ namespace AutomationTool
 			LogP4("", "unshelve " + String.Format("-s {0} ", FromCL) + String.Format("-c {0} ", ToCL) + CommandLine, SpewIsVerbose: SpewIsVerbose);
 		}
 
-        /// <summary>
-        /// Invokes p4 unshelve command.
-        /// </summary>
-        /// <param name="FromCL">Changelist to unshelve.</param>
-        /// <param name="ToCL">Changelist where the checked out files should be added.</param>
-        /// <param name="CommandLine">Commandline for the command.</param>
-        public void Shelve(int FromCL, string CommandLine = "", bool AllowSpew = true)
-        {
-            LogP4("", "shelve " + String.Format("-r -c {0} ", FromCL) + CommandLine, AllowSpew: AllowSpew);
-        }
+		/// <summary>
+		/// Invokes p4 shelve command.
+		/// </summary>
+		/// <param name="FromCL">Changelist to unshelve.</param>
+		/// <param name="ToCL">Changelist where the checked out files should be added.</param>
+		/// <param name="CommandLine">Commandline for the command.</param>
+		public void Shelve(int FromCL, string CommandLine = "", bool AllowSpew = true)
+		{
+			LogP4("", "shelve " + String.Format("-r -c {0} ", FromCL) + CommandLine, AllowSpew: AllowSpew);
+		}
 
 		/// <summary>
-        /// Deletes shelved files from a changelist
+		/// Invokes p4 shelve command, without reverting existing shelf first (overwrites any existing shelved file)
+		/// This means that any files that are already shelved, but not in the FromCL being shelved, they will still exist
 		/// </summary>
-        /// <param name="FromCL">Changelist to unshelve.</param>
-        /// <param name="CommandLine">Commandline for the command.</param>
-        public void DeleteShelvedFiles(int FromCL, bool AllowSpew = true)
+		/// <param name="FromCL">Changelist to unshelve.</param>
+		/// <param name="ToCL">Changelist where the checked out files should be added.</param>
+		/// <param name="CommandLine">Commandline for the command.</param>
+		public void ShelveNoRevert(int FromCL, string CommandLine = "", bool AllowSpew = true)
+		{
+			LogP4("", "shelve " + String.Format("-f -c {0} ", FromCL) + CommandLine, AllowSpew: AllowSpew);
+		}
+
+		/// <summary>
+		/// Deletes shelved files from a changelist
+		/// </summary>
+		/// <param name="FromCL">Changelist to unshelve.</param>
+		/// <param name="CommandLine">Commandline for the command.</param>
+		public void DeleteShelvedFiles(int FromCL, bool AllowSpew = true)
         {
 			string Output;
             if (!LogP4Output(out Output, "", String.Format("shelve -d -c {0}", FromCL), AllowSpew: AllowSpew) && !Output.StartsWith("No shelved files in changelist to delete."))
@@ -2094,6 +2106,57 @@ namespace AutomationTool
 			}
         }
 
+		/// <summary>
+		/// Invoke a command on a list of files, breaking up into multiple commands so as not to blow the max commandline length
+		/// </summary>
+		/// <param name="Comand">The command, like "p4 edit -c 1000"</param>
+		/// <param name="Files"></param>
+		/// <param name="AllowSpew"></param>
+		public  void BatchedCommand(string Command, List<string> Files, bool AllowSpew = true)
+		{
+			// using the limit of ProcessStartInfo.Arguments
+			const int MaxCommandLineLength = 32699;
+
+			StringBuilder CommandLine = new StringBuilder();
+			for (int Idx = 0; Idx < Files.Count; Idx++)
+			{
+				if (CommandLine.Length + Files[Idx].Length + 3 > MaxCommandLineLength)
+				{
+					LogP4("", Command + CommandLine.ToString(), AllowSpew:AllowSpew);
+					CommandLine.Clear();
+				}
+				CommandLine.AppendFormat(" \"{0}\"", Files[Idx]);
+			}
+			if (CommandLine.Length > 0)
+			{
+				LogP4("", Command + CommandLine.ToString(), AllowSpew: AllowSpew);
+			}
+		}
+
+
+		/// <summary>
+		/// Invokes p4 reopen command.
+		/// </summary>
+		/// <param name="CL">Changelist where the checked out files should be moved to.</param>
+		/// <param name="CommandLine">Commandline for the command</param>
+		public void Reopen(int CL, string CommandLine, bool AllowSpew = true)
+		{
+			LogP4("", $"reopen -c {CL} {CommandLine}", AllowSpew: AllowSpew);
+		}
+
+		/// <summary>
+		/// Invokes p4 reopen command with a list of files. IMPORTANT: because of commandline limits
+		/// the file list will be broken up across multiple commands. This will fail if a "move" operation
+		/// (that is made up of two files) is split between commandlines - they must be moved in the
+		/// same operation
+		/// </summary>
+		/// <param name="CL">Changelist where the checked out files should be moved to.</param>
+		/// <param name="Files">List of files to be moved</param>
+		public void Reopen(int CL, List<string> Files, bool AllowSpew = true)
+		{
+			BatchedCommand($"reopen -c {CL}", Files, AllowSpew: AllowSpew);
+		}
+
 		/// <summary>
 		/// Invokes p4 edit command.
 		/// </summary>
@@ -2111,26 +2174,7 @@ namespace AutomationTool
 		/// <param name="CommandLine">Commandline for the command.</param>
 		public void Edit(int CL, List<string> Files, bool AllowSpew = true)
 		{
-			const int MaxCommandLineLength = 1024;
-
-			StringBuilder CommandLine = new StringBuilder();
-			for(int Idx = 0; Idx < Files.Count; Idx++)
-			{
-				if(CommandLine.Length + Files[Idx].Length + 3 > MaxCommandLineLength)
-				{
-					LogP4("", String.Format("edit -c {0} {1}", CL, CommandLine.ToString()), AllowSpew: AllowSpew);
-					CommandLine.Clear();
-				}
-				if(CommandLine.Length > 0)
-				{
-					CommandLine.Append(" ");
-				}
-				CommandLine.AppendFormat("\"{0}\"", Files[Idx]);
-			}
-			if(CommandLine.Length > 0)
-			{
-				LogP4("", String.Format("edit -c {0} {1}", CL, CommandLine.ToString()), AllowSpew: AllowSpew);
-			}
+			BatchedCommand($"edit -c {CL}", Files, AllowSpew: AllowSpew);
 		}
 
 		/// <summary>
@@ -2611,7 +2655,10 @@ namespace AutomationTool
 				int EndOffset = CmdOutput.LastIndexOf(EndStr);
 				if (Offset == 0 && Offset < EndOffset)
 				{
-					CommandUtils.LogInformation("Change {0} does not exist", CL);
+					if (AllowSpew)
+					{
+						CommandUtils.LogInformation("Change {0} does not exist", CL);
+					}
 					return false;
 				}
 
@@ -2625,7 +2672,10 @@ namespace AutomationTool
 				}
 
 				string Status = CmdOutput.Substring(StatusOffset + StatusStr.Length).TrimStart().Split('\n')[0].TrimEnd();
-				CommandUtils.LogInformation("Change {0} exists ({1})", CL, Status);
+				if (AllowSpew)
+				{
+					CommandUtils.LogInformation("Change {0} exists ({1})", CL, Status);
+				}
 				Pending = (Status == "pending");
 				return true;
 			}
diff --git a/Engine/Source/Programs/AutomationTool/Scripts/MultiShelve.cs b/Engine/Source/Programs/AutomationTool/Scripts/MultiShelve.cs
new file mode 100644
--- /dev/null
+++ b/Engine/Source/Programs/AutomationTool/Scripts/MultiShelve.cs
@@ -0,0 +1,84 @@
+﻿// Copyright Epic Games, Inc. All Rights Reserved.
+
+using AutomationTool;
+using System;
+using System.Collections.Generic;
+using System.Linq;
+using System.IO;
+using Microsoft.Extensions.Logging;
+using EpicGames.Core;
+using UnrealBuildBase;
+
+
+namespace BuildScripts.Automation
+{
+	[RequireP4]
+	[DoesNotNeedP4CL]
+	[Help("Shelves multiple changelists into one (useful for preflighting multiple changes together")]
+	[Help("TargetCL=<changelist>", "If not specified, the script will create a new changelist to shelve into, and write out the CL number")]
+	[Help("SourceCLs=<cl,cl,...>", "List of changelists to shelve into the TargetCL")]
+	class MultiShelve : BuildCommand
+	{
+		public override void ExecuteBuild()
+		{
+			string SourceCLsParam = ParseRequiredStringParam("SourceCLs");
+			IEnumerable<int> SourceCLs = SourceCLsParam.Split(",", StringSplitOptions.TrimEntries).Select(x => int.Parse(x));
+			int TargetCL = ParseParamInt("TargetCL", -1);
+			bool bAllowSpew = ParseParam("AllowSpew");
+			bool bCleanTarget = ParseParam("CleanTarget");
+
+			P4Connection P4 = new P4Connection(null, null);
+
+
+			if (TargetCL == -1)
+			{
+				TargetCL = P4.CreateChange(Description: $"Shelving changelists " + string.Join(", ", SourceCLs));
+				Log.TraceInformation($"Creating target changelist >> {TargetCL} << to shelve into");
+				Log.TraceInformation($"");
+			}
+			// make sure there aren't files already in the target CL
+			else
+			{
+				if (P4.ChangeFiles(TargetCL, out _, AllowSpew: bAllowSpew).Count > 0)
+				{
+					Log.TraceInformation($"");
+					Log.TraceInformation($"Changelist {TargetCL} already contains files. Note that these files will be reshelved for each source CL.");
+				}
+
+				if (bCleanTarget)
+				{
+					Log.TraceInformation($"");
+					Log.TraceInformation($"Deleting shelved files in target CL {TargetCL}");
+					P4.DeleteShelvedFiles(TargetCL, AllowSpew: bAllowSpew);
+				}	
+			}
+
+			foreach (int CL in SourceCLs)
+			{
+				List<string> Files = P4.ChangeFiles(CL, out _, AllowSpew: bAllowSpew);
+
+				Log.TraceInformation($"");
+				if (Files.Count == 0)
+				{
+					Log.TraceInformation($"Skipping CL {CL}, it had no files!");
+					continue;
+				}
+				Log.TraceInformation($"Shelving from {CL}:");
+
+
+				Log.TraceInformation($"  Moving {Files.Count} files from {CL}");
+				P4.Reopen(TargetCL, Files, AllowSpew: bAllowSpew);
+
+				Log.TraceInformation($"  Shelving...");
+				P4.ShelveNoRevert(TargetCL, AllowSpew: bAllowSpew);
+
+				Log.TraceInformation($"  Moving back...");
+				P4.Reopen(CL, Files, AllowSpew: bAllowSpew);
+			}
+			Log.TraceInformation($"");
+			Log.TraceInformation($"All done!");
+			Log.TraceInformation($"");
+			Log.TraceInformation($"");
+		}
+	}
+}
