Make OpenEditor BuildCommand accept rooted overrides for UnrealEditorApp and Project paths instead of trying to find them using Host or Project utils.

diff --git a/Engine/Source/Programs/AutomationTool/Scripts/OpenEditor.Automation.cs b/Engine/Source/Programs/AutomationTool/Scripts/OpenEditor.Automation.cs
--- a/Engine/Source/Programs/AutomationTool/Scripts/OpenEditor.Automation.cs
+++ b/Engine/Source/Programs/AutomationTool/Scripts/OpenEditor.Automation.cs
@@ -41,7 +41,15 @@ namespace AutomationTool
 
 		public override ExitCode Execute()
 		{
-			string EditorPath = HostPlatform.Current.GetUnrealExePath(UnrealEditorApp);
+			string EditorPath = string.Empty;
+			if (Path.IsPathRooted(UnrealEditorApp) && File.Exists(UnrealEditorApp))
+			{
+				EditorPath = UnrealEditorApp;
+			}
+			else
+			{
+				EditorPath = HostPlatform.Current.GetUnrealExePath(UnrealEditorApp);
+			}
 
 			string EditorArgs = "";
 
@@ -54,7 +62,11 @@ namespace AutomationTool
 			}
 			else
 			{
-				if (!String.IsNullOrEmpty(ProjectName))
+				if (Path.IsPathRooted(ProjectName) && File.Exists(ProjectName))
+				{
+					EditorArgs = ProjectName;
+				}
+				else if (!string.IsNullOrEmpty(ProjectName))
 				{
 					FileReference ProjectFile = ProjectUtils.FindProjectFileFromName(ProjectName);
 
