Adding a way to force compression during staging

diff --git a/Engine/Source/Programs/AutomationTool/AutomationUtils/ProjectParams.cs b/Engine/Source/Programs/AutomationTool/AutomationUtils/ProjectParams.cs
--- a/Engine/Source/Programs/AutomationTool/AutomationUtils/ProjectParams.cs
+++ b/Engine/Source/Programs/AutomationTool/AutomationUtils/ProjectParams.cs
@@ -1849,6 +1849,11 @@ namespace AutomationTool
 		/// Do not compress packages during cook, override game ProjectPackagingSettings to force it off
 		/// </summary>
 		public bool ForceUncompressed;
+		
+		/// <summary>
+		/// Compress packages during cook, override game ProjectPackagingSettings and Platform Hardware Compression settings to force it on
+		/// </summary>
+		public bool ForceCompressed;
 
 		/// <summary>
 		/// Additional parameters when generating the PAK file
diff --git a/Engine/Source/Programs/AutomationTool/Scripts/CopyBuildToStagingDirectory.Automation.cs b/Engine/Source/Programs/AutomationTool/Scripts/CopyBuildToStagingDirectory.Automation.cs
--- a/Engine/Source/Programs/AutomationTool/Scripts/CopyBuildToStagingDirectory.Automation.cs
+++ b/Engine/Source/Programs/AutomationTool/Scripts/CopyBuildToStagingDirectory.Automation.cs
@@ -2629,7 +2629,20 @@ namespace AutomationScripts
 				bCompressed_ProjectPackagingSettings = false;
 			}
 
-			if (CompressionFormats == "None" || CompressionFormats == "none")
+			bool bIsCompressionFormatNone = (CompressionFormats == "None" || CompressionFormats == "none");
+
+			if (Params.ForceCompressed)
+			{
+				Params.Compressed = true;
+				
+				// If we are forcing compression make sure we have a valid compression format
+				if (bIsCompressionFormatNone)
+				{
+					PlatformGameConfig.GetString("/Script/UnrealEd.ProjectPackagingSettings", "PackageCompressionFormat", out CompressionFormats);
+					// empty CompressionFormats string means "use default" eg. zlib;
+				}	
+			}
+			else if (bIsCompressionFormatNone)
 			{
 				// rather than pass "-compressed" + format=None through to UnrealPak/iostore
 				//	just turn off compression
