BenchmarkBuild

- Fix profiling of -xge and -noxge in the same pass resulting in redundant compile steps (UE-173565)
- Make -noclean also imply -nopostclean
- Add support for -ubtargs / -ubt2args / -ubt3args etc to allow profiling of different UBT args (similar to -cookargs / -cook2args etc)

Ex -

RunUAT BenchmarkBuild -project=Unreal -client -compile -ubt2args="-nodebuginfo" -ubt3args="-nopch" -ubt4args="-disableunity" -noxge

Task Unreal Build Client Win64: 00:08:12
Task Unreal Build Client Win64 (-nodebuginfo): 00:07:41
Task Unreal Build Client Win64 (-nopch): 00:13:43
Task Unreal Build Client Win64 (-disableunity): 00:22:49

diff --git a/Engine/Source/Programs/AutomationTool/Scripts/BenchmarkBuild/BenchmarkBuild.Automation.cs b/Engine/Source/Programs/AutomationTool/Scripts/BenchmarkBuild/BenchmarkBuild.Automation.cs
--- a/Engine/Source/Programs/AutomationTool/Scripts/BenchmarkBuild/BenchmarkBuild.Automation.cs
+++ b/Engine/Source/Programs/AutomationTool/Scripts/BenchmarkBuild/BenchmarkBuild.Automation.cs
@@ -48,8 +48,10 @@ namespace AutomationTool.Benchmark
 	[Help("All", "Shorthand for -editor -client -AllCompile -AllEditor -AllCook -AllDDC")]
 	[Help("editorxge", "Do a pass with XGE for editor DDC (default)")]
 	[Help("noeditorxge", "Do a pass without XGE for editor DDC")]
-	[Help("CookArgs=", "Extra args to use when cooking. -CookArgs1=\"-foo\" -CookArgs2=\"-bar\" will run two cooks with each argument set")]
-	[Help("PIEArgs=", "Extra args to use when running the editor. -PIEArgs1=\"-foo\" -PIEArgs2=\"-bar\" will run two PIE tests with each argument set")]
+	[Help("UBTArgs=", "Extra args to use when compiling. -UBTArgs=\"-foo\" -UBT2Args=\"-bar\" will run two compile passes with -foo and -bar")]
+	[Help("CookArgs=", "Extra args to use when cooking. -CookArgs=\"-foo\" -Cook2Args=\"-bar\" will run two cook passes with -foo and -bar")]
+	[Help("LaunchArgs=", "Extra args to use for launching. -LaunchArgs=\"-foo\" -Launch2Args=\"-bar\" will run two launch passes with -foo and -bar")]
+	[Help("PIEArgs=", "Extra args to use for PIE. -PIEArgs=\"-foo\" -PIE2Args=\"-bar\" will run two PIE passes with -foo and -bar")]
 	[Help("iterations=<n>", "How many times to perform each test)")]
 	[Help("wait=<n>", "How many seconds to wait between each test)")]
 	[Help("csv", "Name/path of file to write CSV results to. If empty the local machine name will be used")]
@@ -95,6 +97,7 @@ namespace AutomationTool.Benchmark
 			public UBTBuildOptions BuildOptions = UBTBuildOptions.None;
 			public int				TimeBetweenTasks = 0;
 
+			public List<string> UBTArgs = new List<string>();
 			public List<string> CookArgs = new List<string>();
 			public List<string> PIEArgs = new List<string>();
 			public List<string> LaunchArgs = new List<string>();
@@ -179,30 +182,62 @@ namespace AutomationTool.Benchmark
 				Iterations = ParseParamInt("Iterations", Iterations);
 				TimeBetweenTasks = ParseParamInt("Wait", TimeBetweenTasks);
 
-				// allow up to 10 cook & PIE variations. -Cook1Args=etc -Cook2Args=etc2 etc
+				// allow up to 10 UBT, Cook, PIE via -UBTArgs=etc, -UBT2Args=etc2, -CookArgs=etc -Cook2Args=etc2 etc
 				for (int i = 0; i < 10; i++)
 				{
-					string PostFix = i == 0 ? "" : i.ToString();
+					string PostFix = i == 0 ? "" : (i+1).ToString();
+
+					// Parse CookArgs, Cook2Args etc
 					string CookParam = ParseParamValue("Cook" + PostFix + "Args", null);
 
 					if (CookParam != null)
 					{
 						CookArgs.Add(CookParam);
 					}
+					else if (i == 0)
+					{
+						// add a default for the first cook
+						CookArgs.Add("");
+					}
 
+					// Parse PIEArgs, PIE22Args etc
 					string PIEParam = ParseParamValue("PIE" + PostFix + "Args", null);
 
 					if (PIEParam != null)
 					{
 						PIEArgs.Add(PIEParam);
 					}
+					else if (i == 0)
+					{
+						// add a default for the first PIE
+						PIEArgs.Add("");
+					}
 
+					// Parse LaunchArgs, Launch2Args etc
 					string LaunchParam = ParseParamValue("Launch" + PostFix + "Args", null);
 
 					if (!string.IsNullOrEmpty(LaunchParam))
 					{
 						LaunchArgs.Add(LaunchParam);
 					}
+					else if (i == 0)
+					{
+						// add a default for the first launch
+						LaunchArgs.Add("");
+					}
+
+					// Parse UBTArgs, UBT2Args etc
+					string UBTParam = ParseParamValue("UBT" + PostFix + "Args", null);
+
+					if (!string.IsNullOrEmpty(UBTParam))
+					{
+						UBTArgs.Add(UBTParam);
+					}
+					else if (i == 0)
+					{
+						// add a default for the first compile
+						UBTArgs.Add("");
+					}
 				}
 
 				FileName = ParseParamValue("csv", FileName);
@@ -255,7 +290,7 @@ namespace AutomationTool.Benchmark
 					BuildOptions |= UBTBuildOptions.PreClean;
 				}
 				// post-clean if we're building a lot of stuff
-				if (!ParseParam("nopostclean")
+				if (!(ParseParam("nopostclean") && !ParseParam("noclean"))
 					/*&& (PlatformsToTest.Count() > 1 || ProjectsToTest.Count() > 1)*/)
 				{
 					BuildOptions |= UBTBuildOptions.PostClean;
@@ -375,10 +410,6 @@ namespace AutomationTool.Benchmark
 
 			Dictionary<BenchmarkTaskBase, List<BenchmarkResult>> Results = new Dictionary<BenchmarkTaskBase, List<BenchmarkResult>>();
 
-
-			var XGECompileOptions = Options.XGEOptions.Where(Opt => (Opt == XGETaskOptions.WithXGE || Opt == XGETaskOptions.NoXGE));
-			var XGEEditorOptions = Options.XGEOptions.Where(Opt => (Opt == XGETaskOptions.WithEditorXGE || Opt == XGETaskOptions.NoEditorXGE));
-
 			for (int ProjectIndex = 0; ProjectIndex < Options.ProjectsToTest.Count(); ProjectIndex++)
 			{
 				string Project = Options.ProjectsToTest.ElementAt(ProjectIndex);
@@ -396,25 +427,26 @@ namespace AutomationTool.Benchmark
 				ProjectTargetInfo EditorTarget = new ProjectTargetInfo(ProjectFile, BuildHostPlatform.Current.Platform, TargetIsClientBuild);
 
 				// Do compile tests of editor and platforms
-				foreach (var XGEOption in XGECompileOptions)
+
+				if (Options.DoBuildEditorTests)
 				{
-					if (Options.DoBuildEditorTests)
-					{
-						Tasks.AddRange(AddBuildTests(ProjectFile, BuildHostPlatform.Current.Platform, "Editor", Options));
-					}
+					Tasks.AddRange(AddBuildTests(ProjectFile, BuildHostPlatform.Current.Platform, "Editor", Options));
+				}
 
+				if (Options.DoBuildClientTests)
+				{
 					foreach (var ClientPlatform in Options.PlatformsToTest)
 					{
 						ProjectTargetInfo PlatformTarget = new ProjectTargetInfo(ProjectFile, ClientPlatform, TargetIsClientBuild);
 
-						if (Options.DoBuildClientTests)
-						{
-							// do build tests
-							Tasks.AddRange(AddBuildTests(ProjectFile, ClientPlatform, TargetIsClientBuild ? "Client" : "Game", Options));
-						}
+
+						// do build tests
+						Tasks.AddRange(AddBuildTests(ProjectFile, ClientPlatform, TargetIsClientBuild ? "Client" : "Game", Options));
 					}
 				}
 
+				var XGEEditorOptions = Options.XGEOptions.Where(Opt => (Opt == XGETaskOptions.WithEditorXGE || Opt == XGETaskOptions.NoEditorXGE));
+
 				List<BenchmarkTaskBase> EditorTasks = new List<BenchmarkTaskBase>();
 
 				if (Options.DoLaunchEditorTests)
@@ -580,16 +612,24 @@ namespace AutomationTool.Benchmark
 
 			if (InOptions.DoCompileTests)
 			{
+				IEnumerable<string> UBTArgList = InOptions.UBTArgs.Any() ? InOptions.UBTArgs : new[] { "" };
+
 				if (InOptions.XGEOptions.Contains(XGETaskOptions.WithXGE))
 				{
-					NewTasks.Add(new BenchmarkBuildTask(InProjectFile, InTargetName, InPlatform, XGETaskOptions.WithXGE, "", 0, InOptions.BuildOptions));
+					foreach (string UBTArgs in UBTArgList)
+					{
+						NewTasks.Add(new BenchmarkBuildTask(InProjectFile, InTargetName, InPlatform, XGETaskOptions.WithXGE, UBTArgs, 0, InOptions.BuildOptions));
+					}
 				}
 
 				if (InOptions.XGEOptions.Contains(XGETaskOptions.NoXGE))
 				{
 					foreach (int ProcessorCount in InOptions.CoresForLocalJobs)
 					{
-						NewTasks.Add(new BenchmarkBuildTask(InProjectFile, InTargetName, InPlatform, XGETaskOptions.NoXGE, "", ProcessorCount, InOptions.BuildOptions));
+						foreach (string UBTArgs in UBTArgList)
+						{
+							NewTasks.Add(new BenchmarkBuildTask(InProjectFile, InTargetName, InPlatform, XGETaskOptions.NoXGE, UBTArgs, ProcessorCount, InOptions.BuildOptions));
+						}
 					}
 				}
 			}
