Fix ProjectName override bug in OpenEditor build command

Running ParseParamValue with a preset project path adds extra quotes to the output and breaks the IsPathRooted check

diff --git a/Engine/Source/Programs/AutomationTool/Scripts/OpenEditor.Automation.cs b/Engine/Source/Programs/AutomationTool/Scripts/OpenEditor.Automation.cs
--- a/Engine/Source/Programs/AutomationTool/Scripts/OpenEditor.Automation.cs
+++ b/Engine/Source/Programs/AutomationTool/Scripts/OpenEditor.Automation.cs
@@ -17,7 +17,7 @@ namespace AutomationTool
 	public class OpenEditor : BuildCommand
 	{
 		// exposed as a property so projects can derive and set this directly
-		public string ProjectName { get; set; }
+		public string ProjectName { get; set; } = string.Empty;
 
 		private string _UnrealEditorApp = "UnrealEditor";
 		public string UnrealEditorApp
@@ -53,7 +53,7 @@ namespace AutomationTool
 
 			string EditorArgs = "";
 
-			ProjectName = ParseParamValue("project", ProjectName);
+			ProjectName = (string.IsNullOrEmpty(ProjectName) || !Path.IsPathRooted(ProjectName)) ? ParseParamValue("project", ProjectName) : ProjectName;
 			IEnumerable<string> ParamList = null;
 
 			if (bNoProject)
