Allow cooked editor staging to exclude Content/Python directories

diff --git a/Engine/Source/Programs/AutomationTool/CookedEditor/CookedEditor.Automation.cs b/Engine/Source/Programs/AutomationTool/CookedEditor/CookedEditor.Automation.cs
--- a/Engine/Source/Programs/AutomationTool/CookedEditor/CookedEditor.Automation.cs
+++ b/Engine/Source/Programs/AutomationTool/CookedEditor/CookedEditor.Automation.cs
@@ -83,6 +83,7 @@ public class ModifyStageContext
 	public bool bStagePlatformDirs = true;
 	public bool bStageUAT = false;
 	public bool bIsForExternalDistribution = false;
+	public bool bStagePython = false;
 
 	public ConfigHelper ConfigHelper;
 
@@ -113,6 +114,7 @@ public class ModifyStageContext
 		bStagePlatformDirs = ConfigHelper.GetBool("bStagePlatformDirs");
 		bStageUAT = ConfigHelper.GetBool("bStageUAT");
 		bIsForExternalDistribution = ConfigHelper.GetBool("bIsForExternalDistribution");
+		bStagePython = ConfigHelper.GetBool("bStagePython");
 
 		// cache some useful properties
 		ProjectDirectory = Params.RawProjectPath.Directory;
