Stage Engine localization data for cooked editors

Some of the editor UI text actually comes from the engine

#rnx

diff --git a/Engine/Source/Programs/AutomationTool/CookedEditor/CookedEditor.Automation.cs b/Engine/Source/Programs/AutomationTool/CookedEditor/CookedEditor.Automation.cs
--- a/Engine/Source/Programs/AutomationTool/CookedEditor/CookedEditor.Automation.cs
+++ b/Engine/Source/Programs/AutomationTool/CookedEditor/CookedEditor.Automation.cs
@@ -472,7 +472,7 @@ public class MakeCookedEditor : BuildCommand
 		{
 			List<string> CulturesToStage = Project.GetCulturesToStage(Params, ConfigHelper.GameConfig);
 
-			string[] EditorLocalizationTargetsToStage = { "Category", "Editor", "EditorTutorials", "Keywords", "PropertyNames", "ToolTips" };
+			string[] EditorLocalizationTargetsToStage = { "Category", "Engine", "Editor", "EditorTutorials", "Keywords", "PropertyNames", "ToolTips" };
 			foreach (string EditorLocalizationTargetToStage in EditorLocalizationTargetsToStage)
 			{
 				if (Project.ShouldStageLocalizationTarget(SC, null, EditorLocalizationTargetToStage))
