AutomationTool: Fix resolving relative .uproject path for -ScriptsForProject

#rnx

diff --git a/Engine/Source/Programs/AutomationTool/Program.cs b/Engine/Source/Programs/AutomationTool/Program.cs
--- a/Engine/Source/Programs/AutomationTool/Program.cs
+++ b/Engine/Source/Programs/AutomationTool/Program.cs
@@ -50,12 +50,23 @@ namespace AutomationToolDriver
 				{
 					throw new Exception("The -ScriptsForProject argument may only be specified once");
 				}
+
 				string ProjectFileName = CurrentParam.Substring(CurrentParam.IndexOf('=') + 1).Replace("\"", "");
-				if (!File.Exists(ProjectFileName))
+				FileReference ProjectReference = NativeProjectsBase.FindProjectFile(ProjectFileName, Logger);
+
+				if (ProjectReference != null)
+				{
+					AutomationToolCommandLine.SetUnchecked(Option_ScriptsForProject, ProjectReference.FullName);
+					Logger.LogDebug("Found project file: {0}", ProjectReference.FullName);
+				}
+				else if (Path.IsPathFullyQualified(ProjectFileName))
 				{
 					throw new Exception($"Project '{ProjectFileName}' does not exist");
 				}
-				AutomationToolCommandLine.SetUnchecked(Option_ScriptsForProject, Path.GetFullPath(ProjectFileName));
+				else
+				{
+					throw new Exception($"Project '{ProjectFileName}' does not exist relative to any entries in *.uprojectdirs");
+				}
 			}
 			else if (CurrentParam.StartsWith(Option_ScriptDir + "=", StringComparison.InvariantCultureIgnoreCase))
 			{
@@ -63,13 +74,29 @@ namespace AutomationToolDriver
 				if (Directory.Exists(ScriptDir))
 				{
 					List<string> OutAdditionalScriptDirectories = (List<string>)AutomationToolCommandLine.GetValueUnchecked(Option_ScriptDir) ?? new List<string>();
-					OutAdditionalScriptDirectories.Add(ScriptDir);
+					OutAdditionalScriptDirectories.Add(Path.GetFullPath(ScriptDir));
 					AutomationToolCommandLine.SetUnchecked(Option_ScriptDir, OutAdditionalScriptDirectories);
-					Logger.LogDebug("Found additional script dir: {0}", ScriptDir);
+					Logger.LogDebug("Found additional script dir: {0}", Path.GetFullPath(ScriptDir));
 				}
 				else
 				{
-					throw new Exception($"Specified ScriptDir doesn't exist: {ScriptDir}");
+					DirectoryReference ScriptDirReference = NativeProjectsBase.FindRelativeDirectoryReference(ScriptDir, Logger);
+
+					if (ScriptDirReference != null)
+					{
+						List<string> OutAdditionalScriptDirectories = (List<string>)AutomationToolCommandLine.GetValueUnchecked(Option_ScriptDir) ?? new List<string>();
+						OutAdditionalScriptDirectories.Add(ScriptDirReference.FullName);
+						AutomationToolCommandLine.SetUnchecked(Option_ScriptDir, OutAdditionalScriptDirectories);
+						Logger.LogDebug("Found additional script dir: {0}", ScriptDirReference.FullName);
+					}
+					else if (Path.IsPathFullyQualified(ScriptDir))
+					{
+						throw new Exception($"Specified ScriptDir doesn't exist: {ScriptDir}");
+					}
+					else
+					{
+						throw new Exception($"Specified ScriptDir doesn't exist relative to any entries in *.uprojectdirs: {ScriptDir}");
+					}
 				}
 			}
 			else if (CurrentParam.StartsWith(Option_Telemetry + "=", StringComparison.InvariantCultureIgnoreCase))
diff --git a/Engine/Source/Programs/Shared/EpicGames.Build/System/NativeProjectsBase.cs b/Engine/Source/Programs/Shared/EpicGames.Build/System/NativeProjectsBase.cs
--- a/Engine/Source/Programs/Shared/EpicGames.Build/System/NativeProjectsBase.cs
+++ b/Engine/Source/Programs/Shared/EpicGames.Build/System/NativeProjectsBase.cs
@@ -125,6 +125,59 @@ namespace UnrealBuildBase
 			return CachedProjectFiles;
 		}
 
+		/// <summary>
+		/// Searches base directories for an existing relative pathed file
+		/// </summary>
+		/// <param name="File">File to search for</param>
+		/// <param name="Logger"></param>
+		/// <returns>A FileReference for the existing file, otherwise null</returns>
+		public static FileReference? FindRelativeFileReference(string File, ILogger Logger)
+		{
+			return EnumerateBaseDirectories(Logger)
+					.Select(x => FileReference.Combine(x, File))
+					.FirstOrDefault(x => FileReference.Exists(x));
+		}
+
+		/// <summary>
+		/// Searches base directories for an existing relative pathed directory
+		/// </summary>
+		/// <param name="Directory">Directory to search for</param>
+		/// <param name="Logger"></param>
+		/// <returns>A DirectoryReference for the existing directory, otherwise null</returns>
+		public static DirectoryReference? FindRelativeDirectoryReference(string Directory, ILogger Logger)
+		{
+			return EnumerateBaseDirectories(Logger)
+					.Select(x => DirectoryReference.Combine(x, Directory))
+					.FirstOrDefault(x => DirectoryReference.Exists(x));
+		}
+
+		/// <summary>
+		/// Takes a project name (e.g "ShooterGame") or path and attempt to find the existing uproject file in the base directories
+		/// </summary>
+		/// <param name="FilePath"></param>
+		/// <returns>A FileReference to an existing .uproject file, otherwise null</returns>
+		public static FileReference? FindProjectFile(string Project, ILogger Logger)
+		{
+			// Handle absolute paths, or relative paths from the current working directory
+			if (File.Exists(Project))
+			{
+				return new FileReference(Project);
+			}
+
+			if (Path.IsPathFullyQualified(Project))
+			{
+				// Absolute path not found, return null instead of searching
+				return null;
+			}
+
+			string ProjectName = Path.GetFileNameWithoutExtension(Project);
+
+			// Search known .uprojects by name, then as relative path, then relative path for Project/Project.uproject
+			return EnumerateProjectFiles(Logger).FirstOrDefault(x => string.Equals(x.GetFileNameWithoutExtension(), ProjectName, StringComparison.OrdinalIgnoreCase)) ??
+				FindRelativeFileReference(Project, Logger) ??
+				FindRelativeFileReference(Path.Combine(ProjectName, $"{ProjectName}.uproject"), Logger);
+		}
+
 		/// <summary>
 		/// Finds all target files under a given folder, and add them to the target name to project file map
 		/// </summary>
