#UE Pak performance. Save IoStore response files in the Engine saved folder instead of the Logs folder so they are not considered logs for build systems that look for and process logs.

diff --git a/Engine/Source/Programs/AutomationTool/Scripts/CopyBuildToStagingDirectory.Automation.cs b/Engine/Source/Programs/AutomationTool/Scripts/CopyBuildToStagingDirectory.Automation.cs
--- a/Engine/Source/Programs/AutomationTool/Scripts/CopyBuildToStagingDirectory.Automation.cs
+++ b/Engine/Source/Programs/AutomationTool/Scripts/CopyBuildToStagingDirectory.Automation.cs
@@ -238,7 +238,7 @@ namespace AutomationScripts
 			// implementation we will just assume that we require maximum security for this data.
 			bool bForceEncryption = !string.IsNullOrEmpty(EncryptionKeyGuid);
 			string PakName = Path.GetFileNameWithoutExtension(OutputLocation.FullName);
-			string UnrealPakResponseFileName = CombinePaths(CmdEnv.LogFolder, "PakList_" + PakName + ".txt");
+			string UnrealPakResponseFileName = CombinePaths(CmdEnv.EngineSavedFolder, "ResponseFiles", "PakList_" + PakName + ".txt");
 			WritePakResponseFile(UnrealPakResponseFileName, UnrealPakResponseFile, Compressed, RehydrateAssets, CryptoSettings, bForceEncryption);
 			CmdLine.AppendFormat(" -create={0}", CommandUtils.MakePathSafeToUseWithCommandLine(UnrealPakResponseFileName));
 
@@ -301,7 +301,7 @@ namespace AutomationScripts
 			// implementation we will just assume that we require maximum security for this data.
 			bool bForceEncryption = !string.IsNullOrEmpty(EncryptionKeyGuid);
 			bool RehydrateAssets = false;
-			string UnrealPakResponseFileName = CombinePaths(CmdEnv.LogFolder, "PakListIoStore_" + ContainerName + ".txt");
+			string UnrealPakResponseFileName = CombinePaths(CmdEnv.EngineSavedFolder, "ResponseFiles", "PakListIoStore_" + ContainerName + ".txt");
 			WritePakResponseFile(UnrealPakResponseFileName, UnrealPakResponseFile, bCompressed, RehydrateAssets, CryptoSettings, bForceEncryption);
 			CmdLine.AppendFormat(" -ResponseFile={0}", CommandUtils.MakePathSafeToUseWithCommandLine(UnrealPakResponseFileName));
 
@@ -3111,7 +3111,7 @@ namespace AutomationScripts
 				{
 					if (FileReference.Exists(OutputLocation) && !Params.IgnorePaksFromDifferentCookSource)
 					{
-						string UnrealPakResponseFileName = CombinePaths(CmdEnv.LogFolder, "PakList_" + OutputLocation.GetFileNameWithoutExtension() + ".txt");
+						string UnrealPakResponseFileName = CombinePaths(CmdEnv.EngineSavedFolder, "ResponseFiles", "PakList_" + OutputLocation.GetFileNameWithoutExtension() + ".txt");
 						if (File.Exists(UnrealPakResponseFileName) && FileReference.GetLastWriteTimeUtc(OutputLocation) > File.GetLastWriteTimeUtc(UnrealPakResponseFileName))
 						{
 							bCopiedExistingPak = true;
