Stage all Engine/Editor localization targets for a cooked editor, even if they've been disabled for use in game

#rnx

diff --git a/Engine/Source/Programs/AutomationTool/CookedEditor/CookedEditor.Automation.cs b/Engine/Source/Programs/AutomationTool/CookedEditor/CookedEditor.Automation.cs
--- a/Engine/Source/Programs/AutomationTool/CookedEditor/CookedEditor.Automation.cs
+++ b/Engine/Source/Programs/AutomationTool/CookedEditor/CookedEditor.Automation.cs
@@ -475,7 +475,8 @@ public class MakeCookedEditor : BuildCommand
 			string[] EditorLocalizationTargetsToStage = { "Category", "Engine", "Editor", "EditorTutorials", "Keywords", "PropertyNames", "ToolTips" };
 			foreach (string EditorLocalizationTargetToStage in EditorLocalizationTargetsToStage)
 			{
-				if (Project.ShouldStageLocalizationTarget(SC, null, EditorLocalizationTargetToStage))
+				// Note: We skip "ShouldStageLocalizationTarget" below as games may disable certain targets that they don't need at runtime (eg, Engine), but all of these targets are still needed for an editor!
+				//if (Project.ShouldStageLocalizationTarget(SC, null, EditorLocalizationTargetToStage))
 				{
 					Project.StageLocalizationDataForTarget(SC, CulturesToStage, DirectoryReference.Combine(Unreal.EngineDirectory, "Content", "Localization", EditorLocalizationTargetToStage));
 				}
