Fix Gauntlet issue matching Unreal Linux executable

diff --git a/Engine/Source/Programs/AutomationTool/Gauntlet/Unreal/BuildSource/Gauntlet.StagedBuildSource.cs b/Engine/Source/Programs/AutomationTool/Gauntlet/Unreal/BuildSource/Gauntlet.StagedBuildSource.cs
--- a/Engine/Source/Programs/AutomationTool/Gauntlet/Unreal/BuildSource/Gauntlet.StagedBuildSource.cs
+++ b/Engine/Source/Programs/AutomationTool/Gauntlet/Unreal/BuildSource/Gauntlet.StagedBuildSource.cs
@@ -258,6 +258,11 @@ namespace Gauntlet
 						{
 							// Look at files & directories since apps on Mac are bundles
 							FileSystemInfo[] AppFiles = Di.GetFileSystemInfos(FileMatch);
+							if(string.IsNullOrEmpty(InExecutableExtension))
+							{
+								// Special case for empty extension (linux), filter out files with extension
+								AppFiles = AppFiles.Where(F => string.IsNullOrEmpty(F.Extension)).ToArray();
+							}
 							Binaries.AddRange(AppFiles);
 						}
 					}
