- Added support for DefaultGameTarget and DefaultClientTarget in Engine.ini like we have for DefaultEditorTarget and DefaultServerTarget (in ProjectParams)
- Made use of it for Lyra to select LyraGame as default (instead of LyraGameEOS)

diff --git a/Engine/Source/Programs/AutomationTool/AutomationUtils/ProjectParams.cs b/Engine/Source/Programs/AutomationTool/AutomationUtils/ProjectParams.cs
--- a/Engine/Source/Programs/AutomationTool/AutomationUtils/ProjectParams.cs
+++ b/Engine/Source/Programs/AutomationTool/AutomationUtils/ProjectParams.cs
@@ -2328,7 +2328,7 @@ namespace AutomationTool
 			{
 				if (AvailableTargets.Count > 1)
 				{
-					throw new AutomationException($"Project contains multiple {Type} targets but no {ConfigKey} is set in the [/Script/BuildSettings.BuildSettings] section of DefaultEngine.ini");
+					throw new AutomationException($"Project contains multiple {Type} targets ({string.Join(", ", AvailableTargets)}) but no {ConfigKey} is set in the [/Script/BuildSettings.BuildSettings] section of DefaultEngine.ini");
 				}
 
 				if (AvailableTargets.Count > 0)
@@ -2464,38 +2464,12 @@ namespace AutomationTool
 
 				if (Client)
 				{
-					GameTarget = ChooseTarget(AvailableClientTargets, TargetType.Client);
+					SelectDefaultTarget(AvailableClientTargets, TargetType.Client, ref GameTarget);
 					ProjectType = TargetType.Client;
 				}
 				else if (AvailableGameTargets.Count > 0 && !NoClient)
 				{
-					if (AvailableGameTargets.Count > 1)
-					{
-						StringBuilder TargetMessage = new StringBuilder("Multiple game targets found for project. Specify the desired target using the -Target=... argument.");
-
-						List<SingleTargetProperties> Targets = DetectedTargets.FindAll(Target => Target.Rules.Type == TargetType.Game);
-						foreach (SingleTargetProperties Target in Targets)
-						{
-							// search the list of script files to see if we can find a likely source for this class
-							// {TargetName}.Target.cs is expected to contain a definition for a class {TargetName}Target
-							// So we can do an imperfect reverse-lookup, and try to find a source file that has the expected pattern.
-
-							List<FileReference> PossibleScriptFiles = Properties.TargetScripts.FindAll(File => String.Equals(File.GetFileNameWithoutAnyExtensions(), Target.TargetName));
-
-							if (PossibleScriptFiles.Count > 0)
-							{
-								TargetMessage.Append($"\n  Could be \"{Target.TargetName}\" ({String.Join(" or ", PossibleScriptFiles)})");
-							}
-							else
-							{
-								TargetMessage.Append($"\n  Could be \"{Target.TargetName}\"");
-							}
-						}
-
-						throw new AutomationException(TargetMessage.ToString());
-					}
-
-					GameTarget = AvailableGameTargets.First();
+					SelectDefaultTarget(AvailableGameTargets, TargetType.Game, ref GameTarget);
 				}
 
 				SelectDefaultTarget(AvailableEditorTargets, TargetType.Editor, ref EditorTarget);
