Removed the building of C++ UHT from UAT code

diff --git a/Engine/Source/Programs/AutomationTool/BuildGraph/Scripts/InstalledBuild.cs b/Engine/Source/Programs/AutomationTool/BuildGraph/Scripts/InstalledBuild.cs
--- a/Engine/Source/Programs/AutomationTool/BuildGraph/Scripts/InstalledBuild.cs
+++ b/Engine/Source/Programs/AutomationTool/BuildGraph/Scripts/InstalledBuild.cs
@@ -142,13 +142,9 @@ namespace AutomationTool
 			BgNode VersionFilesNode = EditorWin64
 				.AddNode(x => UpdateVersionFilesAsync(x));
 
-			BgNode<BgFileSet> WinUhtNode = EditorWin64
-				.AddNode(x => CompileUnrealHeaderToolAsync(x, UnrealTargetPlatform.Win64))
-				.Requires(VersionFilesNode);
-
 			BgNode<BgFileSet> WinEditorNode = EditorWin64
 				.AddNode(x => CompileUnrealEditorWin64Async(x, CrashReporterCompileArgs, EmbedSrcSrvInfo, CompileDatasmithPlugins, WithFullDebugInfo, SignExecutables))
-				.Requires(WinUhtNode);
+				.Requires(VersionFilesNode);
 
 			Aggregates.Add(new BgAggregate("Win64 Editor", WinEditorNode, label: "Editors/Win64"));
 
@@ -161,7 +157,7 @@ namespace AutomationTool
 
 			BgNode<BgFileSet> WinGame = TargetWin64
 				.AddNode(x => CompileUnrealGameWin64(x, GameConfigurations, EmbedSrcSrvInfo, WithFullDebugInfo, SignExecutables))
-				.Requires(WinUhtNode);
+				.Requires(VersionFilesNode);
 
 			Aggregates.Add(new BgAggregate("TargetPlatforms_Win64", WinGame));
 
@@ -180,7 +176,7 @@ namespace AutomationTool
 
 			BgNode<BgFileSet> WinTools = ToolsGroupWin64
 				.AddNode(x => BuildToolsWin64Async(x, CrashReporterCompileArgs))
-				.Requires(WinUhtNode);
+				.Requires(VersionFilesNode);
 
 			BgNode<BgFileSet> CsTools = ToolsGroupWin64
 				.AddNode(x => BuildToolsCSAsync(x, SignExecutables))
@@ -195,7 +191,7 @@ namespace AutomationTool
 			DDCPlatformsWin64 = DDCPlatformsWin64.If(WithWin64, x => x.Add("Windows"));
 
 			BgNode DdcNode = DDCGroupWin64
-				.AddNode(x => BuildDDCWin64Async(x, DDCPlatformsWin64, BgList<BgFileSet>.Create(WinUhtNode.Output, WinEditorNode.Output, WinTools.Output)))
+				.AddNode(x => BuildDDCWin64Async(x, DDCPlatformsWin64, BgList<BgFileSet>.Create(WinEditorNode.Output, WinTools.Output)))
 				.Requires(WinEditorNode, WinTools);
 
 
@@ -205,7 +201,6 @@ namespace AutomationTool
 			BgAgent WinStageAgent = new BgAgent("Installed Build Group Win64", "Win64_Licensee");
 
 			BgList<BgFileSet> WinInstalledFiles = BgList<BgFileSet>.Empty;
-			WinInstalledFiles = WinInstalledFiles.Add(WinUhtNode.Output);
 			WinInstalledFiles = WinInstalledFiles.Add(WinEditorNode.Output);
 			WinInstalledFiles = WinInstalledFiles.Add(WinTools.Output);
 			WinInstalledFiles = WinInstalledFiles.Add(CsTools.Output);
@@ -239,15 +234,6 @@ namespace AutomationTool
 			}
 		}
 
-		/// <summary>
-		/// Builds UHT for any platform
-		/// </summary>
-		[BgNodeName("Compile UnrealHeaderTool {Platform}")]
-		static async Task<BgFileSet> CompileUnrealHeaderToolAsync(BgContext State, UnrealTargetPlatform Platform)
-		{
-			return await CompileAsync("UnrealHeaderTool", Platform, UnrealTargetConfiguration.Development, Arguments: "-precompile -allmodules");
-		}
-
 		/// <summary>
 		/// Builds the Windows editor
 		/// </summary>
diff --git a/Engine/Source/Programs/AutomationTool/Scripts/BuildCommonTools.Automation.cs b/Engine/Source/Programs/AutomationTool/Scripts/BuildCommonTools.Automation.cs
--- a/Engine/Source/Programs/AutomationTool/Scripts/BuildCommonTools.Automation.cs
+++ b/Engine/Source/Programs/AutomationTool/Scripts/BuildCommonTools.Automation.cs
@@ -91,7 +91,6 @@ public class BuildCommonTools : BuildCommand
 		if(Platforms.Contains(UnrealBuildTool.UnrealTargetPlatform.Win64))
 		{
 			Agenda.AddTarget("CrashReportClient", UnrealBuildTool.UnrealTargetPlatform.Win64, UnrealBuildTool.UnrealTargetConfiguration.Development);
-			Agenda.AddTarget("UnrealHeaderTool", UnrealBuildTool.UnrealTargetPlatform.Win64, UnrealBuildTool.UnrealTargetConfiguration.Development);
 			Agenda.AddTarget("UnrealPak", UnrealBuildTool.UnrealTargetPlatform.Win64, UnrealBuildTool.UnrealTargetConfiguration.Development);
 			Agenda.AddTarget("UnrealLightmass", UnrealBuildTool.UnrealTargetPlatform.Win64, UnrealBuildTool.UnrealTargetConfiguration.Development);
 			Agenda.AddTarget("InterchangeWorker", UnrealBuildTool.UnrealTargetPlatform.Win64, UnrealBuildTool.UnrealTargetConfiguration.Development);
diff --git a/Engine/Source/Programs/AutomationTool/Scripts/BuildPluginCommand.Automation.cs b/Engine/Source/Programs/AutomationTool/Scripts/BuildPluginCommand.Automation.cs
--- a/Engine/Source/Programs/AutomationTool/Scripts/BuildPluginCommand.Automation.cs
+++ b/Engine/Source/Programs/AutomationTool/Scripts/BuildPluginCommand.Automation.cs
@@ -209,10 +209,6 @@ public sealed class BuildPlugin : BuildCommand
 			CommandUtils.LogInformation("Building plugin for host platforms: {0}", String.Join(", ", HostPlatforms));
 			foreach (UnrealTargetPlatform HostPlatform in HostPlatforms)
 			{
-				if (Plugin.SupportedPrograms != null && Plugin.SupportedPrograms.Contains("UnrealHeaderTool"))
-				{
-					CompilePluginWithUBT(UBTExe, HostProjectFile, HostProjectPluginFile, Plugin, "UnrealHeaderTool", TargetType.Program, HostPlatform, UnrealTargetConfiguration.Development, ManifestFileNames, String.Format("{0} -plugin={1}", AdditionalArgs, CommandUtils.MakePathSafeToUseWithCommandLine(HostProjectPluginFile.FullName)));
-				}
 				CompilePluginWithUBT(UBTExe, HostProjectFile, HostProjectPluginFile, Plugin, "UnrealEditor", TargetType.Editor, HostPlatform, UnrealTargetConfiguration.Development, ManifestFileNames, AdditionalArgs);
 			}
 		}
@@ -250,10 +246,6 @@ public sealed class BuildPlugin : BuildCommand
 			CommandUtils.LogInformation("Building plugin for host platforms: {0}", String.Join(", ", HostPlatforms));
 			foreach (UnrealTargetPlatform HostPlatform in HostPlatforms)
 			{
-				if (Plugin.SupportedPrograms != null && Plugin.SupportedPrograms.Contains("UnrealHeaderTool"))
-				{
-					CompilePluginWithUBT(UnrealBuildToolDll, HostProjectFile, HostProjectPluginFile, Plugin, "UnrealHeaderTool", TargetType.Program, HostPlatform, UnrealTargetConfiguration.Development, ManifestFileNames, String.Format("{0} -plugin={1}", AdditionalArgs, CommandUtils.MakePathSafeToUseWithCommandLine(HostProjectPluginFile.FullName)));
-				}
 				CompilePluginWithUBT(UnrealBuildToolDll, HostProjectFile, HostProjectPluginFile, Plugin, "UnrealEditor", TargetType.Editor, HostPlatform, UnrealTargetConfiguration.Development, ManifestFileNames, AdditionalArgs);
 			}
 		}
diff --git a/Engine/Source/Programs/AutomationTool/Scripts/BuildProjectCommand.Automation.cs b/Engine/Source/Programs/AutomationTool/Scripts/BuildProjectCommand.Automation.cs
--- a/Engine/Source/Programs/AutomationTool/Scripts/BuildProjectCommand.Automation.cs
+++ b/Engine/Source/Programs/AutomationTool/Scripts/BuildProjectCommand.Automation.cs
@@ -95,11 +95,6 @@ namespace AutomationScripts
 				if (!Unreal.IsEngineInstalled())
 				{
 					CrashReportPlatforms.Add(EditorPlatform);
-					if (Params.EditorTargets.Contains("UnrealHeaderTool") == false)
-					{
-						Agenda.AddTargets(new string[] { "UnrealHeaderTool" }, EditorPlatform, EditorConfiguration, InAddArgs:ProgramArchitecture);
-						Agenda.AddTargets(new string[] { "UnrealHeaderTool" }, EditorPlatform, EditorConfiguration, Params.CodeBasedUprojectPath, InAddArgs: ProgramArchitecture);
-					}
 					if (Params.EditorTargets.Contains("ShaderCompileWorker") == false)
 					{
 						Agenda.AddTargets(new string[] { "ShaderCompileWorker" }, EditorPlatform, EditorConfiguration, InAddArgs: ProgramArchitecture);
diff --git a/Engine/Source/Programs/AutomationTool/Scripts/CopySharedCookedBuild.Automation.cs b/Engine/Source/Programs/AutomationTool/Scripts/CopySharedCookedBuild.Automation.cs
--- a/Engine/Source/Programs/AutomationTool/Scripts/CopySharedCookedBuild.Automation.cs
+++ b/Engine/Source/Programs/AutomationTool/Scripts/CopySharedCookedBuild.Automation.cs
@@ -129,7 +129,6 @@ class CopySharedCookedBuild : BuildCommand
 				UnrealTargetPlatform CurrentPlatform = HostPlatform.Current.HostEditorPlatform;
 
 				UnrealBuild.BuildAgenda Agenda = new UnrealBuild.BuildAgenda();
-				Agenda.AddTarget("UnrealHeaderTool", CurrentPlatform, UnrealTargetConfiguration.Development);
 				Agenda.AddTarget(EditorTargetName, CurrentPlatform, UnrealTargetConfiguration.Development, ProjectFileName.EndsWith(".uproject", StringComparison.InvariantCultureIgnoreCase)? new FileReference(ProjectFileName) : null);
 				Agenda.AddTarget("ShaderCompileWorker", CurrentPlatform, UnrealTargetConfiguration.Development);
 				Agenda.AddTarget("UnrealLightmass", CurrentPlatform, UnrealTargetConfiguration.Development);
diff --git a/Engine/Source/Programs/AutomationTool/Scripts/MegaXGE.Automation.cs b/Engine/Source/Programs/AutomationTool/Scripts/MegaXGE.Automation.cs
--- a/Engine/Source/Programs/AutomationTool/Scripts/MegaXGE.Automation.cs
+++ b/Engine/Source/Programs/AutomationTool/Scripts/MegaXGE.Automation.cs
@@ -38,22 +38,6 @@ class MegaXGE : BuildCommand
 
 		var Agenda = new UnrealBuild.BuildAgenda();
 
-		// we need to always build UHT when we use mega XGE
-		var ProgramTargets = new string[] 
-		{
-			"UnrealHeaderTool",
-		};
-		Agenda.AddTargets(ProgramTargets, UnrealTargetPlatform.Win64, UnrealTargetConfiguration.Development);
-		if (Clean)
-		{
-			LogSetProgress(ShowProgress, "Cleaning previous builds...");
-			foreach (var CurTarget in ProgramTargets)
-			{
-				string Args = String.Format("{0} {1} {2}", CurTarget, UnrealTargetPlatform.Win64.ToString(), UnrealTargetConfiguration.Development.ToString());
-				RunAndLog(CmdEnv, CleanToolLocation, Args);
-			}
-		}
-
 		LogInformation("*************************");
 		for (int Arg = 1; Arg < 100; Arg++)
 		{
diff --git a/Engine/Source/Programs/AutomationTool/Scripts/RebuildHLODCommand.Automation.cs b/Engine/Source/Programs/AutomationTool/Scripts/RebuildHLODCommand.Automation.cs
--- a/Engine/Source/Programs/AutomationTool/Scripts/RebuildHLODCommand.Automation.cs
+++ b/Engine/Source/Programs/AutomationTool/Scripts/RebuildHLODCommand.Automation.cs
@@ -87,7 +87,6 @@ namespace AutomationScripts.Automation
 		{
 			LogInformation("Running Step:- RebuildHLOD::BuildNecessaryTargets");
 			UnrealBuild.BuildAgenda Agenda = new UnrealBuild.BuildAgenda();
-            Agenda.AddTarget("UnrealHeaderTool", UnrealBuildTool.UnrealTargetPlatform.Win64, UnrealBuildTool.UnrealTargetConfiguration.Development);
 			Agenda.AddTarget("ShaderCompileWorker", UnrealBuildTool.UnrealTargetPlatform.Win64, UnrealBuildTool.UnrealTargetConfiguration.Development);			
 			Agenda.AddTarget(CommandletTargetName, UnrealBuildTool.UnrealTargetPlatform.Win64, UnrealBuildTool.UnrealTargetConfiguration.Development);
 
diff --git a/Engine/Source/Programs/AutomationTool/Scripts/RebuildLightMapsCommand.Automation.cs b/Engine/Source/Programs/AutomationTool/Scripts/RebuildLightMapsCommand.Automation.cs
--- a/Engine/Source/Programs/AutomationTool/Scripts/RebuildLightMapsCommand.Automation.cs
+++ b/Engine/Source/Programs/AutomationTool/Scripts/RebuildLightMapsCommand.Automation.cs
@@ -68,7 +68,6 @@ namespace AutomationScripts.Automation
 		{
 			LogInformation("Running Step:- RebuildLightMaps::BuildNecessaryTargets");
 			UnrealBuild.BuildAgenda Agenda = new UnrealBuild.BuildAgenda();
-            Agenda.AddTarget("UnrealHeaderTool", UnrealBuildTool.UnrealTargetPlatform.Win64, UnrealBuildTool.UnrealTargetConfiguration.Development);
 			Agenda.AddTarget("ShaderCompileWorker", UnrealBuildTool.UnrealTargetPlatform.Win64, UnrealBuildTool.UnrealTargetConfiguration.Development);
 			Agenda.AddTarget("UnrealLightmass", UnrealBuildTool.UnrealTargetPlatform.Win64, UnrealBuildTool.UnrealTargetConfiguration.Development);
 			Agenda.AddTarget(CommandletTargetName, UnrealBuildTool.UnrealTargetPlatform.Win64, UnrealBuildTool.UnrealTargetConfiguration.Development);
diff --git a/Engine/Source/Programs/AutomationTool/Scripts/ResavePackagesCommand.Automation.cs b/Engine/Source/Programs/AutomationTool/Scripts/ResavePackagesCommand.Automation.cs
--- a/Engine/Source/Programs/AutomationTool/Scripts/ResavePackagesCommand.Automation.cs
+++ b/Engine/Source/Programs/AutomationTool/Scripts/ResavePackagesCommand.Automation.cs
@@ -70,7 +70,6 @@ namespace AutomationScripts.Automation
 		{
 			LogInformation("Running Step:- RebuildLightMaps::BuildNecessaryTargets");
 			UnrealBuild.BuildAgenda Agenda = new UnrealBuild.BuildAgenda();
-            Agenda.AddTarget("UnrealHeaderTool", UnrealBuildTool.UnrealTargetPlatform.Win64, UnrealBuildTool.UnrealTargetConfiguration.Development);
 			Agenda.AddTarget("ShaderCompileWorker", UnrealBuildTool.UnrealTargetPlatform.Win64, UnrealBuildTool.UnrealTargetConfiguration.Development);
 			Agenda.AddTarget("UnrealLightmass", UnrealBuildTool.UnrealTargetPlatform.Win64, UnrealBuildTool.UnrealTargetConfiguration.Development);
 			Agenda.AddTarget(CommandletTargetName, UnrealBuildTool.UnrealTargetPlatform.Win64, UnrealBuildTool.UnrealTargetConfiguration.Development);
diff --git a/Engine/Source/Programs/AutomationTool/Scripts/UnrealBuildUtils.cs b/Engine/Source/Programs/AutomationTool/Scripts/UnrealBuildUtils.cs
--- a/Engine/Source/Programs/AutomationTool/Scripts/UnrealBuildUtils.cs
+++ b/Engine/Source/Programs/AutomationTool/Scripts/UnrealBuildUtils.cs
@@ -28,22 +28,6 @@ public class UnrealBuildUtils : CommandUtils
 			});
 	}
 
-	/// <summary>
-	/// Builds UnrealHeaderTool for the specified platform.
-	/// </summary>
-	/// <param name="Command"></param>
-	/// <param name="InPlatform"></param>
-	public static void BuildUnrealHeaderTool(BuildCommand Command, UnrealBuildTool.UnrealTargetPlatform InPlatform)
-	{
-		BuildProduct(Command, new UnrealBuild.BuildTarget()
-		{
-			UprojectPath = null,
-			TargetName = "UnrealHeaderTool",
-			Platform = InPlatform,
-			Config = UnrealBuildTool.UnrealTargetConfiguration.Development,
-		});
-	}
-
 	private static void BuildProduct(BuildCommand Command, UnrealBuild.BuildTarget Target)
 	{
 		if (Target == null)
