If cleaning up test artifacts fails, set each test artifact to have "normal" attributes.
The some artifacts are read-only, so they fail to get cleaned up by [REDACTED_PERSON]() unless the read-only attribute is removed.

diff --git a/Engine/Source/Programs/AutomationTool/Gauntlet/Framework/Devices/Gauntlet.DevicePool.cs b/Engine/Source/Programs/AutomationTool/Gauntlet/Framework/Devices/Gauntlet.DevicePool.cs
--- a/Engine/Source/Programs/AutomationTool/Gauntlet/Framework/Devices/Gauntlet.DevicePool.cs
+++ b/Engine/Source/Programs/AutomationTool/Gauntlet/Framework/Devices/Gauntlet.DevicePool.cs
@@ -509,6 +509,34 @@ namespace Gauntlet
 			Reservation.ReportDeviceError(DeviceURL, ServiceDeviceName, ErrorMessage);
 		}
 
+		/// <summary>
+		/// Force clean the cache path by setting all of the files in the cache path
+		/// to have "normal" attributes (i.e. not read-only) before attempting to delete.
+		/// </summary>
+		public bool ForceCleanCachePath(string ClientTempDir)
+		{
+			DirectoryInfo ClientTempDirInfo = new DirectoryInfo(ClientTempDir) { Attributes = FileAttributes.Normal };
+
+			Log.Info(KnownLogEvents.Gauntlet_DeviceEvent, "Setting files in {0} to have normal attributes (no longer read-only).", ClientTempDir);
+			foreach (FileSystemInfo info in ClientTempDirInfo.GetFileSystemInfos("*", SearchOption.AllDirectories))
+			{
+				info.Attributes = FileAttributes.Normal;
+			}
+
+			try
+			{
+				Log.Info(KnownLogEvents.Gauntlet_DeviceEvent, "Clearing artifact path {0} (force)", ClientTempDir);
+				Directory.Delete(ClientTempDir, true);
+			}
+			catch (Exception Ex)
+			{
+				Log.Warning(KnownLogEvents.Gauntlet_DeviceEvent, "Failed to force delete {File}. {Exception}", ClientTempDir, Ex.Message);
+				return false;
+			}
+
+			return true;
+		}
+
 		/// <summary>
 		/// Construct a path to hold cache files and make sure it is empty
 		/// </summary>
@@ -529,9 +557,13 @@ namespace Gauntlet
 				}
 				catch (Exception Ex)
 				{
-					// warn and use a different directory
-					Log.Warning(KnownLogEvents.Gauntlet_DeviceEvent, "Failed to delete {Folder}. {Message}", ClientTempDir, Ex.Message);
-					ClientTempDir = Path.Combine(PlatformCache, string.Format("{0}_{1}", InDeviceDefiniton.Name, CleanAttempts++));
+					Log.Info(KnownLogEvents.Gauntlet_DeviceEvent, "Clearing artifact path {0} failed - attempting to force clean", ClientTempDir);
+					if (!ForceCleanCachePath(ClientTempDir))
+					{
+						// warn and use a different directory
+						Log.Warning(KnownLogEvents.Gauntlet_DeviceEvent, "Failed to delete {Folder}. {Message}", ClientTempDir, Ex.Message);
+						ClientTempDir = Path.Combine(PlatformCache, string.Format("{0}_{1}", InDeviceDefiniton.Name, CleanAttempts++));
+					}
 				}
 			}
 			// create this path
