UBT: Set human-friendly titles for XGE tasks

If not set, XGE infers the task title from filename which resulted in "UAT_XGE".

diff --git a/Engine/Source/Programs/AutomationTool/AutomationUtils/UnrealBuild.cs b/Engine/Source/Programs/AutomationTool/AutomationUtils/UnrealBuild.cs
--- a/Engine/Source/Programs/AutomationTool/AutomationUtils/UnrealBuild.cs
+++ b/Engine/Source/Programs/AutomationTool/AutomationUtils/UnrealBuild.cs
@@ -1118,7 +1118,7 @@ namespace AutomationTool
 			string Args = null;
 			if (XGETool != null) 
 			{
-				Args = "\"" + TaskFilePath + "\" /Rebuild /NoLogo /ShowAgent /ShowTime";
+				Args = "\"" + TaskFilePath + "\" /Rebuild /NoLogo /ShowAgent /ShowTime /Title=\"UnrealBuildTool Compile\"";
 				if (ParseParam("StopOnErrors"))
 				{
 					Args += " /StopOnErrors";
diff --git a/Engine/Source/Programs/UnrealBuildTool/Executors/XGE.cs b/Engine/Source/Programs/UnrealBuildTool/Executors/XGE.cs
--- a/Engine/Source/Programs/UnrealBuildTool/Executors/XGE.cs
+++ b/Engine/Source/Programs/UnrealBuildTool/Executors/XGE.cs
@@ -643,6 +643,7 @@ namespace UnrealBuildTool
 					bXGENoWatchdogThread ? "/no_watchdog_thread" : "")
 				);
 			XGEStartInfo.UseShellExecute = false;
+			XGEStartInfo.Arguments += " /Title=\"UnrealBuildTool Compile\"";
 
 			// Use the IDE-integrated Incredibuild monitor to display progress.
 			XGEStartInfo.Arguments += " /UseIdeMonitor";
