UAT: Make sure to use the specified unreal executable instead of a hard coded development one when launching the cook on the fly server.
Fix bug where quick launching with cotf from a debug editor failed to launch the cook server due to a missing or old development editor executable.
Since the default executable is the the console -cmd.exe one, implement ShellExecute to open a separate shell for the server output, also fixes duplicated stdout lines in that window.
Improve BuildCookRun log messages for booting and shutting down server and client.

diff --git a/Engine/Source/Programs/AutomationTool/AutomationUtils/ProcessUtils.cs b/Engine/Source/Programs/AutomationTool/AutomationUtils/ProcessUtils.cs
--- a/Engine/Source/Programs/AutomationTool/AutomationUtils/ProcessUtils.cs
+++ b/Engine/Source/Programs/AutomationTool/AutomationUtils/ProcessUtils.cs
@@ -768,6 +768,12 @@ namespace AutomationTool
 			/// </summary>
 			UTF16Output = 1 << 10,
 
+			/// <summary>
+			/// If set, then use the shell to create the process, else the process will be created directly from the executable.
+			/// Will force NoStdOutRedirect since UseShellExecute is incompatible with redirection of stdout.
+			/// </summary>
+			UseShellExecute = 1 << 11,
+
 			Default = AllowSpew | AppMustExist,
 		}
 
@@ -871,7 +877,8 @@ namespace AutomationTool
 				LogWithVerbosity(SpewVerbosity,"Running: " + App + " " + (String.IsNullOrEmpty(CommandLine) ? "" : CommandLine));
 			}
 
-			bool bRedirectStdOut = !Options.HasFlag(ERunOptions.NoStdOutRedirect);
+			bool bUseShellExecute = Options.HasFlag(ERunOptions.UseShellExecute);
+			bool bRedirectStdOut = !bUseShellExecute && !Options.HasFlag(ERunOptions.NoStdOutRedirect);
 			bool bAllowSpew = bRedirectStdOut && Options.HasFlag(ERunOptions.AllowSpew);
 			bool bCaptureSpew = bRedirectStdOut && !Options.HasFlag(ERunOptions.NoStdOutCapture);
 			IProcessResult Result = ProcessManager.CreateProcess(App, bAllowSpew, bCaptureSpew, Env, SpewVerbosity: SpewVerbosity, SpewFilterCallback: SpewFilterCallback, WorkingDir: WorkingDir);
@@ -887,7 +894,7 @@ namespace AutomationTool
 				// also see UE-102580
 				Proc.StartInfo.Arguments = String.IsNullOrEmpty(CommandLine) ? "" : CommandLine.Replace('\'', '\"');
 
-				Proc.StartInfo.UseShellExecute = false;
+				Proc.StartInfo.UseShellExecute = bUseShellExecute;
 				if (bRedirectStdOut)
 				{
 					Proc.StartInfo.RedirectStandardOutput = true;
diff --git a/Engine/Source/Programs/AutomationTool/Scripts/CookCommand.Automation.cs b/Engine/Source/Programs/AutomationTool/Scripts/CookCommand.Automation.cs
--- a/Engine/Source/Programs/AutomationTool/Scripts/CookCommand.Automation.cs
+++ b/Engine/Source/Programs/AutomationTool/Scripts/CookCommand.Automation.cs
@@ -271,7 +271,7 @@ namespace AutomationScripts
 				}
 
 				var CookServerLogFile = CombinePaths(LogFolderOutsideOfSandbox, "CookServer.log");
-				CookServerProcess = RunCookOnTheFlyServer(Params.RawProjectPath, Params.NoClient ? "" : CookServerLogFile, COTFCommandLine);
+				CookServerProcess = RunCookOnTheFlyServer(Params, Params.NoClient ? "" : CookServerLogFile, COTFCommandLine);
 			}
 			else
 			{
diff --git a/Engine/Source/Programs/AutomationTool/Scripts/RunProjectCommand.Automation.cs b/Engine/Source/Programs/AutomationTool/Scripts/RunProjectCommand.Automation.cs
--- a/Engine/Source/Programs/AutomationTool/Scripts/RunProjectCommand.Automation.cs
+++ b/Engine/Source/Programs/AutomationTool/Scripts/RunProjectCommand.Automation.cs
@@ -74,7 +74,7 @@ namespace AutomationScripts
 		}
 
 		/// <summary>
-		/// Fot not-installed runs, copies all logs from the temp log folder back to the UAT log folder.
+		/// For not-installed runs, copies all logs from the temp log folder back to the UAT log folder.
 		/// </summary>
 		private static void CopyLogsBackToLogFolder()
 		{
@@ -82,6 +82,7 @@ namespace AutomationScripts
 			{
 				var LogFolderOutsideOfSandbox = GetLogFolderOutsideOfSandbox();
 				var TempLogFiles = FindFiles_NoExceptions("*", false, LogFolderOutsideOfSandbox);
+				LogInformation("Found {0} temp logs to copy from {1} to {2}", TempLogFiles.Length, LogFolderOutsideOfSandbox, CmdEnv.LogFolder);
 				foreach (var LogFilename in TempLogFiles)
 				{
 					var DestFilename = CombinePaths(CmdEnv.LogFolder, Path.GetFileName(LogFilename));
@@ -93,9 +94,15 @@ namespace AutomationScripts
 		private static bool WaitForProcessReady(IProcessResult Process, string Name, string LogFile, string[] ReadyTexts, int MaxLogWaitTimeInSeconds = -1)
 		{
 			var StartTime = DateTime.UtcNow;
+			var bFirst = true;
 			while (!FileExists(LogFile) && !Process.HasExited)
 			{
-				if (MaxLogWaitTimeInSeconds > 0)
+				if (bFirst)
+				{
+					LogInformation("Waiting for {0} to start logging at: {1}", Name, LogFile);
+					bFirst = false;
+				}
+				else if (MaxLogWaitTimeInSeconds > 0)
 				{
 					var Duration = (DateTime.UtcNow - StartTime).Seconds;
 					var TimeLeft = MaxLogWaitTimeInSeconds - Duration;
@@ -106,19 +113,19 @@ namespace AutomationScripts
 					}
 					else
 					{
-						LogInformation("Waiting for {0} seconds for {1} logging process to start...", TimeLeft, Name);
+						LogInformation("Waiting for {0} seconds for {1} to start logging...", TimeLeft, Name);
 					}
 				}
 				else
 				{
-					LogInformation("Waiting for {0} logging process to start...", Name);
+					LogInformation("Waiting for {0} to start logging...", Name);
 				}
 				Thread.Sleep(2000);
 			}
 
 			if (!FileExists(LogFile))
 			{
-				throw new AutomationException("{0} exited without creating a log file.", Name);
+				throw new AutomationException("{0} exited without creating a log file at: {1}", Name, LogFile);
 			}
 
 			LogInformation("Logging started for {0} at: {1}", Name, LogFile);
@@ -156,7 +163,7 @@ namespace AutomationScripts
 					Thread.Sleep(2000);
 				}
 			}
-			throw new AutomationException("{0} exited before we asked it to.", Name);
+			throw new AutomationException("{0} exited before we asked it to (see {1} for more info)", Name, LogFile);
 		}
 
 		public static void Run(ProjectParams Params)
@@ -193,15 +200,25 @@ namespace AutomationScripts
 			{
 				if (!GlobalCommandLine.NoKill)
 				{
-					if (CookServerProcess != null && !CookServerProcess.HasExited)
+					if (CookServerProcess != null)
 					{
-						LogInformation("Stopping cook server...");
-						CookServerProcess.StopProcess();
+						if (!CookServerProcess.HasExited)
+						{
+							LogInformation("Stopping cook server...");
+							CookServerProcess.StopProcess();
+						}
+						LogInformation("Cook server exited with error code: {0} (see {1} for more info)",
+								CookServerProcess.ExitCode, CookServerLogFile);
 					}
-					if (DedicatedServerProcess != null && !DedicatedServerProcess.HasExited)
+					if (DedicatedServerProcess != null)
 					{
-						LogInformation("Stopping dedicated server...");
-						DedicatedServerProcess.StopProcess();
+						if (!DedicatedServerProcess.HasExited)
+						{
+							LogInformation("Stopping dedicated server...");
+							DedicatedServerProcess.StopProcess();
+						}
+						LogInformation("Dedicated server exited with error code: {0} (see {1} for more info)",
+								DedicatedServerProcess.ExitCode, DedicatedServerLogFile);
 					}
 				}
 				CopyLogsBackToLogFolder();
@@ -339,7 +356,7 @@ namespace AutomationScripts
 					{
 						LogInformation("Client exited, waiting for stdout...");
 						ClientProcess.WaitForExit();
-						LogInformation("Client exited, logging done!");
+						LogInformation("Client exited, logging done! (see {0} for more info)", ClientLogFile);
 						bKeepReading = false;
 					}
 
@@ -514,7 +531,7 @@ namespace AutomationScripts
 				// In unattended/-testexit mode we only throw if testexit text was not found
 				if (!bTestExitTextFound)
 				{
-					throw new AutomationException("Client exited before we asked it to.");
+					throw new AutomationException("Client exited before we asked it to (see {0} for more info)", ClientLogFile);
 				}
 			}
 			else
@@ -523,10 +540,10 @@ namespace AutomationScripts
 				// already thrown a more specific exception or given a more specific ErrorCode, but this catches the rest.
 				if (ClientProcess != null && ClientProcess.ExitCode != 0)
 				{
-					throw new AutomationException("Client exited with error code: " + ClientProcess.ExitCode);
+					throw new AutomationException("Client exited with error code: {0} (see {1} for more info)", ClientProcess.ExitCode, ClientLogFile);
 				}
 			}
-			LogInformation("Client exited with error code: " + ClientProcess.ExitCode);
+			LogInformation("Client exited with error code: {0} (see {1} for more info)", ClientProcess.ExitCode, ClientLogFile);
 		}
 
 		private static void SetupClientParams(List<DeploymentContext> DeployContextList, ProjectParams Params, string ClientLogFile, out ERunOptions ClientRunFlags, out string ClientApp, out string ClientCmdLine)
@@ -956,11 +973,11 @@ namespace AutomationScripts
 			return Result;
 		}
 
-		private static IProcessResult RunCookOnTheFlyServer(FileReference ProjectName, string ServerLogFile, string AdditionalCommandLine)
+		private static IProcessResult RunCookOnTheFlyServer(ProjectParams Params, string ServerLogFile, string AdditionalCommandLine)
 		{
-			var ServerApp = HostPlatform.Current.GetUnrealExePath("UnrealEditor.exe");
-			var Args = String.Format("{0} -run=cook -cookonthefly -unattended -CrashForUAT -log",
-				CommandUtils.MakePathSafeToUseWithCommandLine(ProjectName.FullName));
+			var ServerApp = HostPlatform.Current.GetUnrealExePath(Params.UnrealExe);
+			var Args = String.Format("{0} -run=cook -cookonthefly -unattended -CrashForUAT -AllowStdOutLogVerbosity",
+				CommandUtils.MakePathSafeToUseWithCommandLine(Params.RawProjectPath.FullName));
 			if (!String.IsNullOrEmpty(ServerLogFile))
 			{
 				// Issue with dotnet not allowing any files with an exclusive advisory lock to be opened for read-only or copied
@@ -978,9 +995,9 @@ namespace AutomationScripts
 				Args += " " + AdditionalCommandLine;
 			}
 
-			// Run the server (Without NoStdOutRedirect -log doesn't log anything to the window)
+			// Run the server with shell execute to launch it in a separate shell with stdout
 			PushDir(Path.GetDirectoryName(ServerApp));
-			var Result = Run(ServerApp, Args, null, ERunOptions.Default | ERunOptions.NoWaitForExit | ERunOptions.NoStdOutRedirect);
+			var Result = Run(ServerApp, Args, null, ERunOptions.Default | ERunOptions.NoWaitForExit | ERunOptions.UseShellExecute);
 			PopDir();
 
 			LogInformation("Running CookServer@Process:{0}@{1}", ServerApp, Result.ProcessObject.Id);
