Fixed up Androd installation via Turnkey:
- Moved running AndroidStudio to install the commandline tools into the InstallSDKfunction, out of the Prerequisites, which now just now makes sure AndroidStudio is installed
- Added direct AndroidStudio download to Turnkey, usable by all licensees (no need to host the installer at each studio now)

diff --git a/Engine/Build/Turnkey/TurnkeyManifest.xml b/Engine/Build/Turnkey/TurnkeyManifest.xml
--- a/Engine/Build/Turnkey/TurnkeyManifest.xml
+++ b/Engine/Build/Turnkey/TurnkeyManifest.xml
@@ -7,4 +7,16 @@
 		<Manifest>file:$(EngineDir)/Restricted/NoRedist/Build/Turnkey/TurnkeyManifest.xml</Manifest>
 		<Manifest>$(UE_STUDIO_TURNKEY_LOCATION)</Manifest>
 	</AdditionalManifests>
+
+    
+    <!-- Public file sources below we can share with all licensees -->
+    <FileSource>
+        <Platform>Android</Platform>
+        <Type>Misc</Type>
+        <Name>AndroidStudio</Name>
+        <Source HostPlatform="Win64">https://redirector.gvt1.com/edgedl/android/studio/install/4.0.2.0/android-studio-ide-193.6821437-windows.exe</Source>
+        <Source HostPlatform="Linux">https://redirector.gvt1.com/edgedl/android/studio/ide-zips/4.0.2.0/android-studio-ide-193.6821437-linux.tar.gz</Source>
+        <Source HostPlatform="Mac">https://redirector.gvt1.com/edgedl/android/studio/install/4.0.2.0/android-studio-ide-193.6821437-mac.dmg</Source>
+    </FileSource>
+
 </TurnkeyManifest>
diff --git a/Engine/Extras/Android/SetupAndroid.command b/Engine/Extras/Android/SetupAndroid.command
--- a/Engine/Extras/Android/SetupAndroid.command
+++ b/Engine/Extras/Android/SetupAndroid.command
@@ -84,9 +84,6 @@ echo Android Studio Path: $STUDIO_PATH
 if [ "$STUDIO_SDK_PATH" == "" ]; then
 	STUDIO_SDK_PATH=~/Library/Android/sdk
 fi
-if [ "$1" != "" ]; then
-	STUDIO_SDK_PATH=$1
-fi
 if [ ! -d "$STUDIO_SDK_PATH" ]; then
 	echo Android SDK not found at: $STUDIO_SDK_PATH
 	echo Unable to locate local Android SDK location. Did you run Android Studio after installing?
diff --git a/Engine/Source/Programs/AutomationTool/Android/AndroidPlatform.Automation.cs b/Engine/Source/Programs/AutomationTool/Android/AndroidPlatform.Automation.cs
--- a/Engine/Source/Programs/AutomationTool/Android/AndroidPlatform.Automation.cs
+++ b/Engine/Source/Programs/AutomationTool/Android/AndroidPlatform.Automation.cs
@@ -15,6 +15,8 @@
 using AutomationUtils.Automation;
 using System.Text.RegularExpressions;
 using AutomationScripts;
+using System.Drawing;
+using System.Security.Policy;
 
 public class AndroidPlatform : Platform
 {
@@ -175,15 +177,72 @@
 		return AndroidSDK != null ? new string[] { AndroidSDK.GetMainVersion() } : Array.Empty<string>();
 	}
 
-	public override bool GetSDKInstallCommand(out string Command, out string Params, ref bool bRequiresPrivilegeElevation, ref bool bCreateWindow, ITurnkeyContext TurnkeyContext)
+	// Android has a more complex sdk installation, so perform it manually
+	public override bool InstallSDK(BuildCommand BuildCommand, ITurnkeyContext TurnkeyContext, DeviceInfo Device, bool bUnattended, bool bSdkAlreadyInstalled)
 	{
+		if (Device != null)
+		{
+			return base.InstallSDK(BuildCommand, TurnkeyContext, Device, bUnattended, bSdkAlreadyInstalled);
+		}
+
+		string SdkDir = GetSdkDir();
+		bool bIsInstalled = Directory.Exists(SdkDir);
+
+		if (!bIsInstalled)
+		{
+			int Option = 2;
+			while (Option == 2)
+			{
+				string Prompt = $"The Android Sdk directory was not found (expected to find it at '{SdkDir}'\n" +
+					"Android Studio can install it for you, but you will need to manually perform some steps (if desired, you can get detailed help with option 2):\n" +
+					"  - Wait for Android Studio to start, you will see an initial dialog asking how to proceed (called \"Welcome to Android Studio\")\n" +
+					"  - Click the \"Configure\" dropdown in the bottom right, and select \"SDK Manager\"\n" +
+					"  - Click on the \"SDK Tools\" tab near the top middle of the right pane\n" +
+					"  - Check the box next to Android SDK COmmand-line Tools (latest)\n" +
+					"  - Click OK in the bottom right\n" +
+					"  - It will probably ask for you to accept a license - you MUST do this\n" +
+					"  - Once installation has completed, close/quit Android Studio to continue\n";
+
+				List<string> Options = new()
+				{
+					"Run Android Studio to install the Command Line Tools",
+					"Get detailed step by step guide",
+				};
+
+				Option = TurnkeyContext.ReadInputInt(Prompt, Options, true, 1);
+
+				if (Option == 0)
+				{
+					return false;
+				}
+				if (Option == 2)
+				{
+					// @todo: we need to redo the documentation DRAMATICALLY on here
+					string URL = "https://docs.unrealengine.com/5.1/en-US/how-to-set-up-android-sdk-and-ndk-for-your-unreal-engine-development-environment/";
+					Process.Start(new ProcessStartInfo { FileName = URL, UseShellExecute = true });
+				}
+			}
+
+			string AndroidStudioExe = GetAndroidStudioExe();
+
+			if (HostPlatform.Platform == UnrealTargetPlatform.Mac)
+			{
+				TurnkeyContext.RunExternalCommand("open", $"-W \"{GetAndroidStudioExe()}\"", false, true, true);
+			}
+			else
+			{
+				TurnkeyContext.RunExternalCommand(GetAndroidStudioExe(), "", false, true, false);
+			}
+		}
+
 		// run the Setup.bat in the engine, not coming from a normal FileSource
 
-		if (HostPlatform.Current.HostEditorPlatform == UnrealTargetPlatform.Win64)
+		string Command;
+		if (HostPlatform.Platform == UnrealTargetPlatform.Win64)
 		{
 			Command = "$(EngineDir)/Extras/Android/SetupAndroid.bat";
 		}
-		else if (HostPlatform.Current.HostEditorPlatform == UnrealTargetPlatform.Mac)
+		else if (HostPlatform.Platform == UnrealTargetPlatform.Mac)
 		{
 			Command = "$(EngineDir)/Extras/Android/SetupAndroid.command";
 		}
@@ -199,16 +258,14 @@
 		string CMakeVersion = AndroidSDK.GetPlatformSpecificVersion("cmake");
 		string NDKVersion = AndroidSDK.GetPlatformSpecificVersion("ndk");
 
-		Params = $"{PlatformsVersion} {BuildToolsVersion} {CMakeVersion} {NDKVersion} -noninteractive";
+		string Params = $"{PlatformsVersion} {BuildToolsVersion} {CMakeVersion} {NDKVersion} -noninteractive";
 
 		// because this may bring up a license acceptance message that needs the user to respond, so we make a new window
-		bCreateWindow = true;
-
-		return true;
+		int ExitCode = TurnkeyContext.RunExternalCommand(Command, Params, bRequiresPrivilegeElevation: false, bUnattended, bCreateWindow: true);
+		return ExitCode == 0;
 	}
 
 
-
 	private static string GetAndroidStudioExe()
 	{
 		if (OperatingSystem.IsLinux())
@@ -274,16 +331,17 @@
 				{
 					if (BashProfileContents[LineIndex].StartsWith("export " + SdkKey + "="))
 					{
-						string PathVar = BashProfileContents[LineIndex].Split
-('=')[1].Replace("\"", "");
-Log.TraceInformation("ANDROID_HOME = {0}", PathVar);
+						string PathVar = BashProfileContents[LineIndex].Split('=')[1].Replace("\"", "");
 						return PathVar;
 					}
 
 				}
 			}
 
-			return Environment.GetEnvironmentVariable("ANDROID_HOME");
+			string UserHome = Environment.GetEnvironmentVariable("HOME");
+			string AndroidSdkPath = Path.Combine(UserHome, "Library", "Android", "Sdk");
+
+			return AndroidSdkPath;
 		}
 
 		Debug.Assert(OperatingSystem.IsWindows());
@@ -296,30 +354,25 @@
 	public override bool UpdateHostPrerequisites(BuildCommand Command, ITurnkeyContext TurnkeyContext, bool bVerifyOnly)
 	{
 		string AndroidStudioExe = GetAndroidStudioExe();
-		string SdkDir = GetSdkDir();
 
 		bool bIsMac = (HostPlatform.Current.HostEditorPlatform == UnrealTargetPlatform.Mac);
 		bool bHaveAndroidStudio = (bIsMac && Directory.Exists(AndroidStudioExe)) || 
 					(!bIsMac && FileExists(AndroidStudioExe));
 
-		bool bIsInstalled = bHaveAndroidStudio && Directory.Exists(SdkDir);
-
 		// if we are only verifying, just return the status, and if it's installed, we are done!
-		if (bVerifyOnly || bIsInstalled)
+		if (bVerifyOnly)
 		{
 			if (!bHaveAndroidStudio)
 			{
 				TurnkeyContext.ReportError("Android Studio is not installed correctly.");
 			}
-			if (!Directory.Exists(SdkDir))
-			{
-				TurnkeyContext.ReportError("Android SDK directory is not set correctly.");
-			}
-			return bIsInstalled;
+			return bHaveAndroidStudio;
 		}
 
 		if (!bHaveAndroidStudio)
 		{
+			TurnkeyContext.PauseForUser("Android Studio was not found on this machine. Press Enter to download and install Android Studio which is required to use Android.");
+
 			// get AS installer
 			string OutputPath = TurnkeyContext.RetrieveFileSource("AndroidStudio");
 
@@ -338,7 +391,7 @@
 
 			if (OutputPath == null)
 			{
-				TurnkeyContext.PauseForUser("Unable to find Android Studio installer. Please download and install Android Studio 4.0.2 from https://developer.android.com/studio/archive before continuing.\n\nMake sure to use the Run Android Studio and complete the first-time setup!\n\nFollow the steps to install command-line tools (latest):\nhttps://docs.unrealengine.com/5.0/en-US/how-to-setup-up-android-sdk-and-ndk-for-your-unreal-engine-development-environment");
+				TurnkeyContext.PauseForUser("Unable to find Android Studio installer. Please download and install Android Studio 4.0.2 from https://developer.android.com/studio/archive to standard location before continuing.");
 			}
 			else
 			{
@@ -416,9 +469,8 @@
 				}
 				else if (HostPlatform.Current.HostEditorPlatform == UnrealTargetPlatform.Win64)
 				{
-					TurnkeyContext.PauseForUser("Running the Android Studio installer, and then Android Studio for first-time setup!\n\nFollow the steps to install command-line tools (latest):\nhttps://docs.unrealengine.com/5.0/en-US/how-to-setup-up-android-sdk-and-ndk-for-your-unreal-engine-development-environment");
+					// install AS with the /S switch
 					int ExitCode = TurnkeyContext.RunExternalCommand(OutputPath, "/S", false, true, true);
-					//				Utils.RunLocalProcessAndReturnStdOut(OutputPath, "", out ExitCode, true);
 
 					// AS installer returns 1223 even on success ("user canceled" even tho there's no UI to cancel it...) when running with /S
 					if (ExitCode != 0 && ExitCode != 1223)
@@ -432,37 +484,21 @@
 					TurnkeyContext.ReportError($"Invalid host platform");
 					return false;
 				}
-
-				// re-query for AS location
-				TurnkeyContext.RunExternalCommand(GetAndroidStudioExe(), "", false, true, bIsMac);
 			}
 		}
-		else
-		{
-			TurnkeyContext.PauseForUser("The Sdk directory was not found. Running the Android Studio to perform first-time setup.\n\nFollow the steps to install command-line tools (latest):\nhttps://docs.unrealengine.com/5.0/en-US/how-to-setup-up-android-sdk-and-ndk-for-your-unreal-engine-development-environment");
 
-			TurnkeyContext.RunExternalCommand(AndroidStudioExe, "", false, true, bIsMac);
-		}
-
 		// check to see if the installation worked. If so, continue on!
 		AndroidStudioExe = GetAndroidStudioExe();
-		SdkDir = GetSdkDir();
 
 		bHaveAndroidStudio = (bIsMac && Directory.Exists(AndroidStudioExe)) || 
 					(!bIsMac && FileExists(AndroidStudioExe));
 
-		bIsInstalled = bHaveAndroidStudio && Directory.Exists(SdkDir);
-
 		if (!bHaveAndroidStudio)
 		{
 			TurnkeyContext.ReportError("Android Studio is not installed correctly, after attempted installation.");
 		}
-		if (!Directory.Exists(SdkDir))
-		{
-			TurnkeyContext.ReportError("Android SDK directory is not set correctly, after attempted installation.");
-		}
 
-		return bIsInstalled;
+		return bHaveAndroidStudio;
 	}
 
 
diff --git a/Engine/Source/Programs/AutomationTool/AutomationUtils/Platform.cs b/Engine/Source/Programs/AutomationTool/AutomationUtils/Platform.cs
--- a/Engine/Source/Programs/AutomationTool/AutomationUtils/Platform.cs
+++ b/Engine/Source/Programs/AutomationTool/AutomationUtils/Platform.cs
@@ -20,7 +20,8 @@ namespace AutomationTool
 		int RunExternalCommand(string Command, string Params, bool bRequiresPrivilegeElevation, bool bUnattended, bool bCreateWindow);
 		void Log(string Message);
 		void ReportError(string Message);
-		void PauseForUser(string Message); 
+		void PauseForUser(string Message);
+		int ReadInputInt(string Prompt, List<string> Options, bool bIsCancellable, int DefaultValue = -1);
 	}
 
 	//public interface InputOutput
diff --git a/Engine/Source/Programs/AutomationTool/Turnkey/CopyProvider.cs b/Engine/Source/Programs/AutomationTool/Turnkey/CopyProvider.cs
--- a/Engine/Source/Programs/AutomationTool/Turnkey/CopyProvider.cs
+++ b/Engine/Source/Programs/AutomationTool/Turnkey/CopyProvider.cs
@@ -315,5 +315,11 @@ namespace Turnkey
 		{
 			TurnkeyUtils.PauseForUser(Message);
 		}
+
+		public int ReadInputInt(string Prompt, List<string> Options, bool bIsCancellable, int DefaultValue)
+		{
+			return TurnkeyUtils.ReadInputInt(Prompt, Options, bIsCancellable, DefaultValue);
+		}
+
 	}
 }
diff --git a/Engine/Source/Programs/UnrealBuildTool/Platform/Android/AndroidPlatformSDK.cs b/Engine/Source/Programs/UnrealBuildTool/Platform/Android/AndroidPlatformSDK.cs
--- a/Engine/Source/Programs/UnrealBuildTool/Platform/Android/AndroidPlatformSDK.cs
+++ b/Engine/Source/Programs/UnrealBuildTool/Platform/Android/AndroidPlatformSDK.cs
@@ -50,9 +50,7 @@ namespace UnrealBuildTool
 					{
 						if (BashProfileContents[LineIndex].StartsWith("export " + SdkKey + "="))
 						{
-							string PathVar = BashProfileContents[LineIndex].Split
-('=')[1].Replace("\"", "");
-Log.TraceInformation("ANDROID_HOME = {0}", PathVar);
+							string PathVar = BashProfileContents[LineIndex].Split('=')[1].Replace("\"", "");
 							return PathVar;
 						}
 
@@ -140,17 +138,17 @@ Log.TraceInformation("ANDROID_HOME = {0}", PathVar);
 		}
 
 
-		public override SDKStatus PrintSDKInfoAndReturnValidity(LogEventType Verbosity, LogFormatOptions Options, LogEventType ErrorVerbosity, LogFormatOptions ErrorOptions)
-		{
-			SDKStatus Validity = base.PrintSDKInfoAndReturnValidity(Verbosity, Options, ErrorVerbosity, ErrorOptions);
+		//public override SDKStatus PrintSDKInfoAndReturnValidity(LogEventType Verbosity, LogFormatOptions Options, LogEventType ErrorVerbosity, LogFormatOptions ErrorOptions)
+		//{
+		//	SDKStatus Validity = base.PrintSDKInfoAndReturnValidity(Verbosity, Options, ErrorVerbosity, ErrorOptions);
 
-			if (GetInstalledVersion() != GetMainVersion())
-			{
-				Log.WriteLine(Verbosity, Options, "Note: Android toolchain NDK {0} recommended", GetMainVersion());
-			}
+		//	if (GetInstalledVersion() != GetMainVersion())
+		//	{
+		//		Log.WriteLine(Verbosity, Options, "Note: Android toolchain NDK {0} recommended ('{1}' was found)", GetMainVersion(), GetInstalledVersion());
+		//	}
 
-			return Validity;
-		}
+		//	return Validity;
+		//}
 
 
 		protected override bool PlatformSupportsAutoSDKs()
@@ -298,21 +296,5 @@ Log.TraceInformation("ANDROID_HOME = {0}", PathVar);
 			}
 			return true;
 		}
-
-		protected override SDKStatus HasRequiredManualSDKInternal()
-		{
-			// if any autosdk setup has been done then the local process environment is suspect
-			if (HasSetupAutoSDK())
-			{
-				return SDKStatus.Invalid;
-			}
-
-			if (HasAnySDK())
-			{
-				return SDKStatus.Valid;
-			}
-
-			return SDKStatus.Invalid;
-		}
 	}
 }
