If cleaning up test artifacts fails, set each test artifact to have "normal" attributes.
Some artifacts are read-only, and so they fail to get cleaned up by gauntlet.
This was previously fixed for the Gauntlet Device pool, but should be fixed here as well.

diff --git a/Engine/Source/Programs/AutomationTool/Gauntlet/Platform/Linux/Gauntlet.TargetDeviceLinux.cs b/Engine/Source/Programs/AutomationTool/Gauntlet/Platform/Linux/Gauntlet.TargetDeviceLinux.cs
--- a/Engine/Source/Programs/AutomationTool/Gauntlet/Platform/Linux/Gauntlet.TargetDeviceLinux.cs
+++ b/Engine/Source/Programs/AutomationTool/Gauntlet/Platform/Linux/Gauntlet.TargetDeviceLinux.cs
@@ -74,18 +74,43 @@ namespace Gauntlet
 			return Device.Run(this);
 		}
 
+		public bool ForceCleanDeviceArtifacts()
+		{
+			DirectoryInfo ClientTempDirInfo = new DirectoryInfo(ArtifactPath) { Attributes = FileAttributes.Normal };
+			Log.Info(KnownLogEvents.Gauntlet_DeviceEvent, "Setting files in device artifacts {0} to have normal attributes (no longer read-only).", ArtifactPath);
+			foreach (FileSystemInfo info in ClientTempDirInfo.GetFileSystemInfos("*", SearchOption.AllDirectories))
+			{
+				info.Attributes = FileAttributes.Normal;
+			}
+			try
+			{
+				Log.Info(KnownLogEvents.Gauntlet_DeviceEvent, "Clearing device artifact path {0} (force)", ArtifactPath);
+				Directory.Delete(ArtifactPath, true);
+			}
+			catch (Exception Ex)
+			{
+				Log.Warning(KnownLogEvents.Gauntlet_DeviceEvent, "Failed to force delete artifact path {File}. {Exception}", ArtifactPath, Ex.Message);
+				return false;
+			}
+			return true;
+		}
+
 		public virtual void CleanDeviceArtifacts()
 		{
 			if (!string.IsNullOrEmpty(ArtifactPath) && Directory.Exists(ArtifactPath))
 			{
 				try
 				{
-					Log.Info("Clearing actifact path {0} for {1}", ArtifactPath, Device.Name);
+					Log.Info("Clearing device artifacts path {0} for {1}", ArtifactPath, Device.Name);
 					Directory.Delete(ArtifactPath, true);
 				}
 				catch (Exception Ex)
 				{
-					Log.Warning(KnownLogEvents.Gauntlet_DeviceEvent, "Failed to delete {File}. {Exception}", ArtifactPath, Ex.Message);
+					Log.Info(KnownLogEvents.Gauntlet_DeviceEvent, "First attempt at clearing artifact path {0} failed - trying again", ArtifactPath);
+					if (!ForceCleanDeviceArtifacts())
+					{
+						Log.Warning(KnownLogEvents.Gauntlet_DeviceEvent, "Failed to delete {File}. {Exception}", ArtifactPath, Ex.Message);
+					}
 				}
 			}
 		}
diff --git a/Engine/Source/Programs/AutomationTool/Gauntlet/Platform/Mac/Gauntlet.TargetDeviceMac.cs b/Engine/Source/Programs/AutomationTool/Gauntlet/Platform/Mac/Gauntlet.TargetDeviceMac.cs
--- a/Engine/Source/Programs/AutomationTool/Gauntlet/Platform/Mac/Gauntlet.TargetDeviceMac.cs
+++ b/Engine/Source/Programs/AutomationTool/Gauntlet/Platform/Mac/Gauntlet.TargetDeviceMac.cs
@@ -56,18 +56,43 @@ namespace Gauntlet
 			return Device.Run(this);
 		}
 
+		public bool ForceCleanDeviceArtifacts()
+		{
+			DirectoryInfo ClientTempDirInfo = new DirectoryInfo(ArtifactPath) { Attributes = FileAttributes.Normal };
+			Log.Info(KnownLogEvents.Gauntlet_DeviceEvent, "Setting files in device artifacts {0} to have normal attributes (no longer read-only).", ArtifactPath);
+			foreach (FileSystemInfo info in ClientTempDirInfo.GetFileSystemInfos("*", SearchOption.AllDirectories))
+			{
+				info.Attributes = FileAttributes.Normal;
+			}
+			try
+			{
+				Log.Info(KnownLogEvents.Gauntlet_DeviceEvent, "Clearing device artifact path {0} (force)", ArtifactPath);
+				Directory.Delete(ArtifactPath, true);
+			}
+			catch (Exception Ex)
+			{
+				Log.Warning(KnownLogEvents.Gauntlet_DeviceEvent, "Failed to force delete artifact path {File}. {Exception}", ArtifactPath, Ex.Message);
+				return false;
+			}
+			return true;
+		}
+
 		public virtual void CleanDeviceArtifacts()
 		{
 			if (!string.IsNullOrEmpty(ArtifactPath) && Directory.Exists(ArtifactPath))
 			{
 				try
 				{
-					Log.Info("Clearing actifact path {0} for {1}", ArtifactPath, Device.Name);
+					Log.Info("Clearing device artifacts path {0} for {1}", ArtifactPath, Device.Name);
 					Directory.Delete(ArtifactPath, true);
 				}
 				catch (Exception Ex)
 				{
-					Log.Warning(KnownLogEvents.Gauntlet_DeviceEvent, "Failed to delete {File}. {Exception}", ArtifactPath, Ex.Message);
+					Log.Info(KnownLogEvents.Gauntlet_DeviceEvent, "First attempt at clearing artifact path {0} failed - trying again", ArtifactPath);
+					if (!ForceCleanDeviceArtifacts())
+					{
+						Log.Warning(KnownLogEvents.Gauntlet_DeviceEvent, "Failed to delete {File}. {Exception}", ArtifactPath, Ex.Message);
+					}
 				}
 			}
 		}
diff --git a/Engine/Source/Programs/AutomationTool/Gauntlet/Platform/Windows/Gauntlet.TargetDeviceWindows.cs b/Engine/Source/Programs/AutomationTool/Gauntlet/Platform/Windows/Gauntlet.TargetDeviceWindows.cs
--- a/Engine/Source/Programs/AutomationTool/Gauntlet/Platform/Windows/Gauntlet.TargetDeviceWindows.cs
+++ b/Engine/Source/Programs/AutomationTool/Gauntlet/Platform/Windows/Gauntlet.TargetDeviceWindows.cs
@@ -176,18 +176,43 @@ namespace Gauntlet
 			return Device.Run(this);
 		}
 
+		public bool ForceCleanDeviceArtifacts()
+		{
+			DirectoryInfo ClientTempDirInfo = new DirectoryInfo(ArtifactPath) { Attributes = FileAttributes.Normal };
+			Log.Info(KnownLogEvents.Gauntlet_DeviceEvent, "Setting files in device artifacts {0} to have normal attributes (no longer read-only).", ArtifactPath);
+			foreach (FileSystemInfo info in ClientTempDirInfo.GetFileSystemInfos("*", SearchOption.AllDirectories))
+			{
+				info.Attributes = FileAttributes.Normal;
+			}
+			try
+			{
+				Log.Info(KnownLogEvents.Gauntlet_DeviceEvent, "Clearing device artifact path {0} (force)", ArtifactPath);
+				Directory.Delete(ArtifactPath, true);
+			}
+			catch (Exception Ex)
+			{
+				Log.Warning(KnownLogEvents.Gauntlet_DeviceEvent, "Failed to force delete artifact path {File}. {Exception}", ArtifactPath, Ex.Message);
+				return false;
+			}
+			return true;
+		}
+
 		public virtual void CleanDeviceArtifacts()
 		{
 			if (!string.IsNullOrEmpty(ArtifactPath) && Directory.Exists(ArtifactPath))
 			{
 				try
 				{
-					Log.Info("Clearing actifact path {0} for {1}", ArtifactPath, Device.Name);
+					Log.Info("Clearing device artifacts path {0} for {1}", ArtifactPath, Device.Name);
 					Directory.Delete(ArtifactPath, true);
 				}
 				catch (Exception Ex)
 				{
-					Log.Warning(KnownLogEvents.Gauntlet_DeviceEvent, "Failed to delete {File}. {Exception}", ArtifactPath, Ex.Message);
+					Log.Info(KnownLogEvents.Gauntlet_DeviceEvent, "First attempt at clearing artifact path {0} failed - trying again", ArtifactPath);
+					if (!ForceCleanDeviceArtifacts())
+					{
+						Log.Warning(KnownLogEvents.Gauntlet_DeviceEvent, "Failed to delete {File}. {Exception}", ArtifactPath, Ex.Message);
+					}
 				}
 			}
 		}
