Gauntlet - fix inconsistent failure between desktop platforms
+ Add support for AddressSanitizer callstack capture in Unreal log parser

#rnx

diff --git a/Engine/Source/Programs/AutomationTool/Gauntlet/Platform/Linux/Gauntlet.TargetDeviceLinux.cs b/Engine/Source/Programs/AutomationTool/Gauntlet/Platform/Linux/Gauntlet.TargetDeviceLinux.cs
--- a/Engine/Source/Programs/AutomationTool/Gauntlet/Platform/Linux/Gauntlet.TargetDeviceLinux.cs
+++ b/Engine/Source/Programs/AutomationTool/Gauntlet/Platform/Linux/Gauntlet.TargetDeviceLinux.cs
@@ -9,6 +9,7 @@ using AutomationTool;
 using UnrealBuildTool;
 using System.Text.RegularExpressions;
 using EpicGames.Core;
+using static AutomationTool.ProcessResult;
 
 namespace Gauntlet
 {
@@ -212,7 +213,6 @@ namespace Gauntlet
 			}
 
 			IProcessResult Result = null;
-			string ProcessLogFile = null;
 
 			lock (Globals.MainLock)
 			{
@@ -225,48 +225,9 @@ namespace Gauntlet
 
 				string CmdLine = LinuxApp.CommandArguments;
 
-				// Look in app parameters if abslog is specified, if so use it
-				/* TODO for linux
-				Regex CLRegex = new Regex(@"(--?[a-zA-Z]+)[:\s=]?([A-Z]:(?:\\[\w\s-]+)+\\?(?=\s-)|\""[^\""]*\""|[^-][^\s]*)?");
-				foreach (Match M in CLRegex.Matches(CmdLine))
-				{
-					if (M.Groups.Count == 3 && M.Groups[1].Value == "-abslog")
-					{
-						ProcessLogFile = M.Groups[2].Value;
-					}
-				}
-				*/
-
-				// explicitly set log file when not already defined
-				if (string.IsNullOrEmpty(ProcessLogFile))
-				{
-					string LogFolder = string.Format(@"{0}/Logs", LinuxApp.ArtifactPath);
-
-					if (!Directory.Exists(LogFolder))
-					{
-						Directory.CreateDirectory(LogFolder);
-					}
-
-					ProcessLogFile = string.Format("{0}/{1}.log", LogFolder, LinuxApp.ProjectName);
-					CmdLine = string.Format("{0} -abslog=\"{1}\"", CmdLine, ProcessLogFile);
-				}
-
-				// cleanup any existing log file
-				try
-				{
-					if (File.Exists(ProcessLogFile))
-					{
-						File.Delete(ProcessLogFile);
-					}
-				}
-				catch (Exception Ex)
-				{
-					throw new AutomationException("Unable to delete existing log file {0} {1}", ProcessLogFile, Ex.Message);
-				}
-
 				Log.Verbose("\t{0}", CmdLine);
 
-				Result = CommandUtils.Run(LinuxApp.ExecutablePath, CmdLine, Options: LinuxApp.RunOptions | (ProcessLogFile != null ? CommandUtils.ERunOptions.NoStdOutRedirect : 0 ));
+				Result = CommandUtils.Run(LinuxApp.ExecutablePath, CmdLine, Options: LinuxApp.RunOptions, SpewFilterCallback: new SpewFilterCallbackType(delegate(string M) { return null; }) /* make sure stderr does not spew in the stdout */);
 
 				if (Result.HasExited && Result.ExitCode != 0)
 				{
@@ -276,7 +237,7 @@ namespace Gauntlet
 				Environment.CurrentDirectory = OldWD;
 			}
 
-			return new LinuxAppInstance(LinuxApp, Result, ProcessLogFile);
+			return new LinuxAppInstance(LinuxApp, Result);
 		}
 
 		private void CopyAdditionalFiles(UnrealAppConfig AppConfig)
diff --git a/Engine/Source/Programs/AutomationTool/Gauntlet/Platform/Mac/Gauntlet.TargetDeviceMac.cs b/Engine/Source/Programs/AutomationTool/Gauntlet/Platform/Mac/Gauntlet.TargetDeviceMac.cs
--- a/Engine/Source/Programs/AutomationTool/Gauntlet/Platform/Mac/Gauntlet.TargetDeviceMac.cs
+++ b/Engine/Source/Programs/AutomationTool/Gauntlet/Platform/Mac/Gauntlet.TargetDeviceMac.cs
@@ -8,6 +8,7 @@ using UnrealBuildTool;
 using System.Text.RegularExpressions;
 using EpicGames.Core;
 using UnrealBuildBase;
+using static AutomationTool.ProcessResult;
 
 namespace Gauntlet
 {
@@ -410,7 +411,7 @@ namespace Gauntlet
 				Log.Info("Launching {0} on {1}", App.Name, ToString());
 				Log.Verbose("\t{0}", MacInstall.CommandArguments);
 
-				Result = CommandUtils.Run(MacInstall.ExecutablePath, MacInstall.CommandArguments, Options: MacInstall.RunOptions);
+				Result = CommandUtils.Run(MacInstall.ExecutablePath, MacInstall.CommandArguments, Options: MacInstall.RunOptions, SpewFilterCallback: new SpewFilterCallbackType(delegate (string M) { return null; }) /* make sure stderr does not spew in the stdout */);
 
 				if (Result.HasExited && Result.ExitCode != 0)
 				{
diff --git a/Engine/Source/Programs/AutomationTool/Gauntlet/Platform/Windows/Gauntlet.TargetDeviceWindows.cs b/Engine/Source/Programs/AutomationTool/Gauntlet/Platform/Windows/Gauntlet.TargetDeviceWindows.cs
--- a/Engine/Source/Programs/AutomationTool/Gauntlet/Platform/Windows/Gauntlet.TargetDeviceWindows.cs
+++ b/Engine/Source/Programs/AutomationTool/Gauntlet/Platform/Windows/Gauntlet.TargetDeviceWindows.cs
@@ -9,6 +9,7 @@ using AutomationTool;
 using UnrealBuildTool;
 using System.Text.RegularExpressions;
 using EpicGames.Core;
+using static AutomationTool.ProcessResult;
 
 namespace Gauntlet
 {
@@ -337,8 +338,9 @@ namespace Gauntlet
 					}
 				}
 
-				// explicitly set log file when not already defined
-				if (string.IsNullOrEmpty(ProcessLogFile))
+				// Explicitly set log file when not already defined if not build machine
+				// -abslog makes sure Unreal dymanicaly update the log window when using -log
+				if (!AutomationTool.Automation.IsBuildMachine && string.IsNullOrEmpty(ProcessLogFile))
 				{
 					string LogFolder = string.Format(@"{0}\Logs", WinApp.ArtifactPath);
 
@@ -367,7 +369,7 @@ namespace Gauntlet
 
 				Log.Verbose("\t{0}", CmdLine);
 
-				Result = CommandUtils.Run(WinApp.ExecutablePath, CmdLine, Options: WinApp.RunOptions | (ProcessLogFile != null ? CommandUtils.ERunOptions.NoStdOutRedirect : 0 ));
+				Result = CommandUtils.Run(WinApp.ExecutablePath, CmdLine, Options: WinApp.RunOptions | (ProcessLogFile != null ? CommandUtils.ERunOptions.NoStdOutRedirect : 0 ), SpewFilterCallback: new SpewFilterCallbackType(delegate (string M) { return null; }) /* make sure stderr does not spew in the stdout */);
 
 				if (Result.HasExited && Result.ExitCode != 0)
 				{
diff --git a/Engine/Source/Programs/AutomationTool/Gauntlet/Unreal/Base/Gauntlet.UnrealTestNode.cs b/Engine/Source/Programs/AutomationTool/Gauntlet/Unreal/Base/Gauntlet.UnrealTestNode.cs
--- a/Engine/Source/Programs/AutomationTool/Gauntlet/Unreal/Base/Gauntlet.UnrealTestNode.cs
+++ b/Engine/Source/Programs/AutomationTool/Gauntlet/Unreal/Base/Gauntlet.UnrealTestNode.cs
@@ -1928,7 +1928,7 @@ namespace Gauntlet
 				{
 					Callstack.Add("Could not parse callstack. See log for full callstack");
 				}
-				Log.Error(KnownLogEvents.Gauntlet_TestEvent, " * Fatal Error: {Summary}\n{Callstack}", Event.Summary, string.Join("\n", Callstack.Select(C=>"    "+C)));
+				Log.Error(KnownLogEvents.Gauntlet_FatalEvent, " * Fatal Error: {Summary}\n{Callstack}", Event.Summary, string.Join("\n", Callstack.Select(C=>"    "+C)));
 				Log.Info("");
 			}
 
diff --git a/Engine/Source/Programs/AutomationTool/Gauntlet/Unreal/Utils/Gauntlet.UnrealLogParser.cs b/Engine/Source/Programs/AutomationTool/Gauntlet/Unreal/Utils/Gauntlet.UnrealLogParser.cs
--- a/Engine/Source/Programs/AutomationTool/Gauntlet/Unreal/Utils/Gauntlet.UnrealLogParser.cs
+++ b/Engine/Source/Programs/AutomationTool/Gauntlet/Unreal/Utils/Gauntlet.UnrealLogParser.cs
@@ -545,13 +545,77 @@ namespace Gauntlet
 		{
 			string[] ErrorMsgMatches = new string[] { @"(Fatal Error:.+)", @"Critical error: =+[\s\n]+(?:.+?\s*Error:\s*)?(.+)", @"(Assertion Failed:.+)", @"(Unhandled Exception:.+)", @"(LowLevelFatalError.+)" };
 
-			var Traces = ParseTracedErrors(ErrorMsgMatches);
+			var Traces = ParseTracedErrors(ErrorMsgMatches).Concat(GetASanErrors());
 
 			// Keep the one with the most information.
 			Traces.OrderBy(T => T.Callstack.Length);
 			return Traces.Count() > 0 ? Traces.Last() : null;
 		}
 
+		/// <summary>
+		/// Parse the log for Address Sanitizer error.
+		/// </summary>
+		/// <returns></returns>
+		public IEnumerable<UnrealLog.CallstackMessage> GetASanErrors()
+		{
+			/// Match:
+			/// ==5077==ERROR: AddressSanitizer: alloc - dealloc - mismatch(operator new vs free) on 0x602014ab4790
+			/// Then for gathering the callstack, match
+			/// ==5077==ABORTING
+			/// remove anything inside the callstack starting with [2022.12.02-15.22.40:688][618]
+
+			List<UnrealLog.CallstackMessage> ASanReports = new List<UnrealLog.CallstackMessage>();
+			string InitPattern = @"==([0-9]+)==ERROR: ([^\n]+)";
+			string LineStartsWithTimeStamp = @"^\[[0-9.:-]+\]\[[0-9]+\]";
+
+			MatchCollection Matches = Regex.Matches(Content, InitPattern, RegexOptions.IgnoreCase);
+
+			foreach (Match TraceInitMatch in Matches)
+			{
+				string TraceID = TraceInitMatch.Groups[1].Value;
+				string EndPattern = @$"=={TraceID}==ABORTING";
+				int TraceInitIndex = TraceInitMatch.Index + TraceInitMatch.Length;
+
+				UnrealLog.CallstackMessage NewTrace = new UnrealLog.CallstackMessage();
+				NewTrace.Position = TraceInitMatch.Index;
+				NewTrace.Message = TraceInitMatch.Groups[2].Value;
+
+				// If the regex matches the very end of the string, the substring will get an invalid range.
+				string ErrorContent = Content.Length <= TraceInitIndex
+					? string.Empty : Content.Substring(TraceInitIndex + 1);
+
+				if (!string.IsNullOrEmpty(ErrorContent))
+				{
+					Match MsgMatch = Regex.Match(ErrorContent, EndPattern);
+					if (MsgMatch.Success)
+					{
+						string MsgContent = ErrorContent.Substring(0, MsgMatch.Index);
+						// Prune the line with time stamp
+						List<string> Backtrace = new List<string>();
+						foreach (string Line in MsgContent.Split("\n"))
+						{
+							Match IsTimeStampLine = Regex.Match(Line, LineStartsWithTimeStamp);
+							if (!IsTimeStampLine.Success)
+							{
+								Backtrace.Add(Line);
+							}
+						}
+
+						NewTrace.Callstack = Backtrace.ToArray();
+					}
+				}
+
+				if(NewTrace.Callstack == null)
+				{
+					NewTrace.Callstack = new string[] { "Unable to parse callstack from log" };
+				}
+
+				ASanReports.Add(NewTrace);
+			}
+
+			return ASanReports;
+		}
+
 		/// <summary>
 		/// Returns true if the log contains a test complete marker
 		/// </summary>
diff --git a/Engine/Source/Programs/Shared/EpicGames.Core/KnownLogEvents.cs b/Engine/Source/Programs/Shared/EpicGames.Core/KnownLogEvents.cs
--- a/Engine/Source/Programs/Shared/EpicGames.Core/KnownLogEvents.cs
+++ b/Engine/Source/Programs/Shared/EpicGames.Core/KnownLogEvents.cs
@@ -170,6 +170,11 @@ namespace EpicGames.Core
 		/// </summary>
 		public static EventId Gauntlet_BuildDropEvent { get; } = new EventId(604);
 
+		/// <summary>
+		/// Fatal Error message from Gauntlet events
+		/// </summary>
+		public static EventId Gauntlet_FatalEvent { get; } = new EventId(605);
+
 		/// <summary>
 		/// A systemic event, relating to the health of the farm
 		/// </summary>
