- Set default Mac architecture to universal (arm64+x64)
- Refactored Android to match with Mac way of setting default

diff --git a/Engine/Source/Programs/AutomationTool/Scripts/BuildPluginCommand.Automation.cs b/Engine/Source/Programs/AutomationTool/Scripts/BuildPluginCommand.Automation.cs
--- a/Engine/Source/Programs/AutomationTool/Scripts/BuildPluginCommand.Automation.cs
+++ b/Engine/Source/Programs/AutomationTool/Scripts/BuildPluginCommand.Automation.cs
@@ -23,7 +23,9 @@ using UnrealBuildBase;
 [Help("Architecture_<Platform>=<Architecture[s]>", "Control architecture to compile for a platform (eg. -Architecture_Mac=arm64+x86). Default is to use UBT defaults for the platform.")]
 public sealed class BuildPlugin : BuildCommand
 {
-	const string AndroidArchitectures = "arm64";
+	const string MacDefaultArchitectures = "arm64+x64";
+	const string AndroidDefaultArchitectures = "arm64";
+
 	string UnrealBuildToolDllRelativePath = @"Engine/Binaries/DotNET/UnrealBuildTool/UnrealBuildTool.dll";
 	FileReference UnrealBuildToolDll;
 	static private Dictionary<UnrealTargetPlatform, string> PlatformToArchitectureMap = new Dictionary<UnrealTargetPlatform, string>();
@@ -115,7 +117,18 @@ public sealed class BuildPlugin : BuildCommand
 		// check if any architectures were specified
 		foreach (UnrealTargetPlatform Platform in UnrealTargetPlatform.GetValidPlatforms())
 		{
-			PlatformToArchitectureMap[Platform] = ParseParamValue($"architecture_{Platform}");
+			// by default, don't specify any architecture when building (unless user requested with -architecture_Platform=), except for
+			// any special cases set at the top of this file
+			string DefaultValue = null;
+			if (Platform == UnrealTargetPlatform.Mac)
+			{
+				DefaultValue = MacDefaultArchitectures;
+			}
+			else if (Platform == UnrealTargetPlatform.Android)
+			{
+				DefaultValue = AndroidDefaultArchitectures;
+			}
+			PlatformToArchitectureMap[Platform] = ParseParamValue($"architecture_{Platform}", DefaultValue);
 		}
 
 		// Compile the plugin for all the target platforms
@@ -355,10 +368,6 @@ public sealed class BuildPlugin : BuildCommand
 				{
 					Arguments += String.Format(" -architecture={0}", SpecifiedArchitecture);
 				}
-				else if (Platform == UnrealTargetPlatform.Android)
-				{
-					Arguments += String.Format(" -architectures={0}", AndroidArchitectures);
-				}
 
 				if (!String.IsNullOrEmpty(InAdditionalArgs))
 				{
