BuildGraph: Add a parameter to the cook task to disable tagging outputs. This can take a long time for large cooks.

diff --git a/Engine/Source/Programs/AutomationTool/BuildGraph/Tasks/CookTask.cs b/Engine/Source/Programs/AutomationTool/BuildGraph/Tasks/CookTask.cs
--- a/Engine/Source/Programs/AutomationTool/BuildGraph/Tasks/CookTask.cs
+++ b/Engine/Source/Programs/AutomationTool/BuildGraph/Tasks/CookTask.cs
@@ -58,6 +58,12 @@ namespace AutomationTool.Tasks
 		[TaskParameter(Optional = true)]
 		public string EditorExe = "";
 
+		/// <summary>
+		/// Whether to tag the output from the cook. Since cooks produce a lot of files, it can be detrimental to spend time tagging them if we don't need them in a dependent node.
+		/// </summary>
+		[TaskParameter(Optional = true)]
+		public bool TagOutput = true;
+
 		/// <summary>
 		/// Tag to be applied to build products of this task.
 		/// </summary>
@@ -117,27 +123,31 @@ namespace AutomationTool.Tasks
 
 			// Find all the cooked files
 			List<FileReference> CookedFiles = new List<FileReference>();
-			foreach(string Platform in Parameters.Platform.Split('+'))
+			if (Parameters.TagOutput)
 			{
-				DirectoryReference PlatformCookedDirectory = DirectoryReference.Combine(ProjectFile.Directory, "Saved", "Cooked", Platform);
-				if(!DirectoryReference.Exists(PlatformCookedDirectory))
-				{
-					throw new AutomationException("Cook output directory not found ({0})", PlatformCookedDirectory.FullName);
-				}
-				List<FileReference> PlatformCookedFiles = DirectoryReference.EnumerateFiles(PlatformCookedDirectory, "*", System.IO.SearchOption.AllDirectories).ToList();
-				if(PlatformCookedFiles.Count == 0)
+				foreach (string Platform in Parameters.Platform.Split('+'))
 				{
-					throw new AutomationException("Cooking did not produce any files in {0}", PlatformCookedDirectory.FullName);
-				}
-				CookedFiles.AddRange(PlatformCookedFiles);
-
-				DirectoryReference PackagingFilesDirectory = DirectoryReference.Combine(ProjectFile.Directory, "Saved", "TmpPackaging", Platform);
-				if (DirectoryReference.Exists(PackagingFilesDirectory))
-				{
-					List<FileReference> PackagingFiles = DirectoryReference.EnumerateFiles(PackagingFilesDirectory, "*", System.IO.SearchOption.AllDirectories).ToList();
-					CookedFiles.AddRange(PackagingFiles);
+					DirectoryReference PlatformCookedDirectory = DirectoryReference.Combine(ProjectFile.Directory, "Saved", "Cooked", Platform);
+					if (!DirectoryReference.Exists(PlatformCookedDirectory))
+					{
+						throw new AutomationException("Cook output directory not found ({0})", PlatformCookedDirectory.FullName);
+					}
+					List<FileReference> PlatformCookedFiles = DirectoryReference.EnumerateFiles(PlatformCookedDirectory, "*", System.IO.SearchOption.AllDirectories).ToList();
+					if (PlatformCookedFiles.Count == 0)
+					{
+						throw new AutomationException("Cooking did not produce any files in {0}", PlatformCookedDirectory.FullName);
+					}
+					CookedFiles.AddRange(PlatformCookedFiles);
+
+					DirectoryReference PackagingFilesDirectory = DirectoryReference.Combine(ProjectFile.Directory, "Saved", "TmpPackaging", Platform);
+					if (DirectoryReference.Exists(PackagingFilesDirectory))
+					{
+						List<FileReference> PackagingFiles = DirectoryReference.EnumerateFiles(PackagingFilesDirectory, "*", System.IO.SearchOption.AllDirectories).ToList();
+						CookedFiles.AddRange(PackagingFiles);
+					}
 				}
 			}
+
 			// Apply the optional tag to the build products
 			foreach(string TagName in FindTagNamesFromList(Parameters.Tag))
 			{
