- Fixed non-Development plist naming when packaging before doing a Development package

diff --git a/Engine/Source/Programs/AutomationTool/IOS/IOSPlatform.Automation.cs b/Engine/Source/Programs/AutomationTool/IOS/IOSPlatform.Automation.cs
--- a/Engine/Source/Programs/AutomationTool/IOS/IOSPlatform.Automation.cs
+++ b/Engine/Source/Programs/AutomationTool/IOS/IOSPlatform.Automation.cs
@@ -1357,7 +1357,7 @@ public class IOSPlatform : ApplePlatform
 							TargetConfiguration,
 							(SC.IsCodeBasedProject ? SC.ProjectRoot : DirectoryReference.Combine(SC.LocalRoot, "Engine")),
 							!SC.IsCodeBasedProject,
-							(SC.IsCodeBasedProject ? SC.StageExecutables[0] : "UnrealGame"),
+							(SC.IsCodeBasedProject ? TargetName : "UnrealGame"),
 							SC.IsCodeBasedProject ? false : Params.Client, // Code based projects will have Client in their executable name already
 							SC.ShortProjectName, DirectoryReference.Combine(SC.LocalRoot, "Engine"),
 							DirectoryReference.Combine((SC.IsCodeBasedProject ? SC.ProjectRoot : DirectoryReference.Combine(SC.LocalRoot, "Engine")), "Binaries", PlatformName, "Payload", (SC.IsCodeBasedProject ? SC.ShortProjectName : "UnrealGame") + ".app"),
diff --git a/Engine/Source/Programs/AutomationTool/TVOS/TVOSPlatform.Automation.cs b/Engine/Source/Programs/AutomationTool/TVOS/TVOSPlatform.Automation.cs
--- a/Engine/Source/Programs/AutomationTool/TVOS/TVOSPlatform.Automation.cs
+++ b/Engine/Source/Programs/AutomationTool/TVOS/TVOSPlatform.Automation.cs
@@ -116,7 +116,7 @@ public class TVOSPlatform : IOSPlatform
 						TargetConfiguration,
 						(SC.IsCodeBasedProject ? SC.ProjectRoot : SC.EngineRoot),
 						!SC.IsCodeBasedProject,
-						(SC.IsCodeBasedProject ? SC.StageExecutables[0] : "UnrealGame"),
+						(SC.IsCodeBasedProject ? TargetName : "UnrealGame"),
 						Params.Client,
 						SC.ShortProjectName,
 						SC.EngineRoot,
diff --git a/Engine/Source/Programs/UnrealBuildTool/ProjectFiles/Xcode/XcodeProjectLegacy.cs b/Engine/Source/Programs/UnrealBuildTool/ProjectFiles/Xcode/XcodeProjectLegacy.cs
--- a/Engine/Source/Programs/UnrealBuildTool/ProjectFiles/Xcode/XcodeProjectLegacy.cs
+++ b/Engine/Source/Programs/UnrealBuildTool/ProjectFiles/Xcode/XcodeProjectLegacy.cs
@@ -915,7 +915,7 @@ namespace UnrealBuildTool.XcodeProjectLegacy
 			Content.Append("\t\t\t\tALWAYS_SEARCH_USER_PATHS = NO;" + ProjectFileGenerator.NewLine);
 			if (Unreal.IsEngineInstalled())
 			{
-				Content.Append("\t\t\t\tARCHS[sdk=macosx*] = x86_64;" + ProjectFileGenerator.NewLine);
+				Content.Append("\t\t\t\t\"ARCHS[sdk=macosx*]\" = x86_64;" + ProjectFileGenerator.NewLine);
 			}
 			Content.Append("\t\t\t\tCLANG_CXX_LANGUAGE_STANDARD = \"c++14\";" + ProjectFileGenerator.NewLine);
 			Content.Append("\t\t\t\tGCC_ENABLE_CPP_RTTI = NO;" + ProjectFileGenerator.NewLine);
