Remove unnecessary Suspend Log parsing.

#rnx

diff --git a/Engine/Source/Programs/AutomationTool/Gauntlet/Framework/Gauntlet.Utils.cs b/Engine/Source/Programs/AutomationTool/Gauntlet/Framework/Gauntlet.Utils.cs
--- a/Engine/Source/Programs/AutomationTool/Gauntlet/Framework/Gauntlet.Utils.cs
+++ b/Engine/Source/Programs/AutomationTool/Gauntlet/Framework/Gauntlet.Utils.cs
@@ -1220,10 +1220,7 @@ namespace Gauntlet
 							}
 							else
 							{
-								using (var PauseEC = new ScopedSuspendECErrorParsing())
-								{
-									Log.Error("File Copy failed with {0}.", ex.Message);
-								}
+								Log.Info("File Copy failed with {0}.", ex.Message);
 
 								// Warn with message if we're exceeding long path, otherwise throw an exception
 								const int MAX_PATH = 260;
