AutomationTool: Trace P4 init

#rnx

diff --git a/Engine/Source/Programs/AutomationTool/AutomationUtils/Automation.cs b/Engine/Source/Programs/AutomationTool/AutomationUtils/Automation.cs
--- a/Engine/Source/Programs/AutomationTool/AutomationUtils/Automation.cs
+++ b/Engine/Source/Programs/AutomationTool/AutomationUtils/Automation.cs
@@ -146,8 +146,11 @@ namespace AutomationTool
 				if (CommandUtils.P4Enabled)
 				{
 					Log.TraceLog("Setting up Perforce environment.");
-					CommandUtils.InitP4Environment();
-					CommandUtils.InitDefaultP4Connection();
+					using (GlobalTracer.Instance.BuildSpan("InitP4").StartActive())
+					{
+						CommandUtils.InitP4Environment();
+						CommandUtils.InitDefaultP4Connection();
+					}
 				}
 
 				try
