Gauntlet - Fix Linux file not being collected when parsing callstack

diff --git a/Engine/Source/Programs/AutomationTool/Gauntlet/SelfTest/TestGauntlet.cs b/Engine/Source/Programs/AutomationTool/Gauntlet/SelfTest/TestGauntlet.cs
--- a/Engine/Source/Programs/AutomationTool/Gauntlet/SelfTest/TestGauntlet.cs
+++ b/Engine/Source/Programs/AutomationTool/Gauntlet/SelfTest/TestGauntlet.cs
@@ -14,11 +14,11 @@ namespace Gauntlet
 	public class TestGauntlet : BuildCommand
 	{
 
-		[AutoParamWithNames("", "Test", "Tests")]
-		public string Tests;
+		[AutoParamWithNames("Test", "Tests")]
+		public string Tests = "";
 
-		[AutoParamWithNames("", "Group")]
-		public string Group;
+		[AutoParamWithNames("Group")]
+		public string Group = "";
 
 		[AutoParam(false)]
 		public bool Verbose;
diff --git a/Engine/Source/Programs/AutomationTool/Gauntlet/SelfTest/Unreal/Gauntlet.SelfTest.LogParserTest.cs b/Engine/Source/Programs/AutomationTool/Gauntlet/SelfTest/Unreal/Gauntlet.SelfTest.LogParserTest.cs
--- a/Engine/Source/Programs/AutomationTool/Gauntlet/SelfTest/Unreal/Gauntlet.SelfTest.LogParserTest.cs
+++ b/Engine/Source/Programs/AutomationTool/Gauntlet/SelfTest/Unreal/Gauntlet.SelfTest.LogParserTest.cs
@@ -129,13 +129,16 @@ namespace Gauntlet.SelfTest
 	[TestGroup("Framework")]
 	class LogParserTestFatalError : TestUnrealLogParserBase
 	{
+		private int IncompleteLineLength = "0x0000000000236999  [Unknown File]".Length;
 		public override void TickTest()
 		{
 			HashSet<string> Platforms = new HashSet<string>();
 			Platforms.Add("Win64");
+			Platforms.Add("Linux");
 			Platforms.Add(Gauntlet.Globals.Params.ParseValue("Platform", "Win64"));
 			foreach (var Platform in Platforms)
 			{
+				Log.Info("Processing log: {Path}", Platform + "FatalError" + ".txt");
 				UnrealLogParser Parser = new UnrealLogParser(GetFileContents(Platform + "FatalError" + ".txt"));
 
 				UnrealLog.CallstackMessage FatalError = Parser.GetFatalError();
@@ -144,6 +147,15 @@ namespace Gauntlet.SelfTest
 				{
 					throw new TestException("LogParser returned incorrect assert info for {0}", Platform);
 				}
+				else
+				{
+					int IncompleteCallstackLineCount = FatalError.Callstack.Where(L => L.Length <= IncompleteLineLength).Count();
+					if(IncompleteCallstackLineCount > 0)
+					{
+						string Lines = string.Join("\n", FatalError.Callstack.Where(L => L.Length <= IncompleteLineLength));
+						throw new TestException("LogParser returned some incomplete callstack lines for {0}:\n{1}", Platform, Lines);
+					}
+				}
 			}			
 
 			MarkComplete(TestResult.Passed);
diff --git a/Engine/Source/Programs/AutomationTool/Gauntlet/Unreal/Utils/Gauntlet.UnrealLogParser.cs b/Engine/Source/Programs/AutomationTool/Gauntlet/Unreal/Utils/Gauntlet.UnrealLogParser.cs
--- a/Engine/Source/Programs/AutomationTool/Gauntlet/Unreal/Utils/Gauntlet.UnrealLogParser.cs
+++ b/Engine/Source/Programs/AutomationTool/Gauntlet/Unreal/Utils/Gauntlet.UnrealLogParser.cs
@@ -28,7 +28,8 @@ namespace Gauntlet
 			Verbose,
 			VeryVerbose,
 			Warning,
-			Error
+			Error,
+			Fatal
 		}
 
 		/// <summary>
@@ -810,7 +811,7 @@ namespace Gauntlet
 						// module address is optional, must be in quotes() after address
 						// Note - Unreal callstacks are always meant to omit all three with placeholders for missing values, but
 						// we'll assume that may not happen...
-						Match CSMatch = Regex.Match(Line, @"(0[xX][0-9A-f]{8,16})(?:\s+\(0[xX][0-9A-f]{8,16}\))?(?:\s+([^\s]+))?(?:\s+\[(.*?)\])?", RegexOptions.IgnoreCase);
+						Match CSMatch = Regex.Match(Line, @"(0[xX][0-9A-f]{8,16})(?:\s+\(0[xX][0-9A-f]{8,16}\))?\s+(.+)\s+\[(.*?)\]$", RegexOptions.IgnoreCase);
 
 						if (CSMatch.Success)
 						{
@@ -848,7 +849,7 @@ namespace Gauntlet
 
 					if (Backtrace.Count > 0)
 					{
-						NewTrace.Callstack = Backtrace.ToArray();
+						NewTrace.Callstack = Backtrace.Distinct().ToArray();
 					}
 					else
 					{
