UAT: Improve performance of UAT finding the last code CL.

diff --git a/Engine/Source/Programs/AutomationTool/AutomationUtils/P4Environment.cs b/Engine/Source/Programs/AutomationTool/AutomationUtils/P4Environment.cs
--- a/Engine/Source/Programs/AutomationTool/AutomationUtils/P4Environment.cs
+++ b/Engine/Source/Programs/AutomationTool/AutomationUtils/P4Environment.cs
@@ -6,6 +6,9 @@ using System.Linq;
 using System.Text;
 using UnrealBuildTool;
 using EpicGames.Core;
+using EpicGames.Perforce;
+using System.Threading.Tasks;
+using System.Diagnostics;
 
 namespace AutomationTool
 {
@@ -210,9 +213,11 @@ namespace AutomationTool
 			string CodeChangelistString = CommandUtils.GetEnvVar(EnvVarNames.CodeChangelist);
 			if(String.IsNullOrEmpty(CodeChangelistString) && CommandUtils.P4CLRequired)
 			{
-				P4Connection Connection = new P4Connection(User, Client, ServerAndPort);
-				CodeChangelistString = DetectCurrentCodeCL(Connection, ClientRoot, Changelist);
+				Stopwatch Timer = Stopwatch.StartNew();
+				using PerforceConnection Connection = new PerforceConnection(new PerforceSettings(ServerAndPort, User) { ClientName = Client }, Log.Logger);
+				CodeChangelistString = DetectCurrentCodeCL(Connection, Changelist).Result.ToString();
 				CommandUtils.SetEnvVar(EnvVarNames.CodeChangelist, CodeChangelistString);
+				Log.TraceVerbose("Took {0}ms to query last code change", Timer.ElapsedMilliseconds);
 			}
 			if(!String.IsNullOrEmpty(CodeChangelistString))
 			{
@@ -449,34 +454,50 @@ namespace AutomationTool
 		/// <summary>
 		/// Detects the current code changelist the workspace is synced to.
 		/// </summary>
-		/// <param name="ClientRootPath">Workspace path.</param>
 		/// <returns>Changelist number as a string.</returns>
-		private static string DetectCurrentCodeCL(P4Connection Connection, string ClientRootPath, int Changelist)
+		private static async Task<int> DetectCurrentCodeCL(PerforceConnection Connection, int Changelist)
 		{
 			CommandUtils.LogVerbose("uebp_CodeCL not set, detecting last code CL...");
 
-			// Retrieve the current changelist
-			StringBuilder P4Cmd = new StringBuilder("changes -m 1");
-
-			foreach (string CodeExtension in EpicGames.Perforce.PerforceUtils.CodeExtensions)
+			// Start by just testing whether the current change is a code change, so we can early out without any expensive p4 calls.
+			int[] Changes = new[] { Changelist };
+			while (Changes.Length > 0)
 			{
-				P4Cmd.AppendFormat(" \"{0}/...{1}@<={2}\"", CommandUtils.CombinePaths(PathSeparator.Depot, ClientRootPath), CodeExtension, Changelist);
-			}
-
-			IProcessResult P4Result = Connection.P4(P4Cmd.ToString(), AllowSpew: false);
-
-			// Loop through all the lines of the output. Even though we requested one result, we'll get one for each search pattern.
-			int CL = 0;
-			foreach(string Line in P4Result.Output.Split('\n'))
-			{
-				string[] Tokens = Line.Trim().Split(new char[] { ' ' }, StringSplitOptions.RemoveEmptyEntries);
-				if(Tokens.Length >= 2)
+				const int NumChangesPerDescribeCall = 10;
+				foreach (IReadOnlyList<int> ChangeBatch in Changes.Batch(NumChangesPerDescribeCall))
 				{
-					int LineCL = Int32.Parse(Tokens[1]);
-					CL = Math.Max(CL, LineCL);
+					const int InitialMaxFiles = 50;
+
+					List<DescribeRecord> Descriptions = await Connection.DescribeAsync(DescribeOptions.None, InitialMaxFiles, ChangeBatch.ToArray());
+					foreach (DescribeRecord Description in Descriptions)
+					{
+						DescribeRecord CurrentDescription = Description;
+
+						// The initial p4 describe call only queries up to InitialMaxFiles in the response. Fetching all files for large merge changelists can be prohibitively slow,
+						// so we query an increasing number of files with the goal of earlying out as soon as we hit a code change.
+						for (int MaxFiles = InitialMaxFiles; ;)
+						{
+							if (CurrentDescription.Files.Any(x => x.DepotFile != null && PerforceUtils.IsCodeFile(x.DepotFile)))
+							{
+								return CurrentDescription.Number;
+							}
+							if (CurrentDescription.Files.Count < MaxFiles)
+							{
+								break;
+							}
+
+							MaxFiles *= 10;
+							CurrentDescription = await Connection.DescribeAsync(DescribeOptions.None, MaxFiles, CurrentDescription.Number);
+						}
+					}
 				}
+
+				// Query the last NumChanges, then split it into batches for calling p4 describe
+				const int NumChanges = 30;
+				List<ChangesRecord> NextChanges = await Connection.GetChangesAsync(ChangesOptions.None, NumChanges, ChangeStatus.Submitted, $"//{Connection.ClientName}/...@<{Changes.Min()}");
+				Changes = NextChanges.Select(x => x.Number).ToArray();
 			}
-			return CL.ToString();
+			return 0;
 		}
 
 		/// <summary>
diff --git a/Engine/Source/Programs/Shared/EpicGames.Perforce/PerforceUtils.cs b/Engine/Source/Programs/Shared/EpicGames.Perforce/PerforceUtils.cs
--- a/Engine/Source/Programs/Shared/EpicGames.Perforce/PerforceUtils.cs
+++ b/Engine/Source/Programs/Shared/EpicGames.Perforce/PerforceUtils.cs
@@ -37,6 +37,23 @@ namespace EpicGames.Perforce
 			".sln"
 		};
 
+		/// <summary>
+		/// Tests if a path is a code file
+		/// </summary>
+		/// <param name="path">Path to test</param>
+		/// <returns>True if the path is a code file</returns>
+		public static bool IsCodeFile(string path)
+		{
+			foreach (string codeExtension in CodeExtensions)
+			{
+				if (path.EndsWith(codeExtension, StringComparison.OrdinalIgnoreCase))
+				{
+					return true;
+				}
+			}
+			return false;
+		}
+
 		/// <summary>
 		/// Escape a path to Perforce syntax
 		/// </summary>
