BuildCMakeLib: Added [DoesNotNeedP4CL] to reduce startup time

#rnx

diff --git a/Engine/Source/Programs/AutomationTool/AutomationUtils/P4Utils.cs b/Engine/Source/Programs/AutomationTool/AutomationUtils/P4Utils.cs
--- a/Engine/Source/Programs/AutomationTool/AutomationUtils/P4Utils.cs
+++ b/Engine/Source/Programs/AutomationTool/AutomationUtils/P4Utils.cs
@@ -15,10 +15,16 @@ using UnrealBuildBase;
 
 namespace AutomationTool
 {
+	/// <summary>
+	/// Declares that the command type requires P4Environment.
+	/// </summary>
 	public class RequireP4Attribute : Attribute
 	{
 	}
 
+	/// <summary>
+	/// Declares that the command type does not access Changelist or CodeChangelist from P4Environment.
+	/// </summary>
 	public class DoesNotNeedP4CLAttribute : Attribute
 	{
 	}
diff --git a/Engine/Source/Programs/AutomationTool/Scripts/BuildCMakeLib.Automation.cs b/Engine/Source/Programs/AutomationTool/Scripts/BuildCMakeLib.Automation.cs
--- a/Engine/Source/Programs/AutomationTool/Scripts/BuildCMakeLib.Automation.cs
+++ b/Engine/Source/Programs/AutomationTool/Scripts/BuildCMakeLib.Automation.cs
@@ -31,6 +31,7 @@ using UnrealBuildBase;
 [Help("SkipSubmit", "Do not perform P4 submit of source or libs. If this argument is not supplied source and libs will be automatically submitted to Perforce. If SkipCreateChangelist is specified, this argument applies by default.")]
 [Help("RoboMerge", "Which RoboMerge action to apply to the submission. If we're skipping submit, this is not used.")]
 [RequireP4]
+[DoesNotNeedP4CL]
 public sealed class BuildCMakeLib : BuildCommand
 {
 	public class TargetLib
