Adding AssetRegistryCacheRootFolder commandline option support to the BuildCookFort command and the MakeCookedEditor command.
Adding the AssetRegistryCacheRootFolder to the ProfileCook,  asset scan, and fast cook test to improve performance.

diff --git a/Engine/Source/Programs/AutomationTool/CookedEditor/CookedEditor.Automation.cs b/Engine/Source/Programs/AutomationTool/CookedEditor/CookedEditor.Automation.cs
--- a/Engine/Source/Programs/AutomationTool/CookedEditor/CookedEditor.Automation.cs
+++ b/Engine/Source/Programs/AutomationTool/CookedEditor/CookedEditor.Automation.cs
@@ -1084,6 +1084,12 @@ public class MakeCookedEditor : BuildCommand
 		// and choose to cook them or not (generally for editoronly assets)
 		Params.AdditionalCookerOptions += " -DlcReevaluateUncookedAssets";
 
+		string AssetRegistryCacheRootFolder = ParseParamValue("AssetRegistryCacheRootFolder", "");
+		if (!string.IsNullOrEmpty(AssetRegistryCacheRootFolder))
+		{
+			Params.AdditionalCookerOptions += string.Format(" -AssetRegistryCacheRootFolder={0}", AssetRegistryCacheRootFolder);
+		}
+
 		// set up cooking against a client, as DLC
 		if (bIsDLC)
 		{
