Device Cleanup

diff --git a/Engine/Source/Programs/AutomationTool/Gauntlet/Platform/Android/Gauntlet.TargetDeviceAndroid.cs b/Engine/Source/Programs/AutomationTool/Gauntlet/Platform/Android/Gauntlet.TargetDeviceAndroid.cs
--- a/Engine/Source/Programs/AutomationTool/Gauntlet/Platform/Android/Gauntlet.TargetDeviceAndroid.cs
+++ b/Engine/Source/Programs/AutomationTool/Gauntlet/Platform/Android/Gauntlet.TargetDeviceAndroid.cs
@@ -908,11 +908,21 @@ namespace Gauntlet
 			RunAdbDeviceCommand(string.Format("shell rm -r {0}", DeviceExternalStorageSavedPath));
 			RunAdbDeviceCommand(string.Format("shell rm -r {0}", DeviceExternalFilesSavedPath));
 
-			bool SkipDeploy = Globals.Params.ParseParam("SkipDeploy");
+			if (Globals.Params.ParseParam("fullclean"))
+			{
+				Log.Info("Fully cleaning console before install...");
+				RunAdbDeviceCommand(string.Format("shell rm -r {0}/UnrealGame/*", StorageLocation));
+				RunAdbDeviceCommand(string.Format("shell rm -r {0}/Android/data/{1}/*", StorageLocation, Build.AndroidPackageName));
+				RunAdbDeviceCommand(string.Format("shell rm -r {0}/Android/obb/{1}/*", StorageLocation, Build.AndroidPackageName));
+				RunAdbDeviceCommand(string.Format("shell rm -r {0}/Download/*", StorageLocation));
+			}
+
+				bool SkipDeploy = Globals.Params.ParseParam("SkipDeploy");
 
 			if (SkipDeploy == false)
 			{
-				if (Globals.Params.ParseParam("cleandevice"))
+				if (Globals.Params.ParseParam("cleandevice")
+					|| Globals.Params.ParseParam("fullclean"))
 				{
 					Log.Info("Cleaning previous builds due to presence of -cleandevice");
 
