Remote Prepare to Debug

diff --git a/Engine/Build/Rsync/RsyncEngine.txt b/Engine/Build/Rsync/RsyncEngine.txt
--- a/Engine/Build/Rsync/RsyncEngine.txt
+++ b/Engine/Build/Rsync/RsyncEngine.txt
@@ -68,6 +68,7 @@ protect /Binaries/TVOS/
 - /Binaries/ThirdParty/DotNet/**
 + /Binaries/ThirdParty/VisualStudio/
 + /Binaries/ThirdParty/VisualStudio/**
++ /Binaries/ThirdParty/fastJSON/***
 + /Binaries/ThirdParty/Intel/
 + /Binaries/ThirdParty/Intel/TBB/
 + /Binaries/ThirdParty/Intel/TBB/Mac/
@@ -154,9 +155,7 @@ protect /Saved/
 + /Source/Programs/DotNETCommon/
 + /Source/Programs/UnrealBuildTool/
 + /Source/Programs/Shared/
-+ /Source/Programs/Shared/EpicGames.Core/
-+ /Source/Programs/Shared/EpicGames.Build/
-+ /Source/Programs/Shared/EpicGames.UHT/
++ /Source/Programs/Shared/EpicGames.*/
 - /Source/Programs/Shared/*/
 + /Source/Programs/UnrealPak/
 + /Source/Programs/UnrealHeaderTool/
@@ -170,10 +169,7 @@ protect /Saved/
 protect /Source/Programs/DotNETCommon/obj/**
 protect /Source/Programs/DotNETCommon/BuildUtilities/bin/**
 protect /Source/Programs/UnrealBuildTool/obj/**
-protect /Source/Programs/Shared/EpicGames.Core/obj/**
-protect /Source/Programs/Shared/EpicGames.Core/bin/**
-protect /Source/Programs/Shared/EpicGames.Build/obj/**
-protect /Source/Programs/Shared/EpicGames.Build/bin/**
+protect /Source/Programs/Shared/EpicGames.*/***
 
 # Source/ThirdParty: Allow specific file types
 + /Source/ThirdParty/**/
@@ -187,6 +183,7 @@ protect /Source/Programs/Shared/EpicGames.Build/bin/**
 + /Source/ThirdParty/**.inl
 + /Source/ThirdParty/**.a
 + /Source/ThirdParty/**.dylib
++ /Source/ThirdParty/**.plist
 + /Source/ThirdParty/**.embeddedframework.zip
 + /Source/ThirdParty/libJPG/*.cpp
 + /Source/ThirdParty/libSampleRate/**.cpp
diff --git a/Engine/Source/Editor/TurnkeySupport/Private/TurnkeySupportModule.cpp b/Engine/Source/Editor/TurnkeySupport/Private/TurnkeySupportModule.cpp
--- a/Engine/Source/Editor/TurnkeySupport/Private/TurnkeySupportModule.cpp
+++ b/Engine/Source/Editor/TurnkeySupport/Private/TurnkeySupportModule.cpp
@@ -513,6 +513,7 @@ public:
 		}
         else if (Mode == EPrepareContentMode::PrepareForDebugging)
          {
+#if PLATFORM_MAC || PLATFORM_WINDOWS
              if (IniPlatformName == TEXT("IOS") || IniPlatformName == TEXT("TvOS"))
              {
                  FString CommandLine = FString::Printf(TEXT("WrangleContentForDebugging -platform=%s -ProjectFilePath=%s"), *IniPlatformName.ToString().ToLower(), *ProjectPath);
@@ -520,8 +521,8 @@ public:
                  FTurnkeyEditorSupport::RunUAT(CommandLine, PlatformInfo->DisplayName, ContentPrepDescription, ContentPrepTaskName, ContentPrepIcon, &AnalyticsParamArray);
              
                  return;
-
              }
+#endif
  		}
 		else if (Mode == EPrepareContentMode::CookOnly)
 		{
@@ -1119,7 +1120,7 @@ static void MakeTurnkeyPlatformMenu(UToolMenu* ToolMenu, FName IniPlatformName,
             Section.AddMenuEntry(
                                  NAME_None,
                                  LOCTEXT("Turnkey_PrepareForDebugging", "Prepare For Debugging"),
-                                 LOCTEXT("TurnkeyTooltip_PrepareForDebuggingIOS", "Prepare this project for debugging. Expects an IPA package with the same name as the project file in the Build/IOS/ folder."),
+                                 LOCTEXT("TurnkeyTooltip_PrepareForDebuggingIOS", "Prepare this project for debugging. Expects an IPA package with the same name as the project file in the Binaries/IOS/ folder."),
                                  FSlateIcon(),
                                  FUIAction(
                                            FExecuteAction::CreateStatic(&FTurnkeySupportCallbacks::CookOrPackage, IniPlatformName, EPrepareContentMode::PrepareForDebugging),
@@ -1132,7 +1133,7 @@ static void MakeTurnkeyPlatformMenu(UToolMenu* ToolMenu, FName IniPlatformName,
             Section.AddMenuEntry(
                                  NAME_None,
                                  LOCTEXT("Turnkey_PrepareForDebugging", "Prepare For Debugging"),
-                                 LOCTEXT("TurnkeyTooltip_PrepareForDebuggingTvOS", "Prepare this project for debugging. Expects an IPA package with the same name as the project file in the Build/TVOS/ folder."),
+                                 LOCTEXT("TurnkeyTooltip_PrepareForDebuggingTvOS", "Prepare this project for debugging. Expects an IPA package with the same name as the project file in the Binaries/TVOS/ folder."),
                                  FSlateIcon(),
                                  FUIAction(
                                            FExecuteAction::CreateStatic(&FTurnkeySupportCallbacks::CookOrPackage, IniPlatformName, EPrepareContentMode::PrepareForDebugging),
@@ -1140,6 +1141,23 @@ static void MakeTurnkeyPlatformMenu(UToolMenu* ToolMenu, FName IniPlatformName,
                                            )
                                  );
         }
+#elif PLATFORM_WINDOWS
+		if (IniPlatformName.ToString() == "IOS" || IniPlatformName.ToString() == "TvOS")
+		{
+
+			FString Tooltip = "";
+			
+			Section.AddMenuEntry(
+				NAME_None,
+				LOCTEXT("Turnkey_PrepareForDebugging", "Prepare For Debugging"),
+				FText::Format(LOCTEXT("TurnkeyTooltip_PrepareForDebugging", "Expects a working remote toolchain and an IPA package with the same name as the project file in the Binaries/ {0} / folder."), FText::FromString(IniPlatformName.ToString().ToUpper())),
+				FSlateIcon(),
+				FUIAction(
+					FExecuteAction::CreateStatic(&FTurnkeySupportCallbacks::CookOrPackage, IniPlatformName, EPrepareContentMode::PrepareForDebugging),
+					FCanExecuteAction::CreateStatic(&FTurnkeySupportCallbacks::CanCookOrPackage, IniPlatformName, EPrepareContentMode::PrepareForDebugging)
+				)
+			);
+		}
 #endif
 
 		UProjectPackagingSettings* PackagingSettings = FTurnkeySupportCallbacks::GetPackagingSettingsForPlatform(IniPlatformName);
diff --git a/Engine/Source/Programs/AutomationTool/IOS/IOSPlatform.Automation.cs b/Engine/Source/Programs/AutomationTool/IOS/IOSPlatform.Automation.cs
--- a/Engine/Source/Programs/AutomationTool/IOS/IOSPlatform.Automation.cs
+++ b/Engine/Source/Programs/AutomationTool/IOS/IOSPlatform.Automation.cs
@@ -2296,74 +2296,74 @@ public class IOSPlatform : ApplePlatform
 
 	public override void PrepareForDebugging(string SourcePackage, string ProjectFilePath, string ClientPlatform)
 	{
-		if (HostPlatform.Current.HostEditorPlatform != UnrealTargetPlatform.Mac)
+		if (HostPlatform.Current.HostEditorPlatform == UnrealTargetPlatform.Win64 || HostPlatform.Current.HostEditorPlatform == UnrealTargetPlatform.Mac)
 		{
-			LogInformation("Wrangling data for debug for an iOS/tvOS app for XCode is a Mac only feature. Aborting command.");
-			return;
-		}
-
-		int StartPos = ProjectFilePath.LastIndexOf("/");
-		int StringLength = ProjectFilePath.Length - 10; // 9 for .uproject, 1 for the /
-		string PackageName = ProjectFilePath.Substring(StartPos + 1, StringLength - StartPos);
+			int StartPos = ProjectFilePath.LastIndexOf("/");
+			int StringLength = ProjectFilePath.Length - 10; // 9 for .uproject, 1 for the /
+			string PackageName = ProjectFilePath.Substring(StartPos + 1, StringLength - StartPos);
 			if (string.IsNullOrEmpty(SourcePackage))
-		{
-			SourcePackage = ProjectFilePath;
-			SourcePackage = SourcePackage.Substring(0, SourcePackage.LastIndexOf('/'));
-			SourcePackage = SourcePackage + "/Build/" + ClientPlatform + '/' + PackageName + ".ipa";
-		}
+			{
+				SourcePackage = Path.Combine(Path.GetDirectoryName(ProjectFilePath), "Binaries", ClientPlatform, PackageName + ".ipa");
+			}
 
-		string ZipFile = SourcePackage.Replace(".ipa", ".zip");
+			string ZipFile = Path.ChangeExtension(SourcePackage, "zip");
 
-		string PayloadPath = SourcePackage;
-		PayloadPath = PayloadPath.Substring(0, PayloadPath.LastIndexOf('/'));
-		PayloadPath += "/Payload/";
-		string CookedDataDirectory = PayloadPath + PackageName + ".app/cookeddata/";
+			string PayloadPath = SourcePackage;
+			PayloadPath = PayloadPath.Substring(0, PayloadPath.LastIndexOf('\\'));
+			string CookedDataDirectory = Path.Combine(Path.GetDirectoryName(PayloadPath), ClientPlatform, "Payload", PackageName + ".app", "cookeddata");
 
-		LogInformation("ClientPlatform : {0}", ClientPlatform);
-		LogInformation("ProjectFilePath : {0}", ProjectFilePath);
-		LogInformation("Source : {0}", SourcePackage);
-		LogInformation("ZipFile {0}", ZipFile);
-		LogInformation("PackageName {0}", PackageName);
-		LogInformation("PayloadPath {0}", PayloadPath);
+			LogInformation("ClientPlatform : {0}", ClientPlatform);
+			LogInformation("ProjectFilePath : {0}", ProjectFilePath);
+			LogInformation("Source : {0}", SourcePackage);
+			LogInformation("ZipFile {0}", ZipFile);
+			LogInformation("PackageName {0}", PackageName);
+			LogInformation("PayloadPath {0}", PayloadPath);
 
-		if (File.Exists(ZipFile))
-		{
-			LogInformation("Deleting previously present ZIP file created from IPA");
-			File.Delete(ZipFile);
-		}
+			if (File.Exists(ZipFile))
+			{
+				LogInformation("Deleting previously present ZIP file created from IPA");
+				File.Delete(ZipFile);
+			}
 
-		File.Copy(SourcePackage, ZipFile);
-		UnzipPackage(ZipFile);
+			File.Copy(SourcePackage, ZipFile);
+			UnzipPackage(ZipFile);
 
-		string ProjectPath = ProjectFilePath;
-		int Index = ProjectFilePath.IndexOf("?");
-		if (Index >= 0)
-		{
-		   ProjectPath = ProjectPath.Substring(0, Index);
-		}
-		if (Directory.Exists(ProjectPath + "/Binaries/" + ClientPlatform + '/' + "Payload/" + PackageName + ".app"))
-		{
-			CopyDirectory_NoExceptions(CookedDataDirectory, ProjectPath + "/Binaries/" + ClientPlatform + "/Payload/" + PackageName + ".app/cookeddata/", true);
+			if (HostPlatform.Current.HostEditorPlatform == UnrealTargetPlatform.Win64)
+			{
+				IOSExports.PrepareRemoteMacForDebugging(CookedDataDirectory, new FileReference(ProjectFilePath), Log.Logger);
+			}
+			else
+			{
+				string ProjectPath = ProjectFilePath;
+				if (Directory.Exists(Path.Combine(ProjectPath, "Binaries", ClientPlatform, "Payload", PackageName, ".app")))
+				{
+					CopyDirectory_NoExceptions(CookedDataDirectory, Path.Combine(ProjectPath, "Binaries", ClientPlatform, "Payload", PackageName, ".app/cookeddata/"), true);
+				}
+				else
+				{
+					string ProjectRoot = SourcePackage;
+					ProjectRoot = ProjectRoot.Substring(0, ProjectRoot.LastIndexOf('/'));
+					CopyFile(SourcePackage, ProjectRoot + "/Binaries/" + ClientPlatform + "/Payload/" + PackageName + ".ipa", true);
+				}
+			}
+			//cleanup
+			LogInformation("Deleting temp files ...");
+			File.Delete(ZipFile);
+			LogInformation("{0} deleted", ZipFile);
+			Directory.Delete(PayloadPath, true);
+			LogInformation("{0} deleted", PayloadPath);
 		}
 		else
 		{
-			string ProjectRoot = SourcePackage;
-			ProjectRoot = ProjectRoot.Substring(0, ProjectRoot.LastIndexOf('/'));
-			CopyFile(SourcePackage, ProjectRoot + "/Binaries/" + ClientPlatform + "/Payload/" + PackageName + ".ipa", true);
+			LogInformation("Wrangling data for debug for an iOS/tvOS app for XCode is a Mac and Windows (Remote) only feature. Aborting command.");
+			return;
 		}
-
-		//cleanup
-		LogInformation("Deleting temp files ...");
-		File.Delete(ZipFile);
-		LogInformation("{0} deleted", ZipFile);
-		Directory.Delete(PayloadPath, true);
-		LogInformation("{0} deleted", PayloadPath);
 	}
 
 	public void UnzipPackage(string PackageToUnzip)
 	{
 		string UnzipPath = PackageToUnzip;
-		UnzipPath = UnzipPath.Substring(0, UnzipPath.LastIndexOf('/'));
+		UnzipPath = UnzipPath.Substring(0, UnzipPath.LastIndexOf('\\'));
 		LogInformation("Unzipping to {0}", UnzipPath);
 
 		using (Ionic.Zip.ZipFile Zip = new Ionic.Zip.ZipFile(PackageToUnzip))
diff --git a/Engine/Source/Programs/AutomationTool/Scripts/WrangleContentForDebugging.Automation.cs b/Engine/Source/Programs/AutomationTool/Scripts/WrangleContentForDebugging.Automation.cs
--- a/Engine/Source/Programs/AutomationTool/Scripts/WrangleContentForDebugging.Automation.cs
+++ b/Engine/Source/Programs/AutomationTool/Scripts/WrangleContentForDebugging.Automation.cs
@@ -14,7 +14,6 @@ namespace AutomationScripts
 	{
 		public override void ExecuteBuild()
 		{
-			LogInformation("======= WrangleContentForDebugging - Start =======");
 
 			if (ParseParam("usage"))
 			{
@@ -22,12 +21,14 @@ namespace AutomationScripts
 			}
 
 			FileReference ProjectFile = new FileReference(ParseParamValue("ProjectFilePath=", null));
-			//ProjectFile = ParseProjectParam();
 			var Params = new ProjectParams
 			(
 				Command: this,
 				RawProjectPath: ProjectFile
 			);
+
+			LogInformation("======= WrangleContentForDebugging - Start =======");
+
 			var SC = Project.CreateDeploymentContext(Params, false);
 
 			//Project newProject;
diff --git a/Engine/Source/Programs/UnrealBuildTool/Platform/IOS/IOSExports.cs b/Engine/Source/Programs/UnrealBuildTool/Platform/IOS/IOSExports.cs
--- a/Engine/Source/Programs/UnrealBuildTool/Platform/IOS/IOSExports.cs
+++ b/Engine/Source/Programs/UnrealBuildTool/Platform/IOS/IOSExports.cs
@@ -182,6 +182,19 @@
 		}
 
 		/// <summary>
+		/// Prepare the build and project on a Remote Mac to be able to debug an iOS or tvOS package built remotely
+		/// </summary>
+		/// <param name="ClientPlatform">TargetPlatform, iOS or tvOS</param>
+		/// <param name="ProjectFile">Location of .uproject file</param>
+		/// <param name="Logger">A logger</param>
+		/// <returns></returns>
+		public static void PrepareRemoteMacForDebugging(string ClientPlatform, FileReference ProjectFile, ILogger Logger)
+        {
+			RemoteMac Remote = new RemoteMac(ProjectFile, Logger);
+			Remote.PrepareToDebug(ClientPlatform, ProjectFile, Logger);
+		}
+
+		/// <summary>
 		/// Reads the per-platform .PackageVersionCounter file to get a version, and modifies it with updated value so that next build will get new version
 		/// </summary>
 		/// <param name="UProjectFile">Location of .uproject file (or null for the engine project)</param>
@@ -305,7 +318,7 @@
 				if (bAutomaticSigning)
 				{
 					Arguments.Add("CODE_SIGN_IDENTITY=" + (bForDistribution ? "\"iPhone Distribution\"" : "\"iPhone Developer\""));
-					Arguments.Add("CODE_SIGN_STYLE=\"Automatic\"");
+					Arguments.Add("CODE_SIGN_STYLE=\"Automatic\"");
 					Arguments.Add("-allowProvisioningUpdates");
 					Arguments.Add($"DEVELOPMENT_TEAM={Team}");
 				}
diff --git a/Engine/Source/Programs/UnrealBuildTool/ToolChain/RemoteMac.cs b/Engine/Source/Programs/UnrealBuildTool/ToolChain/RemoteMac.cs
--- a/Engine/Source/Programs/UnrealBuildTool/ToolChain/RemoteMac.cs
+++ b/Engine/Source/Programs/UnrealBuildTool/ToolChain/RemoteMac.cs
@@ -711,6 +711,25 @@ namespace UnrealBuildTool
 			DownloadFile(OutputFile, Logger);
 		}
 
+		public void PrepareToDebug(string CookedDataDirectory, FileReference ProjectFile, ILogger Logger)
+        {
+			Logger.LogInformation("[Remote] Uploading the default uprojectdirs ...");
+			UploadFile(new FileReference("Default.uprojectdirs"), Logger);
+
+			Logger.LogInformation("[Remote] Generating an Xcode Project file...");
+			string CommandLine = EscapeShellArgument(GetRemotePath(Unreal.EngineDirectory));
+			CommandLine += "/Build/BatchFiles/Mac/GenerateProjectFiles.sh";
+			StringBuilder BuildCommandLine = new StringBuilder(CommandLine);			
+			
+			Execute(GetRemotePath(Unreal.RootDirectory), CommandLine, Logger);
+
+			Logger.LogInformation("[Remote] Uploading cooked data directory...");
+			Logger.LogInformation(CookedDataDirectory);
+			UploadDirectory(new DirectoryReference(CookedDataDirectory), Logger);
+
+			Logger.LogInformation("An Xcode project file has been generated on your remote Mac at the above address. Please access it to debug your IPA.");
+		}
+
 		/// <summary>
 		/// Convers a remote path into local form
 		/// </summary>
