Update build scripts to support and use latest BPT v1.5.2 by default

diff --git a/Engine/Source/Programs/AutomationTool/AutomationUtils/MCPPublic.cs b/Engine/Source/Programs/AutomationTool/AutomationUtils/MCPPublic.cs
--- a/Engine/Source/Programs/AutomationTool/AutomationUtils/MCPPublic.cs
+++ b/Engine/Source/Programs/AutomationTool/AutomationUtils/MCPPublic.cs
@@ -642,8 +642,10 @@ namespace EpicGames.MCP.Automation
 			/// </summary>
 			Online_v140,
 			Online_v150,
+			Online_v151,
+			Online_v152,
 
-			Online_Live = Online_v150
+			Online_Live = Online_v152
 		}
 
 		/// <summary>
