adding warning to runaut buildplugin when HostPlatforms contains unsupported platforms

https://p4-swarm.epicgames.net/reviews/23960551

diff --git a/Engine/Source/Programs/AutomationTool/Scripts/BuildPluginCommand.Automation.cs b/Engine/Source/Programs/AutomationTool/Scripts/BuildPluginCommand.Automation.cs
--- a/Engine/Source/Programs/AutomationTool/Scripts/BuildPluginCommand.Automation.cs
+++ b/Engine/Source/Programs/AutomationTool/Scripts/BuildPluginCommand.Automation.cs
@@ -132,7 +132,7 @@ public sealed class BuildPlugin : BuildCommand
 		}
 
 		// Compile the plugin for all the target platforms
-		List<UnrealTargetPlatform> HostPlatforms = GetHostPlatforms(this);
+		IReadOnlyList<UnrealTargetPlatform> HostPlatforms = GetHostPlatforms(this);
 		List<UnrealTargetPlatform> TargetPlatforms = GetTargetPlatforms(this, BuildHostPlatform.Current.Platform);
 		FileReference[] BuildProducts = CompilePlugin(UnrealBuildToolDll, HostProjectFile, HostProjectPluginFile, Plugin, HostPlatforms, TargetPlatforms, AdditionalArgs.ToString());
 
@@ -249,7 +249,7 @@ public sealed class BuildPlugin : BuildCommand
 		return BuildProducts.ToArray();
 	}
 
-	public static FileReference[] CompilePlugin(FileReference UnrealBuildToolDll, FileReference HostProjectFile, FileReference HostProjectPluginFile, PluginDescriptor Plugin, List<UnrealTargetPlatform> HostPlatforms, List<UnrealTargetPlatform> TargetPlatforms, string AdditionalArgs = "")
+	public static FileReference[] CompilePlugin(FileReference UnrealBuildToolDll, FileReference HostProjectFile, FileReference HostProjectPluginFile, PluginDescriptor Plugin, IReadOnlyList<UnrealTargetPlatform> HostPlatforms, List<UnrealTargetPlatform> TargetPlatforms, string AdditionalArgs = "")
 	{
 		List<FileReference> ManifestFileNames = new List<FileReference>();
 
@@ -501,66 +501,55 @@ public sealed class BuildPlugin : BuildCommand
 		return TargetPlatforms;
 	}
 
-	static List<UnrealTargetPlatform> GetHostPlatforms(BuildCommand Command)
+	static IReadOnlyList<UnrealTargetPlatform> GetHostPlatforms(BuildCommand Command)
 	{
-		List<UnrealTargetPlatform> HostPlatforms = new List<UnrealTargetPlatform>();
-		if (!Command.ParseParam("NoHostPlatform"))
+		if (Command.ParseParam("NoHostPlatform"))
 		{
-			var CurrentPlatform = BuildHostPlatform.Current.Platform;
-
-			if (Command.ParseParam("HostPlatforms"))
-			{
-				string HostPlatformFilter = Command.ParseParamValue("HostPlatforms", null);
-				if(HostPlatformFilter != null)
-				{
-					// Only interested in building for Platforms that support code projects
-					HostPlatforms = PlatformExports.GetRegisteredPlatforms().Where(x => InstalledPlatformInfo.IsValidPlatform(x, EProjectType.Code)).ToList();
-
-					// only build Mac on Mac
-					if (CurrentPlatform != UnrealTargetPlatform.Mac && HostPlatforms.Contains(UnrealTargetPlatform.Mac))
-					{
-						HostPlatforms.Remove(UnrealTargetPlatform.Mac);
-					}
-					// only build Windows on Windows
-					if (CurrentPlatform != UnrealTargetPlatform.Win64 && HostPlatforms.Contains(UnrealTargetPlatform.Win64))
-					{
-						HostPlatforms.Remove(UnrealTargetPlatform.Win64);
-					}
-					// build Linux on Windows and Linux
-					if (CurrentPlatform != UnrealTargetPlatform.Win64 && CurrentPlatform != UnrealTargetPlatform.Linux)
-					{
-						if (HostPlatforms.Contains(UnrealTargetPlatform.Linux))
-							HostPlatforms.Remove(UnrealTargetPlatform.Linux);
+			return Array.Empty<UnrealTargetPlatform>();
+		}
+		var CurrentPlatform = BuildHostPlatform.Current.Platform;
+		string HostPlatformFilter = Command.ParseParamValue("HostPlatforms", null);
+		if (HostPlatformFilter == null)
+		{
+			return new[] { CurrentPlatform };
+		}
+		// Only interested in building for Platforms that support code projects
+		HashSet<UnrealTargetPlatform> SupportedHostPlatforms = PlatformExports.GetRegisteredPlatforms().Where(x => InstalledPlatformInfo.IsValidPlatform(x, EProjectType.Code)).ToHashSet();
 
-						if (HostPlatforms.Contains(UnrealTargetPlatform.LinuxArm64))
-							HostPlatforms.Remove(UnrealTargetPlatform.LinuxArm64);
-					}
+		// only build Mac on Mac
+		if (CurrentPlatform != UnrealTargetPlatform.Mac && SupportedHostPlatforms.Contains(UnrealTargetPlatform.Mac))
+		{
+			SupportedHostPlatforms.Remove(UnrealTargetPlatform.Mac);
+		}
+		// only build Windows on Windows
+		if (CurrentPlatform != UnrealTargetPlatform.Win64 && SupportedHostPlatforms.Contains(UnrealTargetPlatform.Win64))
+		{
+			SupportedHostPlatforms.Remove(UnrealTargetPlatform.Win64);
+		}
+		// build Linux on Windows and Linux
+		if (CurrentPlatform != UnrealTargetPlatform.Win64 && CurrentPlatform != UnrealTargetPlatform.Linux)
+		{
+			SupportedHostPlatforms.Remove(UnrealTargetPlatform.Linux);
+			SupportedHostPlatforms.Remove(UnrealTargetPlatform.LinuxArm64);
+		}
 
-					List<UnrealTargetPlatform> NewHostPlatforms = new List<UnrealTargetPlatform>();
-					foreach (string HostPlatformName in HostPlatformFilter.Split(new char[]{ '+' }, StringSplitOptions.RemoveEmptyEntries))
-					{
-						UnrealTargetPlatform HostPlatform;
-						if (!UnrealTargetPlatform.TryParse(HostPlatformName, out HostPlatform))
-						{
-							throw new AutomationException("Unknown host platform '{0}' specified on command line", HostPlatformName);
-						}
-						if(HostPlatforms.Contains(HostPlatform))
-						{
-							NewHostPlatforms.Add(HostPlatform);
-						}
-					}
-					HostPlatforms = NewHostPlatforms;
-				}
-				else
-				{
-					HostPlatforms.Add(CurrentPlatform);
-				}
+		List<UnrealTargetPlatform> NewHostPlatforms = new List<UnrealTargetPlatform>();
+		foreach (string HostPlatformName in HostPlatformFilter.Split(new char[] { '+' }, StringSplitOptions.RemoveEmptyEntries))
+		{
+			UnrealTargetPlatform HostPlatform;
+			if (!UnrealTargetPlatform.TryParse(HostPlatformName, out HostPlatform))
+			{
+				throw new AutomationException("Unknown host platform '{0}' specified on command line", HostPlatformName);
+			}
+			if (SupportedHostPlatforms.Contains(HostPlatform))
+			{
+				NewHostPlatforms.Add(HostPlatform);
 			}
 			else
 			{
-				HostPlatforms.Add(CurrentPlatform);
+				LogWarning($"HostPlatform {HostPlatformName} not supported on current platform {CurrentPlatform}");
 			}
 		}
-		return HostPlatforms;
+		return NewHostPlatforms;
 	}
 }
