Add the ability to filter based on extensions
#rnx

diff --git a/Engine/Source/Programs/AutomationTool/BuildGraph/Tasks/FindModifiedFiles.cs b/Engine/Source/Programs/AutomationTool/BuildGraph/Tasks/FindModifiedFiles.cs
--- a/Engine/Source/Programs/AutomationTool/BuildGraph/Tasks/FindModifiedFiles.cs
+++ b/Engine/Source/Programs/AutomationTool/BuildGraph/Tasks/FindModifiedFiles.cs
@@ -21,10 +21,10 @@ namespace AutomationTool.Tasks
 	public class FindModifiedFilesTaskParameters
 	{
 		/// <summary>
-		/// The C# project file to compile. Using semicolons, more than one project file can be specified.
+		/// Filter Path (default is ...)
 		/// </summary>
-		[TaskParameter]
-		public string Paths;
+		[TaskParameter(Optional = true)]
+		public string Path = "...";
 
 		/// <summary>
 		/// The configuration to compile.
@@ -81,7 +81,7 @@ namespace AutomationTool.Tasks
 		{
 			using IPerforceConnection Connection = await PerforceConnection.CreateAsync(CommandUtils.P4Settings, Log.Logger);
 
-			StringBuilder Filter = new StringBuilder($"//{Connection.Settings.ClientName}/...");
+			StringBuilder Filter = new StringBuilder($"//{Connection.Settings.ClientName}/{Parameters.Path}");
 			if (Parameters.Change > 0)
 			{
 				Filter.Append($"@={Parameters.Change}");
