Added logging to the staging step that patches resources into the bootstrap executable.

#rnx

diff --git a/Engine/Source/Programs/AutomationTool/Win/WinPlatform.Automation.cs b/Engine/Source/Programs/AutomationTool/Win/WinPlatform.Automation.cs
--- a/Engine/Source/Programs/AutomationTool/Win/WinPlatform.Automation.cs
+++ b/Engine/Source/Programs/AutomationTool/Win/WinPlatform.Automation.cs
@@ -193,6 +193,8 @@ public class Win64Platform : Platform
 			// currently the icon updating doesn't run under mono
 			if (UnrealBuildTool.BuildHostPlatform.Current.Platform == UnrealTargetPlatform.Win64)
 			{
+				LogInformation("Patching bootstrap executable; {0}", IntermediateFile.FullName);
+
 				// Get the icon from the build directory if possible
 				GroupIconResource GroupIcon = null;
 				if(FileReference.Exists(FileReference.Combine(SC.ProjectRoot, "Build/Windows/Application.ico")))
@@ -217,6 +219,10 @@ public class Win64Platform : Platform
 					Update.SetData(ExecArgsResourceId, ResourceType.RawData, Encoding.Unicode.GetBytes(StagedArguments + "\0"));
 				}
 			}
+			else
+			{
+				LogInformation("Skipping patching of bootstrap executable (unsupported host platform)");
+			}
 
 			// Copy it to the staging directory
 			SC.StageFile(StagedFileType.SystemNonUFS, IntermediateFile, new StagedFileReference(ExeName));
diff --git a/Engine/Source/Programs/AutomationTool/Win/WinResources.cs b/Engine/Source/Programs/AutomationTool/Win/WinResources.cs
--- a/Engine/Source/Programs/AutomationTool/Win/WinResources.cs
+++ b/Engine/Source/Programs/AutomationTool/Win/WinResources.cs
@@ -1,4 +1,5 @@
 // Copyright Epic Games, Inc. All Rights Reserved.
+using EpicGames.Core;
 using System;
 using System.Collections.Generic;
 using System.Linq;
@@ -245,6 +246,7 @@ public class ModuleResourceUpdate : IDisposable
 	public ModuleResourceUpdate(string OutputFile, bool bRemoveExisting)
 	{
 		UpdateHandle = BeginUpdateResource(OutputFile, bRemoveExisting);
+		EpicGames.Core.Log.WriteLine(EpicGames.Core.LogEventType.Console, "Begin update resource '{0}' ({1})", OutputFile, Marshal.GetLastWin32Error());
 	}
 
 	public void SetData(int ResourceId, ResourceType Type, byte[] Data)
@@ -254,7 +256,11 @@ public class ModuleResourceUpdate : IDisposable
 
 		Marshal.Copy(Data, 0, UnmanagedPointer, Data.Length);
 
-		if(!UpdateResource(UpdateHandle, new IntPtr((int)Type), new IntPtr(ResourceId), DefaultLanguage, UnmanagedPointer, (uint)Data.Length))
+		bool UpdateReturn = UpdateResource(UpdateHandle, new IntPtr((int)Type), new IntPtr(ResourceId), DefaultLanguage, UnmanagedPointer, (uint)Data.Length);
+
+		EpicGames.Core.Log.WriteLine(EpicGames.Core.LogEventType.Console, "Add resource {0}, length {1} ({2})", ResourceId, Data.Length, UpdateReturn);
+
+		if (!UpdateReturn)
 		{
 			throw new Exception("Couldn't update resource");
 		}
@@ -274,7 +280,9 @@ public class ModuleResourceUpdate : IDisposable
 
 	public void Dispose()
 	{
-		EndUpdateResource(UpdateHandle, false);
+		bool EndReturn = EndUpdateResource(UpdateHandle, false);
+		EpicGames.Core.Log.WriteLine(EpicGames.Core.LogEventType.Console, "End update resource ({0}/{1})", EndReturn, Marshal.GetLastWin32Error());
+
 		foreach(IntPtr UnmanagedPointer in UnmanagedPointers)
 		{
 			Marshal.FreeHGlobal(UnmanagedPointer);
