Ondemand PakFile rules

This change adds support for defining ondemand pak file rule(s). These rules
overrides existing rules created from asset chunking and is currenlty only meant to be used for .ubulk/.uptnl files.
Ondemand content is serialized to loose files with a new I/O store writer to /MyGame/ContentOnDemand/Paks/{MyPakRuleName}.

#rnx

diff --git a/Engine/Source/Developer/IoStoreUtilities/Private/IoStoreOnDemand.cpp b/Engine/Source/Developer/IoStoreUtilities/Private/IoStoreOnDemand.cpp
new file mode 100644
--- /dev/null
+++ b/Engine/Source/Developer/IoStoreUtilities/Private/IoStoreOnDemand.cpp
@@ -0,0 +1,251 @@
+// Copyright Epic Games, Inc. All Rights Reserved.
+
+#include "IoStoreOnDemand.h"
+
+#include "HAL/FileManager.h"
+#include "IO/IoStore.h"
+#include "Misc/FileHelper.h"
+#include "Misc/Paths.h"
+#include "Serialization/JsonReader.h"
+#include "Serialization/JsonSerializer.h"
+#include "Serialization/JsonWriter.h"
+#include "Tasks/Task.h"
+#include "Tasks/Pipe.h"
+
+namespace
+{
+
+using FJsonWriter = TSharedPtr<TJsonWriter<TCHAR, TPrettyJsonPrintPolicy<TCHAR>>>; 
+using FJsonWriterFactory = TJsonWriterFactory<TCHAR, TPrettyJsonPrintPolicy<TCHAR>>;
+
+class FOnDemandIoStoreWriter final
+	: public IIoStoreWriter
+{
+	struct FWriteStats
+	{
+		std::atomic<uint32> PendingCount{0};
+		std::atomic<uint32> TotalCount{0};
+		std::atomic<uint64> TotalByteCount{0};
+	};
+
+	struct FPendingWrite
+	{
+		FPendingWrite(const FIoChunkId& Id, IIoStoreWriteRequest* Request, const FIoWriteOptions& Options)
+			: ChunkId(Id)
+			, WriteRequest(Request)
+			, WriteOptions(Options)
+			, TaskPipe(UE_SOURCE_LOCATION)
+		{ }
+		
+		FIoChunkId ChunkId;
+		TUniquePtr<IIoStoreWriteRequest> WriteRequest;
+		FIoWriteOptions WriteOptions;
+		UE::Tasks::FPipe TaskPipe;
+		FString ErrorText;
+		uint64 UncompressedSize{0};
+		uint64 CompressedSize{0};
+	};
+
+	static constexpr int32 MaxConcurrentWrites{64};
+
+public:
+	FOnDemandIoStoreWriter(const FOnDemandWriterSettings& WriterSettings)
+		: Settings(WriterSettings)
+	{
+		Settings.TocFilePath = FPaths::SetExtension(Settings.TocFilePath, TEXT("utocmanifest"));
+	}
+
+	virtual ~FOnDemandIoStoreWriter()
+	{ }
+
+	virtual void SetReferenceChunkDatabase(TSharedPtr<IIoStoreWriterReferenceChunkDatabase> ReferenceChunkDatabase) override
+	{ }
+
+	virtual void SetHashDatabase(TSharedPtr<IIoStoreWriterHashDatabase> HashDatabase, bool bVerifyHashDatabase) override
+	{ }
+
+	virtual void EnableDiskLayoutOrdering(const TArray<TUniquePtr<FIoStoreReader>>& PatchSourceReaders = TArray<TUniquePtr<FIoStoreReader>>())
+	{ }
+
+	virtual void Append(const FIoChunkId& ChunkId, FIoBuffer Chunk, const FIoWriteOptions& WriteOptions, uint64 OrderHint = MAX_uint64) override
+	{
+		struct FWriteRequest
+			: IIoStoreWriteRequest
+		{
+			FWriteRequest(FIoBuffer InSourceBuffer, uint64 InOrderHint)
+				: OrderHint(InOrderHint)
+			{
+				SourceBuffer = InSourceBuffer;
+				SourceBuffer.MakeOwned();
+			}
+
+			virtual ~FWriteRequest() = default;
+
+			void PrepareSourceBufferAsync(FGraphEventRef CompletionEvent) override
+			{
+				CompletionEvent->DispatchSubsequents();
+			}
+
+			const FIoBuffer* GetSourceBuffer() override
+			{
+				return &SourceBuffer;
+			}
+
+			void FreeSourceBuffer() override
+			{
+			}
+
+			uint64 GetOrderHint() override
+			{
+				return OrderHint;
+			}
+
+			TArrayView<const FFileRegion> GetRegions()
+			{
+				return TArrayView<const FFileRegion>();
+			}
+
+			FIoBuffer SourceBuffer;
+			uint64 OrderHint;
+		};
+
+		Append(ChunkId, new FWriteRequest(Chunk, OrderHint), WriteOptions);
+	}
+
+	virtual void Append(const FIoChunkId& ChunkId, IIoStoreWriteRequest* Request, const FIoWriteOptions& WriteOptions) override
+	{
+		check(Request);
+
+		for(;;)
+		{
+			{
+				FScopeLock _(&PendingLock);
+				if (WriteStats.PendingCount.load() < MaxConcurrentWrites)
+				{
+					WriteStats.PendingCount++;
+					FPendingWrite* PendingWrite = PendingWrites.Add_GetRef(MakeUnique<FPendingWrite>(ChunkId, Request, WriteOptions)).Get();
+					
+					// Setup the task pipe when holding the lock to make it easy to flush all pending task pipe(s)
+
+					FGraphEventRef Event = FGraphEvent::CreateGraphEvent();
+					PendingWrite->WriteRequest->PrepareSourceBufferAsync(Event);
+					UE::Tasks::FTask ReadChunkTask = PendingWrite->TaskPipe.Launch(TEXT("ReadChunk"), [PendingWrite, Event]() mutable
+					{
+						Event->Wait();
+					});
+
+					// Add compression/encryption tasks if needed
+					
+					UE::Tasks::FTask WriteChunkTask = PendingWrite->TaskPipe.Launch(TEXT("WriteChunk"), [this, PendingWrite]() mutable
+					{
+						if (const FIoBuffer* SourceBuffer = PendingWrite->WriteRequest->GetSourceBuffer())
+						{
+							FString FileName = LexToString(PendingWrite->ChunkId) + TEXT(".iochunk");
+							FString FilePath = Settings.OutputDirectory / FileName; 
+
+							if (TUniquePtr<FArchive> Ar(IFileManager::Get().CreateFileWriter(*FilePath)); Ar.IsValid())
+							{
+								const uint64 ChunkSize = SourceBuffer->GetSize();
+								const uint32 CurrentCount = WriteStats.TotalCount.fetch_add(1) + 1; 
+								UE_CLOG((CurrentCount % 128 == 0), LogIoStore, Display,
+									TEXT("Writing ondemand chunk #%u '%s' -> '%s' (%llu Bytes)"),
+									CurrentCount, *PendingWrite->WriteOptions.FileName, *FileName, ChunkSize);
+
+								Ar->Serialize((void*)SourceBuffer->GetData(), ChunkSize);
+								PendingWrite->UncompressedSize = PendingWrite->CompressedSize = ChunkSize;
+								PendingWrite->WriteRequest->FreeSourceBuffer();
+								WriteStats.TotalByteCount += ChunkSize;
+							}
+							else
+							{
+								PendingWrite->ErrorText = FString::Printf(TEXT("Failed to create file '%s'"), *FileName);
+							}
+						}
+						else
+						{
+							PendingWrite->ErrorText = TEXT("Invalid source buffer");
+						}
+
+						PendingWrite->WriteRequest.Reset();
+
+						WriteStats.PendingCount--;
+						WriteCompletedEvent->Trigger();
+					}, ReadChunkTask);
+					
+					break;
+				}
+			}
+
+			WriteCompletedEvent->Wait();
+		}
+	}
+
+	virtual TIoStatusOr<FIoStoreWriterResult> GetResult() override
+	{
+		FScopeLock _(&PendingLock);
+
+		for (TUniquePtr<FPendingWrite>& PendingWrite : PendingWrites)
+		{
+			PendingWrite->TaskPipe.WaitUntilEmpty();
+		}
+		
+		FIoStoreWriterResult WriteResult;
+		WriteResult.ContainerId = FIoContainerId::FromName(Settings.ContainerName);
+		WriteResult.ContainerName = Settings.ContainerName.ToString();
+		
+		FString Json;
+		FJsonWriter JsonWriter = FJsonWriterFactory::Create(&Json);
+		JsonWriter->WriteObjectStart();
+		JsonWriter->WriteValue(TEXT("ContainerName"), Settings.ContainerName.ToString());
+		JsonWriter->WriteArrayStart(TEXT("Chunks"));
+		
+		for (TUniquePtr<FPendingWrite>& PendingWrite : PendingWrites)
+		{
+			if (PendingWrite->ErrorText.IsEmpty() == false)
+			{
+				UE_LOG(LogIoStore, Display, TEXT("Failed to write on demand chunk '%s', reason '%s'"), *LexToString(PendingWrite->ChunkId), *PendingWrite->ErrorText);
+				continue;
+			}
+			JsonWriter->WriteObjectStart();
+			JsonWriter->WriteValue(TEXT("IoChunkId"), LexToString(PendingWrite->ChunkId));
+			JsonWriter->WriteValue(TEXT("FileName"), *PendingWrite->WriteOptions.FileName);
+			JsonWriter->WriteObjectEnd();
+			
+			WriteResult.UncompressedContainerSize += PendingWrite->UncompressedSize;
+			WriteResult.CompressedContainerSize += PendingWrite->CompressedSize;
+			WriteResult.TocEntryCount++;
+		}
+
+		JsonWriter->WriteArrayEnd();
+		JsonWriter->WriteObjectEnd();
+		JsonWriter->Close();
+
+		UE_LOG(LogIoStore, Display, TEXT("Saving ondemand JSON manifest '%s'"), *Settings.TocFilePath);
+		if (!FFileHelper::SaveStringToFile(Json, *Settings.TocFilePath))
+		{
+			return FIoStatus(EIoErrorCode::CorruptToc, TEXTVIEW("Failed writing ondemand JSON manifest")); 
+		}
+		
+		WriteResult.TocSize = sizeof(TCHAR) + Json.Len();
+
+		return WriteResult;
+	}
+
+	virtual void EnumerateChunks(TFunction<bool(const FIoStoreTocChunkInfo&)>&& Callback) const override
+	{
+	}
+
+private:
+	FOnDemandWriterSettings Settings;
+	FWriteStats WriteStats;
+	FCriticalSection PendingLock;
+	TArray<TUniquePtr<FPendingWrite>> PendingWrites;
+	FEventRef WriteCompletedEvent;
+};
+
+} // namespace
+
+TUniquePtr<IIoStoreWriter> MakeOnDemandIoStoreWriter(const FOnDemandWriterSettings& WriterSettings)
+{
+	return MakeUnique<FOnDemandIoStoreWriter>(WriterSettings);
+}
diff --git a/Engine/Source/Developer/IoStoreUtilities/Private/IoStoreOnDemand.h b/Engine/Source/Developer/IoStoreUtilities/Private/IoStoreOnDemand.h
new file mode 100644
--- /dev/null
+++ b/Engine/Source/Developer/IoStoreUtilities/Private/IoStoreOnDemand.h
@@ -0,0 +1,18 @@
+// Copyright Epic Games, Inc. All Rights Reserved.
+
+#pragma once
+
+#include "Containers/UnrealString.h"
+#include "Templates/UniquePtr.h"
+#include "UObject/NameTypes.h"
+
+class IIoStoreWriter;
+
+struct FOnDemandWriterSettings
+{
+	FString TocFilePath;
+	FString OutputDirectory;
+	FName ContainerName;
+};
+
+TUniquePtr<IIoStoreWriter> MakeOnDemandIoStoreWriter(const FOnDemandWriterSettings& WriterSettings);
diff --git a/Engine/Source/Developer/IoStoreUtilities/Private/IoStoreUtilities.cpp b/Engine/Source/Developer/IoStoreUtilities/Private/IoStoreUtilities.cpp
--- a/Engine/Source/Developer/IoStoreUtilities/Private/IoStoreUtilities.cpp
+++ b/Engine/Source/Developer/IoStoreUtilities/Private/IoStoreUtilities.cpp
@@ -2,6 +2,7 @@
 
 #include "IoStoreUtilities.h"
 
+#include "IoStoreOnDemand.h"
 #include "Async/AsyncWork.h"
 #include "HAL/FileManager.h"
 #include "HAL/PlatformFileManager.h"
@@ -417,6 +418,7 @@ struct FContainerSourceSpec
 	FName Name;
 	FString OutputPath;
 	FString OptionalOutputPath;
+	FString OnDemandOutputPath;
 	TArray<FContainerSourceFile> SourceFiles;
 	FString PatchTargetFile;
 	TArray<FString> PatchSourceContainerFiles;
@@ -887,6 +889,7 @@ struct FContainerTargetSpec
 	FGuid EncryptionKeyGuid;
 	FString OutputPath;
 	FString OptionalSegmentOutputPath;
+	FString OnDemandOutputPath;
 	TSharedPtr<IIoStoreWriter> IoStoreWriter;
 	TSharedPtr<IIoStoreWriter> OptionalSegmentIoStoreWriter;
 	TArray<FContainerTargetFile> TargetFiles;
@@ -2358,6 +2361,7 @@ void InitializeContainerTargetsAndPackages(
 	{
 		FContainerTargetSpec* ContainerTarget = AddContainer(ContainerSource.Name, ContainerTargets);
 		ContainerTarget->OutputPath = ContainerSource.OutputPath;
+		ContainerTarget->OnDemandOutputPath = ContainerSource.OnDemandOutputPath;
 		ContainerTarget->bGenerateDiffPatch = ContainerSource.bGenerateDiffPatch;
 		if (Arguments.bSign)
 		{
@@ -3419,34 +3423,46 @@ int32 CreateTarget(const FIoStoreArguments& Arguments, const FIoStoreWriterSetti
 			check(ContainerTarget->ContainerId.IsValid());
 			if (!ContainerTarget->OutputPath.IsEmpty())
 			{
-				FIoContainerSettings ContainerSettings;
-				ContainerSettings.ContainerId = ContainerTarget->ContainerId;
-				if (Arguments.bCreateDirectoryIndex)
+				if (!ContainerTarget->OnDemandOutputPath.IsEmpty())
 				{
-					ContainerSettings.ContainerFlags = ContainerTarget->ContainerFlags | EIoContainerFlags::Indexed;
+					FOnDemandWriterSettings WriterSettings;
+					WriterSettings.ContainerName = ContainerTarget->Name;
+					WriterSettings.TocFilePath = ContainerTarget->OutputPath;
+					WriterSettings.OutputDirectory = ContainerTarget->OnDemandOutputPath;
+					ContainerTarget->IoStoreWriter = TSharedPtr<IIoStoreWriter>(MakeOnDemandIoStoreWriter(WriterSettings).Release());
+					IoStoreWriters.Add(ContainerTarget->IoStoreWriter);
 				}
-				if (EnumHasAnyFlags(ContainerTarget->ContainerFlags, EIoContainerFlags::Encrypted))
-				{
-					const FNamedAESKey* Key = Arguments.KeyChain.GetEncryptionKeys().Find(ContainerTarget->EncryptionKeyGuid);
-					check(Key);
-					ContainerSettings.EncryptionKeyGuid = ContainerTarget->EncryptionKeyGuid;
-					ContainerSettings.EncryptionKey = Key->Key;
-				}
-				if (EnumHasAnyFlags(ContainerTarget->ContainerFlags, EIoContainerFlags::Signed))
-				{
-					ContainerSettings.SigningKey = Arguments.KeyChain.GetSigningKey();
-					ContainerSettings.ContainerFlags |= EIoContainerFlags::Signed;
-				}
-				ContainerSettings.bGenerateDiffPatch = ContainerTarget->bGenerateDiffPatch;
-				ContainerTarget->IoStoreWriter = IoStoreWriterContext->CreateContainer(*ContainerTarget->OutputPath, ContainerSettings);
-				ContainerTarget->IoStoreWriter->EnableDiskLayoutOrdering(ContainerTarget->PatchSourceReaders);
-				ContainerTarget->IoStoreWriter->SetReferenceChunkDatabase(ChunkDatabase);
-				ContainerTarget->IoStoreWriter->SetHashDatabase(HashDatabase, Arguments.bVerifyHashDatabase);
-				IoStoreWriters.Add(ContainerTarget->IoStoreWriter);
-				if (!ContainerTarget->OptionalSegmentOutputPath.IsEmpty())
+				else
 				{
-					ContainerTarget->OptionalSegmentIoStoreWriter = IoStoreWriterContext->CreateContainer(*ContainerTarget->OptionalSegmentOutputPath, ContainerSettings);
-					IoStoreWriters.Add(ContainerTarget->OptionalSegmentIoStoreWriter);
+					FIoContainerSettings ContainerSettings;
+					ContainerSettings.ContainerId = ContainerTarget->ContainerId;
+					if (Arguments.bCreateDirectoryIndex)
+					{
+						ContainerSettings.ContainerFlags = ContainerTarget->ContainerFlags | EIoContainerFlags::Indexed;
+					}
+					if (EnumHasAnyFlags(ContainerTarget->ContainerFlags, EIoContainerFlags::Encrypted))
+					{
+						const FNamedAESKey* Key = Arguments.KeyChain.GetEncryptionKeys().Find(ContainerTarget->EncryptionKeyGuid);
+						check(Key);
+						ContainerSettings.EncryptionKeyGuid = ContainerTarget->EncryptionKeyGuid;
+						ContainerSettings.EncryptionKey = Key->Key;
+					}
+					if (EnumHasAnyFlags(ContainerTarget->ContainerFlags, EIoContainerFlags::Signed))
+					{
+						ContainerSettings.SigningKey = Arguments.KeyChain.GetSigningKey();
+						ContainerSettings.ContainerFlags |= EIoContainerFlags::Signed;
+					}
+					ContainerSettings.bGenerateDiffPatch = ContainerTarget->bGenerateDiffPatch;
+					ContainerTarget->IoStoreWriter = IoStoreWriterContext->CreateContainer(*ContainerTarget->OutputPath, ContainerSettings);
+					ContainerTarget->IoStoreWriter->EnableDiskLayoutOrdering(ContainerTarget->PatchSourceReaders);
+					ContainerTarget->IoStoreWriter->SetReferenceChunkDatabase(ChunkDatabase);
+					ContainerTarget->IoStoreWriter->SetHashDatabase(HashDatabase, Arguments.bVerifyHashDatabase);
+					IoStoreWriters.Add(ContainerTarget->IoStoreWriter);
+					if (!ContainerTarget->OptionalSegmentOutputPath.IsEmpty())
+					{
+						ContainerTarget->OptionalSegmentIoStoreWriter = IoStoreWriterContext->CreateContainer(*ContainerTarget->OptionalSegmentOutputPath, ContainerSettings);
+						IoStoreWriters.Add(ContainerTarget->OptionalSegmentIoStoreWriter);
+					}
 				}
 			}
 		}
@@ -6586,6 +6602,7 @@ bool ParseContainerGenerationArguments(FIoStoreArguments& Arguments, FIoStoreWri
 			ContainerSpec.OutputPath = FPaths::ChangeExtension(ContainerSpec.OutputPath, TEXT(""));
 
 			FParse::Value(*Command, TEXT("OptionalOutput="), ContainerSpec.OptionalOutputPath);
+			FParse::Value(*Command, TEXT("OnDemandOutput="), ContainerSpec.OnDemandOutputPath);
 
 			FString ContainerName;
 			if (FParse::Value(*Command, TEXT("ContainerName="), ContainerName))
diff --git a/Engine/Source/Programs/AutomationTool/Scripts/CopyBuildToStagingDirectory.Automation.cs b/Engine/Source/Programs/AutomationTool/Scripts/CopyBuildToStagingDirectory.Automation.cs
--- a/Engine/Source/Programs/AutomationTool/Scripts/CopyBuildToStagingDirectory.Automation.cs
+++ b/Engine/Source/Programs/AutomationTool/Scripts/CopyBuildToStagingDirectory.Automation.cs
@@ -276,6 +276,7 @@ namespace AutomationScripts
 			string ContainerName,
 			FileReference PakOutputLocation,
 			DirectoryReference OptionalOutputLocation,
+			DirectoryReference OnDemandOutputLocation,
 			bool bCompressed,
 			EncryptionAndSigning.CryptoSettings CryptoSettings,
 			string EncryptionKeyGuid,
@@ -284,7 +285,17 @@ namespace AutomationScripts
 			bool bIsDLC)
 		{
 			StringBuilder CmdLine = new StringBuilder();
-			CmdLine.AppendFormat("-Output={0}", MakePathSafeToUseWithCommandLine(Path.ChangeExtension(PakOutputLocation.FullName, ".utoc")));
+
+			if (OnDemandOutputLocation != null)
+			{
+				CmdLine.AppendFormat("-Output={0}", MakePathSafeToUseWithCommandLine(Path.ChangeExtension(PakOutputLocation.FullName, ".tocmanifest")));
+				CmdLine.AppendFormat(" -OnDemandOutput={0}", MakePathSafeToUseWithCommandLine(OnDemandOutputLocation.FullName));
+			}
+			else
+			{
+				CmdLine.AppendFormat("-Output={0}", MakePathSafeToUseWithCommandLine(Path.ChangeExtension(PakOutputLocation.FullName, ".utoc")));
+			}
+
 			if (OptionalOutputLocation != null)
 			{
 				CmdLine.AppendFormat(" -OptionalOutput={0}", MakePathSafeToUseWithCommandLine(OptionalOutputLocation.FullName));
@@ -2041,6 +2052,18 @@ namespace AutomationScripts
 
 			// List of pak files to use instead of manifest
 			public List<string> OverridePaks;
+
+			// Whether this rule is defined for content-on-demand
+			public bool bOnDemand;
+
+			public static bool IsMatch(PakFileRules PakRules, KeyValuePair<string, string> StagingFile)
+			{
+				bool bMatched = !PakRules.bDisabled &&
+					((!PakRules.bUseExactFilePathMatch && PakRules.Filter != null && PakRules.Filter.Matches(StagingFile.Key)) ||
+					(PakRules.bUseExactFilePathMatch && PakRules.ExactFileMatches != null && PakRules.ExactFileMatches.Contains(StagingFile.Value)));
+				
+				return bMatched;
+			}
 		};
 
 		/// <summary>
@@ -2173,6 +2196,7 @@ namespace AutomationScripts
 				PakRulesConfig.TryGetValue(SectionName, "bExcludeFromPaks", out PakRules.bExcludeFromPaks);
 				PakRulesConfig.TryGetValue(SectionName, "bOverrideChunkManifest", out PakRules.bOverrideChunkManifest);
 				PakRulesConfig.TryGetValue(SectionName, "bDisabled", out PakRules.bDisabled);
+				PakRulesConfig.TryGetValue(SectionName, "bOnDemand", out PakRules.bOnDemand);
 				string PakString;
 				PakRulesConfig.TryGetValue(SectionName, "OverridePaks", out PakString);
 
@@ -2185,12 +2209,20 @@ namespace AutomationScripts
 					PakRules.OverridePaks = null;
 				}
 
-				if (PakRules.bExcludeFromPaks && PakRules.OverridePaks != null)
+				if (PakRules.bExcludeFromPaks)
 				{
-					LogWarning("Error in PakFileRules {0}, set to exclude but also sets override!", PakRules.Name);
-					continue;
+					if (PakRules.OverridePaks != null)
+					{
+						LogWarning("Error in PakFileRules {0}, set to exclude but also sets override!", PakRules.Name);
+						continue;
+					}
+					if (PakRules.bOnDemand)
+					{
+						LogWarning("Error in PakFileRules {0}, set to exclude but also sets to ondemand!", PakRules.Name);
+						continue;
+					}
 				}
-				else if (!PakRules.bExcludeFromPaks && PakRules.OverridePaks == null)
+				else if (PakRules.OverridePaks == null && !PakRules.bOnDemand)
 				{
 					LogWarning("Error in PakFileRules {0}, set to include but did not specify paks!", PakRules.Name);
 					continue;
@@ -2236,14 +2268,48 @@ namespace AutomationScripts
 				return;
 			}
 
+			// Process content-on-demand rules first
+			if (ModifyPakList != null)
+			{
+				foreach (PakFileRules Rule in RulesList.Where(Rule => Rule.bOnDemand))
+				{
+					if (PakFileRules.IsMatch(Rule, StagingFile))
+					{
+						ModifyPakList.Clear();
+
+						string ChunkName = Rule.Name;
+						if (Rule.OverridePaks != null && Rule.OverridePaks.Count > 0)
+						{
+							string OverrideChunkName = Rule.OverridePaks[0];
+							if (ChunkNameToDefinition.ContainsKey(OverrideChunkName))
+							{
+								ChunkName = OverrideChunkName;
+							}
+							else
+							{
+								LogWarning("Undefined PAK override '{0}' in PAK rule '{1}'", OverrideChunkName, Rule.Name);
+							}
+						}
+
+						ChunkDefinition Chunk = null;
+						if (ChunkNameToDefinition.TryGetValue(ChunkName, out Chunk))
+						{
+							ModifyPakList.Add(Chunk);
+						}
+						else
+						{
+							LogWarning("Undefined chunk name '{0}' in PAK rule '{1}'", ChunkName, Rule.Name);
+						}
+						
+						return;
+					}
+				}
+			}
+
 			// Search in order, return on first match
 			foreach (var PakRules in RulesList)
 			{
-				bool bMatched = !PakRules.bDisabled &&
-					((!PakRules.bUseExactFilePathMatch && PakRules.Filter != null && PakRules.Filter.Matches(StagingFile.Key)) ||
-					(PakRules.bUseExactFilePathMatch && PakRules.ExactFileMatches != null && PakRules.ExactFileMatches.Contains(StagingFile.Value)));
-
-				if (bMatched)
+				if (PakFileRules.IsMatch(PakRules, StagingFile))
 				{
 					bool bOverrideChunkAssignment = false;
 					if (ModifyPakList != null && ModifyPakList.Count > 0)
@@ -2297,7 +2363,6 @@ namespace AutomationScripts
 			}
 		}
 
-
 		/// <summary>
 		/// Creates a pak file using staging context (single manifest)
 		/// </summary>
@@ -2476,6 +2541,11 @@ namespace AutomationScripts
 			/// Whether to rehydrate the assets when creating the pak file or not
 			/// </summary>
 			public bool bRehydrateAssets;
+			
+			/// <summary>
+			// Whether this pak file is for content-on-demand content delivery or not
+			/// </summary>
+			public bool bOnDemand;
 
 			/// <summary>
 			/// GUID of the encryption key for this pak file
@@ -2488,13 +2558,14 @@ namespace AutomationScripts
 			/// <param name="PakName">Path to the base output file for this pak file</param>
 			/// <param name="UnrealPakResponseFile">Map of files within the pak file to their source file on disk</param>
 			/// <param name="bCompressed">Whether to enable compression</param>
-			public CreatePakParams(string PakName, Dictionary<string, string> UnrealPakResponseFile, bool bCompressed, bool RehydrateAssets, string EncryptionKeyGuid)
+			public CreatePakParams(string PakName, Dictionary<string, string> UnrealPakResponseFile, bool bCompressed, bool RehydrateAssets, string EncryptionKeyGuid, bool bOnDemand = false)
 			{
 				this.PakName = PakName;
 				this.UnrealPakResponseFile = UnrealPakResponseFile;
 				this.bCompressed = bCompressed;
 				this.bRehydrateAssets = RehydrateAssets;
 				this.EncryptionKeyGuid = EncryptionKeyGuid;
+				this.bOnDemand = bOnDemand;
 			}
 		}
 
@@ -3002,11 +3073,13 @@ namespace AutomationScripts
 				OutputFilename = OutputFilename + PostFix;
 
 				StagedFileReference OutputRelativeLocation;
+				StagedDirectoryReference OnDemandOutputRelativeLocation;
 				if (Params.HasDLCName)
 				{
 					if (Params.DLCOverrideStagedSubDir != null)
 					{
 						OutputRelativeLocation = StagedFileReference.Combine(SC.RelativeProjectRootForStage, Params.DLCOverrideStagedSubDir, "Content", "Paks", Params.DLCFile.GetFileNameWithoutExtension() + OutputFilename + ".pak");
+						OnDemandOutputRelativeLocation = StagedDirectoryReference.Combine(SC.RelativeProjectRootForStage, Params.DLCOverrideStagedSubDir, "ContentOnDemand", "Paks", Params.DLCFile.GetFileNameWithoutExtension() + OutputFilename);
 					}
 					else
 					{
@@ -3021,23 +3094,29 @@ namespace AutomationScripts
 							PluginSubdirectory = Params.DLCFile.GetFileNameWithoutAnyExtensions();
 						}
 						OutputRelativeLocation = StagedFileReference.Combine(SC.RelativeProjectRootForStage, PluginSubdirectory, "Content", "Paks", SC.FinalCookPlatform, Params.DLCFile.GetFileNameWithoutExtension() + OutputFilename + ".pak");
+						OnDemandOutputRelativeLocation = StagedDirectoryReference.Combine(SC.RelativeProjectRootForStage, PluginSubdirectory, "ContentOnDemand", "Paks", SC.FinalCookPlatform, Params.DLCFile.GetFileNameWithoutExtension() + OutputFilename);
 					}
 				}
 				else
 				{
 					OutputRelativeLocation = StagedFileReference.Combine(SC.RelativeProjectRootForStage, "Content", "Paks", OutputFilename + OutputFilenameExtension);
+					OnDemandOutputRelativeLocation = StagedDirectoryReference.Combine(SC.RelativeProjectRootForStage, "ContentOnDemand", "Paks", OutputFilename);
 				}
 				if (SC.StageTargetPlatform.DeployLowerCaseFilenames(StagedFileType.UFS))
 				{
 					OutputRelativeLocation = OutputRelativeLocation.ToLowerInvariant();
 				}
 				OutputRelativeLocation = SC.StageTargetPlatform.Remap(OutputRelativeLocation);
+				
+				//TODO: Do we need to remap the ondemand output location?
+				//OnDemandOutputRelativeLocation = SC.StageTargetPlatform.Remap(OnDemandOutputRelativeLocation);
 
 				FileReference OutputLocation = FileReference.Combine(SC.RuntimeRootDir, OutputRelativeLocation.Name);
+				DirectoryReference OnDemandOutputLocation = DirectoryReference.Combine(SC.RuntimeRootDir, OnDemandOutputRelativeLocation.Name);
 
 				bool bCopiedExistingPak = false;
 
-				if (SC.StageTargetPlatform != SC.CookSourcePlatform && !Params.IgnorePaksFromDifferentCookSource)
+				if (!PakParams.bOnDemand && SC.StageTargetPlatform != SC.CookSourcePlatform && !Params.IgnorePaksFromDifferentCookSource)
 				{
 					// Check to see if we have an existing pak file we can use
 
@@ -3174,6 +3253,7 @@ namespace AutomationScripts
 								OutputLocation.GetFileNameWithoutAnyExtensions(),
 								OutputLocation,
 								SC.OptionalFileStageDirectory,
+								PakParams.bOnDemand ? OnDemandOutputLocation : null,
 								bCompressContainers,
 								CryptoSettings,
 								PakParams.EncryptionKeyGuid,
@@ -3694,6 +3774,7 @@ namespace AutomationScripts
 				ResponseFile = new ConcurrentDictionary<string, string>();
 				Manifest = null;
 				bCompressed = false;
+				bOnDemand = false;
 			}
 
 			// Name of pak file without extension, ie pakchunk0
@@ -3709,6 +3790,9 @@ namespace AutomationScripts
 			public bool bCompressed;
 			public string EncryptionKeyGuid;
 			public string RequestedEncryptionKeyGuid;
+
+			// Whether the chunk is used for content-on-demand content delivery or not
+			public bool bOnDemand;
 		}
 
 		/// <summary>
@@ -3808,6 +3892,31 @@ namespace AutomationScripts
 					}
 				}
 
+				// Add ondemand chunk definitions
+				if (PakRulesList != null)
+				{
+					foreach (PakFileRules Rule in PakRulesList.Where(R => R.bOnDemand))
+					{
+						string ChunkName = Rule.Name;
+						if (Rule.OverridePaks != null && Rule.OverridePaks.Count > 0)
+						{
+							ChunkName = Rule.OverridePaks[0];
+						}
+
+						if (ChunkDefinitions.Any(D => D.ChunkName == ChunkName))
+						{
+							continue;
+						}
+
+						Log.TraceInformation("Creating content-on-demand chunk definition '{0}' for PAK rule '{1}'", ChunkName, Rule.Name);
+
+						//TODO: Add encryption/compression settings
+						ChunkDefinition Chunk = new ChunkDefinition(ChunkName);
+						Chunk.bOnDemand = true;
+						ChunkDefinitions.Add(Chunk);
+					}
+				}
+
 				ConcurrentDictionary<string, ChunkDefinition> ChunkNameToDefinition = new ConcurrentDictionary<string, ChunkDefinition>(StringComparer.InvariantCultureIgnoreCase);
 				foreach (ChunkDefinition Chunk in ChunkDefinitions)
 				{
@@ -3863,19 +3972,23 @@ namespace AutomationScripts
 					{
 						ChunkDefinition TargetChunk = Chunk;
 
-						string OrigExt = Path.GetExtension(OriginalFilename);
-						if (OrigExt.Equals(OptionalBulkDataFileExtension))
+						// Don't separate out any bulk data for ondemand chunks
+						if (Chunk.bOnDemand == false)
 						{
-							// any optional files encountered we want to put in a separate pak file
-							string OptionalChunkName = Chunk.ChunkName + "optional";
-							if (!OptionalChunks.ContainsKey(OptionalChunkName))
+							string OrigExt = Path.GetExtension(OriginalFilename);
+							if (OrigExt.Equals(OptionalBulkDataFileExtension))
 							{
-								ChunkDefinition OptionalChunk = new ChunkDefinition(OptionalChunkName);
-								OptionalChunk.RequestedEncryptionKeyGuid = Chunk.RequestedEncryptionKeyGuid;
-								OptionalChunk.EncryptionKeyGuid = Chunk.EncryptionKeyGuid;
-								OptionalChunks.TryAdd(OptionalChunkName, OptionalChunk);
+								// any optional files encountered we want to put in a separate pak file
+								string OptionalChunkName = Chunk.ChunkName + "optional";
+								if (!OptionalChunks.ContainsKey(OptionalChunkName))
+								{
+									ChunkDefinition OptionalChunk = new ChunkDefinition(OptionalChunkName);
+									OptionalChunk.RequestedEncryptionKeyGuid = Chunk.RequestedEncryptionKeyGuid;
+									OptionalChunk.EncryptionKeyGuid = Chunk.EncryptionKeyGuid;
+									OptionalChunks.TryAdd(OptionalChunkName, OptionalChunk);
+								}
+								TargetChunk = OptionalChunks[OptionalChunkName];
 							}
-							TargetChunk = OptionalChunks[OptionalChunkName];
 						}
 
 						TargetChunk.ResponseFile.TryAdd(StagingFile.Key, StagingFile.Value);
@@ -3976,7 +4089,7 @@ namespace AutomationScripts
 					string EncryptionKeyToUse = Params.SkipEncryption ? "" : Chunk.EncryptionKeyGuid;
 					PakInputs.Add(new CreatePakParams(Chunk.ChunkName,
 						Chunk.ResponseFile.ToDictionary(entry => entry.Key, entry => entry.Value),
-						Params.Compressed || Chunk.bCompressed, Params.RehydrateAssets, EncryptionKeyToUse));
+						Params.Compressed || Chunk.bCompressed, Params.RehydrateAssets, EncryptionKeyToUse, Chunk.bOnDemand));
 				}
 			}
 
