UAT: Replace parallel build functionality in UAT with a call to UBT with a list of targets.

diff --git a/Engine/Source/Programs/AutomationTool/AutomationUtils/UnrealBuild.cs b/Engine/Source/Programs/AutomationTool/AutomationUtils/UnrealBuild.cs
--- a/Engine/Source/Programs/AutomationTool/AutomationUtils/UnrealBuild.cs
+++ b/Engine/Source/Programs/AutomationTool/AutomationUtils/UnrealBuild.cs
@@ -13,6 +13,8 @@ using OpenTracing;
 using OpenTracing.Util;
 using UnrealBuildBase;
 
+using CommandLine = UnrealBuildBase.CommandLine;
+
 namespace AutomationTool
 {
 	class UnrealBuildException : AutomationException
@@ -95,19 +97,7 @@ namespace AutomationTool
 			}
 		}
 
-		public class XGEItem
-		{
-			public FileReference ManifestFile;
-			public BuildManifest Manifest;
-			public UnrealTargetPlatform Platform;
-			public UnrealTargetConfiguration Config;
-			public string TargetName;
-			public FileReference UProjectPath;
-			public List<string> XgeXmlFiles;
-			public string OutputCaption;
-		}
-
-		XGEItem XGEPrepareBuildWithUBT(string TargetName, UnrealTargetPlatform Platform, UnrealTargetConfiguration Config, FileReference UprojectPath, string InAddArgs)
+		public void CleanWithUBT(string TargetName, UnrealTargetPlatform Platform, UnrealTargetConfiguration Config, FileReference UprojectPath, string InAddArgs = "")
 		{
 			string AddArgs = "";
 			if (UprojectPath != null)
@@ -118,123 +108,108 @@ namespace AutomationTool
 			AddArgs += " " + InAddArgs;
 
 			PrepareUBT();
-
-            FileReference ManifestFile = GetManifestFile(UprojectPath);
-			CommandUtils.DeleteFile(ManifestFile);
-
-			ClearExportedXGEXML();
-
-			string UHTArg = "";
-
-			CommandUtils.RunUBT(CommandUtils.CmdEnv, UnrealBuildToolDll: UnrealBuildToolDll, Project: UprojectPath, Target: TargetName, Platform: Platform, Config: Config, AdditionalArgs: String.Format("-Manifest={0} {1} -NoHotReload -xgeexport {2}", CommandUtils.MakePathSafeToUseWithCommandLine(ManifestFile.FullName), UHTArg, AddArgs));
-
-			XGEItem Result = new XGEItem();
-			Result.Platform = Platform;
-			Result.Config = Config;
-			Result.TargetName = TargetName;
-			Result.UProjectPath = UprojectPath;
-			Result.ManifestFile = ManifestFile;
-			Result.Manifest = CommandUtils.ReadManifest(ManifestFile);
-			Result.OutputCaption = String.Format("{0}-{1}-{2}", TargetName, Platform.ToString(), Config.ToString());
-			CommandUtils.DeleteFile(ManifestFile);
-
-			Result.XgeXmlFiles = new List<string>();
-			foreach (var XGEFile in FindXGEFiles())
+			using (IScope Scope = GlobalTracer.Instance.BuildSpan("Compile").StartActive())
 			{
-				if (!CommandUtils.FileExists_NoExceptions(XGEFile))
-				{
-					throw new UnrealBuildException("Couldn't find file: {0}", XGEFile);
-				}
-				int FileNum = 0;
-				string OutFile;
-				while (true)
-				{
-					OutFile = CommandUtils.CombinePaths(CommandUtils.CmdEnv.LogFolder, String.Format("UBTExport.{0}.xge.xml", FileNum));
-					FileInfo ItemInfo = new FileInfo(OutFile);
-					if (!ItemInfo.Exists)
-					{
-						break;
-					}
-					FileNum++;
-				}
-				CommandUtils.CopyFile(XGEFile, OutFile);
-				Result.XgeXmlFiles.Add(OutFile);
+				Scope.Span.SetTag("target", TargetName);
+				Scope.Span.SetTag("platform", Platform.ToString());
+				Scope.Span.SetTag("config", Config.ToString());
+				CommandUtils.RunUBT(CommandUtils.CmdEnv, UnrealBuildToolDll: UnrealBuildToolDll, Project: UprojectPath, Target: TargetName, Platform: Platform, Config: Config, AdditionalArgs: "-Clean -NoHotReload" + AddArgs);
 			}
-			ClearExportedXGEXML();
-			return Result;
-		}
+        }
 
-		void XGEFinishBuildWithUBT(XGEItem Item)
+		List<FileReference> GetManifestFiles(List<BuildTarget> Targets)
 		{
-			// run the deployment steps, if necessary
-			foreach(string DeployTargetFile in Item.Manifest.DeployTargetFiles)
+			List<FileReference> ManifestFiles = new List<FileReference>();
+			for (int Idx = 0; Idx < Targets.Count; Idx++)
 			{
-				CommandUtils.RunUBT(CommandUtils.CmdEnv, UnrealBuildToolDll, String.Format("-Mode=Deploy -Receipt=\"{0}\"", DeployTargetFile));
+				BuildTarget Target = Targets[Idx];
+				DirectoryReference ManifestDir = GetManifestDir(Target.UprojectPath);
+				ManifestFiles.Add(FileReference.Combine(ManifestDir, $"Manifest-{Idx + 1}-{Target.TargetName}-{Target.Platform}-{Target.Config}.xml"));
 			}
 
-			foreach (string ManifestItem in Item.Manifest.BuildProducts)
-			{
-				if (!CommandUtils.FileExists_NoExceptions(ManifestItem))
-				{
-					throw new UnrealBuildException($"XGEFinishBuildWithUBT: \"{ManifestItem}\" was in manifest \"{Item.ManifestFile}\" but could not be found");
-				}
-				AddBuildProduct(ManifestItem);
-			}
+			return ManifestFiles;
 		}
 
-		public void CleanWithUBT(string TargetName, UnrealTargetPlatform Platform, UnrealTargetConfiguration Config, FileReference UprojectPath, string InAddArgs = "")
+		FileReference GetManifestFile(BuildTarget Target, int Index)
 		{
-			string AddArgs = "";
-			if (UprojectPath != null)
+			DirectoryReference ManifestDir = GetManifestDir(Target.UprojectPath);
+			if (Index == -1)
 			{
-				AddArgs += " " + CommandUtils.MakePathSafeToUseWithCommandLine(UprojectPath.FullName);
+				return FileReference.Combine(ManifestDir, "Manifest.xml");
 			}
-			AddArgs += " -NoUBTMakefiles";
-			AddArgs += " " + InAddArgs;
-
-			PrepareUBT();
-			using (IScope Scope = GlobalTracer.Instance.BuildSpan("Compile").StartActive())
+			else
 			{
-				Scope.Span.SetTag("target", TargetName);
-				Scope.Span.SetTag("platform", Platform.ToString());
-				Scope.Span.SetTag("config", Config.ToString());
-				CommandUtils.RunUBT(CommandUtils.CmdEnv, UnrealBuildToolDll: UnrealBuildToolDll, Project: UprojectPath, Target: TargetName, Platform: Platform, Config: Config, AdditionalArgs: "-Clean -NoHotReload" + AddArgs);
+				return FileReference.Combine(ManifestDir, $"Manifest-{Index + 1}-{Target.TargetName}-{Target.Platform}-{Target.Config}.xml");
 			}
-        }
+		}
 
-		BuildManifest BuildWithUBT(string TargetName, UnrealTargetPlatform TargetPlatform, UnrealTargetConfiguration Config, FileReference UprojectPath, bool ForceFlushMac = false, bool DisableXGE = false, bool AllCores = false, string InAddArgs = "")
+		void BuildWithUBT(List<BuildTarget> Targets, Dictionary<BuildTarget, BuildManifest> TargetToManifest, bool DisableXGE, bool AllCores)
 		{
-			string AddArgs = "";
-			if (UprojectPath != null)
+			List<FileReference> ManifestFiles = new List<FileReference>(Targets.Count);
+
+			StringBuilder FullCommandLine = new StringBuilder();
+			if (Targets.Count == 1)
 			{
-				AddArgs += " " + CommandUtils.MakePathSafeToUseWithCommandLine(UprojectPath.FullName);
+				ManifestFiles.Add(GetManifestFile(Targets[0], -1));
+				FullCommandLine.Append(GetTargetArguments(Targets[0], ManifestFiles[0]));
 			}
-			AddArgs += " -NoUBTMakefiles";
-			AddArgs += " " + InAddArgs;
-			if (ForceFlushMac)
+			else
 			{
-				AddArgs += " -flushmac";
+				List<string> Arguments = new List<string>();
+				for (int Idx = 0; Idx < Targets.Count; Idx++)
+				{
+					ManifestFiles.Add(GetManifestFile(Targets[0], Idx));
+					Arguments.Add("-Target=" + GetTargetArguments(Targets[Idx], ManifestFiles[Idx]));
+				}
+				FullCommandLine.Append(CommandLine.FormatCommandLine(Arguments));
 			}
+
 			if (DisableXGE)
 			{
-				AddArgs += " -noxge";
+				FullCommandLine.Append(" -NoXGE");
 			}
 			if (AllCores)
 			{
-				AddArgs += " -allcores";
+				FullCommandLine.Append(" -AllCores");
 			}
 
 			PrepareUBT();
 
-			FileReference ManifestFile = GetManifestFile(UprojectPath);
-			CommandUtils.DeleteFile(ManifestFile);
+			for (int Idx = 0; Idx < ManifestFiles.Count; Idx++)
+			{
+				CommandUtils.DeleteFile(ManifestFiles[Idx]);
+			}
 
-			CommandUtils.RunUBT(CommandUtils.CmdEnv, UnrealBuildToolDll: UnrealBuildToolDll, Project: UprojectPath, Target: TargetName, Platform: TargetPlatform, Config: Config, AdditionalArgs: String.Format("{0} -Manifest={1} -NoHotReload", AddArgs, CommandUtils.MakePathSafeToUseWithCommandLine(ManifestFile.FullName)));
+			CommandUtils.RunUBT(CommandUtils.CmdEnv, UnrealBuildToolDll, FullCommandLine.ToString());
+
+			for (int Idx = 0; Idx < Targets.Count; Idx++)
+			{
+				FileReference ManifestFile = ManifestFiles[Idx];
+				BuildManifest Manifest = AddBuildProductsFromManifest(ManifestFile);
+				CommandUtils.DeleteFile(ManifestFile);
+				TargetToManifest?.Add(Targets[Idx], Manifest);
+			}
+		}
 
-			BuildManifest Manifest = AddBuildProductsFromManifest(ManifestFile);
-			CommandUtils.DeleteFile(ManifestFile);
+		string GetTargetArguments(BuildTarget Target, FileReference ManifestFile)
+		{
+			List<string> Arguments = new List<string> { Target.TargetName, Target.Platform.ToString(), Target.Config.ToString() };
+			if (Target.UprojectPath != null)
+			{
+				Arguments.Add($"-Project={Target.UprojectPath}");
+			}
+			if (ManifestFile != null)
+			{
+				Arguments.Add($"-Manifest={ManifestFile}");
+			}
 
-			return Manifest;
+			string CmdLine = CommandLine.FormatCommandLine(Arguments);
+			if (!String.IsNullOrEmpty(Target.UBTArgs))
+			{
+				CmdLine = $"{CmdLine} {Target.UBTArgs}";
+			}
+
+			return CmdLine;
 		}
 
 		string[] DotNetProductExtenstions()
@@ -519,437 +494,6 @@ namespace AutomationTool
 			BuildProductFiles.Clear();
 		}
 
-		public List<string> FindXGEFiles()
-		{
-			var Result = new List<string>();
-			var Root = CommandUtils.CombinePaths(CommandUtils.CmdEnv.LocalRoot, @"\Engine\Intermediate\Build");			
-			Result.AddRange(CommandUtils.FindFiles_NoExceptions("*.xge.xml", false, Root));
-			Result.Sort();
-			return Result;
-		}
-
-		public bool ProcessXGEItems(List<XGEItem> Actions, string XGETool, string Args, string TaskFilePath, bool ShowProgress, bool AllCores)
-		{
-			IScope CombineXGEScope = GlobalTracer.Instance.BuildSpan("CombineXGEItemFiles").WithTag("xgeTool", Path.GetFileNameWithoutExtension(XGETool)).StartActive();
-
-			XmlDocument XGETaskDocument;	
-			if (!CombineXGEItemFiles(Actions, TaskFilePath, out XGETaskDocument))
-			{
-				CommandUtils.LogVerbose("Incremental build, apparently everything was up to date, no XGE jobs produced.");
-			}
-			else
-			{
-				XmlElement EnvironmentsElement = XGETaskDocument.CreateElement("Environments");// Make sure all the tasks have a unique prefix
-				if(ShowProgress)
-				{
-					List<XmlElement> AllToolElements = new List<XmlElement>();
-					foreach(XmlElement EnvironmentElement in EnvironmentsElement.GetElementsByTagName("Environment"))
-					{
-						foreach(XmlElement ToolsElement in EnvironmentElement.GetElementsByTagName("Tools"))
-						{
-							foreach(XmlElement ToolElement in ToolsElement.GetElementsByTagName("Tool"))
-							{
-								AllToolElements.Add(ToolElement);
-							}
-						}
-					}
-					for(int Idx = 0; Idx < AllToolElements.Count; Idx++)
-					{
-						XmlElement ToolElement = AllToolElements[Idx];
-						if (ToolElement.HasAttribute("OutputPrefix"))
-						{
-							ToolElement.SetAttribute("OutputPrefix", ToolElement.Attributes["OutputPrefix"].Value + String.Format(" [@progress increment 1/{0}]", AllToolElements.Count));
-						}
-						else
-						{
-							ToolElement.SetAttribute("OutputPrefix", String.Format(" [@progress increment 1/{0} skipline]", AllToolElements.Count));
-						}
-					}
-				}
-
-				// Write the XGE task XML to a temporary file.
-				using (FileStream OutputFileStream = new FileStream(TaskFilePath, FileMode.Create, FileAccess.Write))
-				{
-					XGETaskDocument.Save(OutputFileStream);
-				}
-				if (!CommandUtils.FileExists(TaskFilePath))
-				{
-					throw new UnrealBuildException("Unable to find xge xml: " + TaskFilePath);
-				}
-
-				CombineXGEScope.Span.Finish();
-
-				if(XGETool == null)
-				{
-					CommandUtils.PushDir(CommandUtils.CombinePaths(CommandUtils.CmdEnv.LocalRoot, @"\Engine\Source"));
-					try
-					{
-						int ExitCode = ParallelExecutor.Execute(TaskFilePath, OwnerCommand.ParseParam("StopOnErrors"), AllCores);
-						if(ExitCode != 0)
-						{
-							return false;
-						}
-					}
-					finally
-					{
-						CommandUtils.PopDir();
-					}
-				}
-				else
-				{
-					using (IScope Scope = GlobalTracer.Instance.BuildSpan("ProcessXGE").WithTag("xgeTool", Path.GetFileNameWithoutExtension(XGETool)).StartActive())
-					{
-						int ConnectionRetries = 4;
-						while (true)
-						{
-							CommandUtils.LogVerbose("Running {0} *******", XGETool);
-							CommandUtils.PushDir(CommandUtils.CombinePaths(CommandUtils.CmdEnv.LocalRoot, @"\Engine\Source"));
-							int SuccesCode;
-							string LogFile = CommandUtils.GetRunAndLogOnlyName(CommandUtils.CmdEnv, XGETool);
-							string Output = CommandUtils.RunAndLog(XGETool, Args, out SuccesCode, LogFile);
-							bool bOutputContainsProject = Output.Contains("------Project:");
-							CommandUtils.PopDir();
-							if (ConnectionRetries > 0 && (SuccesCode == 4 || SuccesCode == 2) && !bOutputContainsProject && XGETool != null)
-							{
-								CommandUtils.LogWarning(String.Format("{0} failure on the local connection timeout", XGETool));
-								if (ConnectionRetries < 2)
-								{
-									System.Threading.Thread.Sleep(60000);
-								}
-								ConnectionRetries--;
-								continue;
-							}
-							else if (SuccesCode != 0)
-							{
-								throw new UnrealBuildException("Command failed (Result:{3}): {0} {1}. See logfile for details: '{2}' ",
-																XGETool, Args, Path.GetFileName(LogFile), SuccesCode);
-							}
-							CommandUtils.LogVerbose("{0} {1} Done *******", XGETool, Args);
-							break;
-						}
-					}
-				}
-			}
-			foreach (var Item in Actions)
-			{
-				XGEFinishBuildWithUBT(Item);
-			}
-			return true;
-		}
-
-        private static bool CombineXGEItemFiles(List<XGEItem> Actions, string TaskFilePath, out XmlDocument XGETaskDocument)
-        {
-            XGETaskDocument = new XmlDocument();
-
-			// <BuildSet FormatVersion="1">...</BuildSet>
-			XmlElement BuildSetElement = XGETaskDocument.CreateElement("BuildSet");
-			XGETaskDocument.AppendChild(BuildSetElement);
-			BuildSetElement.SetAttribute("FormatVersion", "1");
-
-			// <Environments>...</Environments>
-			XmlElement EnvironmentsElement = XGETaskDocument.CreateElement("Environments");
-			BuildSetElement.AppendChild(EnvironmentsElement);
-
-			int Job = 0;
-			int Env = 0;
-			Dictionary<string, XmlElement> EnvStringToEnv = new Dictionary<string, XmlElement>();
-			Dictionary<string, XmlElement> EnvStringToProject = new Dictionary<string, XmlElement>();
-			Dictionary<string, string> ParamsToTool = new Dictionary<string, string>();
-			Dictionary<string, XmlElement> ParamsToToolElement = new Dictionary<string, XmlElement>();
-			Dictionary<string, string> ToolToAction = new Dictionary<string, string>();
-			foreach (var Item in Actions)
-			{
-				var CurrentDependencies = new List<string>();
-				foreach (var XGEFile in Item.XgeXmlFiles)
-				{
-					if (!CommandUtils.FileExists_NoExceptions(XGEFile))
-					{
-						throw new UnrealBuildException("Couldn't find file: {0}", XGEFile);
-					}
-					var TargetFile = TaskFilePath + "." + Path.GetFileName(XGEFile);
-					CommandUtils.CopyFile(XGEFile, TargetFile);
-					CommandUtils.CopyFile_NoExceptions(XGEFile, TaskFilePath);
-
-					XmlReaderSettings XmlSettings = new XmlReaderSettings();
-					XmlSettings.DtdProcessing = DtdProcessing.Ignore;
-					XmlSettings.XmlResolver = null;
-
-					XmlDocument UBTTask = new XmlDocument();
-                    using (XmlReader Reader = XmlReader.Create(XGEFile, XmlSettings))
-					{
-						UBTTask.Load(Reader);
-					}
-
-					CommandUtils.DeleteFile(XGEFile);
-
-					var All = new List<string>();
-					{
-						var Elements = UBTTask.GetElementsByTagName("Variable");
-						foreach (XmlElement Element in Elements)
-						{
-							string Pair = Element.Attributes["Name"].Value + "=" + Element.Attributes["Value"].Value;
-							All.Add(Pair);
-						}
-					}
-					All.Sort();
-					string AllString = "";
-					foreach (string Element in All)
-					{
-						AllString += Element + "\n";
-					}
-					XmlElement ToolsElement;
-					XmlElement ProjectElement;
-
-					if (EnvStringToEnv.ContainsKey(AllString))
-					{
-						ToolsElement = EnvStringToEnv[AllString];
-						ProjectElement = EnvStringToProject[AllString];
-					}
-					else
-					{
-						string EnvName = string.Format("Env_{0}", Env);
-						Env++;
-						// <Environment Name="Win32">...</Environment>
-						XmlElement EnvironmentElement = XGETaskDocument.CreateElement("Environment");
-						EnvironmentsElement.AppendChild(EnvironmentElement);
-						EnvironmentElement.SetAttribute("Name", EnvName);
-
-						// <Tools>...</Tools>
-						ToolsElement = XGETaskDocument.CreateElement("Tools");
-						EnvironmentElement.AppendChild(ToolsElement);
-
-						{
-							// <Variables>...</Variables>
-							XmlElement VariablesElement = XGETaskDocument.CreateElement("Variables");
-							EnvironmentElement.AppendChild(VariablesElement);
-
-							var Elements = UBTTask.GetElementsByTagName("Variable");
-							foreach (XmlElement Element in Elements)
-							{
-								// <Variable>...</Variable>
-								XmlElement VariableElement = XGETaskDocument.CreateElement("Variable");
-								VariablesElement.AppendChild(VariableElement);
-								VariableElement.SetAttribute("Name", Element.Attributes["Name"].Value);
-								VariableElement.SetAttribute("Value", Element.Attributes["Value"].Value);
-							}
-						}
-
-						// <Project Name="Default" Env="Default">...</Project>
-						ProjectElement = XGETaskDocument.CreateElement("Project");
-						BuildSetElement.AppendChild(ProjectElement);
-						ProjectElement.SetAttribute("Name", EnvName);
-						ProjectElement.SetAttribute("Env", EnvName);
-
-						EnvStringToEnv.Add(AllString, ToolsElement);
-						EnvStringToProject.Add(AllString, ProjectElement);
-
-					}
-
-					Dictionary<string, string> ToolToTool = new Dictionary<string, string>();
-					Dictionary<string, string> ActionToAction = new Dictionary<string, string>();
-
-					{
-						var Elements = UBTTask.GetElementsByTagName("Tool");
-						foreach (XmlElement Element in Elements)
-						{
-							string Key = Element.Attributes["Path"].Value;
-							Key += " ";
-							Key += Element.Attributes["Params"].Value;
-
-							//hack hack hack
-							string ElementParams = Element.Attributes["Params"].Value;
-							if (!String.IsNullOrEmpty(ElementParams))
-							{
-								int YcIndex = ElementParams.IndexOf(" /Yc\"");
-								if (YcIndex >= 0)
-								{
-									// /Fp&quot;D:\BuildFarm\buildmachine_++depot+UE4\Engine\Intermediate\BuildData\Win64\UnrealEditor\Development\SharedPCHs\CoreUObject.h.pch&quot
-									string Fp = " /Fp\"";
-									int FpIndex = ElementParams.IndexOf(Fp, YcIndex);
-									if (FpIndex >= 0)
-									{
-										int EndIndex = ElementParams.IndexOf("\"", FpIndex + Fp.Length);
-										if (EndIndex >= 0)
-										{
-											string PCHFileName = ElementParams.Substring(FpIndex + Fp.Length, EndIndex - FpIndex - Fp.Length);
-											if (PCHFileName.Contains(@"\SharedPCHs\") && PCHFileName.Contains(@"\UnrealEditor\"))
-											{
-												Key = "SharedEditorPCH$ " + PCHFileName;
-												CommandUtils.LogLog("Hack: detected Shared PCH, which will use a different key {0}", Key);
-											}
-										}
-									}
-								}
-
-							}
-
-							string ToolName = string.Format("{0}_{1}", Element.Attributes["Name"].Value, Job);
-							string OriginalToolName = ToolName;
-
-							if (ParamsToTool.ContainsKey(Key))
-							{
-								ToolName = ParamsToTool[Key];
-								ToolToTool.Add(OriginalToolName, ToolName);
-
-								XmlElement ToolElement = ParamsToToolElement[Key];
-								ToolElement.SetAttribute("GroupPrefix", ToolElement.Attributes["GroupPrefix"].Value + " + " + Item.OutputCaption);
-							}
-							else
-							{
-								// <Tool ... />
-								XmlElement ToolElement = XGETaskDocument.CreateElement("Tool");
-								ToolsElement.AppendChild(ToolElement);
-
-								ParamsToTool.Add(Key, ToolName);
-								ParamsToToolElement.Add(Key, ToolElement);
-
-								ToolElement.SetAttribute("Name", ToolName);
-								ToolElement.SetAttribute("AllowRemote", Element.Attributes["AllowRemote"].Value);
-								if (Element.HasAttribute("OutputPrefix"))
-								{
-									ToolElement.SetAttribute("OutputPrefix", Element.Attributes["OutputPrefix"].Value);
-								}
-								ToolElement.SetAttribute("GroupPrefix", "** For " + Item.OutputCaption);
-
-								ToolElement.SetAttribute("Params", Element.Attributes["Params"].Value);
-								ToolElement.SetAttribute("Path", Element.Attributes["Path"].Value);
-								if(Element.HasAttribute("VCCompiler"))
-								{
-									ToolElement.SetAttribute("VCCompiler", Element.Attributes["VCCompiler"].Value);
-								}
-								ToolElement.SetAttribute("SkipIfProjectFailed", Element.Attributes["SkipIfProjectFailed"].Value);
-								if (Element.HasAttribute("AutoReserveMemory"))
-								{
-									ToolElement.SetAttribute("AutoReserveMemory", Element.Attributes["AutoReserveMemory"].Value);
-								}
-								ToolElement.SetAttribute("OutputFileMasks", Element.Attributes["OutputFileMasks"].Value);
-								if(Element.HasAttribute("AutoRecover"))
-								{
-									ToolElement.SetAttribute("AutoRecover", Element.Attributes["AutoRecover"].Value);
-								}
-								//ToolElement.SetAttribute("AllowRestartOnLocal", "false");  //vs2012 can't really restart, so we will go with this for now
-								if (Element.Attributes["OutputFileMasks"].Value == "PCLaunch.rc.res")
-								{
-									// total hack, when doing clean compiles, this output directory does not exist, we need to create it now
-									string Parms = Element.Attributes["Params"].Value;
-									string Start = "/fo \"";
-									int StartIndex = Parms.IndexOf(Start);
-									if (StartIndex >= 0)
-									{
-										Parms = Parms.Substring(StartIndex + Start.Length);
-										int EndIndex = Parms.IndexOf("\"");
-										if (EndIndex > 0)
-										{
-											string ResLocation = CommandUtils.CombinePaths(Parms.Substring(0, EndIndex));
-											if (!CommandUtils.DirectoryExists_NoExceptions(CommandUtils.GetDirectoryName(ResLocation)))
-											{
-												CommandUtils.CreateDirectory(CommandUtils.GetDirectoryName(ResLocation));
-											}
-										}
-									}
-								}
-							}
-						}
-					}
-					{
-						var NextDependencies = new List<string>();
-
-						var Elements = UBTTask.GetElementsByTagName("Task");
-						foreach (XmlElement Element in Elements)
-						{
-							string ToolName = string.Format("{0}_{1}", Element.Attributes["Tool"].Value, Job);
-							string ActionName = string.Format("{0}_{1}", Element.Attributes["Name"].Value, Job);
-							string OriginalActionName = ActionName;
-
-							if (ToolToTool.ContainsKey(ToolName))
-							{
-								ToolName = ToolToTool[ToolName];
-								ActionName = ToolToAction[ToolName];
-								ActionToAction.Add(OriginalActionName, ActionName);
-							}
-							else
-							{
-								ActionToAction.Add(OriginalActionName, ActionName);
-								ToolToAction.Add(ToolName, ActionName);
-								// <Task ... />
-								XmlElement TaskElement = XGETaskDocument.CreateElement("Task");
-								ProjectElement.AppendChild(TaskElement);
-
-								TaskElement.SetAttribute("SourceFile", Element.Attributes["SourceFile"].Value);
-								if (Element.HasAttribute("Caption"))
-								{
-									TaskElement.SetAttribute("Caption", Element.Attributes["Caption"].Value);
-								}
-								TaskElement.SetAttribute("Name", ActionName);
-								NextDependencies.Add(ActionName);
-								TaskElement.SetAttribute("Tool", ToolName);
-								TaskElement.SetAttribute("WorkingDir", Element.Attributes["WorkingDir"].Value);
-								TaskElement.SetAttribute("SkipIfProjectFailed", Element.Attributes["SkipIfProjectFailed"].Value);
-
-								string NewDepends = "";
-								if (Element.HasAttribute("DependsOn"))
-								{
-									string Depends = Element.Attributes["DependsOn"].Value;
-									while (Depends.Length > 0)
-									{
-										string ThisAction = Depends;
-										int Semi = Depends.IndexOf(";");
-										if (Semi >= 0)
-										{
-											ThisAction = Depends.Substring(0, Semi);
-											Depends = Depends.Substring(Semi + 1);
-										}
-										else
-										{
-											Depends = "";
-										}
-										if (ThisAction.Length > 0)
-										{
-											if (NewDepends.Length > 0)
-											{
-												NewDepends += ";";
-											}
-											string ActionJob = ThisAction + string.Format("_{0}", Job);
-											if (!ActionToAction.ContainsKey(ActionJob))
-											{
-												throw new UnrealBuildException("Action not found '{0}' in {1}->{2}", ActionJob, XGEFile, TargetFile);
-												// the XGE schedule is not topologically sorted. Hmmm. Well, we make a scary assumption here that this 
-											}
-											NewDepends += ActionToAction[ActionJob];
-										}
-									}
-								}
-								foreach (var Dep in CurrentDependencies)
-								{
-									if (NewDepends.Length > 0)
-									{
-										NewDepends += ";";
-									}
-									NewDepends += Dep;
-								}
-								if (NewDepends != "")
-								{
-									TaskElement.SetAttribute("DependsOn", NewDepends);
-								}
-							}
-
-						}
-						CurrentDependencies.AddRange(NextDependencies);
-					}
-					Job++;
-				}
-			}
-			return (Job > 0);
-		}
-
-		public void ClearExportedXGEXML()
-		{
-			foreach (var XGEFile in FindXGEFiles())
-			{
-				CommandUtils.DeleteFile(XGEFile);
-			}
-		}
-
 		public bool CanUseXGE(UnrealTargetPlatform Platform)
 		{
 			return PlatformExports.CanUseXGE(Platform, Log.Logger);
@@ -1049,102 +593,8 @@ namespace AutomationTool
 				}
 			}
 
-			// Filter the targets into those which can be built in parallel, vs those that must be executed serially
-			List<BuildTarget> NormalTargets = new List<BuildTarget>();
-			List<BuildTarget> ParallelXgeTargets = new List<BuildTarget>();
-			List<BuildTarget> ParallelTargets = new List<BuildTarget>();
-			foreach (BuildTarget Target in Agenda.Targets)
-			{
-				if(Target.TargetName == "UnrealHeaderTool")
-				{
-					NormalTargets.Insert(0, Target);
-				}
-				else if(bCanUseXGE && CanUseXGE(Target.Platform))
-				{
-					ParallelXgeTargets.Add(Target);
-				}
-				else if(bCanUseParallelExecutor && CanUseParallelExecutor(Target.Platform))
-				{
-					ParallelTargets.Add(Target);
-				}
-				else
-				{
-					NormalTargets.Add(Target);
-				}
-			}
-
-			// Execute all the serial targets
-			foreach(BuildTarget Target in NormalTargets)
-			{
-				// When building a target for Mac or iOS, use UBT's -flushmac option to clean up the remote builder
-				bool bForceFlushMac = DeleteBuildProducts && (Target.Platform == UnrealTargetPlatform.Mac || Target.Platform == UnrealTargetPlatform.IOS);
-				BuildManifest Manifest = BuildWithUBT(Target.TargetName, Target.Platform, Target.Config, Target.UprojectPath, bForceFlushMac, bDisableXGE, InAllCores, Target.UBTArgs);
-				if(InTargetToManifest != null)
-				{
-					InTargetToManifest[Target] = Manifest;
-				}
-			}
-
-			// Execute all the XGE targets
-			if(ParallelXgeTargets.Count > 0)
-			{
-				BuildParallelTargets(ParallelXgeTargets, InShowProgress, InAllCores, XGEConsole, InTargetToManifest);
-			}
-
-			// Execute all the parallel targets
-			if(ParallelTargets.Count > 0)
-			{
-				BuildParallelTargets(ParallelTargets, InShowProgress, InAllCores, null, InTargetToManifest);
-			}
-		}
-
-		private void BuildParallelTargets(List<BuildTarget> ParallelTargets, bool InShowProgress, bool InAllCores, string XGETool, Dictionary<BuildTarget, BuildManifest> InTargetToManifest)
-		{
-			string TaskFilePath = CommandUtils.CombinePaths(CommandUtils.CmdEnv.LogFolder, @"UAT_XGE.xml");
-
-			CommandUtils.LogSetProgress(InShowProgress, "Generating headers...");
-
-			List<XGEItem> XGEItems = new List<XGEItem>();
-			foreach (BuildTarget Target in ParallelTargets)
-			{
-				XGEItem Item = XGEPrepareBuildWithUBT(Target.TargetName, Target.Platform, Target.Config, Target.UprojectPath, Target.UBTArgs);
-				if(InTargetToManifest != null)
-				{
-					InTargetToManifest[Target] = Item.Manifest;
-				}
-				XGEItems.Add(Item);
-			}
-
-			string Args = null;
-			if (XGETool != null) 
-			{
-				Args = "\"" + TaskFilePath + "\" /Rebuild /NoLogo /ShowAgent /ShowTime /Title=\"UnrealBuildTool Compile\"";
-				if (ParseParam("StopOnErrors"))
-				{
-					Args += " /StopOnErrors";
-				}
-
-				if (OperatingSystem.IsWindows())
-				{
-					// A bug in the UCRT can cause XGE to hang on VS2015 builds. Figure out if this hang is likely to effect this build and workaround it if able.
-					string XGEVersion = Microsoft.Win32.Registry.GetValue(@"HKEY_LOCAL_MACHINE\SOFTWARE\Wow6432Node\Xoreax\IncrediBuild\Builder", "Version", null) as string;
-					if (XGEVersion != null)
-					{
-						// Per Xoreax support, subtract 1001000 from the registry value to get the build number of the installed XGE.
-						int XGEBuildNumber;
-						if (Int32.TryParse(XGEVersion, out XGEBuildNumber) && XGEBuildNumber - 1001000 >= 1659)
-						{
-							Args += " /no_watchdog_thread";
-						}
-					}
-				}
-			}
-
-			CommandUtils.LogSetProgress(InShowProgress, "Building...");
-			if (!ProcessXGEItems(XGEItems, XGETool, Args, TaskFilePath, InShowProgress, InAllCores))
-			{
-				throw new UnrealBuildException("{0} failed, retries not enabled:", XGETool);
-			}
+			// Build all the targets
+			BuildWithUBT(Agenda.Targets, InTargetToManifest, bDisableXGE, InAllCores);
 		}
 
 		/// <summary>
@@ -1262,17 +712,22 @@ namespace AutomationTool
 			}
 		}
 
-		FileReference GetManifestFile(FileReference ProjectFile)
+		DirectoryReference GetManifestDir(FileReference ProjectFile)
 		{
 			// Can't write to Engine directory on installed builds
 			if (Unreal.IsEngineInstalled() && ProjectFile != null)
 			{
-				return FileReference.Combine(ProjectFile.Directory, "Intermediate", "Build", "Manifest.xml");
+				return DirectoryReference.Combine(ProjectFile.Directory, "Intermediate", "Build");
 			}
 			else
 			{
-				return FileReference.Combine(Unreal.EngineDirectory, "Intermediate", "Build", "Manifest.xml");
+				return DirectoryReference.Combine(Unreal.EngineDirectory, "Intermediate", "Build");
+			}
 			}
+
+		FileReference GetManifestFile(FileReference ProjectFile)
+		{
+			return FileReference.Combine(GetManifestDir(ProjectFile), "Manifest.xml");
 		}
 
 		[Obsolete("Deprecated in UE5.1; use UnrealBuildToolDll")]
