Fix AAB archiving issue (incorrectly added '-' when architecture not specified)

#android

diff --git a/Engine/Source/Programs/AutomationTool/Android/AndroidPlatform.Automation.cs b/Engine/Source/Programs/AutomationTool/Android/AndroidPlatform.Automation.cs
--- a/Engine/Source/Programs/AutomationTool/Android/AndroidPlatform.Automation.cs
+++ b/Engine/Source/Programs/AutomationTool/Android/AndroidPlatform.Automation.cs
@@ -194,7 +194,7 @@
 			while (Option == 2)
 			{
 				string Prompt = $"The Android Sdk directory was not found (expected to find it at '{SdkDir}'\n" +
-					"Android Studio can install it for you, but you will need to manually perform some steps (if desired, you can get detailed help with option 2):\n" +
+					"Android Studio can install it for you, but you will need to manually perform some steps (if desired, you can get detailed help with option 2):\n" +
 					"  - Wait for Android Studio to start, you will see an initial dialog asking how to proceed (called \"Welcome to Android Studio\")\n" +
 					"  - Click the \"Configure\" dropdown in the bottom right, and select \"SDK Manager\"\n" +
 					"  - Click on the \"SDK Tools\" tab near the top middle of the right pane\n" +
@@ -525,8 +525,8 @@
 		}
 
 		// Apk's go to project location, not necessarily where the .so is (content only packages need to output to their directory)
-		string ArchName = Architecture == null ? "" : Architecture.ToString();
-		string ApkName = Path.Combine(ProjectDir, $"{DecoratedExeName}-{ArchName}.apk");
+		string ArchName = Architecture == null ? "" : "-" + Architecture.ToString();
+		string ApkName = Path.Combine(ProjectDir, $"{DecoratedExeName}{ArchName}.apk");
 
 		// if the source binary was UnrealGame, handle using it or switching to project name
 		if (Path.GetFileNameWithoutExtension(Params.GetProjectExeForPlatform(UnrealTargetPlatform.Android).ToString()) == "UnrealGame")
