Add mkdir for DeviceExternalFilesSavedPath before deploying on Android

diff --git a/Engine/Source/Programs/AutomationTool/Gauntlet/Platform/Android/Gauntlet.TargetDeviceAndroid.cs b/Engine/Source/Programs/AutomationTool/Gauntlet/Platform/Android/Gauntlet.TargetDeviceAndroid.cs
--- a/Engine/Source/Programs/AutomationTool/Gauntlet/Platform/Android/Gauntlet.TargetDeviceAndroid.cs
+++ b/Engine/Source/Programs/AutomationTool/Gauntlet/Platform/Android/Gauntlet.TargetDeviceAndroid.cs
@@ -940,6 +940,7 @@ namespace Gauntlet
 
 				// remote dir on the device, create it if it doesn't exist
 				RunAdbDeviceCommand(string.Format("shell mkdir -p {0}/", DeviceExternalStorageSavedPath));
+				RunAdbDeviceCommand(string.Format("shell mkdir -p {0}/", DeviceExternalFilesSavedPath));
 
 				// path to the APK to install.
 				string ApkPath = Build.SourceApkPath;
