This change adds support for defining ondemand pak file rule(s). These rules
overrides existing rules created from asset chunking and is currenlty only meant to be used for .ubulk/.uptnl files.
Ondemand content is serialized to loose files with a new I/O store writer to /MyGame/ContentOnDemand/Paks/{MyPakRuleName}.

#rnx

diff --git a/Engine/Source/Programs/AutomationTool/Scripts/CopyBuildToStagingDirectory.Automation.cs b/Engine/Source/Programs/AutomationTool/Scripts/CopyBuildToStagingDirectory.Automation.cs
--- a/Engine/Source/Programs/AutomationTool/Scripts/CopyBuildToStagingDirectory.Automation.cs
+++ b/Engine/Source/Programs/AutomationTool/Scripts/CopyBuildToStagingDirectory.Automation.cs
@@ -2056,6 +2056,7 @@ namespace AutomationScripts
 			// Whether this rule is defined for content-on-demand
 			public bool bOnDemand;
 
+
 			public static bool IsMatch(PakFileRules PakRules, KeyValuePair<string, string> StagingFile)
 			{
 				bool bMatched = !PakRules.bDisabled &&
@@ -2306,6 +2307,43 @@ namespace AutomationScripts
 				}
 			}
 
+			// Search in order, return on first match
+			foreach (var PakRules in RulesList)
+			{
+				if (PakFileRules.IsMatch(PakRules, StagingFile))
+				{
+					{
+						ModifyPakList.Clear();
+
+						string ChunkName = Rule.Name;
+						if (Rule.OverridePaks != null && Rule.OverridePaks.Count > 0)
+						{
+							string OverrideChunkName = Rule.OverridePaks[0];
+							if (ChunkNameToDefinition.ContainsKey(OverrideChunkName))
+							{
+								ChunkName = OverrideChunkName;
+							}
+							else
+							{
+								LogWarning("Undefined PAK override '{0}' in PAK rule '{1}'", OverrideChunkName, Rule.Name);
+							}
+						}
+
+						ChunkDefinition Chunk = null;
+						if (ChunkNameToDefinition.TryGetValue(ChunkName, out Chunk))
+						{
+							ModifyPakList.Add(Chunk);
+						}
+						else
+						{
+							LogWarning("Undefined chunk name '{0}' in PAK rule '{1}'", ChunkName, Rule.Name);
+						}
+						
+						return;
+					}
+				}
+			}
+
 			// Search in order, return on first match
 			foreach (var PakRules in RulesList)
 			{
