Fixed issue with ucas and utoc files not being copied to the output directory when creating chunk installs

diff --git a/Engine/Source/Programs/AutomationTool/Scripts/CopyBuildToStagingDirectory.Automation.cs b/Engine/Source/Programs/AutomationTool/Scripts/CopyBuildToStagingDirectory.Automation.cs
--- a/Engine/Source/Programs/AutomationTool/Scripts/CopyBuildToStagingDirectory.Automation.cs
+++ b/Engine/Source/Programs/AutomationTool/Scripts/CopyBuildToStagingDirectory.Automation.cs
@@ -3479,6 +3479,11 @@ namespace AutomationScripts
 						{
 							InternalUtils.SafeCopyFile(Path.ChangeExtension(OutputLocation.FullName, ".sig"), Path.ChangeExtension(RawDataPakPath, ".sig"), true);
 						}
+						if (bShouldCreateIoStoreContainerFiles)
+						{
+							InternalUtils.SafeCopyFile(Path.ChangeExtension(OutputLocation.FullName, ".utoc"), Path.ChangeExtension(RawDataPakPath, ".utoc"), true);
+							InternalUtils.SafeCopyFile(Path.ChangeExtension(OutputLocation.FullName, ".ucas"), Path.ChangeExtension(RawDataPakPath, ".ucas"), true);
+						}
 						InternalUtils.SafeDeleteFile(OutputLocation.FullName, true);
 
 						if (Params.IsGeneratingPatch)
@@ -3495,6 +3500,11 @@ namespace AutomationScripts
 							{
 								InternalUtils.SafeCopyFile(PatchSourceContentPath, Path.ChangeExtension(SourceRawDataPakPath, ".sig"), true);
 							}
+							if (bShouldCreateIoStoreContainerFiles)
+							{
+								InternalUtils.SafeCopyFile(PatchSourceContentPath, Path.ChangeExtension(SourceRawDataPakPath, ".utoc"), true);
+								InternalUtils.SafeCopyFile(PatchSourceContentPath, Path.ChangeExtension(SourceRawDataPakPath, ".ucas"), true);
+							}
 						}
 
 						string BuildRoot = MakePathSafeToUseWithCommandLine(RawDataPath);
