PR #10116: Add -ubtargs option to MegaXGE to pass to UnrealBuildTool  (Contributed by [REDACTED_PERSON])

diff --git a/Engine/Source/Programs/AutomationTool/Scripts/MegaXGE.Automation.cs b/Engine/Source/Programs/AutomationTool/Scripts/MegaXGE.Automation.cs
--- a/Engine/Source/Programs/AutomationTool/Scripts/MegaXGE.Automation.cs
+++ b/Engine/Source/Programs/AutomationTool/Scripts/MegaXGE.Automation.cs
@@ -7,8 +7,11 @@ using AutomationTool;
 using UnrealBuildTool;
 using EpicGames.Core;
 
-[Help("Compiles a bunch of stuff together with megaxge: Example arguments: -Target1=\"PlatformerGame win32|ios debug|development\"")]
+[Help("Compiles a bunch of stuff together with megaxge: Example arguments: -ubtargs=\"-nopdb\" -Target1=\"PlatformerGame win32|ios debug|development\"")]
 [Help(typeof(UnrealBuild))]
+[Help("ubtargs", "-args -for -ubt")]
+[Help("clean", "Cleans targets before building")]
+[Help("progress", "Reports the current steps to the log")]
 [Help("Target1", "target1[|target2...] platform1[|platform2...] config1[|config2...]")]
 [Help("Target2", "target1[|target2...] platform1[|platform2...] config1[|config2...]")]
 
@@ -27,6 +30,7 @@ class MegaXGE : BuildCommand
 			WorkingCL = P4.CreateChange(P4Env.Client, String.Format("MegaXGE build from changelist {0} - Params: {1}", P4Env.Changelist, CmdLine));
 		}
 
+		string UbtArgs = ParseParamValue("ubtargs", "");
 		LogInformation("************************* MegaXGE");
 
 		bool Clean = ParseParam("Clean");
@@ -130,7 +134,7 @@ class MegaXGE : BuildCommand
 				{
 					foreach (var Configuration in Configurations)
 					{
-						Agenda.AddTargets(new string[] { CurTarget }, Platform, Configuration, ProjectFile);
+						Agenda.AddTargets(new string[] { CurTarget }, Platform, Configuration, ProjectFile, UbtArgs);
 						LogInformation("Target {0} {1} {2}", CurTarget, Platform.ToString(), Configuration.ToString());
 						if (Clean)
 						{
