PR #9104: UAT BuildGraph: Add support for -project flag (Contributed by [REDACTED_PERSON])

diff --git a/Engine/Source/Programs/AutomationTool/BuildGraph/BuildGraph.cs b/Engine/Source/Programs/AutomationTool/BuildGraph/BuildGraph.cs
--- a/Engine/Source/Programs/AutomationTool/BuildGraph/BuildGraph.cs
+++ b/Engine/Source/Programs/AutomationTool/BuildGraph/BuildGraph.cs
@@ -303,6 +303,15 @@ namespace AutomationTool
 				DefaultProperties["EngineCompatibleChange"] = Version.CompatibleChangelist.ToString();
 			}
 
+			// If the -project flag is given, pass useful information into the graph
+			FileReference ProjectFile = ParseProjectParam();
+			if (ProjectFile != null)
+			{
+				DefaultProperties["ProjectName"] = ProjectFile.GetFileNameWithoutExtension();
+				DefaultProperties["ProjectFile"] = ProjectFile.FullName;
+				DefaultProperties["ProjectDir"] = ProjectFile.Directory.FullName;
+			}
+
 			// Add any additional custom arguments from the command line (of the form -Set:X=Y)
 			Dictionary<string, string> Arguments = new Dictionary<string, string>(StringComparer.InvariantCultureIgnoreCase);
 			foreach (string Param in Params)
