Extended gauntlet role to support different build flavor for the same configuration (i.e ASan or others)

[RN] Extended gauntlet role to support different build flavor for the same configuration (i.e ASan or others)

diff --git a/Engine/Source/Programs/AutomationTool/Gauntlet/Framework/Base/Gauntlet.BuildSource.cs b/Engine/Source/Programs/AutomationTool/Gauntlet/Framework/Base/Gauntlet.BuildSource.cs
--- a/Engine/Source/Programs/AutomationTool/Gauntlet/Framework/Base/Gauntlet.BuildSource.cs
+++ b/Engine/Source/Programs/AutomationTool/Gauntlet/Framework/Base/Gauntlet.BuildSource.cs
@@ -20,7 +20,8 @@ namespace Gauntlet
 		CanReplaceExecutable = (1 << 3),		// build can use exes from elsewhere
         Bulk = (1 << 4),						// build is full-content, now startup download
 		NotBulk = (1 << 5),						// build is not bulk
-		ContentOnlyProject = (1 << 6)           // build is a content-only project, e.g. exe is UnrealGame.exe
+		ContentOnlyProject = (1 << 6),          // build is a content-only project, e.g. exe is UnrealGame.exe
+		ASan = (1 << 7)							// build is an ASan one
 	}
 
 	/// <summary>
diff --git a/Engine/Source/Programs/AutomationTool/Gauntlet/Unreal/Base/Gauntlet.UnrealSession.cs b/Engine/Source/Programs/AutomationTool/Gauntlet/Unreal/Base/Gauntlet.UnrealSession.cs
--- a/Engine/Source/Programs/AutomationTool/Gauntlet/Unreal/Base/Gauntlet.UnrealSession.cs
+++ b/Engine/Source/Programs/AutomationTool/Gauntlet/Unreal/Base/Gauntlet.UnrealSession.cs
@@ -198,6 +198,11 @@ namespace Gauntlet
 				{
 					RequiredBuildFlags |= BuildFlags.Loose;
 				}
+
+				if (Globals.Params.ParseParam("asan"))
+				{
+					RequiredBuildFlags |= BuildFlags.ASan;
+				}
 			}
 
 			InstallOnly = false;
