[Backout] - CL24230614
Original CL Desc
-----------------------------------------------------------------
Cosmetic bulk asset streaming

Adds support for marking skeletal meshes and textures as optional files that stream from the cloud, typically used for cosmetics. Meshes and textures will now upload all of their non-inlined LODs as optional bulk files to the cloud, removing them from the pak files. The client will download into the VirtualFileCache in PersistentDownloadDir/VFC, which evicts least recently used LODs. These assets will have a longer delay to update to the higher quality LOD when their LODs are not in cache because they are now downloaded on-demand.

diff --git a/Engine/Plugins/Experimental/DefaultInstallBundleManager/Source/Private/DefaultInstallBundleManager.cpp b/Engine/Plugins/Experimental/DefaultInstallBundleManager/Source/Private/DefaultInstallBundleManager.cpp
--- a/Engine/Plugins/Experimental/DefaultInstallBundleManager/Source/Private/DefaultInstallBundleManager.cpp
+++ b/Engine/Plugins/Experimental/DefaultInstallBundleManager/Source/Private/DefaultInstallBundleManager.cpp
@@ -2922,15 +2922,6 @@ bool FDefaultInstallBundleManager::HasBundleSource(EInstallBundleSourceType Sour
 	return BundleSources.Contains(SourceType);
 }
 
-const TSharedPtr<IInstallBundleSource> FDefaultInstallBundleManager::GetBundleSource(EInstallBundleSourceType SourceType) const
-{
-	if (InitState != EInstallBundleManagerInitState::Succeeded)
-		return nullptr;
-
-	const TSharedPtr<IInstallBundleSource>* Found = BundleSources.Find(SourceType);
-	return Found ? *Found : nullptr;
-}
-
 FDelegateHandle FDefaultInstallBundleManager::PushInitErrorCallback(FInstallBundleManagerInitErrorHandler Callback)
 {
 	InitErrorHandlerStack.Push(MoveTemp(Callback));
diff --git a/Engine/Plugins/Experimental/DefaultInstallBundleManager/Source/Public/DefaultInstallBundleManager.h b/Engine/Plugins/Experimental/DefaultInstallBundleManager/Source/Public/DefaultInstallBundleManager.h
--- a/Engine/Plugins/Experimental/DefaultInstallBundleManager/Source/Public/DefaultInstallBundleManager.h
+++ b/Engine/Plugins/Experimental/DefaultInstallBundleManager/Source/Public/DefaultInstallBundleManager.h
@@ -605,9 +605,6 @@ protected:
 public:
 	virtual bool HasBundleSource(EInstallBundleSourceType SourceType) const override;
 
-	virtual const TSharedPtr<IInstallBundleSource> GetBundleSource(EInstallBundleSourceType SourceType) const override;
-
-
 	virtual FDelegateHandle PushInitErrorCallback(FInstallBundleManagerInitErrorHandler Callback) override;
 	virtual void PopInitErrorCallback() override;
 	virtual void PopInitErrorCallback(FDelegateHandle Handle) override;
diff --git a/Engine/Source/Programs/AutomationTool/AutomationUtils/BundleUtils.cs b/Engine/Source/Programs/AutomationTool/AutomationUtils/BundleUtils.cs
--- a/Engine/Source/Programs/AutomationTool/AutomationUtils/BundleUtils.cs
+++ b/Engine/Source/Programs/AutomationTool/AutomationUtils/BundleUtils.cs
@@ -19,7 +19,6 @@ namespace AutomationUtils.Automation
 			public List<string> Tags { get; set; }
 			public List<string> Dependencies { get; set; }
 			public List<string> FileRegex { get; set; }
-			public List<string> Files { get; set; }
 			public bool bFoundParent { get; set; }
 			public bool bContainsShaderLibrary { get; set; }
 			public int Order { get; set; }
@@ -94,17 +93,6 @@ namespace AutomationUtils.Automation
 						Bundle.FileRegex = new List<string>();
 					}
 				}
-				{
-					List<string> Files;
-					if (BundleConfig.GetArray(SectionName, "Files", out Files))
-					{
-						Bundle.Files = Files;
-					}
-					else
-					{
-						Bundle.Files = new List<string>();
-					}
-				}
 				{
 					bool bContainsShaderLibrary;
 					if (BundleConfig.GetBool(SectionName, "ContainsShaderLibrary", out bContainsShaderLibrary))
@@ -140,7 +128,7 @@ namespace AutomationUtils.Automation
 		public static TPlatformBundleSettings MatchBundleSettings<TPlatformBundleSettings>(
 			string FileName, IReadOnlyDictionary<string, TPlatformBundleSettings> InstallBundles) where TPlatformBundleSettings : BundleSettings
 		{
-			// Try to find a matching chunk regex or exact filename
+			// Try to find a matching chunk regex
 			foreach (var Bundle in InstallBundles.Values)
 			{
 				foreach (string RegexString in Bundle.FileRegex)
@@ -150,13 +138,6 @@ namespace AutomationUtils.Automation
 						return Bundle;
 					}
 				}
-				foreach(string FileString in Bundle.Files)
-				{
-					if (FileString.Equals(FileName, StringComparison.OrdinalIgnoreCase))
-					{
-						return Bundle;
-					}
-				}
 			}
 
 			return null;
diff --git a/Engine/Source/Programs/AutomationTool/AutomationUtils/DeploymentContext.cs b/Engine/Source/Programs/AutomationTool/AutomationUtils/DeploymentContext.cs
--- a/Engine/Source/Programs/AutomationTool/AutomationUtils/DeploymentContext.cs
+++ b/Engine/Source/Programs/AutomationTool/AutomationUtils/DeploymentContext.cs
@@ -41,8 +41,6 @@ public class FilesToStage
 	/// </summary>
 	public Dictionary<StagedFileReference, FileReference> UFSFiles = new Dictionary<StagedFileReference, FileReference>();
 
-	public Dictionary<StagedFileReference, FileReference> LooseFiles = new Dictionary<StagedFileReference, FileReference>();
-
 	/// <summary>
 	/// After staging, this is a map from staged file to source file. These file are binaries, etc and can't go into a pak file.
 	/// </summary>
diff --git a/Engine/Source/Programs/AutomationTool/Scripts/CopyBuildToStagingDirectory.Automation.cs b/Engine/Source/Programs/AutomationTool/Scripts/CopyBuildToStagingDirectory.Automation.cs
--- a/Engine/Source/Programs/AutomationTool/Scripts/CopyBuildToStagingDirectory.Automation.cs
+++ b/Engine/Source/Programs/AutomationTool/Scripts/CopyBuildToStagingDirectory.Automation.cs
@@ -276,6 +276,7 @@ namespace AutomationScripts
 			string ContainerName,
 			FileReference PakOutputLocation,
 			DirectoryReference OptionalOutputLocation,
+			DirectoryReference OnDemandOutputLocation,
 			bool bCompressed,
 			EncryptionAndSigning.CryptoSettings CryptoSettings,
 			string EncryptionKeyGuid,
@@ -284,7 +285,17 @@ namespace AutomationScripts
 			bool bIsDLC)
 		{
 			StringBuilder CmdLine = new StringBuilder();
-			CmdLine.AppendFormat("-Output={0}", MakePathSafeToUseWithCommandLine(Path.ChangeExtension(PakOutputLocation.FullName, ".utoc")));
+
+			if (OnDemandOutputLocation != null)
+			{
+				CmdLine.AppendFormat("-Output={0}", MakePathSafeToUseWithCommandLine(Path.ChangeExtension(PakOutputLocation.FullName, ".tocmanifest")));
+				CmdLine.AppendFormat(" -OnDemandOutput={0}", MakePathSafeToUseWithCommandLine(OnDemandOutputLocation.FullName));
+			}
+			else
+			{
+				CmdLine.AppendFormat("-Output={0}", MakePathSafeToUseWithCommandLine(Path.ChangeExtension(PakOutputLocation.FullName, ".utoc")));
+			}
+
 			if (OptionalOutputLocation != null)
 			{
 				CmdLine.AppendFormat(" -OptionalOutput={0}", MakePathSafeToUseWithCommandLine(OptionalOutputLocation.FullName));
@@ -2042,7 +2053,17 @@ namespace AutomationScripts
 			// List of pak files to use instead of manifest
 			public List<string> OverridePaks;
 
-			public bool bStageLoose;
+			// Whether this rule is defined for content-on-demand
+			public bool bOnDemand;
+
+			public static bool IsMatch(PakFileRules PakRules, KeyValuePair<string, string> StagingFile)
+			{
+				bool bMatched = !PakRules.bDisabled &&
+					((!PakRules.bUseExactFilePathMatch && PakRules.Filter != null && PakRules.Filter.Matches(StagingFile.Key)) ||
+					(PakRules.bUseExactFilePathMatch && PakRules.ExactFileMatches != null && PakRules.ExactFileMatches.Contains(StagingFile.Value)));
+				
+				return bMatched;
+			}
 		};
 
 		/// <summary>
@@ -2173,9 +2194,9 @@ namespace AutomationScripts
 				PakRules.Name = SectionName;
 				PakRulesConfig.TryGetValue(SectionName, "bUseExactFilePathMatch", out PakRules.bUseExactFilePathMatch);
 				PakRulesConfig.TryGetValue(SectionName, "bExcludeFromPaks", out PakRules.bExcludeFromPaks);
-				PakRulesConfig.TryGetValue(SectionName, "bStageLoose", out PakRules.bStageLoose);
 				PakRulesConfig.TryGetValue(SectionName, "bOverrideChunkManifest", out PakRules.bOverrideChunkManifest);
 				PakRulesConfig.TryGetValue(SectionName, "bDisabled", out PakRules.bDisabled);
+				PakRulesConfig.TryGetValue(SectionName, "bOnDemand", out PakRules.bOnDemand);
 				string PakString;
 				PakRulesConfig.TryGetValue(SectionName, "OverridePaks", out PakString);
 
@@ -2188,12 +2209,20 @@ namespace AutomationScripts
 					PakRules.OverridePaks = null;
 				}
 
-				if (PakRules.bExcludeFromPaks && PakRules.OverridePaks != null)
+				if (PakRules.bExcludeFromPaks)
 				{
-					LogWarning("Error in PakFileRules {0}, set to exclude but also sets override!", PakRules.Name);
-					continue;
+					if (PakRules.OverridePaks != null)
+					{
+						LogWarning("Error in PakFileRules {0}, set to exclude but also sets override!", PakRules.Name);
+						continue;
+					}
+					if (PakRules.bOnDemand)
+					{
+						LogWarning("Error in PakFileRules {0}, set to exclude but also sets to ondemand!", PakRules.Name);
+						continue;
+					}
 				}
-				else if (!PakRules.bExcludeFromPaks && PakRules.OverridePaks == null)
+				else if (PakRules.OverridePaks == null && !PakRules.bOnDemand)
 				{
 					LogWarning("Error in PakFileRules {0}, set to include but did not specify paks!", PakRules.Name);
 					continue;
@@ -2230,31 +2259,58 @@ namespace AutomationScripts
 		/// </summary>
 		/// <param name="Params"></param>
 		/// <param name="SC"></param>
-		private static void ApplyPakFileRules(List<PakFileRules> RulesList, KeyValuePair<string, string> StagingFile, HashSet<ChunkDefinition> ModifyPakList, ConcurrentDictionary<string, ChunkDefinition> ChunkNameToDefinition, out bool bExcludeFromPaks, out bool bStageLoose)
+		private static void ApplyPakFileRules(List<PakFileRules> RulesList, KeyValuePair<string, string> StagingFile, HashSet<ChunkDefinition> ModifyPakList, ConcurrentDictionary<string, ChunkDefinition> ChunkNameToDefinition, out bool bExcludeFromPaks)
 		{
 			bExcludeFromPaks = false;
-			bStageLoose = false;
 
 			if (RulesList == null)
 			{
 				return;
 			}
 
-			// Search in order, return on first match
-			foreach (var PakRules in RulesList)
+			// Process content-on-demand rules first
+			if (ModifyPakList != null)
 			{
-				bool bMatched = !PakRules.bDisabled &&
-					((!PakRules.bUseExactFilePathMatch && PakRules.Filter != null && PakRules.Filter.Matches(StagingFile.Key)) ||
-					(PakRules.bUseExactFilePathMatch && PakRules.ExactFileMatches != null && PakRules.ExactFileMatches.Contains(StagingFile.Value)));
-
-				if (bMatched)
+				foreach (PakFileRules Rule in RulesList.Where(Rule => Rule.bOnDemand))
 				{
-					if (PakRules.bStageLoose)
+					if (PakFileRules.IsMatch(Rule, StagingFile))
 					{
-						bStageLoose = true;
+						ModifyPakList.Clear();
+
+						string ChunkName = Rule.Name;
+						if (Rule.OverridePaks != null && Rule.OverridePaks.Count > 0)
+						{
+							string OverrideChunkName = Rule.OverridePaks[0];
+							if (ChunkNameToDefinition.ContainsKey(OverrideChunkName))
+							{
+								ChunkName = OverrideChunkName;
+							}
+							else
+							{
+								LogWarning("Undefined PAK override '{0}' in PAK rule '{1}'", OverrideChunkName, Rule.Name);
+							}
+						}
+
+						ChunkDefinition Chunk = null;
+						if (ChunkNameToDefinition.TryGetValue(ChunkName, out Chunk))
+						{
+							ModifyPakList.Add(Chunk);
+						}
+						else
+						{
+							LogWarning("Undefined chunk name '{0}' in PAK rule '{1}'", ChunkName, Rule.Name);
+						}
+						
 						return;
 					}
+				}
+			}
 
+			// Search in order, return on first match
+			foreach (var PakRules in RulesList)
+			{
+				if (PakFileRules.IsMatch(PakRules, StagingFile))
+				{
 					bool bOverrideChunkAssignment = false;
 					if (ModifyPakList != null && ModifyPakList.Count > 0)
 					{
@@ -2307,7 +2363,6 @@ namespace AutomationScripts
 			}
 		}
 
-
 		/// <summary>
 		/// Creates a pak file using staging context (single manifest)
 		/// </summary>
@@ -2329,8 +2384,7 @@ namespace AutomationScripts
 			foreach (var StagingFile in UnrealPakResponseFile)
 			{
 				bool bExcludeFromPaks = false;
-				bool bStageLoose = false;
-				ApplyPakFileRules(PakRulesList, StagingFile, null, null, out bExcludeFromPaks, out bStageLoose);
+				ApplyPakFileRules(PakRulesList, StagingFile, null, null, out bExcludeFromPaks);
 
 				if (bExcludeFromPaks)
 				{
@@ -2487,6 +2541,11 @@ namespace AutomationScripts
 			/// Whether to rehydrate the assets when creating the pak file or not
 			/// </summary>
 			public bool bRehydrateAssets;
+			
+			/// <summary>
+			// Whether this pak file is for content-on-demand content delivery or not
+			/// </summary>
+			public bool bOnDemand;
 
 			/// <summary>
 			/// GUID of the encryption key for this pak file
@@ -2499,13 +2558,14 @@ namespace AutomationScripts
 			/// <param name="PakName">Path to the base output file for this pak file</param>
 			/// <param name="UnrealPakResponseFile">Map of files within the pak file to their source file on disk</param>
 			/// <param name="bCompressed">Whether to enable compression</param>
-			public CreatePakParams(string PakName, Dictionary<string, string> UnrealPakResponseFile, bool bCompressed, bool RehydrateAssets, string EncryptionKeyGuid)
+			public CreatePakParams(string PakName, Dictionary<string, string> UnrealPakResponseFile, bool bCompressed, bool RehydrateAssets, string EncryptionKeyGuid, bool bOnDemand = false)
 			{
 				this.PakName = PakName;
 				this.UnrealPakResponseFile = UnrealPakResponseFile;
 				this.bCompressed = bCompressed;
 				this.bRehydrateAssets = RehydrateAssets;
 				this.EncryptionKeyGuid = EncryptionKeyGuid;
+				this.bOnDemand = bOnDemand;
 			}
 		}
 
@@ -3013,11 +3073,13 @@ namespace AutomationScripts
 				OutputFilename = OutputFilename + PostFix;
 
 				StagedFileReference OutputRelativeLocation;
+				StagedDirectoryReference OnDemandOutputRelativeLocation;
 				if (Params.HasDLCName)
 				{
 					if (Params.DLCOverrideStagedSubDir != null)
 					{
 						OutputRelativeLocation = StagedFileReference.Combine(SC.RelativeProjectRootForStage, Params.DLCOverrideStagedSubDir, "Content", "Paks", Params.DLCFile.GetFileNameWithoutExtension() + OutputFilename + ".pak");
+						OnDemandOutputRelativeLocation = StagedDirectoryReference.Combine(SC.RelativeProjectRootForStage, Params.DLCOverrideStagedSubDir, "ContentOnDemand", "Paks", Params.DLCFile.GetFileNameWithoutExtension() + OutputFilename);
 					}
 					else
 					{
@@ -3032,23 +3094,29 @@ namespace AutomationScripts
 							PluginSubdirectory = Params.DLCFile.GetFileNameWithoutAnyExtensions();
 						}
 						OutputRelativeLocation = StagedFileReference.Combine(SC.RelativeProjectRootForStage, PluginSubdirectory, "Content", "Paks", SC.FinalCookPlatform, Params.DLCFile.GetFileNameWithoutExtension() + OutputFilename + ".pak");
+						OnDemandOutputRelativeLocation = StagedDirectoryReference.Combine(SC.RelativeProjectRootForStage, PluginSubdirectory, "ContentOnDemand", "Paks", SC.FinalCookPlatform, Params.DLCFile.GetFileNameWithoutExtension() + OutputFilename);
 					}
 				}
 				else
 				{
 					OutputRelativeLocation = StagedFileReference.Combine(SC.RelativeProjectRootForStage, "Content", "Paks", OutputFilename + OutputFilenameExtension);
+					OnDemandOutputRelativeLocation = StagedDirectoryReference.Combine(SC.RelativeProjectRootForStage, "ContentOnDemand", "Paks", OutputFilename);
 				}
 				if (SC.StageTargetPlatform.DeployLowerCaseFilenames(StagedFileType.UFS))
 				{
 					OutputRelativeLocation = OutputRelativeLocation.ToLowerInvariant();
 				}
 				OutputRelativeLocation = SC.StageTargetPlatform.Remap(OutputRelativeLocation);
+				
+				//TODO: Do we need to remap the ondemand output location?
+				//OnDemandOutputRelativeLocation = SC.StageTargetPlatform.Remap(OnDemandOutputRelativeLocation);
 
 				FileReference OutputLocation = FileReference.Combine(SC.RuntimeRootDir, OutputRelativeLocation.Name);
+				DirectoryReference OnDemandOutputLocation = DirectoryReference.Combine(SC.RuntimeRootDir, OnDemandOutputRelativeLocation.Name);
 
 				bool bCopiedExistingPak = false;
 
-				if (SC.StageTargetPlatform != SC.CookSourcePlatform && !Params.IgnorePaksFromDifferentCookSource)
+				if (!PakParams.bOnDemand && SC.StageTargetPlatform != SC.CookSourcePlatform && !Params.IgnorePaksFromDifferentCookSource)
 				{
 					// Check to see if we have an existing pak file we can use
 
@@ -3185,6 +3253,7 @@ namespace AutomationScripts
 								OutputLocation.GetFileNameWithoutAnyExtensions(),
 								OutputLocation,
 								SC.OptionalFileStageDirectory,
+								PakParams.bOnDemand ? OnDemandOutputLocation : null,
 								bCompressContainers,
 								CryptoSettings,
 								PakParams.EncryptionKeyGuid,
@@ -3715,6 +3784,7 @@ namespace AutomationScripts
 				ResponseFile = new ConcurrentDictionary<string, string>();
 				Manifest = null;
 				bCompressed = false;
+				bOnDemand = false;
 			}
 
 			// Name of pak file without extension, ie pakchunk0
@@ -3730,6 +3800,9 @@ namespace AutomationScripts
 			public bool bCompressed;
 			public string EncryptionKeyGuid;
 			public string RequestedEncryptionKeyGuid;
+
+			// Whether the chunk is used for content-on-demand content delivery or not
+			public bool bOnDemand;
 		}
 
 		/// <summary>
@@ -3829,19 +3902,41 @@ namespace AutomationScripts
 					}
 				}
 
+				// Add ondemand chunk definitions
+				if (PakRulesList != null)
+				{
+					foreach (PakFileRules Rule in PakRulesList.Where(R => R.bOnDemand))
+					{
+						string ChunkName = Rule.Name;
+						if (Rule.OverridePaks != null && Rule.OverridePaks.Count > 0)
+						{
+							ChunkName = Rule.OverridePaks[0];
+						}
+
+						if (ChunkDefinitions.Any(D => D.ChunkName == ChunkName))
+						{
+							continue;
+						}
+
+						Log.TraceInformation("Creating content-on-demand chunk definition '{0}' for PAK rule '{1}'", ChunkName, Rule.Name);
+
+						//TODO: Add encryption/compression settings
+						ChunkDefinition Chunk = new ChunkDefinition(ChunkName);
+						Chunk.bOnDemand = true;
+						ChunkDefinitions.Add(Chunk);
+					}
+				}
+
 				ConcurrentDictionary<string, ChunkDefinition> ChunkNameToDefinition = new ConcurrentDictionary<string, ChunkDefinition>(StringComparer.InvariantCultureIgnoreCase);
 				foreach (ChunkDefinition Chunk in ChunkDefinitions)
 				{
 					ChunkNameToDefinition.TryAdd(Chunk.ChunkName, Chunk);
 				}
 
-				ConcurrentDictionary<StagedFileReference, FileReference> LooseFiles = new ConcurrentDictionary<StagedFileReference, FileReference>();
-
 				Parallel.ForEach(StagingManifestResponseFile, StagingFile =>
 				{
 					bool bAddedToChunk = false;
 					bool bExcludeFromPaks = false;
-					bool bStageLoose = false;
 					HashSet<ChunkDefinition> PakList = new HashSet<ChunkDefinition>();
 
 					string OriginalFilename = StagingFile.Key;
@@ -3875,18 +3970,7 @@ namespace AutomationScripts
 					}
 
 					// Now run through the pak rules which may override things
-					ApplyPakFileRules(PakRulesList, StagingFile, PakList, ChunkNameToDefinition, out bExcludeFromPaks, out bStageLoose);
-
-					if (bStageLoose)
-					{
-						string StagingFileValue = StagingFile.Value;
-						if (StagingFileValue.StartsWith(SC.PakFileInternalRoot))
-						{
-							StagingFileValue = StagingFileValue.Substring(SC.PakFileInternalRoot.Length);
-						}
-						LooseFiles.TryAdd(new StagedFileReference(StagingFileValue), new FileReference(StagingFile.Key));
-						return;
-					}
+					ApplyPakFileRules(PakRulesList, StagingFile, PakList, ChunkNameToDefinition, out bExcludeFromPaks);
 
 					if (bExcludeFromPaks)
 					{
@@ -3898,19 +3982,23 @@ namespace AutomationScripts
 					{
 						ChunkDefinition TargetChunk = Chunk;
 
-						string OrigExt = Path.GetExtension(OriginalFilename);
-						if (OrigExt.Equals(OptionalBulkDataFileExtension))
+						// Don't separate out any bulk data for ondemand chunks
+						if (Chunk.bOnDemand == false)
 						{
-							// any optional files encountered we want to put in a separate pak file
-							string OptionalChunkName = Chunk.ChunkName + "optional";
-							if (!OptionalChunks.ContainsKey(OptionalChunkName))
+							string OrigExt = Path.GetExtension(OriginalFilename);
+							if (OrigExt.Equals(OptionalBulkDataFileExtension))
 							{
-								ChunkDefinition OptionalChunk = new ChunkDefinition(OptionalChunkName);
-								OptionalChunk.RequestedEncryptionKeyGuid = Chunk.RequestedEncryptionKeyGuid;
-								OptionalChunk.EncryptionKeyGuid = Chunk.EncryptionKeyGuid;
-								OptionalChunks.TryAdd(OptionalChunkName, OptionalChunk);
+								// any optional files encountered we want to put in a separate pak file
+								string OptionalChunkName = Chunk.ChunkName + "optional";
+								if (!OptionalChunks.ContainsKey(OptionalChunkName))
+								{
+									ChunkDefinition OptionalChunk = new ChunkDefinition(OptionalChunkName);
+									OptionalChunk.RequestedEncryptionKeyGuid = Chunk.RequestedEncryptionKeyGuid;
+									OptionalChunk.EncryptionKeyGuid = Chunk.EncryptionKeyGuid;
+									OptionalChunks.TryAdd(OptionalChunkName, OptionalChunk);
+								}
+								TargetChunk = OptionalChunks[OptionalChunkName];
 							}
-							TargetChunk = OptionalChunks[OptionalChunkName];
 						}
 
 						TargetChunk.ResponseFile.TryAdd(StagingFile.Key, StagingFile.Value);
@@ -3934,11 +4022,6 @@ namespace AutomationScripts
 				{
 					ChunkDefinitions.Add(OptionalChunkIt.Value);
 				}
-
-				foreach (var LooseFile in LooseFiles)
-				{
-					SC.FilesToStage.LooseFiles.Add(LooseFile.Key, LooseFile.Value);
-				}
 			}
 
 			// chunk downloader pak file is a minimal pak file which contains no content.  It can be provided with as a minimal download so that the game can download all the content from another source.
@@ -4016,7 +4099,7 @@ namespace AutomationScripts
 					string EncryptionKeyToUse = Params.SkipEncryption ? "" : Chunk.EncryptionKeyGuid;
 					PakInputs.Add(new CreatePakParams(Chunk.ChunkName,
 						Chunk.ResponseFile.ToDictionary(entry => entry.Key, entry => entry.Value),
-						Params.Compressed || Chunk.bCompressed, Params.RehydrateAssets, EncryptionKeyToUse));
+						Params.Compressed || Chunk.bCompressed, Params.RehydrateAssets, EncryptionKeyToUse, Chunk.bOnDemand));
 				}
 			}
 
@@ -4189,15 +4272,7 @@ namespace AutomationScripts
 
 			foreach (System.IO.FileInfo File in StagingDirectory.EnumerateFiles())
 			{
-				if (File.Extension != ".uptnl")
-				{
-					try { File.Delete(); } catch { }
-					LogInformation("Deleted {0}", File.FullName);
-				}
-				else
-				{
-					LogInformation("Didn't delete {0}", File.FullName);
-				}
+				try { File.Delete(); } catch { }
 			}
 		}
 
@@ -4293,7 +4368,6 @@ namespace AutomationScripts
 
 			if (!SC.Stage || Params.SkipStage)
 			{
-				CopyManifestFilesToStageDir(SC, SC.FilesToStage.LooseFiles, "LooseFiles");
 				return;
 			}
 
diff --git a/Engine/Source/Runtime/Engine/Private/SkeletalMeshLODRenderData.cpp b/Engine/Source/Runtime/Engine/Private/SkeletalMeshLODRenderData.cpp
--- a/Engine/Source/Runtime/Engine/Private/SkeletalMeshLODRenderData.cpp
+++ b/Engine/Source/Runtime/Engine/Private/SkeletalMeshLODRenderData.cpp
@@ -33,12 +33,7 @@ static FAutoConsoleVariableRef CVarStripSkeletalMeshLodsBelowMinLod(
 	TEXT("If set will strip skeletal mesh LODs under the minimum renderable LOD for the target platform during cooking.")
 );
 
-int32 GAllOptionalLODs = 0;
-static FAutoConsoleVariableRef CVarAllOptionalLODs(
-	TEXT("r.Skinned.AllOptionalLODs"),
-	GAllOptionalLODs,
-	TEXT("All LODs are marked as optional")
-);
+
 
 int32 GAllowSkinnedMorphDataStreaming = 0;
 static FAutoConsoleVariableRef CVarAllowSkinnedMorphDataStreaming(
@@ -937,7 +932,7 @@ void FSkeletalMeshLODRenderData::Serialize(FArchive& Ar, UObject* Owner, int32 I
 					SerializeStreamedData(MemWriter, OwnerMesh, Idx, ClassDataStripFlags, bNeedsCPUAccess, bForceKeepCPUResources);
 				}
 
-				bIsLODOptional = GAllOptionalLODs != 0 ? true : bIsBelowMinLOD;
+				bIsLODOptional = bIsBelowMinLOD;
 				const uint32 BulkDataFlags = (bDiscardBulkData ? 0 : BULKDATA_Force_NOT_InlinePayload)
 					| (bIsLODOptional ? BULKDATA_OptionalPayload : 0);
 				const uint32 OldBulkDataFlags = BulkData.GetBulkDataFlags();
diff --git a/Engine/Source/Runtime/Experimental/StreamingFileSystem/Private/StreamingFileSystem.cpp b/Engine/Source/Runtime/Experimental/StreamingFileSystem/Private/StreamingFileSystem.cpp
deleted file mode 100644
diff --git a/Engine/Source/Runtime/Experimental/StreamingFileSystem/Public/StreamingFileSystem.h b/Engine/Source/Runtime/Experimental/StreamingFileSystem/Public/StreamingFileSystem.h
deleted file mode 100644
diff --git a/Engine/Source/Runtime/Experimental/StreamingFileSystem/StreamingFileSystem.Build.cs b/Engine/Source/Runtime/Experimental/StreamingFileSystem/StreamingFileSystem.Build.cs
deleted file mode 100644
diff --git a/Engine/Source/Runtime/InstallBundleManager/InstallBundleManager.Build.cs b/Engine/Source/Runtime/InstallBundleManager/InstallBundleManager.Build.cs
--- a/Engine/Source/Runtime/InstallBundleManager/InstallBundleManager.Build.cs
+++ b/Engine/Source/Runtime/InstallBundleManager/InstallBundleManager.Build.cs
@@ -10,7 +10,7 @@ public class InstallBundleManager : ModuleRules
 			new string[] {
 				"Core",
 				"CoreUObject",
-				//"Engine",
+				"Engine",
 				"InputCore",
 				"Projects",
 				"ApplicationCore",
diff --git a/Engine/Source/Runtime/InstallBundleManager/Private/InstallBundleManagerInterface.cpp b/Engine/Source/Runtime/InstallBundleManager/Private/InstallBundleManagerInterface.cpp
--- a/Engine/Source/Runtime/InstallBundleManager/Private/InstallBundleManagerInterface.cpp
+++ b/Engine/Source/Runtime/InstallBundleManager/Private/InstallBundleManagerInterface.cpp
@@ -48,12 +48,6 @@ TSharedPtr<IInstallBundleManager> IInstallBundleManager::GetPlatformInstallBundl
 	return {};
 }
 
-const TSharedPtr<IInstallBundleSource> IInstallBundleManager::GetBundleSource(EInstallBundleSourceType SourceType) const
-{
-	return {};
-}
-
-
 TValueOrError<FInstallBundleRequestInfo, EInstallBundleResult> IInstallBundleManager::RequestUpdateContent(FName BundleName, EInstallBundleRequestFlags Flags, ELogVerbosity::Type LogVerbosityOverride /*= ELogVerbosity::NoLogging*/)
 {
 	return RequestUpdateContent(MakeArrayView(&BundleName, 1), Flags, LogVerbosityOverride);
diff --git a/Engine/Source/Runtime/InstallBundleManager/Public/InstallBundleManagerInterface.h b/Engine/Source/Runtime/InstallBundleManager/Public/InstallBundleManagerInterface.h
--- a/Engine/Source/Runtime/InstallBundleManager/Public/InstallBundleManagerInterface.h
+++ b/Engine/Source/Runtime/InstallBundleManager/Public/InstallBundleManagerInterface.h
@@ -22,7 +22,6 @@
 #include "UObject/UnrealNames.h"
 
 class IAnalyticsProviderET;
-class IInstallBundleSource;
 
 struct FInstallBundleProgress
 {
@@ -103,9 +102,6 @@ public:
 
 	virtual bool HasBundleSource(EInstallBundleSourceType SourceType) const = 0;
 
-	virtual const TSharedPtr<IInstallBundleSource> GetBundleSource(EInstallBundleSourceType SourceType) const;
-
-
 	virtual FDelegateHandle PushInitErrorCallback(FInstallBundleManagerInitErrorHandler Callback) = 0;
 	virtual void PopInitErrorCallback() = 0;
 	virtual void PopInitErrorCallback(FDelegateHandle Handle) = 0;
diff --git a/Engine/Source/Runtime/InstallBundleManager/Public/InstallBundleSourceInterface.h b/Engine/Source/Runtime/InstallBundleManager/Public/InstallBundleSourceInterface.h
--- a/Engine/Source/Runtime/InstallBundleManager/Public/InstallBundleSourceInterface.h
+++ b/Engine/Source/Runtime/InstallBundleManager/Public/InstallBundleSourceInterface.h
@@ -8,10 +8,6 @@
 class IInstallBundleSource;
 class IAnalyticsProviderET;
 
-class IBuildManifest;
-
-typedef TSharedPtr< class IBuildManifest, ESPMode::ThreadSafe > IBuildManifestPtr;
-
 DECLARE_DELEGATE_TwoParams(FInstallBundleSourceInitDelegate, TSharedRef<IInstallBundleSource> /*Source*/, FInstallBundleSourceAsyncInitInfo /*InitInfo*/);
 
 DECLARE_DELEGATE_TwoParams(FInstallBundleSourceQueryBundleInfoDelegate, TSharedRef<IInstallBundleSource> /*Source*/, FInstallBundleSourceBundleInfoQueryResult /*Result*/);
@@ -32,10 +28,6 @@ public:
 	// Returns a unique id for this source
 	virtual EInstallBundleSourceType GetSourceType() const = 0;
 
-	virtual const TArray<FString> GetDistributionCDNPaths() const { return TArray<FString>(); }
-
-	virtual IBuildManifestPtr GetBuildManifest() { return {}; }
-
 	// Returns the how this source should be weighted when combined with other sources
 	virtual float GetSourceWeight() const { return 1.0f; }
 
diff --git a/Engine/Source/Runtime/PakFile/PakFile.Build.cs b/Engine/Source/Runtime/PakFile/PakFile.Build.cs
--- a/Engine/Source/Runtime/PakFile/PakFile.Build.cs
+++ b/Engine/Source/Runtime/PakFile/PakFile.Build.cs
@@ -10,10 +10,5 @@ public class PakFile : ModuleRules
 		PrivateDependencyModuleNames.Add("Core");
 		PrivateDependencyModuleNames.Add("CoreUObject");
 		PublicDependencyModuleNames.Add("RSA");
-
-		if (Target.Type != TargetType.Program)
-		{
-			PublicDependencyModuleNames.Add("StreamingFileSystem");
-		}
 	}
 }
diff --git a/Engine/Source/Runtime/PakFile/Private/IPlatformFilePak.cpp b/Engine/Source/Runtime/PakFile/Private/IPlatformFilePak.cpp
--- a/Engine/Source/Runtime/PakFile/Private/IPlatformFilePak.cpp
+++ b/Engine/Source/Runtime/PakFile/Private/IPlatformFilePak.cpp
@@ -42,9 +42,6 @@
 #include "FilePackageStore.h"
 #include "Compression/OodleDataCompression.h"
 #include "IO/IoStore.h"
-#if !IS_PROGRAM
-#include "StreamingFileSystem.h"
-#endif
 
 DEFINE_LOG_CATEGORY(LogPakFile);
 
@@ -7535,12 +7532,6 @@ bool FPakPlatformFile::Initialize(IPlatformFile* Inner, const TCHAR* CmdLine)
 		FIoDispatcher& IoDispatcher = FIoDispatcher::Get();
 		IoDispatcherFileBackend = CreateIoDispatcherFileBackend();
 		IoDispatcher.Mount(IoDispatcherFileBackend.ToSharedRef());
-
-#if !IS_PROGRAM && PLATFORM_USE_STREAMING_FILE_CACHE
-		IoDispatcherSFSBackend = IStreamingFileSystem::CreateStreamingFileSystem();
-		IoDispatcher.Mount(IoDispatcherSFSBackend.ToSharedRef());
-#endif
-
 		PackageStoreBackend = MakeShared<FFilePackageStoreBackend>();
 		FPackageStore::Get().Mount(PackageStoreBackend.ToSharedRef());
 
diff --git a/Engine/Source/Runtime/PakFile/Public/IPlatformFilePak.h b/Engine/Source/Runtime/PakFile/Public/IPlatformFilePak.h
--- a/Engine/Source/Runtime/PakFile/Public/IPlatformFilePak.h
+++ b/Engine/Source/Runtime/PakFile/Public/IPlatformFilePak.h
@@ -56,7 +56,6 @@ class FOutputDevice;
 class IAsyncReadFileHandle;
 class IMappedFileHandle;
 struct FIoContainerHeader;
-class IStreamingFileSystem;
 
 PAKFILE_API DECLARE_LOG_CATEGORY_EXTERN(LogPakFile, Log, All);
 DECLARE_FLOAT_ACCUMULATOR_STAT_EXTERN(TEXT("Total pak file read time"), STAT_PakFile_Read, STATGROUP_PakFile, PAKFILE_API);
@@ -2038,10 +2037,7 @@ class PAKFILE_API FPakPlatformFile : public IPlatformFile
 	FString IniFileExtension;
 	/** The filename for the gameusersettings ini file, used for excluding ini files, but not gameusersettings */
 	FString GameUserSettingsIniFilename;
-	
-	TSharedPtr<FFileIoStore>			IoDispatcherFileBackend;
-	TSharedPtr<IStreamingFileSystem>	IoDispatcherSFSBackend;
-
+	TSharedPtr<FFileIoStore> IoDispatcherFileBackend;
 	TSharedPtr<FFilePackageStoreBackend> PackageStoreBackend;
 
 	FTSTicker::FDelegateHandle RetireReadersHandle;
