AutomationTool: Set version file as writeable when P4 is not enabled

diff --git a/Engine/Source/Programs/AutomationTool/AutomationUtils/UnrealBuild.cs b/Engine/Source/Programs/AutomationTool/AutomationUtils/UnrealBuild.cs
--- a/Engine/Source/Programs/AutomationTool/AutomationUtils/UnrealBuild.cs
+++ b/Engine/Source/Programs/AutomationTool/AutomationUtils/UnrealBuild.cs
@@ -356,7 +356,7 @@ namespace AutomationTool
 
 					if (File.Exists(BuildVersionFile.FullName))
 					{
-						VersionFileUpdater.MakeFileWriteable(BuildVersionFile.FullName);
+						VersionFileUpdater.MakeFileWriteable(BuildVersionFile.FullName, true);
 					}
 
 					Version.Write(BuildVersionFile);
diff --git a/Engine/Source/Programs/AutomationTool/AutomationUtils/Utils.cs b/Engine/Source/Programs/AutomationTool/AutomationUtils/Utils.cs
--- a/Engine/Source/Programs/AutomationTool/AutomationUtils/Utils.cs
+++ b/Engine/Source/Programs/AutomationTool/AutomationUtils/Utils.cs
@@ -1253,7 +1253,8 @@ namespace AutomationTool
 		/// submit it to Perforce. For existing writeable files, P4V prompts the user to overwrite the next time they sync.
 		/// </summary>
 		/// <param name="FileName">The file to make writeable</param>
-		public static void MakeFileWriteable(string FileName)
+		/// <param name="bForce">Whether to force the writeable flag if P4 cannot clear it</param>
+		public static void MakeFileWriteable(string FileName, bool bForce = false)
 		{
 			if(CommandUtils.IsReadOnly(FileName))
 			{
@@ -1261,9 +1262,17 @@ namespace AutomationTool
 				{
 					CommandUtils.P4.Sync(String.Format("\"{0}#0\"", FileName), false, false);
 				}
+				
 				if(CommandUtils.FileExists_NoExceptions(FileName) && CommandUtils.IsReadOnly(FileName))
 				{
-					throw new AutomationException("Cannot write to {0}; file is read-only", FileName);
+					if(bForce)
+					{
+						CommandUtils.SetFileAttributes(FileName, ReadOnly: false);
+					}
+					else
+					{
+						throw new AutomationException("Cannot write to {0}; file is read-only", FileName);	
+					}
 				}
 			}
 		}
