Added support for passing new flag to BPT to generate full chunk history, ugraded content worker BPT to BPT_Online_v1.5.2_PF-24214439.zip
#tests local testing iteration on large project

diff --git a/Engine/Source/Programs/AutomationTool/AutomationUtils/MCPPublic.cs b/Engine/Source/Programs/AutomationTool/AutomationUtils/MCPPublic.cs
--- a/Engine/Source/Programs/AutomationTool/AutomationUtils/MCPPublic.cs
+++ b/Engine/Source/Programs/AutomationTool/AutomationUtils/MCPPublic.cs
@@ -693,6 +693,7 @@ namespace EpicGames.MCP.Automation
 			{
 				DataAgeThreshold = DEFAULT_DATA_AGE_THRESHOLD;
 				ChunkWindowSize = 1048576;
+				ResaveKnownChunks = false;
 			}
 
 			/// <summary>
@@ -774,6 +775,10 @@ namespace EpicGames.MCP.Automation
 			/// </summary>
 			public string FeatureLevel;
 			/// <summary>
+			/// Output includes all known chunks, not just new chunks generated
+			/// </summary>
+			public bool ResaveKnownChunks;
+			/// <summary>
 			/// Contains a list of custom string arguments to be embedded in the generated manifest file.
 			/// </summary>
 			public List<KeyValuePair<string, string>> CustomStringArgs;
