Gauntlet - Change log type from DeviceEvent to TestEvent when test raises a warning about resource starvation.

#rnx

diff --git a/Engine/Source/Programs/AutomationTool/Gauntlet/Framework/Gauntlet.TestExecutor.cs b/Engine/Source/Programs/AutomationTool/Gauntlet/Framework/Gauntlet.TestExecutor.cs
--- a/Engine/Source/Programs/AutomationTool/Gauntlet/Framework/Gauntlet.TestExecutor.cs
+++ b/Engine/Source/Programs/AutomationTool/Gauntlet/Framework/Gauntlet.TestExecutor.cs
@@ -271,7 +271,7 @@ namespace Gauntlet
 									double TimeWaiting = (DateTime.Now - NodeInfo.FirstReadyCheckTime).TotalSeconds;
 									if (TimeWaiting >= Options.Wait)
 									{
-										Log.Warning(KnownLogEvents.Gauntlet_DeviceEvent, "Test {Name} has been waiting to run resource-free for {Time:00} seconds. Removing from wait list", Node, TimeWaiting);
+										Log.Warning(KnownLogEvents.Gauntlet_TestEvent, "Test {Name} has been waiting to run resource-free for {Time:00} seconds. Removing from wait list", Node, TimeWaiting);
 										Node.AddTestEvent(new UnrealTestEvent(EventSeverity.Error, "Insufficient devices found", new List<string> {string.Format("Test {0} was unable to find enough devices after trying for {1:00} seconds.", Node, TimeWaiting), "This is not a test-related failure."}));
 										PendingTests[i] = null;
 										NodeInfo.TimeSetupBegan = NodeInfo.TimeSetupEnded = NodeInfo.TimeTestEnded = DateTime.Now;
