Adding the ability to override the oodle dll for packaging in ProjectParams, and use it

diff --git a/Engine/Source/Programs/AutomationTool/AutomationUtils/ProjectParams.cs b/Engine/Source/Programs/AutomationTool/AutomationUtils/ProjectParams.cs
--- a/Engine/Source/Programs/AutomationTool/AutomationUtils/ProjectParams.cs
+++ b/Engine/Source/Programs/AutomationTool/AutomationUtils/ProjectParams.cs
@@ -395,6 +395,7 @@ namespace AutomationTool
 			this.ForceUncompressed = InParams.ForceUncompressed;
 			this.AdditionalPakOptions = InParams.AdditionalPakOptions;
 			this.AdditionalIoStoreOptions = InParams.AdditionalIoStoreOptions;
+			this.ForceOodleDllVersion = InParams.ForceOodleDllVersion;
 			this.Archive = InParams.Archive;
 			this.ArchiveDirectoryParam = InParams.ArchiveDirectoryParam;
 			this.ArchiveMetaData = InParams.ArchiveMetaData;
@@ -464,6 +465,7 @@ namespace AutomationTool
 			bool? ForceUncompressed = null,
 			string AdditionalPakOptions = null,
 			string AdditionalIoStoreOptions = null,
+			string ForceOodleDllVersion = null,
             bool? IterativeCooking = null,
 			string IterateSharedCookedBuild = null,
 			bool? IterateSharedBuildUsePrecompiledExe = null,
@@ -772,6 +774,7 @@ namespace AutomationTool
 			this.ForceUncompressed = GetParamValueIfNotSpecified(Command, ForceUncompressed, this.ForceUncompressed, "ForceUncompressed");
 			this.AdditionalPakOptions = ParseParamValueIfNotSpecified(Command, AdditionalPakOptions, "AdditionalPakOptions");
 			this.AdditionalIoStoreOptions = ParseParamValueIfNotSpecified(Command, AdditionalIoStoreOptions, "AdditionalIoStoreOptions");
+			this.ForceOodleDllVersion = ParseParamValueIfNotSpecified(Command, ForceOodleDllVersion, "ForceOodleDllVersion");
 			this.IterativeCooking = GetParamValueIfNotSpecified(Command, IterativeCooking, this.IterativeCooking, new string[] { "iterativecooking", "iterate" });
 			this.IterateSharedCookedBuild = GetParamValueIfNotSpecified(Command, false, false, "iteratesharedcookedbuild") ? "usesyncedbuild" : null;
 			this.IterateSharedCookedBuild = ParseParamValueIfNotSpecified(Command, IterateSharedCookedBuild, "IterateSharedCookedBuild", String.Empty);
@@ -1872,6 +1875,11 @@ namespace AutomationTool
 		/// </summary>
 		public string AdditionalIoStoreOptions;
 
+		/// <summary>
+		/// If not empty, this is the dll file to use for Oodle compression. Can be "Latest" to use latest version.
+		/// </summary>
+		public string ForceOodleDllVersion;
+
         /// <summary>
         /// Cook: Do not include a version number in the cooked content
         /// </summary>
@@ -3092,6 +3100,7 @@ namespace AutomationTool
 				CommandUtils.LogLog("ForceUncompressed={0}", ForceUncompressed);
 				CommandUtils.LogLog("AdditionalPakOptions={0}", AdditionalPakOptions);
 				CommandUtils.LogLog("AdditionalIoStoreOptions={0}", AdditionalIoStoreOptions);
+				CommandUtils.LogLog("ForceOodleDllVersion={0}", ForceOodleDllVersion);
 				CommandUtils.LogLog("CookOnTheFly={0}", CookOnTheFly);
 				CommandUtils.LogLog("CookOnTheFlyStreaming={0}", CookOnTheFlyStreaming);
 				CommandUtils.LogLog("UnversionedCookedContent={0}", UnversionedCookedContent);
diff --git a/Engine/Source/Programs/AutomationTool/Scripts/CopyBuildToStagingDirectory.Automation.cs b/Engine/Source/Programs/AutomationTool/Scripts/CopyBuildToStagingDirectory.Automation.cs
--- a/Engine/Source/Programs/AutomationTool/Scripts/CopyBuildToStagingDirectory.Automation.cs
+++ b/Engine/Source/Programs/AutomationTool/Scripts/CopyBuildToStagingDirectory.Automation.cs
@@ -2798,6 +2798,17 @@ namespace AutomationScripts
 
 				string OodleCompressDLL;
 				BuildPlatformEngineConfig.GetString("OodleDataCompressionFormat", "OodleCompressDLL", out OodleCompressDLL);
+				if (!string.IsNullOrWhiteSpace(Params.ForceOodleDllVersion))
+				{
+					if (string.Equals(Params.ForceOodleDllVersion, "latest", StringComparison.OrdinalIgnoreCase))
+					{
+						OodleCompressDLL = "";
+					}
+					else
+					{
+						OodleCompressDLL = Params.ForceOodleDllVersion;
+					}
+				}
 				if (!string.IsNullOrWhiteSpace(OodleCompressDLL))
 				{
 					CompressionFormats += " -OodleCompressDLL=" + OodleCompressDLL;
