UnrealBuild: Workaround for iOS & tvOS compile configs needing to be build separately

#rnx

diff --git a/Engine/Source/Programs/AutomationTool/AutomationUtils/UnrealBuild.cs b/Engine/Source/Programs/AutomationTool/AutomationUtils/UnrealBuild.cs
--- a/Engine/Source/Programs/AutomationTool/AutomationUtils/UnrealBuild.cs
+++ b/Engine/Source/Programs/AutomationTool/AutomationUtils/UnrealBuild.cs
@@ -158,7 +158,7 @@ namespace AutomationTool
 				List<string> Arguments = new List<string>();
 				for (int Idx = 0; Idx < Targets.Count; Idx++)
 				{
-					ManifestFiles.Add(GetManifestFile(Targets[0], Idx));
+					ManifestFiles.Add(GetManifestFile(Targets[Idx], Idx));
 					Arguments.Add("-Target=" + GetTargetArguments(Targets[Idx], ManifestFiles[Idx]));
 				}
 				FullCommandLine.Append(CommandLine.FormatCommandLine(Arguments));
@@ -593,8 +593,36 @@ namespace AutomationTool
 				}
 			}
 
+			List<BuildTarget> Targets = new List<BuildTarget>(Agenda.Targets);
+
+			// Temporary hack: iOS & tvOS configs need to build separately
+			if (Targets.Any(x => x.Platform == UnrealTargetPlatform.IOS || x.Platform == UnrealTargetPlatform.TVOS))
+			{
+				List<string> TargetNames = Targets.Select(x => x.TargetName).Distinct().ToList();
+				List<UnrealTargetConfiguration> Configs = Targets.Select(x => x.Config).Distinct().ToList();
+				foreach (string TargetName in TargetNames)
+				{
+					foreach (UnrealTargetConfiguration Config in Configs)
+					{
+						List<BuildTarget> ConfigTargets = Targets.Where(x => (x.Platform == UnrealTargetPlatform.IOS || x.Platform == UnrealTargetPlatform.TVOS) && x.TargetName == TargetName && x.Config == Config).ToList();
+						if (ConfigTargets.Count > 0)
+						{
+							// Build all the targets
+							BuildWithUBT(ConfigTargets, InTargetToManifest, bDisableXGE, InAllCores);
+						}
+					}
+				}
+				Targets.RemoveAll(x => x.Platform == UnrealTargetPlatform.IOS || x.Platform == UnrealTargetPlatform.TVOS);
+			}
+			// End hack
+
+			if (Targets.Count == 0)
+			{
+				return;
+			}
+
 			// Build all the targets
-			BuildWithUBT(Agenda.Targets, InTargetToManifest, bDisableXGE, InAllCores);
+			BuildWithUBT(Targets, InTargetToManifest, bDisableXGE, InAllCores);
 		}
 
 		/// <summary>
