IPA used as a donor for prepare for debug feature error handling

diff --git a/Engine/Source/Editor/TurnkeySupport/Private/TurnkeySupportModule.cpp b/Engine/Source/Editor/TurnkeySupport/Private/TurnkeySupportModule.cpp
--- a/Engine/Source/Editor/TurnkeySupport/Private/TurnkeySupportModule.cpp
+++ b/Engine/Source/Editor/TurnkeySupport/Private/TurnkeySupportModule.cpp
@@ -517,7 +517,7 @@ public:
              if (IniPlatformName == TEXT("IOS") || IniPlatformName == TEXT("TvOS"))
              {
 				 bool bSupportSecondaryMac = false;
-				 bSupportSecondaryMac = GConfig->GetBool(TEXT("/Script/IOSRuntimeSettings.IOSRuntimeSettings"), TEXT("bSupportSecondaryMac"), bSupportSecondaryMac, GEngineIni);
+				 GConfig->GetBool(TEXT("/Script/IOSRuntimeSettings.IOSRuntimeSettings"), TEXT("bSupportSecondaryMac"), bSupportSecondaryMac, GEngineIni);
 				 if (bSupportSecondaryMac)
 				 {
 					 FString OtherCommandLine = FString::Printf(TEXT("SetSecondaryRemoteMac -platform=%s -ProjectFilePath=%s"), *IniPlatformName.ToString().ToLower(), *ProjectPath);
diff --git a/Engine/Source/Programs/AutomationTool/IOS/IOSPlatform.Automation.cs b/Engine/Source/Programs/AutomationTool/IOS/IOSPlatform.Automation.cs
--- a/Engine/Source/Programs/AutomationTool/IOS/IOSPlatform.Automation.cs
+++ b/Engine/Source/Programs/AutomationTool/IOS/IOSPlatform.Automation.cs
@@ -2307,11 +2307,32 @@ public class IOSPlatform : ApplePlatform
 			int StartPos = ProjectFilePath.LastIndexOf("/");
 			int StringLength = ProjectFilePath.Length - 10; // 9 for .uproject, 1 for the /
 			string PackageName = ProjectFilePath.Substring(StartPos + 1, StringLength - StartPos);
+			string PackagePath = Path.Combine(Path.GetDirectoryName(ProjectFilePath), "Binaries", ClientPlatform);
 			if (string.IsNullOrEmpty(SourcePackage))
 			{
-				SourcePackage = Path.Combine(Path.GetDirectoryName(ProjectFilePath), "Binaries", ClientPlatform, PackageName + ".ipa");
+				SourcePackage = Path.Combine(PackagePath, PackageName + ".ipa");
 			}
+			
+			string[] IPAFiles;
+			if (!File.Exists(SourcePackage))
+            {
+				IPAFiles = Directory.GetFiles(PackagePath, "*.ipa");
+				LogWarning("Source package not found : {0}, trying to find another IPA file in the same folder.", SourcePackage);
+				if (IPAFiles.Length == 0)
+                {
+					LogError("No IPA file found in : {0}. Aborting.", PackagePath);
+					throw new AutomationException(ExitCode.Error_MissingExecutable, "No IPA file found in {0}.", PackagePath);
+				}
+				else
+                {
+					if (IPAFiles.Length > 1)
+                    {
+						LogWarning("More than one IPA file found. Taking the first one found, {0}", IPAFiles[0]);
+					}
+					SourcePackage = IPAFiles[0];
+				}					
 
+			}
 			string ZipFile = Path.ChangeExtension(SourcePackage, "zip");
 
 			string PayloadPath = SourcePackage;
@@ -2356,8 +2377,6 @@ public class IOSPlatform : ApplePlatform
 			LogInformation("Deleting temp files ...");
 			File.Delete(ZipFile);
 			LogInformation("{0} deleted", ZipFile);
-			Directory.Delete(PayloadPath, true);
-			LogInformation("{0} deleted", PayloadPath);
 		}
 		else
 		{
