- Fixed scheme name for the codesign call of Xcode

diff --git a/Engine/Source/Programs/AutomationTool/IOS/IOSPlatform.Automation.cs b/Engine/Source/Programs/AutomationTool/IOS/IOSPlatform.Automation.cs
--- a/Engine/Source/Programs/AutomationTool/IOS/IOSPlatform.Automation.cs
+++ b/Engine/Source/Programs/AutomationTool/IOS/IOSPlatform.Automation.cs
@@ -871,12 +871,8 @@ public class IOSPlatform : ApplePlatform
 		Params.AutomaticSigning = bAutomaticSigning;
 
 		// Scheme name and configuration for code signing with Xcode project
-		string SchemeName = Params.IsCodeBasedProject ? Params.RawProjectPath.GetFileNameWithoutExtension() : "UE5";
+		string SchemeName = SC.StageTargets[0].Receipt.TargetName;
 		string SchemeConfiguration = TargetConfiguration.ToString();
-		if (Params.Client)
-		{
-			SchemeConfiguration += " Client";
-		}
 
 		WriteEntitlements(Params, SC);
 
