- Fixed FinalizeInstalledBuild for multi-architecture

diff --git a/Engine/Source/Programs/AutomationTool/Scripts/FinalizeInstalledBuild.cs b/Engine/Source/Programs/AutomationTool/Scripts/FinalizeInstalledBuild.cs
--- a/Engine/Source/Programs/AutomationTool/Scripts/FinalizeInstalledBuild.cs
+++ b/Engine/Source/Programs/AutomationTool/Scripts/FinalizeInstalledBuild.cs
@@ -67,32 +67,18 @@ namespace AutomationTool
 			// Create list of platform configurations installed in a Rocket build
 			List<InstalledPlatformInfo.InstalledPlatformConfiguration> InstalledConfigs = new List<InstalledPlatformInfo.InstalledPlatformConfiguration>();
 
-			// Add the editor platform/architecture for unspecified target, otherwise we'll never be able to run UAT. This currently expects this to be a single architecture for now!
-			// @todo handle multiple architectures, and pass in the target name (and maybe project file?)
-			UnrealArch EditorArchitecture = UnrealArchitectureConfig.ForPlatform(HostPlatform.Current.HostEditorPlatform).ActiveArchitectures(null, null).SingleArchitecture;
-			InstalledConfigs.Add(new InstalledPlatformInfo.InstalledPlatformConfiguration(UnrealTargetConfiguration.Development, HostPlatform.Current.HostEditorPlatform, TargetRules.TargetType.Editor, EditorArchitecture, "", EProjectType.Unknown, false));
-			InstalledConfigs.Add(new InstalledPlatformInfo.InstalledPlatformConfiguration(UnrealTargetConfiguration.DebugGame, HostPlatform.Current.HostEditorPlatform, TargetRules.TargetType.Editor, EditorArchitecture, "", EProjectType.Unknown, false));
+			// Add the editor platform/architecture(s) for unspecified target, otherwise we'll never be able to run UAT
+			UnrealArchitectures EditorArchitectures = GetArchitecturesForPlatform(HostPlatform.Current.HostEditorPlatform);
+			foreach (UnrealArch EditorArchitecture in EditorArchitectures.Architectures)
+			{
+				InstalledConfigs.Add(new InstalledPlatformInfo.InstalledPlatformConfiguration(UnrealTargetConfiguration.Development, HostPlatform.Current.HostEditorPlatform, TargetRules.TargetType.Editor, EditorArchitecture, "", EProjectType.Unknown, false));
+				InstalledConfigs.Add(new InstalledPlatformInfo.InstalledPlatformConfiguration(UnrealTargetConfiguration.DebugGame, HostPlatform.Current.HostEditorPlatform, TargetRules.TargetType.Editor, EditorArchitecture, "", EProjectType.Unknown, false));
+			}
 
 			foreach (UnrealTargetPlatform CodeTargetPlatform in Platforms)
 			{
-				UnrealArchitectureConfig PlatformArchConfig = UnrealArchitectureConfig.ForPlatform(CodeTargetPlatform);
-
-				// Try to parse additional Architectures from the command line
-				string Architectures = ParseParamValue(CodeTargetPlatform.ToString() + "Architectures");
-
 				// Build a list of pre-compiled architecture combinations for this platform if any
-				UnrealArchitectures AllArchitectures;
-
-				if (!String.IsNullOrWhiteSpace(Architectures))
-				{
-					AllArchitectures = UnrealArchitectures.FromString(Architectures, CodeTargetPlatform);
-				}
-				// if there aren't any, use the default
-				else
-				{
-					// get the architectures for the target platform, for unspecified target
-					AllArchitectures = PlatformArchConfig.ActiveArchitectures(null, null);
-				}
+				UnrealArchitectures AllArchitectures = GetArchitecturesForPlatform(CodeTargetPlatform);
 
 				// Check whether this platform should only be used for content based projects
 				EProjectType ProjectType = ContentOnlyPlatforms.Contains(CodeTargetPlatform) ? EProjectType.Content : EProjectType.Any;
@@ -115,7 +101,7 @@ namespace AutomationTool
 						UnrealTargetConfiguration EngineConfiguration = (CodeTargetConfiguration == UnrealTargetConfiguration.DebugGame) ? UnrealTargetConfiguration.Development : CodeTargetConfiguration;
 
 						// if the platform doesn't split up multi-arch into multiple targets (which equate to receipt files) then make multiple configs from one receipt
-						if (PlatformArchConfig.Mode != UnrealArchitectureMode.OneTargetPerArchitecture)
+						if (UnrealArchitectureConfig.ForPlatform(CodeTargetPlatform).Mode != UnrealArchitectureMode.OneTargetPerArchitecture)
 						{
 							FileReference ReceiptFileName = TargetReceipt.GetDefaultPath(new DirectoryReference(OutputEnginePath), CurrentTargetName, CodeTargetPlatform, EngineConfiguration, AllArchitectures);
 							if (FileReference.Exists(ReceiptFileName))
@@ -180,5 +166,29 @@ namespace AutomationTool
 				return PlatformsString.Split(new char[] { ';' }, StringSplitOptions.RemoveEmptyEntries).Select(x => UnrealTargetPlatform.Parse(x)).ToList();
 			}
 		}
+
+		/// <summary>
+		/// Look for a platform's architecture(s) on the commandline, and if not, then ask the platform's config for active architectures
+		/// </summary>
+		/// <param name="Platform">Platform to look for architectures</param>
+		/// <returns>One or more architectures</returns>
+		private UnrealArchitectures GetArchitecturesForPlatform(UnrealTargetPlatform Platform)
+		{
+			// Try to parse additional Architectures from the command line
+			string Architectures = ParseParamValue(Platform.ToString() + "Architectures");
+
+			// Wwas a platform specified on commandline?
+			if (!string.IsNullOrWhiteSpace(Architectures))
+			{
+				return UnrealArchitectures.FromString(Architectures, Platform);
+			}
+			// if there aren't any, use the default
+			else
+			{
+				// get the architectures for the target platform, for unspecified target
+				UnrealArchitectureConfig PlatformArchConfig = UnrealArchitectureConfig.ForPlatform(Platform);
+				return PlatformArchConfig.ActiveArchitectures(null, null);
+			}
+		}
 	}
 }
\ No newline at end of file
