Automation Tool
--Fix directory renames not failing properly

#localization none

diff --git a/Engine/Source/Programs/AutomationTool/AutomationUtils/CommandUtils.cs b/Engine/Source/Programs/AutomationTool/AutomationUtils/CommandUtils.cs
--- a/Engine/Source/Programs/AutomationTool/AutomationUtils/CommandUtils.cs
+++ b/Engine/Source/Programs/AutomationTool/AutomationUtils/CommandUtils.cs
@@ -657,11 +657,11 @@ namespace AutomationTool
 		/// </summary>
 		/// <param name="OldName">Old name</param>
 		/// <param name="NewName">new name</param>
-		public static void RenameDirectory(string OldName, string NewName, bool bQuiet = false)
+		public static void RenameDirectory(string OldName, string NewName, bool bQuiet = false, bool bRetry = true)
 		{
 			var OldNormalized = ConvertSeparators(PathSeparator.Default, OldName);
 			var NewNormalized = ConvertSeparators(PathSeparator.Default, NewName);
-			InternalUtils.SafeRenameDirectory(OldNormalized, NewNormalized, bQuiet);
+			InternalUtils.SafeRenameDirectory(OldNormalized, NewNormalized, bQuiet, bRetry, /*bThrow*/true);
 		}
 
 		/// <summary>
@@ -670,18 +670,11 @@ namespace AutomationTool
 		/// </summary>
 		/// <param name="OldName">Old name</param>
 		/// <param name="NewName">new name</param>
-		public static bool RenameDirectory_NoExceptions(string OldName, string NewName)
+		public static bool RenameDirectory_NoExceptions(string OldName, string NewName, bool bQuiet = false, bool bRetry = true)
 		{
-			try
-			{
-				RenameDirectory(OldName, NewName);
-			}
-			catch (Exception)
-			{
-				LogWarning("Failed to rename/move file '{0}' to '{1}'", OldName, NewName);
-				return false;
-			}
-			return true;
+			var OldNormalized = ConvertSeparators(PathSeparator.Default, OldName);
+			var NewNormalized = ConvertSeparators(PathSeparator.Default, NewName);
+			return InternalUtils.SafeRenameDirectory(OldNormalized, NewNormalized, bQuiet, bRetry, /*bThrow*/false);
 		}
 
 		/// <summary>
diff --git a/Engine/Source/Programs/AutomationTool/AutomationUtils/Utils.cs b/Engine/Source/Programs/AutomationTool/AutomationUtils/Utils.cs
--- a/Engine/Source/Programs/AutomationTool/AutomationUtils/Utils.cs
+++ b/Engine/Source/Programs/AutomationTool/AutomationUtils/Utils.cs
@@ -266,7 +266,7 @@ namespace AutomationTool
 		/// <param name="NewName"></param>
 		/// <param name="bQuiet"></param>
 		/// <returns>True if the directory was moved, false otehrwise</returns>
-		public static bool SafeRenameDirectory(string OldName, string NewName, bool bQuiet = false)
+		public static bool SafeRenameDirectory(string OldName, string NewName, bool bQuiet = false, bool bRetry = true, bool bThrow = false)
 		{
 			if (!bQuiet)
 			{
@@ -287,13 +287,27 @@ namespace AutomationTool
 				{
 					if (Directory.Exists(OldName) == true || Directory.Exists(NewName) == false)
 					{
-						Log.TraceWarning("Failed to rename {0} to {1}", OldName, NewName);
-						Log.TraceWarning(LogUtils.FormatException(Ex));
+						if (!bQuiet)
+						{
+							Log.TraceWarning("Failed to rename {0} to {1}", OldName, NewName);
+							Log.TraceWarning(LogUtils.FormatException(Ex));
+						}
 						Result = false;
 					}
+
+					++Attempts;
+					if (Attempts == MaxAttempts)
+					{
+						if (bThrow)
+						{
+							throw;
+						}
+
+						break;
+					}
 				}
 			}
-			while (Result == false && ++Attempts < MaxAttempts);
+			while (Result == false && bRetry);
 
 			return Result;
 		}
