UAT: Optimize CreatePakUsingStagingManifest by using ParallelFor when processing PakFileRules
From ~300ms to ~30ms for ~17 000 cooked files.

#rnx

diff --git a/Engine/Source/Programs/AutomationTool/Scripts/CopyBuildToStagingDirectory.Automation.cs b/Engine/Source/Programs/AutomationTool/Scripts/CopyBuildToStagingDirectory.Automation.cs
--- a/Engine/Source/Programs/AutomationTool/Scripts/CopyBuildToStagingDirectory.Automation.cs
+++ b/Engine/Source/Programs/AutomationTool/Scripts/CopyBuildToStagingDirectory.Automation.cs
@@ -2378,10 +2378,10 @@ namespace AutomationScripts
 
 			List<PakFileRules> PakRulesList = GetPakFileRules(Params, SC);
 
-			List<string> FilesToRemove = new List<string>();
+			ConcurrentBag<string> FilesToRemove = new ConcurrentBag<string>();
 
 			// Apply the pak file rules, this can remove things but will not override the pak file name
-			foreach (var StagingFile in UnrealPakResponseFile)
+			Parallel.ForEach(UnrealPakResponseFile, StagingFile =>
 			{
 				bool bExcludeFromPaks = false;
 				ApplyPakFileRules(PakRulesList, StagingFile, null, null, out bExcludeFromPaks);
@@ -2390,9 +2390,9 @@ namespace AutomationScripts
 				{
 					FilesToRemove.Add(StagingFile.Key);
 				}
-			}
+			});
 
-			foreach (var FileToRemove in FilesToRemove)
+			foreach (var FileToRemove in FilesToRemove.ToArray())
 			{
 				UnrealPakResponseFile.Remove(FileToRemove);
 			}
