Adding a force-close for ADB to ensure we don't leave processes running and unintentionally clog up system resources

diff --git a/Engine/Source/Programs/AutomationTool/Gauntlet/Platform/Android/Gauntlet.TargetDeviceAndroid.cs b/Engine/Source/Programs/AutomationTool/Gauntlet/Platform/Android/Gauntlet.TargetDeviceAndroid.cs
--- a/Engine/Source/Programs/AutomationTool/Gauntlet/Platform/Android/Gauntlet.TargetDeviceAndroid.cs
+++ b/Engine/Source/Programs/AutomationTool/Gauntlet/Platform/Android/Gauntlet.TargetDeviceAndroid.cs
@@ -52,13 +52,9 @@ namespace Gauntlet
 		{
 			get
 			{
-				if (bHaveSavedArtifacts == false)
+				if (HasExited)
 				{
-					if (HasExited)
-					{
-						SaveArtifacts();
-						bHaveSavedArtifacts = true;
-					}
+					SaveArtifacts();
 				}
 				
 				return Path.Combine(AndroidDevice.LocalCachePath, "Saved");
@@ -215,6 +211,8 @@ namespace Gauntlet
 
 		protected void SaveArtifacts()
 		{
+			if (bHaveSavedArtifacts)
+				return;
 
 			// copy remote artifacts to local
 			if (Directory.Exists(Install.AndroidDevice.LocalCachePath))
@@ -279,6 +277,8 @@ namespace Gauntlet
 			File.WriteAllText(Path.Combine(LocalSaved, LogcatFilename), LogcatResult.Output);
 
 			Install.AndroidDevice.PostRunCleanup();
+
+			bHaveSavedArtifacts = true;
 		}
 	}
 
@@ -573,7 +573,7 @@ namespace Gauntlet
 				finally
 				{
 					disposedValue = true;
-					AdbCredentialCache.RemoveInstance();					
+					AdbCredentialCache.RemoveInstance();
 				}
 
 			}
@@ -1396,6 +1396,12 @@ namespace Gauntlet
 				{
 					Reset();
 					KillAdbServer();
+					// Kill ADB server, just as a safety measure to ensure it closes
+					IProcessResult TaskkillResult = CommandUtils.Run("taskkill", "/f /im adb.exe", null, CommandUtils.ERunOptions.NoLoggingOfRunCommand);
+					if (TaskkillResult.Output.Contains("success", StringComparison.OrdinalIgnoreCase))
+					{
+						Log.Info(TaskkillResult.Output);
+					}
 				}
 			}
 		}
