Add Force=true option to [CookedEditorSettings] EngineExtraStageFiles / ProjectExtraStageFiles to force staging of a file that would otherwise be rejected by the Restricted folder filter.

diff --git a/Engine/Source/Programs/AutomationTool/AutomationUtils/DeploymentContext.cs b/Engine/Source/Programs/AutomationTool/AutomationUtils/DeploymentContext.cs
--- a/Engine/Source/Programs/AutomationTool/AutomationUtils/DeploymentContext.cs
+++ b/Engine/Source/Programs/AutomationTool/AutomationUtils/DeploymentContext.cs
@@ -310,6 +310,12 @@ public class DeploymentContext //: ProjectParams
 	/// </summary>
 	public HashSet<StagedFileReference> ConfigFilesDenyList = new HashSet<StagedFileReference>();
 
+	/// <summary>
+	/// Set files which are allow listed to be staged that would otherwise be excluded by RestrictedFolderNames.
+	/// This list is read from the +ExtraAllowedFiles=... array in the [Staging] section of *Game.ini files.
+	/// </summary>
+	public HashSet<StagedFileReference> ExtraFilesAllowList = new HashSet<StagedFileReference>();
+
 
 	/// <summary>
 	/// List of ini keys to strip when staging
@@ -598,8 +604,9 @@ public class DeploymentContext //: ProjectParams
 		}
 
 		// Read the list of files which are allow listed to be staged
-		ReadConfigFileList(GameConfig, "Staging", "AllowedConfigFiles", ConfigFilesAllowList);
-		ReadConfigFileList(GameConfig, "Staging", "DisallowedConfigFiles", ConfigFilesDenyList);
+		ReadAllowDenyFileList(GameConfig, "Staging", "AllowedConfigFiles", ConfigFilesAllowList);
+		ReadAllowDenyFileList(GameConfig, "Staging", "DisallowedConfigFiles", ConfigFilesDenyList);
+		ReadAllowDenyFileList(GameConfig, "Staging", "ExtraAllowedFiles", ExtraFilesAllowList);
 
 		// Grab the game ini data
 		String PackagingIniPath = "/Script/UnrealEd.ProjectPackagingSettings";
@@ -661,20 +668,20 @@ public class DeploymentContext //: ProjectParams
 	}
 
 	/// <summary>
-	/// Read a list of allowed or denied config files names from a config file
+	/// Read a list of allowed or denied files names from a config file
 	/// </summary>
 	/// <param name="Config">The config hierarchy to read from</param>
 	/// <param name="SectionName">The section name</param>
 	/// <param name="KeyName">The key name to read from</param>
-	/// <param name="ConfigFiles">Receives a list of config file paths</param>
-	private static void ReadConfigFileList(ConfigHierarchy Config, string SectionName, string KeyName, HashSet<StagedFileReference> ConfigFiles)
+	/// <param name="Files">Receives a list of file paths</param>
+	private static void ReadAllowDenyFileList(ConfigHierarchy Config, string SectionName, string KeyName, HashSet<StagedFileReference> FilesRef)
 	{
-		List<string> ConfigFileNames;
-		if(Config.GetArray(SectionName, KeyName, out ConfigFileNames))
+		List<string> FileNames;
+		if(Config.GetArray(SectionName, KeyName, out FileNames))
 		{
-			foreach(string ConfigFileName in ConfigFileNames)
+			foreach(string FileName in FileNames)
 			{
-				ConfigFiles.Add(new StagedFileReference(ConfigFileName));
+				FilesRef.Add(new StagedFileReference(FileName));
 			}
 		}
 	}
@@ -1130,4 +1137,53 @@ public class DeploymentContext //: ProjectParams
 	{
 		return string.Format("Manifest_UFSFiles_{0}{1}.txt", StageTargetPlatform.PlatformType, GetSanitizedDeviceNameSuffix(DeviceName));
 	}
+
+	public static StagedFileReference ApplyDirectoryRemap(DeploymentContext SC, StagedFileReference InputFile)
+	{
+		StagedFileReference CurrentFile = InputFile;
+		foreach (Tuple<StagedDirectoryReference, StagedDirectoryReference> RemapDirectory in SC.RemapDirectories)
+		{
+			StagedFileReference NewFile;
+			if (StagedFileReference.TryRemap(CurrentFile, RemapDirectory.Item1, RemapDirectory.Item2, out NewFile))
+			{
+				CurrentFile = NewFile;
+			}
+		}
+		return CurrentFile;
+	}
+
+	public static StagedFileReference MakeRelativeStagedReference(DeploymentContext SC, FileSystemReference Ref)
+	{
+		return MakeRelativeStagedReference(SC, Ref, out _);
+	}
+
+	public static StagedFileReference MakeRelativeStagedReference(DeploymentContext SC, FileSystemReference Ref, out DirectoryReference RootDir)
+	{
+		if (Ref.IsUnderDirectory(SC.ProjectRoot))
+		{
+			RootDir = SC.ProjectRoot;
+			return ApplyDirectoryRemap(SC, new StagedFileReference(SC.ShortProjectName + "/" + Ref.MakeRelativeTo(SC.ProjectRoot).Replace('\\', '/')));
+		}
+		else if (Ref.IsUnderDirectory(SC.EngineRoot))
+		{
+			RootDir = SC.EngineRoot;
+			return ApplyDirectoryRemap(SC, new StagedFileReference("Engine/" + Ref.MakeRelativeTo(SC.EngineRoot).Replace('\\', '/')));
+		}
+		throw new Exception();
+	}
+	public static FileReference UnmakeRelativeStagedReference(DeploymentContext SC, StagedFileReference Ref)
+	{
+		// paths will be in the form "Engine/Foo" or "{ProjectName}/Foo" (or something that we don't handle, so assert)
+		// So, replace the Engine/ with {EngineDir} and {ProjectName}/ with {ProjectDir}, and then append Foo
+		if (Ref.Name.StartsWith("Engine/"))
+		{
+			// skip over "Engine/" which is 7 chars long
+			return FileReference.Combine(SC.EngineRoot, Ref.Name.Substring(7));
+		}
+		else if (Ref.Name.StartsWith(SC.ShortProjectName + "/"))
+		{
+			return FileReference.Combine(SC.ProjectRoot, Ref.Name.Substring(SC.ShortProjectName.Length + 1));
+		}
+		throw new Exception();
+	}
 }
diff --git a/Engine/Source/Programs/AutomationTool/AutomationUtils/Utils.cs b/Engine/Source/Programs/AutomationTool/AutomationUtils/Utils.cs
--- a/Engine/Source/Programs/AutomationTool/AutomationUtils/Utils.cs
+++ b/Engine/Source/Programs/AutomationTool/AutomationUtils/Utils.cs
@@ -361,7 +361,7 @@ namespace AutomationTool
 		// Characters that can appear at the start of
 		private static char[] IgnoredIniValuePrefixes = { '+', '-', ' ', '\t' };
 
-		private static void FilterIniFile(string SourceName, string TargetName, List<string> IniKeyDenyList, List<string> InSectionDenyList)
+		private static void FilterIniFile(string SourceName, string TargetName, List<string> IniKeyDenyList, List<string> InSectionDenyList, List<string> InSectionAllowList)
 		{
 			string[] Lines = File.ReadAllLines(SourceName);
 			StringBuilder NewLines = new StringBuilder("");
@@ -379,7 +379,7 @@ namespace AutomationTool
 				bool bFiltered = bFilteringSection;
 
 				// look for each filter on each line
-				if (!bFiltered)
+				if (!bFiltered && IniKeyDenyList != null)
 				{
 					string TrimmedLine = Line.TrimStart(IgnoredIniValuePrefixes);
 					foreach (string Filter in IniKeyDenyList)
@@ -392,12 +392,13 @@ namespace AutomationTool
 					}
 				}
 
-				if (InSectionDenyList != null)
+				if (InSectionDenyList != null || InSectionAllowList != null)
 				{
 					if (Line.StartsWith("[") && Line.EndsWith("]"))
 					{
 						string SectionName = Line.Substring(1, Line.Length - 2);
-						bFilteringSection = bFiltered = InSectionDenyList.Contains(SectionName);
+						bFilteringSection = bFiltered = ((InSectionDenyList != null && InSectionAllowList == null && InSectionDenyList.Contains(SectionName)) ||
+						                                 (InSectionAllowList != null && !InSectionAllowList.Contains(SectionName)));
 
 						if (bFilteringSection)
 						{
@@ -430,12 +431,18 @@ namespace AutomationTool
 		/// <param name="SourceName">Source name</param>
 		/// <param name="TargetName">Target name</param>
 		/// <returns>True if the operation was successful, false otherwise.</returns>
-		public static bool SafeCopyFile(string SourceName, string TargetName, bool bQuiet = false, List<string> IniKeyDenyList = null, List<string> IniSectionDenyList = null)
+		public static bool SafeCopyFile(string SourceName, string TargetName, bool bQuiet = false, List<string> IniKeyDenyList = null, List<string> IniSectionDenyList = null, List<string> IniSectionAllowList = null, bool bSafeCreateDirectory = false)
 		{
 			if (!bQuiet)
 			{
 				Log.TraceLog("SafeCopyFile {0} {1}", SourceName, TargetName);
 			}
+
+			if (bSafeCreateDirectory)
+			{
+				SafeCreateDirectory(Path.GetDirectoryName(TargetName), bQuiet);
+			}
+
 			const int MaxAttempts = 10;
 			int Attempts = 0;
 
@@ -449,9 +456,9 @@ namespace AutomationTool
 					bool bSkipSizeCheck = false;
 					// BinaryConfig.ini is a special case with binary data, but with same extension to handle all ini 
 					// file chunking/packaging/etc rules
-					if (IniKeyDenyList != null && Path.GetExtension(SourceName) == ".ini" && Path.GetFileName(SourceName) != "BinaryConfig.ini")
+					if ((IniKeyDenyList != null || IniSectionDenyList != null || IniSectionAllowList != null) && Path.GetExtension(SourceName) == ".ini" && Path.GetFileName(SourceName) != "BinaryConfig.ini")
 					{
-						FilterIniFile(SourceName, TargetName, IniKeyDenyList, IniSectionDenyList);
+						FilterIniFile(SourceName, TargetName, IniKeyDenyList, IniSectionDenyList, IniSectionAllowList);
 						// ini files may change size, don't check
 						bSkipSizeCheck = true;
 					}
diff --git a/Engine/Source/Programs/AutomationTool/CookedEditor/CookedEditor.Automation.cs b/Engine/Source/Programs/AutomationTool/CookedEditor/CookedEditor.Automation.cs
--- a/Engine/Source/Programs/AutomationTool/CookedEditor/CookedEditor.Automation.cs
+++ b/Engine/Source/Programs/AutomationTool/CookedEditor/CookedEditor.Automation.cs
@@ -214,43 +214,7 @@ public class ModifyStageContext
 		SC.FilesToStage.UFSFiles = SC.FilesToStage.UFSFiles.Where(x => !SC.FilesToStage.NonUFSFiles.ContainsKey(x.Key)).ToDictionary(x => x.Key, x => x.Value);
 	}
 	#region Private implementation
-
-	private StagedFileReference MakeRelativeStagedReference(DeploymentContext SC, FileSystemReference Ref)
-	{
-		return MakeRelativeStagedReference(SC, Ref, out _);
-	}
-
-	private StagedFileReference MakeRelativeStagedReference(DeploymentContext SC, FileSystemReference Ref, out DirectoryReference RootDir)
-	{
-		if (Ref.IsUnderDirectory(ProjectDirectory))
-		{
-			RootDir = ProjectDirectory;
-			return Project.ApplyDirectoryRemap(SC, new StagedFileReference(ProjectName + "/" + Ref.MakeRelativeTo(ProjectDirectory).Replace('\\', '/')));
-		}
-		else if (Ref.IsUnderDirectory(EngineDirectory))
-		{
-			RootDir = EngineDirectory;
-			return Project.ApplyDirectoryRemap(SC, new StagedFileReference( "Engine/" + Ref.MakeRelativeTo(EngineDirectory).Replace('\\', '/')));
-		}
-		throw new Exception();
-	}
-
-	private FileReference UnmakeRelativeStagedReference(DeploymentContext SC, StagedFileReference Ref)
-	{
-		// paths will be in the form "Engine/Foo" or "{ProjectName}/Foo" (or something that we don't handle, so assert)
-		// So, replace the Engine/ with {EngineDir} and {ProjectName}/ with {ProjectDir}, and then append Foo
-		if (Ref.Name.StartsWith("Engine/"))
-		{
-			// skip over "Engine/" which is 7 chars long
-			return FileReference.Combine(EngineDirectory, Ref.Name.Substring(7));
-		}
-		else if (Ref.Name.StartsWith(ProjectName + "/"))
-		{
-			return FileReference.Combine(ProjectDirectory, Ref.Name.Substring(ProjectName.Length + 1));
-		}
-		throw new Exception();
-	}
-
+	
 	private void RemoveReleasedFiles(DeploymentContext SC)
 	{
 		HashSet<StagedFileReference> ShippedFiles = new HashSet<StagedFileReference>();
@@ -280,10 +244,10 @@ public class ModifyStageContext
 
 		ShippedFiles.RemoveWhere(x => x.HasExtension(".ttf") && !x.Name.Contains("LastResort"));
 
-		var RemappedNonUFS = NonUFSFilesToStage.Select(x => MakeRelativeStagedReference(SC, x));
+		var RemappedNonUFS = NonUFSFilesToStage.Select(x => DeploymentContext.MakeRelativeStagedReference(SC, x));
 
-		UFSFilesToStage.RemoveAll(x => ShippedFiles.Contains(MakeRelativeStagedReference(SC, x)));
-		NonUFSFilesToStage.RemoveAll(x => ShippedFiles.Contains(MakeRelativeStagedReference(SC, x)));
+		UFSFilesToStage.RemoveAll(x => ShippedFiles.Contains(DeploymentContext.MakeRelativeStagedReference(SC, x)));
+		NonUFSFilesToStage.RemoveAll(x => ShippedFiles.Contains(DeploymentContext.MakeRelativeStagedReference(SC, x)));
 	}
 	private Dictionary<StagedFileReference, FileReference> MimicStageFiles(DeploymentContext SC, List<FileReference> SourceFiles)
 	{
@@ -292,7 +256,7 @@ public class ModifyStageContext
 		foreach (FileReference FileRef in new HashSet<FileReference>(SourceFiles))
 		{
 			DirectoryReference RootDir;
-			StagedFileReference StagedFile = MakeRelativeStagedReference(SC, FileRef, out RootDir);
+			StagedFileReference StagedFile = DeploymentContext.MakeRelativeStagedReference(SC, FileRef, out RootDir);
 
 			// add the mapping
 			Mapping.Add(StagedFile, FileRef);
@@ -309,7 +273,7 @@ public class ModifyStageContext
 			Int32 OrigNumFiles = Files.Count();
 			// remove entries where any restricted folder names are in the name remapped path (if we remap from NFL to non-NFL, then we don't remove it)
 			// If a configuration file has been explicitly allowed, do not omit it either.
-			Files = Files.Where(x => SC.ConfigFilesAllowList.Contains(x.Key) || SC.DirectoriesAllowList.Contains(x.Key.Directory) || !SC.RestrictedFolderNames.Any(y => Project.ApplyDirectoryRemap(SC, x.Key).ContainsName(y))).ToDictionary(x => x.Key, x => x.Value);
+			Files = Files.Where(x => SC.ConfigFilesAllowList.Contains(x.Key) || SC.ExtraFilesAllowList.Contains(x.Key) || SC.DirectoriesAllowList.Contains(x.Key.Directory) || !SC.RestrictedFolderNames.Any(y => DeploymentContext.ApplyDirectoryRemap(SC, x.Key).ContainsName(y))).ToDictionary(x => x.Key, x => x.Value);
 			if (OrigNumFiles != Files.Count())
 			{
 				Log.TraceInformationOnce("Some files were not staged since they have restricted folder names in the remapped path.");
@@ -324,7 +288,7 @@ public class ModifyStageContext
 	private void AddUFSFilesToList(List<FileReference> FileList, string Extension, DeploymentContext SC)
 	{
 		// look in SC and UFSFiles
-		FileList.AddRange(SC.FilesToStage.UFSFiles.Keys.Where(x => x.HasExtension(Extension)).Select(y => UnmakeRelativeStagedReference(SC, y)));
+		FileList.AddRange(SC.FilesToStage.UFSFiles.Keys.Where(x => x.HasExtension(Extension)).Select(y => DeploymentContext.UnmakeRelativeStagedReference(SC, y)));
 		FileList.AddRange(UFSFilesToStage.Where(x => x.GetExtension().Equals(Extension, StringComparison.InvariantCultureIgnoreCase)));
 	}
 
@@ -904,10 +868,14 @@ public class MakeCookedEditor : BuildCommand
 			{
 				SearchMode = SearchOption.TopDirectoryOnly;
 			}
-			// settings with a Dest are special, and have to go right to SC
-			if (Props.ContainsKey("Dest"))
+
+			bool bHasDest = Props.ContainsKey("Dest");
+			bool bHasIniSections = Props.ContainsKey("IniSections");
+
+			// settings with a Dest or IniSections are special, and have to go right to SC
+			if (bHasDest || bHasIniSections)
 			{
-				FileReference SourceFile = FileReference.Combine(BaseDirectory, SubPath);
+				FileReference SourceFile = FileReference.Combine(BaseDirectory, SubPath, FileWildcard);
 
 				// allow a blank Path to mean empty file
 				if (SubPath == "")
@@ -919,12 +887,39 @@ public class MakeCookedEditor : BuildCommand
 
 				if (SearchMode == SearchOption.AllDirectories || FileWildcard.Contains("*"))
 				{
-					throw new AutomationException($"Unable to stage directories with \"Dest\" setting for CookedEditor: '{Entry}'");
+					throw new AutomationException($"Unable to stage directories with \"Dest\" or \"IniSections\" setting for CookedEditor: '{Entry}'");
+				}
+
+				// setup Dest
+				StagedFileReference DestFile;
+				if (bHasDest)
+				{
+					DestFile = new StagedFileReference(Props["Dest"]);
+				}
+				else
+				{
+					DestFile = DeploymentContext.MakeRelativeStagedReference(SC, SourceFile);
+				}
+
+				if (bHasIniSections)
+				{
+					if (!SourceFile.HasExtension(".ini"))
+					{
+						throw new AutomationException($"Unable to stage non-ini file '{Entry}' using \"IniSections\" setting for CookedEditor: ");
+					}
+
+					FileReference IntermediateFile = FileReference.Combine(Context.ProjectDirectory, "Intermediate", "StagedIniSections", DeploymentContext.MakeRelativeStagedReference(SC, SourceFile).ToString());
+					InternalUtils.SafeCopyFile(SourceFile.ToString(), IntermediateFile.ToString(), IniSectionAllowList : Props["IniSections"].Split(',').ToList(), bSafeCreateDirectory : true);
+					SourceFile = IntermediateFile;
 				}
 
 				// now stage it to a different location as specified in the params
 				StagedFileType FileType = (FileList == Context.NonUFSFilesToStage) ? StagedFileType.NonUFS : StagedFileType.UFS;
-				SC.StageFile(FileType, SourceFile, new StagedFileReference(Props["Dest"]));
+				if (Props.ContainsKey("Force") && bool.Parse(Props["Force"]) == true)
+				{
+					SC.ExtraFilesAllowList.Add(DestFile);
+				}
+				SC.StageFile(FileType, SourceFile, DestFile);
 				continue;
 			}
 
@@ -932,7 +927,18 @@ public class MakeCookedEditor : BuildCommand
 			DirectoryReference Dir = DirectoryReference.Combine(BaseDirectory, SubPath);
 			if (DirectoryReference.Exists(Dir))
 			{
-				FileList.AddRange(DirectoryReference.EnumerateFiles(Dir, FileWildcard, SearchMode));
+				if (Props.ContainsKey("Force") && bool.Parse(Props["Force"]) == true)
+				{
+					foreach (FileReference File in DirectoryReference.EnumerateFiles(Dir, FileWildcard, SearchMode))
+					{
+						SC.ExtraFilesAllowList.Add(DeploymentContext.MakeRelativeStagedReference(SC, File));
+						FileList.Add(File);
+					}
+				}
+				else
+				{
+					FileList.AddRange(DirectoryReference.EnumerateFiles(Dir, FileWildcard, SearchMode));
+				}
 			}
 		}
 	}
diff --git a/Engine/Source/Programs/AutomationTool/Scripts/CopyBuildToStagingDirectory.Automation.cs b/Engine/Source/Programs/AutomationTool/Scripts/CopyBuildToStagingDirectory.Automation.cs
--- a/Engine/Source/Programs/AutomationTool/Scripts/CopyBuildToStagingDirectory.Automation.cs
+++ b/Engine/Source/Programs/AutomationTool/Scripts/CopyBuildToStagingDirectory.Automation.cs
@@ -1433,10 +1433,10 @@ namespace AutomationScripts
 			}
 
 			// Apply all the directory mappings
-			SC.FilesToStage.UFSFiles = SC.FilesToStage.UFSFiles.ToDictionary(x => ApplyDirectoryRemap(SC, x.Key), x => x.Value);
-			SC.FilesToStage.NonUFSFiles = SC.FilesToStage.NonUFSFiles.ToDictionary(x => ApplyDirectoryRemap(SC, x.Key), x => x.Value);
-			SC.FilesToStage.NonUFSDebugFiles = SC.FilesToStage.NonUFSDebugFiles.ToDictionary(x => ApplyDirectoryRemap(SC, x.Key), x => x.Value);
-			SC.FilesToStage.NonUFSSystemFiles = SC.FilesToStage.NonUFSSystemFiles.ToDictionary(x => ApplyDirectoryRemap(SC, x.Key), x => x.Value);
+			SC.FilesToStage.UFSFiles = SC.FilesToStage.UFSFiles.ToDictionary(x => DeploymentContext.ApplyDirectoryRemap(SC, x.Key), x => x.Value);
+			SC.FilesToStage.NonUFSFiles = SC.FilesToStage.NonUFSFiles.ToDictionary(x => DeploymentContext.ApplyDirectoryRemap(SC, x.Key), x => x.Value);
+			SC.FilesToStage.NonUFSDebugFiles = SC.FilesToStage.NonUFSDebugFiles.ToDictionary(x => DeploymentContext.ApplyDirectoryRemap(SC, x.Key), x => x.Value);
+			SC.FilesToStage.NonUFSSystemFiles = SC.FilesToStage.NonUFSSystemFiles.ToDictionary(x => DeploymentContext.ApplyDirectoryRemap(SC, x.Key), x => x.Value);
 
 			// Create plugin manifests after the directory mappings
 			if (bCreatePluginManifest && Params.UsePak(SC.StageTargetPlatform))
@@ -1463,6 +1463,7 @@ namespace AutomationScripts
 			}
 			RestrictedFiles.RemoveWhere(RestrictedFile => SC.DirectoriesAllowList.Any(TestDirectory => RestrictedFile.Directory.IsUnderDirectory(TestDirectory)));
 			RestrictedFiles.RemoveWhere(RestrictedFile => SC.ConfigFilesAllowList.Contains(RestrictedFile));
+			RestrictedFiles.RemoveWhere(RestrictedFile => SC.ExtraFilesAllowList.Contains(RestrictedFile));
 
 			if (RestrictedFiles.Count > 0)
 			{
@@ -1811,20 +1812,6 @@ namespace AutomationScripts
 			return null;
 		}
 
-		public static StagedFileReference ApplyDirectoryRemap(DeploymentContext SC, StagedFileReference InputFile)
-		{
-			StagedFileReference CurrentFile = InputFile;
-			foreach (Tuple<StagedDirectoryReference, StagedDirectoryReference> RemapDirectory in SC.RemapDirectories)
-			{
-				StagedFileReference NewFile;
-				if (StagedFileReference.TryRemap(CurrentFile, RemapDirectory.Item1, RemapDirectory.Item2, out NewFile))
-				{
-					CurrentFile = NewFile;
-				}
-			}
-			return CurrentFile;
-		}
-
 		static List<string> ParseInputPaths(List<string> ConfigLines)
 		{
 			// Each string has the format '(Path="TheDirToStage")'
