Deprecate the SyncProject command, and recommend users run the UGS CLI build instead.

diff --git a/Engine/Source/Programs/AutomationTool/AutomationUtils/Matchers/GenericEventMatcher.cs b/Engine/Source/Programs/AutomationTool/AutomationUtils/Matchers/GenericEventMatcher.cs
--- a/Engine/Source/Programs/AutomationTool/AutomationUtils/Matchers/GenericEventMatcher.cs
+++ b/Engine/Source/Programs/AutomationTool/AutomationUtils/Matchers/GenericEventMatcher.cs
@@ -18,7 +18,7 @@ namespace AutomationUtils.Matchers
 			@"^\s*(FATAL|fatal error):");
 
 		static readonly Regex s_warningErrorPattern = new Regex(
-			@"(?<!\w)(?i)(WARNING|ERROR) ?(\([^)]+\)|\[[^\]]+\])?: ");
+			@"(?<!\w)(?i)(WARNING|ERROR) ?(\([^)]+\)|\[[^\]]+\])?:(?: |$)");
 
 		static readonly Regex s_errorPattern = new Regex(
 			@"[Ee]rror [A-Z]\d+\s:");
diff --git a/Engine/Source/Programs/AutomationTool/Scripts/SyncProject.Automation.cs b/Engine/Source/Programs/AutomationTool/Scripts/SyncProject.Automation.cs
--- a/Engine/Source/Programs/AutomationTool/Scripts/SyncProject.Automation.cs
+++ b/Engine/Source/Programs/AutomationTool/Scripts/SyncProject.Automation.cs
@@ -108,6 +108,21 @@ class SyncProject : SyncProjectBase
 {
 	public override ExitCode Execute()
 	{
+		if (!ParseParam("Deprecated"))
+		{
+			LogError("**************************************************************************");
+			LogError("The SyncProject command has been deprecated, and will be removed in an");
+			LogError("upcoming UnrealEngine release.");
+			LogError("");
+			LogError("Similar functionality is available through the UnrealGameSync command-line");
+			LogError("tool, which is supported on Windows, Mac and Linux.");
+			LogError("");
+			LogError("To ignore this warning and continue to using SyncProject anyway, add the");
+			LogError("-Deprecated argument to the command line.");
+			LogError("**************************************************************************");
+			return ExitCode.Error_Unknown;
+		}
+
 		LogInformation("************************* SyncProject");
 
 		// These are files that should always be synced because tools update them
