Override key entry when the key already exists in the TestData object

#rnx

diff --git a/Engine/Source/Programs/AutomationTool/Gauntlet/Framework/Utils/Gauntlet.HordeReport.cs b/Engine/Source/Programs/AutomationTool/Gauntlet/Framework/Utils/Gauntlet.HordeReport.cs
--- a/Engine/Source/Programs/AutomationTool/Gauntlet/Framework/Utils/Gauntlet.HordeReport.cs
+++ b/Engine/Source/Programs/AutomationTool/Gauntlet/Framework/Utils/Gauntlet.HordeReport.cs
@@ -1178,7 +1178,16 @@ namespace Gauntlet
 				DataItem NewDataItem = new DataItem();
 				NewDataItem.Key = string.IsNullOrEmpty(InKey) ? InData.Type : InKey;
 				NewDataItem.Data = InData;
-				Items.Add(NewDataItem);
+
+				var FoundItemIndex = Items.FindIndex(I => I.Key == InKey);
+				if (FoundItemIndex == -1)
+				{
+					Items.Add(NewDataItem);
+				}
+				else
+				{
+					Items[FoundItemIndex] = NewDataItem;
+				}
 
 				var ExtraItems = InData.GetReportDependencies();
 				if (ExtraItems.Count() > 0)
diff --git a/Engine/Source/Programs/AutomationTool/Gauntlet/Unreal/Base/Gauntlet.UnrealTestNode.cs b/Engine/Source/Programs/AutomationTool/Gauntlet/Unreal/Base/Gauntlet.UnrealTestNode.cs
--- a/Engine/Source/Programs/AutomationTool/Gauntlet/Unreal/Base/Gauntlet.UnrealTestNode.cs
+++ b/Engine/Source/Programs/AutomationTool/Gauntlet/Unreal/Base/Gauntlet.UnrealTestNode.cs
@@ -1479,7 +1479,7 @@ namespace Gauntlet
 			);
 			HordeReport.TestDataCollection HordeTestDataCollection = new HordeReport.TestDataCollection();
 			HordeTestDataCollection.AddNewTestReport(Report, GetCachedConfiguration().HordeTestDataKey);
-			HordeTestDataCollection.WriteToJson(HordeTestDataFilePath, true);
+			HordeTestDataCollection.WriteToJson(HordeTestDataFilePath, !AutomationTool.Automation.IsBuildMachine);
 		}
 
 		/// <summary>
