Alternative AFS startup method (solves failed to launch and bind listener issues)

#android

diff --git a/Engine/Plugins/Runtime/AndroidFileServer/Source/AndroidFileServer/Android/java/com/epicgames/unreal/RemoteFileManagerActivity.java b/Engine/Plugins/Runtime/AndroidFileServer/Source/AndroidFileServer/Android/java/com/epicgames/unreal/RemoteFileManagerActivity.java
new file mode 100644
--- /dev/null
+++ b/Engine/Plugins/Runtime/AndroidFileServer/Source/AndroidFileServer/Android/java/com/epicgames/unreal/RemoteFileManagerActivity.java
@@ -0,0 +1,151 @@
+// Copyright Epic Games, Inc. All Rights Reserved.
+
+package com.epicgames.unreal;
+
+import android.app.Activity;
+import android.content.Intent;
+import android.content.pm.ApplicationInfo;
+import android.content.pm.PackageManager;
+import android.os.Bundle;
+import android.util.Log;
+
+import com.epicgames.unreal.GameActivity;
+import com.epicgames.unreal.RemoteFileManager;
+
+import java.io.BufferedReader;
+import java.io.IOException;
+import java.io.InputStreamReader;
+import java.io.PrintWriter;
+import java.net.ServerSocket;
+import java.net.Socket;
+import java.util.Set;
+
+import static android.os.Environment.getExternalStorageDirectory;
+
+public class RemoteFileManagerActivity extends Activity {
+    private static String TAG = "UEFS";
+    private static final int SERVER_PORT = 57099;
+
+	@Override
+	protected void onCreate(Bundle savedInstanceState) {
+		super.onCreate(savedInstanceState);
+
+		Intent intent = getIntent();
+		String action = intent.getAction();
+		if (action != null && action == "com.epicgames.unreal.RemoteFileManager.intent.COMMAND2") {
+			String cmd = intent.getStringExtra("cmd");
+			String token = intent.getStringExtra("token");
+			int port = intent.getIntExtra("port", SERVER_PORT);
+			if (cmd != null && cmd.equals("start") && GameActivity.AndroidFileServer_Verify(token)) {
+
+				// make sure UnrealFileServer not already running
+				boolean bUSBRunning = false;
+				boolean bWiFiRunning = false;
+				
+				Set<Thread> threads = Thread.getAllStackTraces().keySet();
+				for (Thread thread : threads) {
+					String name = thread.getName();
+					if (name.equals("UnrealAFS-USB"))
+					{
+						if (thread.isAlive())
+						{
+							bUSBRunning = true;
+						}
+					}
+					else if (name.equals("UnrealAFS-WiFi"))
+					{
+						if (thread.isAlive())
+						{
+							bWiFiRunning = true;
+						}
+					}
+				}
+
+				String packageName = getPackageName();
+				PackageManager pm = getPackageManager();
+
+				int versionCode = 0;
+				try
+				{
+					versionCode = pm.getPackageInfo(packageName, 0).versionCode;
+				}
+				catch (Exception e)
+				{
+					// if the above failed, then, we can't use obbs
+				}
+
+				String EngineVersion = "5.0.0";
+				String ProjectName = packageName;
+				ProjectName = ProjectName.substring(ProjectName.lastIndexOf('.') + 1);
+				boolean IsShipping = false;
+				boolean UseExternalFilesDir = false;
+				boolean PublicLogFiles = false;
+				try {
+					ApplicationInfo ai = pm.getApplicationInfo(packageName, PackageManager.GET_META_DATA);
+					Bundle bundle = ai.metaData;
+
+					if (bundle.containsKey("com.epicgames.unreal.GameActivity.EngineVersion"))
+					{
+						EngineVersion = bundle.getString("com.epicgames.unreal.GameActivity.EngineVersion");
+					}
+					if (bundle.containsKey("com.epicgames.unreal.GameActivity.ProjectName"))
+					{
+						ProjectName = bundle.getString("com.epicgames.unreal.GameActivity.ProjectName");
+					}
+					if (bundle.containsKey("com.epicgames.unreal.GameActivity.BuildConfiguration"))
+					{
+						String Configuration = bundle.getString("com.epicgames.unreal.GameActivity.BuildConfiguration");
+						IsShipping = Configuration.contains("Shipping");
+					}
+					if (bundle.containsKey("com.epicgames.unreal.GameActivity.bUseExternalFilesDir"))
+					{
+						UseExternalFilesDir = bundle.getBoolean("com.epicgames.unreal.GameActivity.bUseExternalFilesDir");
+					}
+					if (bundle.containsKey("com.epicgames.unreal.GameActivity.bPublicLogFiles"))
+					{
+						PublicLogFiles = bundle.getBoolean("com.epicgames.unreal.GameActivity.bPublicLogFiles");
+					}
+				}
+				catch (PackageManager.NameNotFoundException | NullPointerException e)
+				{
+					Log.e("UE","Error when accessing application metadata", e);
+				}
+
+				String internal = getFilesDir().getAbsolutePath();
+				String external = getExternalFilesDir(null).getAbsolutePath();
+				String storage = getExternalStorageDirectory().getAbsolutePath();
+				String obbdir = getObbDir().getAbsolutePath();
+
+				if (!bUSBRunning)
+				{
+					RemoteFileManager fileManager = new RemoteFileManager(true, port, internal, external, storage, obbdir, packageName, versionCode, ProjectName, EngineVersion, IsShipping, PublicLogFiles);
+					new Thread(fileManager, "UnrealAFS-USB").start();
+				}
+				if (!bWiFiRunning)
+				{
+					RemoteFileManager fileManager = new RemoteFileManager(false, port, internal, external, storage, obbdir, packageName, versionCode, ProjectName, EngineVersion, IsShipping, PublicLogFiles);
+					new Thread(fileManager, "UnrealAFS-WiFi").start();
+				}
+			}
+			else if (cmd != null && cmd.equals("stop")) {
+
+				// send stop request to threads
+				Set<Thread> threads = Thread.getAllStackTraces().keySet();
+				for (Thread thread : threads) {
+					String name = thread.getName();
+					if (name.equals("UnrealAFS-USB") || name.equals("UnrealAFS-WiFi"))
+					{
+						if (thread.isAlive())
+						{
+							Log.d(TAG, "Sent stop request to " + getPackageName() + "/" + name);
+							thread.interrupt();
+						}
+					}
+				}
+            }
+		}
+
+		// now terminate
+		finish();
+    }
+}
\ No newline at end of file
diff --git a/Engine/Plugins/Runtime/AndroidFileServer/Source/AndroidFileServer/AndroidFileServer_UPL.xml b/Engine/Plugins/Runtime/AndroidFileServer/Source/AndroidFileServer/AndroidFileServer_UPL.xml
--- a/Engine/Plugins/Runtime/AndroidFileServer/Source/AndroidFileServer/AndroidFileServer_UPL.xml
+++ b/Engine/Plugins/Runtime/AndroidFileServer/Source/AndroidFileServer/AndroidFileServer_UPL.xml
@@ -63,6 +63,11 @@
 						<if condition="bExternalStartAFS">
 							<true>
 								<addElements tag="application">
+									<activity android:name="com.epicgames.unreal.RemoteFileManagerActivity" android:exported="true" android:screenOrientation="sensor" android:theme="@android:style/Theme.NoDisplay">
+										<intent-filter>
+											<action android:name="com.epicgames.unreal.RemoteFileManager.intent.COMMAND2" />
+										</intent-filter>
+									</activity>
 									<receiver android:name="com.epicgames.unreal.RemoteFileManagerReceiver" android:exported="true">
 										<intent-filter>
 											<action android:name="com.epicgames.unreal.RemoteFileManager.intent.COMMAND" />
@@ -125,13 +130,16 @@
 				<copyFile src="$S(PluginDir)/Android/java/com/epicgames/unreal/RemoteFileManager.java" dst="$S(BuildDir)/src/com/epicgames/unreal/RemoteFileManager.java" />
 				<if condition="bExternalStartAFS">
 					<true>
+						<copyFile src="$S(PluginDir)/Android/java/com/epicgames/unreal/RemoteFileManagerActivity.java" dst="$S(BuildDir)/src/com/epicgames/unreal/RemoteFileManagerActivity.java" />
 						<copyFile src="$S(PluginDir)/Android/java/com/epicgames/unreal/RemoteFileManagerReceiver.java" dst="$S(BuildDir)/src/com/epicgames/unreal/RemoteFileManagerReceiver.java" />
 					</true>
 					<false>
 						<isArch arch="arm64-v8a">
+							<deleteFiles filespec="arm64/src/com/epicgames/unreal/RemoteFileManagerActivity.java" />
 							<deleteFiles filespec="arm64/src/com/epicgames/unreal/RemoteFileManagerReceiver.java" />
 						</isArch>
 						<isArch arch="x86_64">
+							<deleteFiles filespec="x64/src/com/epicgames/unreal/RemoteFileManagerActivity.java" />
 							<deleteFiles filespec="x64/src/com/epicgames/unreal/RemoteFileManagerReceiver.java" />
 						</isArch>
 					</false>
@@ -140,10 +148,12 @@
 			<false>
 				<isArch arch="arm64-v8a">
 					<deleteFiles filespec="arm64/src/com/epicgames/unreal/RemoteFileManager.java" />
+					<deleteFiles filespec="arm64/src/com/epicgames/unreal/RemoteFileManagerActivity.java" />
 					<deleteFiles filespec="arm64/src/com/epicgames/unreal/RemoteFileManagerReceiver.java" />
 				</isArch>
 				<isArch arch="x86_64">
 					<deleteFiles filespec="x64/src/com/epicgames/unreal/RemoteFileManager.java" />
+					<deleteFiles filespec="x64/src/com/epicgames/unreal/RemoteFileManagerActivity.java" />
 					<deleteFiles filespec="x64/src/com/epicgames/unreal/RemoteFileManagerReceiver.java" />
 				</isArch>
 			</false>
@@ -276,13 +286,23 @@ tasks.whenTaskAdded { task ->
 
 	afs_infile = new File(rootProject.projectDir, 'app/src/main/AndroidManifest.xml')
 	afs_outfile = new File(rootProject.projectDir, 'AFSProject/app/src/main/AndroidManifest.xml')
+	def afs_hasActivity = false
 	def afs_hasReceiver = false
 	afs_outfile.withWriter { writer ->
 		afs_infile.withReader { reader ->
 			while ((afs_line = reader.readLine()) != null) {
-				if (afs_line.contains('com.epicgames.unreal.RemoteFileManagerReceiver')) {
+				if (afs_line.contains('com.epicgames.unreal.RemoteFileManagerActivity')) {
+					afs_hasActivity = true
+				} else if (afs_line.contains('com.epicgames.unreal.RemoteFileManagerReceiver')) {
 					afs_hasReceiver = true
 				} else if (afs_line.contains('</application>')) {
+					if (afs_hasActivity == false) {
+						writer.writeLine('    <activity android:name="com.epicgames.unreal.RemoteFileManagerActivity" android:exported="true" android:screenOrientation="sensor">')
+						writer.writeLine('    	<intent-filter>')
+						writer.writeLine('    		<action android:name="com.epicgames.unreal.RemoteFileManager.intent.COMMAND2" />')
+						writer.writeLine('    	</intent-filter>')
+						writer.writeLine('    </activity>')
+					}
 					if (afs_hasReceiver == false) {
 						writer.writeLine('    <receiver android:name="com.epicgames.unreal.RemoteFileManagerReceiver" android:exported="true">')
 						writer.writeLine('      <intent-filter>')
diff --git a/Engine/Source/Programs/AutomationTool/Android/AndroidFileClient.cs b/Engine/Source/Programs/AutomationTool/Android/AndroidFileClient.cs
--- a/Engine/Source/Programs/AutomationTool/Android/AndroidFileClient.cs
+++ b/Engine/Source/Programs/AutomationTool/Android/AndroidFileClient.cs
@@ -18,6 +18,7 @@ using AutomationScripts;
 using System.Net;
 using System.Net.Sockets;
 using System.IO.Compression;
+using Newtonsoft.Json.Linq;
 
 
 namespace AutomationTool
@@ -1353,7 +1354,55 @@ namespace AutomationTool
 			return Result;
 		}
 
-		public static bool GetListenStatus_netstat(string Device, int Port, out bool bUSB, out bool bWifi, out string WifiAddress)
+		public static List<string> GetInstalledActivities(string Device)
+		{
+			List<string> Result = new List<string>();
+
+			String InstalledResult = adb.Shell(Device, "cmd package query-activities --components -a com.epicgames.unreal.RemoteFileManager.intent.COMMAND2");
+			if (InstalledResult.Contains("not found") || InstalledResult.Contains("FAIL"))
+			{
+				InstalledResult = adb.Shell(Device, "dumpsys package");
+				bool bFoundActivity = false;
+				foreach (string Line in InstalledResult.Split('\n'))
+				{
+					if (!bFoundActivity)
+					{
+						if (Line.Contains("com.epicgames.unreal.RemoteFileManager.intent.COMMAND2:"))
+						{
+							bFoundActivity = true;
+						}
+					}
+					else
+					{
+						int SlashIndex = Line.IndexOf("/com.epicgames.unreal.RemoteFileManagerActivity");
+						if (SlashIndex > 0)
+						{
+							int StartIndex = Line.LastIndexOf(' ');
+							if (StartIndex >= 0)
+							{
+								Result.Add(Line.Substring(StartIndex + 1, SlashIndex - StartIndex - 1));
+							}
+						}
+						else if (Line.IndexOf(":") > 0)
+						{
+							break;
+						}
+					}
+				}
+				return Result;
+			}
+			foreach (string Line in InstalledResult.Split('\n'))
+			{
+				int SlashIndex = Line.IndexOf('/');
+				if (SlashIndex > 0 && (Line.Substring(SlashIndex + 1) == "com.epicgames.unreal.RemoteFileManagerActivity"))
+				{
+					Result.Add(Line.Substring(0, SlashIndex));
+				}
+			}
+			return Result;
+		}
+
+		public static bool GetListenStatus_netstat(string Device, int Port, out bool bHaveApp, out bool bUSB, out bool bWifi, out string WifiAddress)
 		{
 			bUSB = false;
 			bWifi = false;
@@ -1363,8 +1412,10 @@ namespace AutomationTool
 			string netstatResult = adb.Shell(Device, "netstat -atn");
 			if (netstatResult.Contains("not found") || netstatResult.Contains("FAIL"))
 			{
+				bHaveApp = false;
 				return false;
 			}
+			bHaveApp = true;
 			foreach (string Line in netstatResult.Split('\n'))
 			{
 				if (Line.Contains("LISTEN") && Line.Contains(Search))
@@ -1389,7 +1440,7 @@ namespace AutomationTool
 			return true;
 		}
 
-		public static bool GetListenStatus_ss(string Device, int Port, out bool bUSB, out bool bWifi, out string WifiAddress)
+		public static bool GetListenStatus_ss(string Device, int Port, out bool bHaveApp, out bool bUSB, out bool bWifi, out string WifiAddress)
 		{
 			bUSB = false;
 			bWifi = false;
@@ -1399,8 +1450,10 @@ namespace AutomationTool
 			string ssResult = adb.Shell(Device, "ss -ta");
 			if (ssResult.Contains("not found") || ssResult.Contains("FAIL"))
 			{
+				bHaveApp = false;
 				return false;
 			}
+			bHaveApp = true;
 			bool bConnectionsPresent = false;
 			bool bFirstLine = true;
 			foreach (string Line in ssResult.Split('\n', StringSplitOptions.RemoveEmptyEntries))
@@ -1469,8 +1522,7 @@ namespace AutomationTool
 			// try ss first and remember if available for next time
 			if (bHave_ss)
 			{
-				bHave_ss = GetListenStatus_ss(Device, Port == 0 ? DefaultPort : Port, out bUSB, out bWifi, out WifiAddress);
-				if (bHave_ss)
+				if (GetListenStatus_ss(Device, Port == 0 ? DefaultPort : Port, out bHave_ss, out bUSB, out bWifi, out WifiAddress))
 				{
 					return true;
 				}
@@ -1479,8 +1531,7 @@ namespace AutomationTool
 			// next fall back to netstat and remember if available for next time
 			if (bHave_netstat)
 			{
-				bHave_netstat = GetListenStatus_netstat(Device, Port == 0 ? DefaultPort : Port, out bUSB, out bWifi, out WifiAddress);
-				if (bHave_netstat)
+				if (GetListenStatus_netstat(Device, Port == 0 ? DefaultPort : Port, out bHave_netstat, out bUSB, out bWifi, out WifiAddress))
 				{
 					return true;
 				}
@@ -1503,6 +1554,14 @@ namespace AutomationTool
 				adb.Shell(Device, "am broadcast -a com.epicgames.unreal.RemoteFileManager.intent.COMMAND -n " + Receiver + "/com.epicgames.unreal.RemoteFileManagerReceiver -e cmd stop");
 			}
 
+			// get a list of installed activities and verify package is available first
+			List<string> InstalledActivities = GetInstalledActivities(Device);
+			foreach (string Activity in InstalledActivities)
+			{
+				bDidSendStops = true;
+				adb.Shell(Device, "am start -a com.epicgames.unreal.RemoteFileManager.intent.COMMAND2 -n " + Activity + "/com.epicgames.unreal.RemoteFileManagerActivity -e cmd stop");
+			}
+
 			// it can take up to 2 seconds for running servers to terminate so check if any binds to port are still active
 			if (bDidSendStops && bWaitForStop)
 			{
@@ -1556,7 +1615,7 @@ namespace AutomationTool
 				}
 			}
 
-
+			/*
 			// get a list of installed clients and verify package is available (up to 6 seconds)
 			int tries = 60;
 			List<string> InstalledReceivers = GetInstalledReceivers(Device);
@@ -1581,6 +1640,32 @@ namespace AutomationTool
 					return false;
 				}
 			}
+			*/
+
+			// get a list of installed clients and verify package is available (up to 6 seconds)
+			int tries = 60;
+			List<string> InstalledActivities = GetInstalledActivities(Device);
+			if (!InstalledActivities.Contains(PackageName))
+			{
+				// if this was called right after installed the activity may not be registered yet, wait and try again
+				bool bFound = false;
+				while (tries-- > 0)
+				{
+					Thread.Sleep(100);
+
+					InstalledActivities = GetInstalledActivities(Device);
+					if (InstalledActivities.Contains(PackageName))
+					{
+						bFound = true;
+						break;
+					}
+				}
+				if (!bFound)
+				{
+					Log.TraceInformation("Did not find package with activity");
+					return false;
+				}
+			}
 
 			// retries up to 8 seconds to start and see listener ready
 			tries = 40;
@@ -1589,7 +1674,8 @@ namespace AutomationTool
 			string WifiAddress;
 
 			// sent start request (won't do anything if already started)
-			string StartCommand = "am broadcast -a com.epicgames.unreal.RemoteFileManager.intent.COMMAND -n " + PackageName + "/com.epicgames.unreal.RemoteFileManagerReceiver -e cmd start -e token " + Token + " -ei port " + ServerPort;
+//			string StartCommand = "am broadcast -a com.epicgames.unreal.RemoteFileManager.intent.COMMAND -n " + PackageName + "/com.epicgames.unreal.RemoteFileManagerReceiver -e cmd start -e token " + Token + " -ei port " + ServerPort;
+			string StartCommand = "am start -a com.epicgames.unreal.RemoteFileManager.intent.COMMAND2 -n " + PackageName + "/com.epicgames.unreal.RemoteFileManagerActivity -e cmd start -e token " + Token + " -ei port " + ServerPort;
 			adb.Shell(Device, StartCommand);
 
 			// see if we can check listen status
diff --git a/Engine/Source/Programs/UnrealAndroidFileTool/Program.cs b/Engine/Source/Programs/UnrealAndroidFileTool/Program.cs
--- a/Engine/Source/Programs/UnrealAndroidFileTool/Program.cs
+++ b/Engine/Source/Programs/UnrealAndroidFileTool/Program.cs
@@ -16,7 +16,7 @@ namespace UnrealAndroidFileTool
 {
 	class Program
 	{
-		private static string AppVersion = "1.0.1";
+		private static string AppVersion = "1.0.2";
 
 		static void ShowHelp(string command)
 		{
@@ -1081,7 +1081,7 @@ namespace UnrealAndroidFileTool
 			Assembly? assembly = Assembly.GetEntryAssembly();
 			if (assembly != null)
 			{
-				string? PathToBinariesDotNET = Path.GetDirectoryName(assembly.GetOriginalLocation());
+				string? PathToBinariesDotNET = Path.GetDirectoryName(AppContext.BaseDirectory);//assembly.GetOriginalLocation());
 				if (PathToBinariesDotNET != null)
 				{
 					AssemblyUtils.InstallAssemblyResolver(PathToBinariesDotNET);
