Include a generated TargetInfo.json file with the pre-built editors, and set the "promoted" flag by default.

This allows running the editor without having to run UBT to determine all the available targets for a project.

diff --git a/Engine/Build/Graph/Examples/BuildEditorAndTools.xml b/Engine/Build/Graph/Examples/BuildEditorAndTools.xml
--- a/Engine/Build/Graph/Examples/BuildEditorAndTools.xml
+++ b/Engine/Build/Graph/Examples/BuildEditorAndTools.xml
@@ -44,6 +44,7 @@
 	<Option Name="OutputDir" DefaultValue="$(RootDir)\LocalBuilds\$(EditorTarget)Binaries" Description ="Output directory for compiled binaries"/>
 	<Option Name="Licensee" Restrict="true|false" DefaultValue="true" Description="Whether to treat the changelist number as being from a licensee Perforce server"/>
 	<Option Name="Versioned" Restrict="true|false" DefaultValue="$(IsBuildMachine)" Description="Whether to embed changelist number into binaries"/>
+	<Option Name="Promoted" Restrict="true|false" DefaultValue="true" Description="Whether to set the promoted flag in version files, allowing the editor to run without UBT."/>
 	<Option Name="ArchiveStream" Restrict="(?://.*)?" DefaultValue="" Description="Stream path that will contain zip file of the compiled binaries, for use by UnrealGameSync"/>
 	<Option Name="ForceSubmit" Restrict="true|false" DefaultValue="false" Description="Forces the submit to happen even if another change has been submitted (resolves in favor of local files)"/>
 	<Option Name="PreferredAgent" DefaultValue="CompileWin64" Description ="Agent type to use for job"/>
@@ -58,6 +59,7 @@
 		<!-- Update the engine version files -->
 		<Node Name="Update Version Files">
 			<SetVersion Change="$(Change)" Branch="$(EscapedBranch)" Licensee="$(Licensee)" If="$(Versioned)"/>
+			<SetVersion Change="0" Branch="$(EscapedBranch)" Licensee="$(Licensee)" If="!$(Versioned) and $(Promoted)"/>
 		</Node>
 		
 		<!-- Compile the tool executables -->
@@ -78,10 +80,14 @@
 		</Node>
 		
 		<!-- Compile the editor executable -->
-		<Node Name="Compile $(EditorTarget) Win64" Requires="Compile Tools Win64" Produces="#EditorBinaries">
+		<Node Name="Compile $(EditorTarget) Win64" Requires="Compile Tools Win64" Produces="#EditorBinaries;#TargetInfo">
 			<Property Name="UProjectParam" Value=""/>
 			<Property Name="UProjectParam" Value="-Project=&quot;$(UProjectPath)&quot;" If="'$(UProjectPath)' != ''"/>
 			<Compile Target="$(EditorTarget)" Platform="Win64" Configuration="Development" Tag="#EditorBinaries" Arguments="$(UProjectParam)"/>
+
+			<Property Name="UProjectArg" Value=""/>
+			<Property Name="UProjectArg" Value="$(UProjectPath)" If="'$(UProjectPath)' != ''"/>
+			<QueryTargets ProjectFile="$(UProjectPath)" Tag="#TargetInfo"/>
 		</Node>
 		
 		<!-- Compile the game targets -->
@@ -97,8 +103,9 @@
 		</ForEach>
 		
 		<!-- Copy all the files to the output directory -->
-		<Node Name="Tag Output Files" Requires="#ToolBinaries;#EditorBinaries;$(GameBinaries)" Produces="#OutputFiles">
-			<Tag Files="#ToolBinaries;#EditorBinaries;$(GameBinaries)" Except=".../Intermediate/..." With="#OutputFiles"/>			
+		<Node Name="Tag Output Files" Requires="#ToolBinaries;#EditorBinaries;#TargetInfo;$(GameBinaries)" Produces="#OutputFiles">
+			<Tag Files="#ToolBinaries;#EditorBinaries;$(GameBinaries)" Except=".../Intermediate/..." With="#OutputFiles"/>
+			<Tag Files="#TargetInfo" With="#OutputFiles"/>
 		</Node>
 
 		<!-- Copy all the files to the output directory -->
diff --git a/Engine/Source/Programs/AutomationTool/BuildGraph/Tasks/QueryTargetsTask.cs b/Engine/Source/Programs/AutomationTool/BuildGraph/Tasks/QueryTargetsTask.cs
new file mode 100644
--- /dev/null
+++ b/Engine/Source/Programs/AutomationTool/BuildGraph/Tasks/QueryTargetsTask.cs
@@ -0,0 +1,148 @@
+// Copyright Epic Games, Inc. All Rights Reserved.
+
+using AutomationTool;
+using AutomationTool.Tasks;
+using EpicGames.BuildGraph;
+using System;
+using System.Collections.Generic;
+using System.Linq;
+using System.Text;
+using System.Threading.Tasks;
+using System.Xml;
+using EpicGames.Core;
+using UnrealBuildTool;
+using UnrealBuildBase;
+
+#nullable enable
+
+namespace AutomationTool.Tasks
+{
+	/// <summary>
+	/// Parameters to query for all the targets in a project
+	/// </summary>
+	public class QueryTargetsTaskParameters
+	{
+		/// <summary>
+		/// Path to the project file to query
+		/// </summary>
+		[TaskParameter(Optional = true)]
+		public FileReference? ProjectFile;
+
+		/// <summary>
+		/// Path to the output file to receive information about the targets
+		/// </summary>
+		[TaskParameter(Optional = true)]
+		public FileReference? OutputFile = null;
+
+		/// <summary>
+		/// Write out all targets, even if a default is specified in the BuildSettings section of the Default*.ini files. 
+		/// </summary>
+		[TaskParameter(Optional = true)]
+		public bool IncludeAllTargets = false;
+
+		/// <summary>
+		/// Tag to be applied to build products of this task.
+		/// </summary>
+		[TaskParameter(Optional = true, ValidationType = TaskParameterValidationType.TagList)]
+		public string? Tag;
+	}
+
+	/// <summary>
+	/// Runs UBT to query all the targets for a particular project
+	/// </summary>
+	[TaskElement("QueryTargets", typeof(QueryTargetsTaskParameters))]
+	public class QueryTargetsTask : BgTaskImpl
+	{
+		/// <summary>
+		/// Parameters for this task
+		/// </summary>
+		QueryTargetsTaskParameters Parameters;
+
+		/// <summary>
+		/// Construct a spawn task
+		/// </summary>
+		/// <param name="InParameters">Parameters for the task</param>
+		public QueryTargetsTask(QueryTargetsTaskParameters InParameters)
+		{
+			Parameters = InParameters;
+		}
+
+		/// <summary>
+		/// Execute the task.
+		/// </summary>
+		/// <param name="Job">Information about the current job</param>
+		/// <param name="BuildProducts">Set of build products produced by this node.</param>
+		/// <param name="TagNameToFileSet">Mapping from tag names to the set of files they include</param>
+		public override Task ExecuteAsync(JobContext Job, HashSet<FileReference> BuildProducts, Dictionary<string, HashSet<FileReference>> TagNameToFileSet)
+		{
+			// Get the output file
+			FileReference? OutputFile = Parameters.OutputFile;
+			if (OutputFile == null)
+			{
+				if (Parameters.ProjectFile == null)
+				{
+					OutputFile = FileReference.Combine(Unreal.EngineDirectory, "Intermediate", "TargetInfo.json");
+				}
+				else
+				{
+					OutputFile = FileReference.Combine(Parameters.ProjectFile.Directory, "Intermediate", "TargetInfo.json");
+				}
+			}
+			FileUtils.ForceDeleteFile(OutputFile);
+
+			// Run UBT to generate the target info
+			List<string> Arguments = new List<string> { "-Mode=QueryTargets" };
+			if (Parameters.ProjectFile != null)
+			{
+				Arguments.Add($"-Project={Parameters.ProjectFile}");
+			}
+			if (Parameters.IncludeAllTargets)
+			{
+				Arguments.Add("-IncludeAllTargets");
+			}
+			CommandUtils.RunUBT(CommandUtils.CmdEnv, Unreal.UnrealBuildToolDllPath, CommandLineArguments.Join(Arguments));
+
+			// Check the output file exists
+			if (!FileReference.Exists(OutputFile))
+			{
+				throw new BuildException($"Missing {OutputFile}");
+			}
+
+			// Apply the optional tag to the build products
+			foreach(string TagName in FindTagNamesFromList(Parameters.Tag))
+			{
+				FindOrAddTagSet(TagNameToFileSet, TagName).Add(OutputFile);
+			}
+
+			// Add the target files to the set of build products
+			BuildProducts.Add(OutputFile);
+			return Task.CompletedTask;
+		}
+
+		/// <summary>
+		/// Output this task out to an XML writer.
+		/// </summary>
+		public override void Write(XmlWriter Writer)
+		{
+			Write(Writer, Parameters);
+		}
+
+		/// <summary>
+		/// Find all the tags which are used as inputs to this task
+		/// </summary>
+		/// <returns>The tag names which are read by this task</returns>
+		public override IEnumerable<string> FindConsumedTagNames()
+		{
+			return Enumerable.Empty<string>();
+		}
+
+		/// <summary>
+		/// Find all the tags which are modified by this task
+		/// </summary>
+		/// <returns>The tag names which are modified by this task</returns>
+		public override IEnumerable<string> FindProducedTagNames()
+		{
+			return FindTagNamesFromList(Parameters.Tag);
+		}
+	}
+}
