UAT: Fix exceptions not being thrown if copying a file fails.

diff --git a/Engine/Source/Programs/AutomationTool/AutomationUtils/CommandUtils.cs b/Engine/Source/Programs/AutomationTool/AutomationUtils/CommandUtils.cs
--- a/Engine/Source/Programs/AutomationTool/AutomationUtils/CommandUtils.cs
+++ b/Engine/Source/Programs/AutomationTool/AutomationUtils/CommandUtils.cs
@@ -1295,7 +1295,7 @@ namespace AutomationTool
 				}
 			}
 
-			for (int AttemptsRemaining = 5; AttemptsRemaining > 0; --AttemptsRemaining)
+			for (int AttemptsRemaining = 5; AttemptsRemaining >= 0; --AttemptsRemaining)
 			{
 				if (InternalUtils.SafeFileExists(Dest, true))
 				{
