PR #10096: Added new argument for UnzipTask to allow disabling Overwrite files (Contributed by [REDACTED_PERSON])

diff --git a/Engine/Source/Programs/AutomationTool/AutomationUtils/CommandUtils.cs b/Engine/Source/Programs/AutomationTool/AutomationUtils/CommandUtils.cs
--- a/Engine/Source/Programs/AutomationTool/AutomationUtils/CommandUtils.cs
+++ b/Engine/Source/Programs/AutomationTool/AutomationUtils/CommandUtils.cs
@@ -2258,8 +2258,9 @@ namespace AutomationTool
 		/// </summary>
 		/// <param name="ZipFileName">Name of the zip file</param>
 		/// <param name="BaseDirectory">Output directory</param>
+		/// <param name="OverwriteFiles">Whether or not to overwrite files during unzip.</param>
 		/// <returns>List of files written</returns>
-		public static IEnumerable<FileReference> UnzipFiles(FileReference ZipFileName, DirectoryReference BaseDirectory)
+		public static IEnumerable<FileReference> UnzipFiles(FileReference ZipFileName, DirectoryReference BaseDirectory, bool OverwriteFiles = true)
 		{
 			List<FileReference> OutputFiles = new List<FileReference>();
 
@@ -2278,8 +2279,11 @@ namespace AutomationTool
 						continue;
 					}
 					FileReference OutputFile = FileReference.Combine(BaseDirectory, Entry.FullName);
-					DirectoryReference.CreateDirectory(OutputFile.Directory);
-					Entry.ExtractToFile_CrossPlatform(OutputFile.FullName, true);
+					if (OverwriteFiles || File.Exists(OutputFile.FullName) == false)
+					{
+						DirectoryReference.CreateDirectory(OutputFile.Directory);
+						Entry.ExtractToFile_CrossPlatform(OutputFile.FullName, OverwriteFiles);
+					}
 					OutputFiles.Add(OutputFile);
 				}
 			}
@@ -2291,16 +2295,21 @@ namespace AutomationTool
 		/// </summary>
 		/// <param name="ZipFileName">Name of the zip file</param>
 		/// <param name="BaseDirectory">Output directory</param>
+		/// <param name="OverwriteFiles">Whether or not to overwrite files during unzip.</param>
 		/// <returns>List of files written</returns>
-		public static IEnumerable<string> LegacyUnzipFiles(string ZipFileName, string BaseDirectory)
+		public static IEnumerable<string> LegacyUnzipFiles(string ZipFileName, string BaseDirectory, bool OverwriteFiles = true)
 		{
 			List<string> OutputFileNames = new List<string>();
 			if (!RuntimePlatform.IsWindows)
 			{
 				CommandUtils.CreateDirectory(BaseDirectory);
 
+				// -u  update files, create if necessary
+				// -o  overwrite files WITHOUT prompting
+				// -n  never overwrite existing files
+				string OverwriteFileUnzipArg = (OverwriteFiles) ? " -u -o" : " -n";
 				// Use system unzip tool as there have been instances of Ionic not being able to open zips created with Mac zip tool
-				string Output = CommandUtils.RunAndLog("unzip", "\"" + ZipFileName + "\" -d \"" + BaseDirectory + "\"", Options: ERunOptions.Default | ERunOptions.SpewIsVerbose);
+				string Output = CommandUtils.RunAndLog("unzip", "\"" + ZipFileName + "\" -d \"" + BaseDirectory + "\"" + OverwriteFileUnzipArg, Options: ERunOptions.Default | ERunOptions.SpewIsVerbose);
 
 				// Split log output into lines
 				string[] Lines = Output.Split(new char[] { '\n', '\r' });
@@ -2337,14 +2346,17 @@ namespace AutomationTool
 				// but that problem is now fixed. Leaving this code as is as we need to return the list of created files anyway.
 				using (Ionic.Zip.ZipFile Zip = new Ionic.Zip.ZipFile(ZipFileName))
 				{
-
+					FileMode OutputFileMode = (OverwriteFiles) ? FileMode.Create : FileMode.CreateNew;
 					foreach (Ionic.Zip.ZipEntry Entry in Zip.Entries.Where(x => !x.IsDirectory))
 					{
 						string OutputFileName = Path.Combine(BaseDirectory, Entry.FileName);
-						Directory.CreateDirectory(Path.GetDirectoryName(OutputFileName));
-						using (FileStream OutputStream = new FileStream(OutputFileName, FileMode.Create, FileAccess.Write))
+						if (OverwriteFiles || File.Exists(OutputFileName) == false)
 						{
-							Entry.Extract(OutputStream);
+							Directory.CreateDirectory(Path.GetDirectoryName(OutputFileName));
+							using (FileStream OutputStream = new FileStream(OutputFileName, OutputFileMode, FileAccess.Write))
+							{
+								Entry.Extract(OutputStream);
+							}
 						}
 						OutputFileNames.Add(OutputFileName);
 					}
diff --git a/Engine/Source/Programs/AutomationTool/BuildGraph/Tasks/UnzipTask.cs b/Engine/Source/Programs/AutomationTool/BuildGraph/Tasks/UnzipTask.cs
--- a/Engine/Source/Programs/AutomationTool/BuildGraph/Tasks/UnzipTask.cs
+++ b/Engine/Source/Programs/AutomationTool/BuildGraph/Tasks/UnzipTask.cs
@@ -37,6 +37,12 @@ namespace AutomationTool.Tasks
 		[TaskParameter(Optional = true)]
 		public bool UseLegacyUnzip = false;
 
+		/// <summary>
+		/// Whether or not to overwrite files during unzip.
+		/// </summary>
+		[TaskParameter(Optional = true)]
+		public bool OverwriteFiles = true;
+
 		/// <summary>
 		/// Tag to be applied to the extracted files.
 		/// </summary>
@@ -83,11 +89,11 @@ namespace AutomationTool.Tasks
 			{
 				if (Parameters.UseLegacyUnzip)
 				{
-					OutputFiles.UnionWith(CommandUtils.LegacyUnzipFiles(ZipFile.FullName, ToDir.FullName).Select(x => new FileReference(x)));
+					OutputFiles.UnionWith(CommandUtils.LegacyUnzipFiles(ZipFile.FullName, ToDir.FullName, Parameters.OverwriteFiles).Select(x => new FileReference(x)));
 				}
 				else
 				{
-					OutputFiles.UnionWith(CommandUtils.UnzipFiles(ZipFile, ToDir));
+					OutputFiles.UnionWith(CommandUtils.UnzipFiles(ZipFile, ToDir, Parameters.OverwriteFiles));
 				}
 			}
 
