TargetDeviceAndroid: Adding extra checks so that the Timeout flow doesn't start extra adb instances we don't need after the device has already been disposed of

diff --git a/Engine/Source/Programs/AutomationTool/Gauntlet/Platform/Android/Gauntlet.TargetDeviceAndroid.cs b/Engine/Source/Programs/AutomationTool/Gauntlet/Platform/Android/Gauntlet.TargetDeviceAndroid.cs
--- a/Engine/Source/Programs/AutomationTool/Gauntlet/Platform/Android/Gauntlet.TargetDeviceAndroid.cs
+++ b/Engine/Source/Programs/AutomationTool/Gauntlet/Platform/Android/Gauntlet.TargetDeviceAndroid.cs
@@ -10,6 +10,7 @@ using System.Linq;
 using System.Text;
 using System.Threading;
 using EpicGames.Core;
+using System.Diagnostics;
 
 namespace Gauntlet
 {
@@ -105,6 +106,11 @@ namespace Gauntlet
 				return false;
 			}
 
+			if(AndroidDevice.Disposed)
+			{
+				return false;
+			}
+
 			ActivityCheckTime = DateTime.UtcNow;
 
 			// get activities filtered by our package name
@@ -201,7 +207,7 @@ namespace Gauntlet
 
 		public void Kill()
 		{
-			if (!HasExited)
+			if (!HasExited && !AndroidDevice.Disposed)
 			{
 				WasKilled = true;
 				Install.AndroidDevice.KillRunningProcess(Install.AndroidPackageName);
@@ -1397,10 +1403,15 @@ namespace Gauntlet
 					Reset();
 					KillAdbServer();
 					// Kill ADB server, just as a safety measure to ensure it closes
-					IProcessResult TaskkillResult = CommandUtils.Run("taskkill", "/f /im adb.exe", null, CommandUtils.ERunOptions.NoLoggingOfRunCommand);
-					if (TaskkillResult.Output.Contains("success", StringComparison.OrdinalIgnoreCase))
+					IEnumerable<Process> ADBProcesses = Process.GetProcesses().Where(p => p.ProcessName.Equals("adb"));
+					if (ADBProcesses.Count() > 0)
 					{
-						Log.Info(TaskkillResult.Output);
+						Log.Info("Terminating {0} ADB Process(es)", ADBProcesses.Count());
+						foreach (Process ADBProcess in ADBProcesses)
+						{
+							Log.Info("Killing ADB process {0}", ADBProcess.Id);
+							ADBProcess.Kill();
+						}
 					}
 				}
 			}
