Improve Gauntlet behavior, promote errors to the output in case if 'InitializationFailure' is happening

diff --git a/Engine/Source/Programs/AutomationTool/Gauntlet/Unreal/Base/Gauntlet.UnrealTestNode.cs b/Engine/Source/Programs/AutomationTool/Gauntlet/Unreal/Base/Gauntlet.UnrealTestNode.cs
--- a/Engine/Source/Programs/AutomationTool/Gauntlet/Unreal/Base/Gauntlet.UnrealTestNode.cs
+++ b/Engine/Source/Programs/AutomationTool/Gauntlet/Unreal/Base/Gauntlet.UnrealTestNode.cs
@@ -1656,7 +1656,7 @@ namespace Gauntlet
 		/// <param name="InRoleArtifacts"></param>
 		/// <param name="InLog"></param>
 		/// <returns></returns>
-		protected virtual IEnumerable<UnrealTestEvent> CreateEventListFromArtifact(StopReason InReason, UnrealRoleArtifacts InRoleArtifacts, UnrealLog InLog)
+		protected virtual IEnumerable<UnrealTestEvent> CreateEventListFromArtifact(StopReason InReason, UnrealRoleArtifacts InRoleArtifacts, UnrealLog InLog, UnrealProcessResult ProcessResult)
 		{
 			List<UnrealTestEvent> EventList = new List<UnrealTestEvent>();
 
@@ -1682,7 +1682,10 @@ namespace Gauntlet
 			}
 
 			bool TrackAllWarnings = GetCachedConfiguration().ShowWarningsInSummary || Flags.HasFlag(BehaviorFlags.PromoteWarnings);
-			bool TrackAllErrors = GetCachedConfiguration().ShowErrorsInSummary || Flags.HasFlag(BehaviorFlags.PromoteErrors);
+			// if we are getting an initialization failure it is still not clear what the reason was, may differ from one case to another,
+			// so we would like to enforce all-error tracking in this situation, even if it is disabled in the config
+			bool TrackAllErrors = GetCachedConfiguration().ShowErrorsInSummary || Flags.HasFlag(BehaviorFlags.PromoteErrors)
+				|| ProcessResult == UnrealProcessResult.InitializationFailure;
 
 			// now look at the log. Add events for warnings/errors if the category is monitored or if this test is flagged to 
 			// promote all warnings/errors
@@ -1722,7 +1725,7 @@ namespace Gauntlet
 			// Give ourselves (and derived classes) a chance to analyze what happened
 			UnrealProcessResult ProcessResult = GetExitCodeAndReason(InReason, LogSummary, InRoleArtifacts, out ExitReason, out ExitCode);
 
-			IEnumerable<UnrealTestEvent> EventList = CreateEventListFromArtifact(InReason, InRoleArtifacts, LogSummary);
+			IEnumerable<UnrealTestEvent> EventList = CreateEventListFromArtifact(InReason, InRoleArtifacts, LogSummary, ProcessResult);
 
 			// if the test is stopping for a reason other than completion, mark this as failing incase derived classes
 			// don't do the right thing
