Fix order of arguments to unzip command on Mac.

diff --git a/Engine/Source/Programs/AutomationTool/AutomationUtils/CommandUtils.cs b/Engine/Source/Programs/AutomationTool/AutomationUtils/CommandUtils.cs
--- a/Engine/Source/Programs/AutomationTool/AutomationUtils/CommandUtils.cs
+++ b/Engine/Source/Programs/AutomationTool/AutomationUtils/CommandUtils.cs
@@ -2307,9 +2307,9 @@ namespace AutomationTool
 				// -u  update files, create if necessary
 				// -o  overwrite files WITHOUT prompting
 				// -n  never overwrite existing files
-				string OverwriteFileUnzipArg = (OverwriteFiles) ? " -u -o" : " -n";
+				string OverwriteFileUnzipArg = (OverwriteFiles) ? "-u -o" : "-n";
 				// Use system unzip tool as there have been instances of Ionic not being able to open zips created with Mac zip tool
-				string Output = CommandUtils.RunAndLog("unzip", "\"" + ZipFileName + "\" -d \"" + BaseDirectory + "\"" + OverwriteFileUnzipArg, Options: ERunOptions.Default | ERunOptions.SpewIsVerbose);
+				string Output = CommandUtils.RunAndLog("unzip", OverwriteFileUnzipArg + " \"" + ZipFileName + "\" -d \"" + BaseDirectory + "\"", Options: ERunOptions.Default | ERunOptions.SpewIsVerbose);
 
 				// Split log output into lines
 				string[] Lines = Output.Split(new char[] { '\n', '\r' });
