- Fixed GenerateDSYM.Automation.cs for UnrealArch/ActiveArchitectures support

diff --git a/Engine/Source/Programs/AutomationTool/Scripts/GenerateDSYM.Automation.cs b/Engine/Source/Programs/AutomationTool/Scripts/GenerateDSYM.Automation.cs
--- a/Engine/Source/Programs/AutomationTool/Scripts/GenerateDSYM.Automation.cs
+++ b/Engine/Source/Programs/AutomationTool/Scripts/GenerateDSYM.Automation.cs
@@ -30,15 +30,14 @@ public class GenerateDSYM : BuildCommand
 				string PlatformName = ParseParamValue("platform");
 				FileReference ProjectFile = ParseProjectParam();
 				string TargetName = ParseOptionalStringParam("target") ?? System.IO.Path.GetFileNameWithoutExtension(ProjectFile?.FullName ?? "");
-				string ArchitectureString = ParseOptionalStringParam("architecture") ?? System.IO.Path.GetFileNameWithoutExtension(ProjectFile?.FullName ?? "");
 				UnrealTargetConfiguration Configuration = ParseOptionalEnumParam<UnrealTargetConfiguration>("config") ?? UnrealTargetConfiguration.Development;
 
 				if (ProjectFile == null || PlatformName == null || TargetName == null)
 				{
 					Log.TraceError("Must specify a file(s) with -file=, or other parameters to find the binaries:-project=<Path or name of project>\n" +
-						"-platform=<Mac|IOS|TVOS>\n-target=<TargetName> (Optional, defaults to the ProjectName)\n" + 
+						"-platform=<Mac|IOS|TVOS>\n" + 
+						"-target=<TargetName> (Optional, defaults to the ProjectName)\n" + 
 						"-config=<Debug|DebugGame|Development|Test|Shipping> (Optional - defaults to Development)\n" +
-						"-architecture= (Optional, defaults to no arch specified)\n\n" +
 						"Ex: -project=EngineTest -platform=Mac -target=EngineTestEditor -config=Test\n\n"+
 						"Other options:\n" +
 						"-flat (Options, if specified, the .dSYMs will be flat files that are easier to copy around between computers/servers/etc\n\n");
@@ -51,10 +50,12 @@ public class GenerateDSYM : BuildCommand
 					Log.TraceError("Platform must be one of Mac, IOS, TVOS");
 					return;
 				}
-				UnrealArchitectures Architecture = UnrealArchitectures.FromString(ArchitectureString, Platform);
 
 
-				FileReference ReceiptFile = TargetReceipt.GetDefaultPath(DirectoryReference.FromFile(ProjectFile) ?? Unreal.EngineDirectory, TargetName, Platform, Configuration, Architecture);
+				// Apple platforms don't need the architecture pass in because they RequiresArchitectureFilenames returns false, but if that ever changes, we pass in the active architecture anyway
+				UnrealArchitectures Architectures = UnrealArchitectureConfig.ForPlatform(Platform).ActiveArchitectures(ProjectFile, TargetName);
+
+				FileReference ReceiptFile = TargetReceipt.GetDefaultPath(DirectoryReference.FromFile(ProjectFile) ?? Unreal.EngineDirectory, TargetName, Platform, Configuration, Architectures);
 				TargetReceipt Receipt = TargetReceipt.Read(ReceiptFile);
 				// get the products that we can dsym (duylibs and executables)
 				IEnumerable<BuildProduct> Products = Receipt.BuildProducts.Where(x => x.Type == BuildProductType.Executable || x.Type == BuildProductType.DynamicLibrary);
