- Fixed targets with names like Foo_Mac to not think they are platform extensions if they aren't under a Mac directory (platform extensions will be named with Platform name always)

diff --git a/Engine/Source/Programs/AutomationTool/AutomationUtils/ProjectUtils.cs b/Engine/Source/Programs/AutomationTool/AutomationUtils/ProjectUtils.cs
--- a/Engine/Source/Programs/AutomationTool/AutomationUtils/ProjectUtils.cs
+++ b/Engine/Source/Programs/AutomationTool/AutomationUtils/ProjectUtils.cs
@@ -784,9 +784,12 @@ namespace AutomationTool
 				{
 					if (FullProjectPath == null || TargetScript.IsUnderDirectory(new DirectoryReference(FullProjectPath)))
 					{
-						// skip target rules that are platform extension or platform group specializations
+						// skip target rules that are platform extension or platform group specializations (don't treat _<Platform> targets as extensions if not under a <Platform> directory)
 						string[] TargetPathSplit = TargetScript.GetFileNameWithoutAnyExtensions().Split(new char[]{'_'}, StringSplitOptions.RemoveEmptyEntries );
-						if (TargetPathSplit.Length > 1 && (UnrealTargetPlatform.IsValidName(TargetPathSplit.Last()) || UnrealPlatformGroup.IsValidName(TargetPathSplit.Last()) ) )
+						if (TargetPathSplit.Length > 1 && 
+							(UnrealTargetPlatform.IsValidName(TargetPathSplit.Last()) || UnrealPlatformGroup.IsValidName(TargetPathSplit.Last())) &&
+							// platform extension targets will always be under a directory of that platform/group name
+							TargetScript.ContainsName(TargetPathSplit.Last(), 0))
 						{
 							continue;
 						}
