UGS: Modify build script to support uploading CLI clients to Horde.

diff --git a/Engine/Source/Programs/AutomationTool/BuildGraph/Tasks/DeployToolTask.cs b/Engine/Source/Programs/AutomationTool/BuildGraph/Tasks/DeployToolTask.cs
new file mode 100644
--- /dev/null
+++ b/Engine/Source/Programs/AutomationTool/BuildGraph/Tasks/DeployToolTask.cs
@@ -0,0 +1,185 @@
+// Copyright Epic Games, Inc. All Rights Reserved.
+
+using AutomationTool;
+using EpicGames.Core;
+using System;
+using System.Collections.Generic;
+using System.Text.Json;
+using System.IO;
+using System.Net.Http;
+using System.Net.Http.Headers;
+using System.Text;
+using System.Threading.Tasks;
+using System.Xml;
+using UnrealBuildBase;
+
+#nullable enable
+
+namespace AutomationTool.Tasks
+{
+	/// <summary>
+	/// Parameters for a DeployTool task
+	/// </summary>
+	public class DeployToolTaskParameters
+	{
+		/// <summary>
+		/// Identifier for the tool
+		/// </summary>
+		[TaskParameter]
+		public string Id = String.Empty;
+
+		/// <summary>
+		/// Settings file to use for the deployment. Should be a JSON file containing server name and access token.
+		/// </summary>
+		[TaskParameter]
+		public string Settings = String.Empty;
+
+		/// <summary>
+		/// Zip file containing the data to upload
+		/// </summary>
+		[TaskParameter(Optional = true)]
+		public string? File;
+
+		/// <summary>
+		/// Directory to zip and upload for the tool
+		/// </summary>
+		[TaskParameter(Optional = true)]
+		public string? Directory;
+	}
+
+	/// <summary>
+	/// Deploys a tool update through Horde
+	/// </summary>
+	[TaskElement("DeployTool", typeof(DeployToolTaskParameters))]
+	public class DeployToolTask : SpawnTaskBase
+	{
+		class DeploySettings
+		{
+			public string Server { get; set; } = String.Empty;
+			public string? Token { get; set; }
+		}
+
+		/// <summary>
+		/// Parameters for this task
+		/// </summary>
+		DeployToolTaskParameters Parameters;
+
+		/// <summary>
+		/// Construct a Helm task
+		/// </summary>
+		/// <param name="InParameters">Parameters for the task</param>
+		public DeployToolTask(DeployToolTaskParameters InParameters)
+		{
+			Parameters = InParameters;
+		}
+
+		/// <summary>
+		/// Execute the task.
+		/// </summary>
+		/// <param name="Job">Information about the current job</param>
+		/// <param name="BuildProducts">Set of build products produced by this node.</param>
+		/// <param name="TagNameToFileSet">Mapping from tag names to the set of files they include</param>
+		public override async Task ExecuteAsync(JobContext Job, HashSet<FileReference> BuildProducts, Dictionary<string, HashSet<FileReference>> TagNameToFileSet)
+		{
+			FileReference? file;
+			if (Parameters.File != null && Parameters.Directory == null)
+			{
+				file = ResolveFile(Parameters.File);
+			}
+			else if (Parameters.File == null && Parameters.Directory != null)
+			{
+				file = FileReference.Combine(Unreal.EngineDirectory, "Intermediate", "DeployTool", "Tool.zip");
+				CommandUtils.DeleteFile(file);
+				CommandUtils.ZipFiles(file, ResolveDirectory(Parameters.Directory), new FileFilter(FileFilterType.Include));
+			}
+			else
+			{
+				throw new AutomationException("Either file or Directory must be specified to the DeployTool task (not both).");
+			}
+
+			FileReference settingsFile = ResolveFile(Parameters.Settings);
+			if (!FileReference.Exists(settingsFile))
+			{
+				throw new AutomationException($"Settings file '{settingsFile}' does not exist");
+			}
+
+			byte[] settingsData = await FileReference.ReadAllBytesAsync(settingsFile);
+			JsonSerializerOptions jsonOptions = new JsonSerializerOptions { AllowTrailingCommas = true, ReadCommentHandling = JsonCommentHandling.Skip, PropertyNameCaseInsensitive = true };
+
+			DeploySettings? settings = JsonSerializer.Deserialize<DeploySettings>(settingsData, jsonOptions);
+			if (settings == null)
+			{
+				throw new AutomationException($"Unable to read settings file {settingsFile}");
+			}
+			else if (settings.Server == null)
+			{
+				throw new AutomationException($"Missing 'server' key from {settingsFile}");
+			}
+
+			using (HttpClient httpClient = new HttpClient())
+			{
+				httpClient.BaseAddress = new Uri(settings.Server);
+				using (HttpRequestMessage request = new HttpRequestMessage(HttpMethod.Post, $"api/v1/tools/{Parameters.Id}/deployments"))
+				{
+					if (settings.Token != null)
+					{
+						request.Headers.Authorization = new AuthenticationHeaderValue("Bearer", settings.Token);
+					}
+
+					using (FileStream stream = FileReference.Open(file, FileMode.Open, FileAccess.Read))
+					{
+						StreamContent streamContent = new StreamContent(stream);
+
+						MultipartFormDataContent content = new MultipartFormDataContent();
+						content.Add(streamContent, "file", file.GetFileName());
+						request.Content = content;
+
+						Log.TraceInformation("Uploading {0} to {1}", file, request.RequestUri);
+						using (HttpResponseMessage response = await httpClient.SendAsync(request))
+						{
+							if (!response.IsSuccessStatusCode)
+							{
+								string? responseContent;
+								try
+								{
+									responseContent = await response.Content.ReadAsStringAsync();
+								}
+								catch
+								{
+									responseContent = "(No message)";
+								}
+								throw new AutomationException($"Upload failed ({response.StatusCode}): {responseContent}");
+							}
+						}
+					}
+				}
+			}
+		}
+
+		/// <summary>
+		/// Output this task out to an XML writer.
+		/// </summary>
+		public override void Write(XmlWriter Writer)
+		{
+			Write(Writer, Parameters);
+		}
+
+		/// <summary>
+		/// Find all the tags which are used as inputs to this task
+		/// </summary>
+		/// <returns>The tag names which are read by this task</returns>
+		public override IEnumerable<string> FindConsumedTagNames()
+		{
+			yield break;
+		}
+
+		/// <summary>
+		/// Find all the tags which are modified by this task
+		/// </summary>
+		/// <returns>The tag names which are modified by this task</returns>
+		public override IEnumerable<string> FindProducedTagNames()
+		{
+			yield break;
+		}
+	}
+}
diff --git a/Engine/Source/Programs/UnrealGameSync/BuildUGS.xml b/Engine/Source/Programs/UnrealGameSync/BuildUGS.xml
--- a/Engine/Source/Programs/UnrealGameSync/BuildUGS.xml
+++ b/Engine/Source/Programs/UnrealGameSync/BuildUGS.xml
@@ -7,10 +7,14 @@
 	<Option Name="SubmitToBranch" DefaultValue="" Description="Branch to submit binaries to"/>
 	<Option Name="SubmitPath" DefaultValue="" Description="Relative path within the stream/branch to submit binaries"/>
 	<Option Name="IncludeInstaller" DefaultValue="false" Description="Whether to include the launcher installer in the regular client distribution"/>
+	<Option Name="HordeSettings" DefaultValue="" Description="Path to a JSON file containing settings for publishing to Horde"/>
 
 	<Property Name="UgsDir" Value="$(RootDir)/Engine/Source/Programs/UnrealGameSync"/>
 	<Property Name="UgsOutputDir" Value="$(RootDir)/LocalBuilds/UnrealGameSync"/>
 	<Property Name="UgsZipOutputDir" Value="$(RootDir)/LocalBuilds/UnrealGameSyncZip"/>
+	<Property Name="UgsCmdOutputDir" Value="$(RootDir)/LocalBuilds/UnrealGameSyncCmd"/>
+	<Property Name="UgsCmdMacOutputDir" Value="$(UgsCmdOutputDir)/Mac"/>
+	<Property Name="UgsCmdLinuxOutputDir" Value="$(UgsCmdOutputDir)/Linux"/>
 	<Property Name="WixDir" Value="$(RootDir)/Engine/Source/ThirdParty/WiX/3.8"/>
 
 	<Property Name="Version" Value="$(EngineMajorVersion).$(EngineMinorVersion).$(EnginePatchVersion)"/>
@@ -64,6 +68,24 @@
 			<Zip FromDir="$(UgsOutputDir)" ZipFile="$(UgsZipPublishDir)/Release.zip"/>
 			<Submit Description="Updating UGS from CL $(Change) (Zip)" Workspace="Temp" Files="$(UgsZipPublishDir)/Release.zip" RootDir="$(UgsZipOutputDir)" Stream="$(SubmitToStream)" Branch="$(SubmitToBranch)"/>
 		</Node>
+		
+		<Node Name="Build UnrealGameSyncCmd Mac">
+			<Delete Files="$(UgsCmdMacOutputDir)/..."/>
+			<DotNet Arguments="publish &quot;$(UgsDir)/UnrealGameSyncCmd/UnrealGameSyncCmd.csproj&quot; --output=&quot;$(UgsCmdMacOutputDir)&quot; --no-self-contained -c Release -r osx-x64 $(VersionArguments)"/>
+		</Node>
+		<Node Name="Publish UnrealGameSyncCmd Mac" Requires="Build UnrealGameSyncCmd Mac">
+			<Error Message="Missing -HordeSettings=... argument" If="'$(HordeSettings)' == ''"/>
+			<DeployTool Id="ugs-mac" Settings="$(HordeSettings)" Directory="$(UgsCmdMacOutputDir)"/>
+		</Node>
+		
+		<Node Name="Build UnrealGameSyncCmd Linux">
+			<Delete Files="$(UgsCmdLinuxOutputDir)/..."/>
+			<DotNet Arguments="publish &quot;$(UgsDir)/UnrealGameSyncCmd/UnrealGameSyncCmd.csproj&quot; --output=&quot;$(UgsCmdLinuxOutputDir)&quot; --no-self-contained -c Release -r linux-x64 $(VersionArguments)"/>
+		</Node>
+		<Node Name="Publish UnrealGameSyncCmd Linux" Requires="Build UnrealGameSyncCmd Linux">
+			<Error Message="Missing -HordeSettings=... argument" If="'$(HordeSettings)' == ''"/>
+			<DeployTool Id="ugs-linux" Settings="$(HordeSettings)" Directory="$(UgsCmdMacOutputDir)"/>
+		</Node>
 	</Agent>
 
 	<Aggregate Name="Publish Unstable With Zip" Requires="Publish Unstable;Publish Unstable Zip"/>
