Fix issues with commas within quotes in -clientargs. -clientargs="-testvalues=\"a,b,c\" " will now work.

diff --git a/Engine/Source/Programs/AutomationTool/Gauntlet/Unreal/RunUnreal.cs b/Engine/Source/Programs/AutomationTool/Gauntlet/Unreal/RunUnreal.cs
--- a/Engine/Source/Programs/AutomationTool/Gauntlet/Unreal/RunUnreal.cs
+++ b/Engine/Source/Programs/AutomationTool/Gauntlet/Unreal/RunUnreal.cs
@@ -214,10 +214,10 @@ namespace Gauntlet
 
 					// look for-args= and then -clientargs= and -editorargs etc
 					Role.ExtraArgs = string.Join(' ', Globals.Params.ParseValues("Args", true));
-					IEnumerable<string> ExtraRoleArgs = Globals.Params.ParseValues(Type.ToString() + "Args", true);
-					if (ExtraRoleArgs.Count() > 0)
+					string ExtraRoleArgs = Globals.Params.ParseValue(Type.ToString() + "Args", string.Empty);
+					if (!string.IsNullOrEmpty(ExtraRoleArgs))
 					{
-						Role.ExtraArgs += " " + string.Join(' ', ExtraRoleArgs);
+						Role.ExtraArgs += " " + ExtraRoleArgs;
 					}
 
 					// look for -clientexeccmds=, -editorexeccmds= etc, these are separate from clientargs for sanity
