* Added an interface that contains ParamHelp attributes for ProjectParam so we can reflect the parameters in other tools.
* Updated LauncherTool to traverse the class interfaces looking for attributes. Using GetCustomAttributes is not good enough because it will not find the HelpAttributes because the base attribute class(DescriptionAttribute) does not explicitly mark itself that the attribute is inherited.

diff --git a/Engine/Source/Programs/AutomationTool/AutomationUtils/ProjectParams.cs b/Engine/Source/Programs/AutomationTool/AutomationUtils/ProjectParams.cs
--- a/Engine/Source/Programs/AutomationTool/AutomationUtils/ProjectParams.cs
+++ b/Engine/Source/Programs/AutomationTool/AutomationUtils/ProjectParams.cs
@@ -11,6 +11,24 @@ using UnrealBuildBase;
 
 namespace AutomationTool
 {
+	[ParamHelp("platform", "Platform of the target", ParamType = typeof(string))]
+	[ParamHelp("config", "Config of the target", ParamType = typeof(string))]
+	[ParamHelp("project", "Config of the target", ParamType = typeof(string))]
+	[ParamHelp("client", "Build, cook and run a client and a server, uses client target configuration", ParamType = typeof(bool))]
+	[ParamHelp("noclient", "Do not run the client, just run the server", ParamType = typeof(bool))]
+	[ParamHelp("server", "Is this a server target?", ParamType = typeof(bool))]
+	[ParamHelp("build", "True if build step should be executed", ParamType = typeof(bool))]
+	[ParamHelp("ubtargs", "Extra options to pass to ubt")]
+	[ParamHelp("cook", "Determines if the build is going to use cooked data", ParamType = typeof(bool))]
+	[ParamHelp("AdditionalCookerOptions", "Additional arguments sent to the cooking step", ParamType = typeof(string))]
+	[ParamHelp("pak", "Generate a pak file", ParamType = typeof(bool))]
+	[ParamHelp("deploy", "Deploy the project for the target platform", ParamType = typeof(bool))]
+	[ParamHelp("stage", "Put this build in a stage director", ParamType = typeof(bool))]
+	[ParamHelp("run", "Run the game after it is built (including server, if -server)", ParamType = typeof(bool))]
+	public interface IProjectParamsHelpers
+	{
+	}
+
 	[Help("targetplatform=PlatformName", "target platform for building, cooking and deployment (also -Platform)")]
 	[Help("servertargetplatform=PlatformName", "target platform for building, cooking and deployment of the dedicated server (also -ServerPlatform)")]
 	public class ProjectParams
diff --git a/Engine/Source/Programs/AutomationTool/Scripts/BuildCookRun.Automation.cs b/Engine/Source/Programs/AutomationTool/Scripts/BuildCookRun.Automation.cs
--- a/Engine/Source/Programs/AutomationTool/Scripts/BuildCookRun.Automation.cs
+++ b/Engine/Source/Programs/AutomationTool/Scripts/BuildCookRun.Automation.cs
@@ -20,7 +20,7 @@ If no DefaultMap can be found, the command falls back to /Engine/Maps/Entry.")]
 [Help(typeof(ProjectParams))]
 [Help(typeof(UnrealBuild))]
 [Help(typeof(CodeSign))]
-public class BuildCookRun : BuildCommand
+public class BuildCookRun : BuildCommand, IProjectParamsHelpers
 {
 	public override void ExecuteBuild()
 	{
