UAT: Include a version number when deploying a new tool version.

diff --git a/Engine/Source/Programs/AutomationTool/BuildGraph/Tasks/DeployToolTask.cs b/Engine/Source/Programs/AutomationTool/BuildGraph/Tasks/DeployToolTask.cs
--- a/Engine/Source/Programs/AutomationTool/BuildGraph/Tasks/DeployToolTask.cs
+++ b/Engine/Source/Programs/AutomationTool/BuildGraph/Tasks/DeployToolTask.cs
@@ -34,6 +34,12 @@ namespace AutomationTool.Tasks
 		[TaskParameter]
 		public string Settings = String.Empty;
 
+		/// <summary>
+		/// Version number for the new tool
+		/// </summary>
+		[TaskParameter]
+		public string Version = String.Empty;
+
 		/// <summary>
 		/// Zip file containing the data to upload
 		/// </summary>
@@ -131,6 +137,7 @@ namespace AutomationTool.Tasks
 						StreamContent streamContent = new StreamContent(stream);
 
 						MultipartFormDataContent content = new MultipartFormDataContent();
+						content.Add(new StringContent(Parameters.Version), "version");
 						content.Add(streamContent, "file", file.GetFileName());
 						request.Content = content;
 
