CIS: Add missing copyright notice to Virtualization.Automation.cs

#rnx

diff --git a/Engine/Source/Programs/AutomationTool/Scripts/Virtualization.Automation.cs b/Engine/Source/Programs/AutomationTool/Scripts/Virtualization.Automation.cs
--- a/Engine/Source/Programs/AutomationTool/Scripts/Virtualization.Automation.cs
+++ b/Engine/Source/Programs/AutomationTool/Scripts/Virtualization.Automation.cs
@@ -1,3 +1,5 @@
+// Copyright Epic Games, Inc. All Rights Reserved.
+
 using System;
 using System.Collections.Generic;
 using System.Linq;
