Gauntlet roles can be marked as 'Deferred Launch' so they must be launched manually from TickTest.

 - this can be used to update the command line for a particular role based on output from another role, for example.
 - add new interface to Windows app install to allow the command line to be updated before launch.

#rnx

diff --git a/Engine/Source/Programs/AutomationTool/Gauntlet/Framework/Base/Gauntlet.AppInstall.cs b/Engine/Source/Programs/AutomationTool/Gauntlet/Framework/Base/Gauntlet.AppInstall.cs
--- a/Engine/Source/Programs/AutomationTool/Gauntlet/Framework/Base/Gauntlet.AppInstall.cs
+++ b/Engine/Source/Programs/AutomationTool/Gauntlet/Framework/Base/Gauntlet.AppInstall.cs
@@ -26,5 +26,22 @@ namespace Gauntlet
 		/// The device that we'll be run on
 		/// </summary>
 		ITargetDevice Device { get; }
+
+
+
+		/// <summary>
+		/// IAppInstallDynamicCommandLine can be implemented by IAppInstall classes where the platform
+		/// allows the command line to be updated after installation
+		/// </summary>
+		public interface IDynamicCommandLine
+		{
+			/// <summary>
+			/// Appends the given string to the installed application's command line
+			/// </summary>
+			/// <param name="AdditionalCommandLine"></param>
+			void AppendCommandline(string AdditionalCommandLine);
+		}
 	}
+
+
 }
diff --git a/Engine/Source/Programs/AutomationTool/Gauntlet/Platform/Windows/Gauntlet.TargetDeviceWindows.cs b/Engine/Source/Programs/AutomationTool/Gauntlet/Platform/Windows/Gauntlet.TargetDeviceWindows.cs
--- a/Engine/Source/Programs/AutomationTool/Gauntlet/Platform/Windows/Gauntlet.TargetDeviceWindows.cs
+++ b/Engine/Source/Programs/AutomationTool/Gauntlet/Platform/Windows/Gauntlet.TargetDeviceWindows.cs
@@ -143,7 +143,7 @@ namespace Gauntlet
 	}
 
 
-	class WindowsAppInstall : IAppInstall
+	class WindowsAppInstall : IAppInstall, IAppInstall.IDynamicCommandLine
 	{
 		public string Name { get; private set; }
 
@@ -177,6 +177,11 @@ namespace Gauntlet
 			return Device.Run(this);
 		}
 
+		public void AppendCommandline(string AdditionalCommandline)
+		{
+			CommandArguments += AdditionalCommandline;
+		}
+
 		public bool ForceCleanDeviceArtifacts()
 		{
 			DirectoryInfo ClientTempDirInfo = new DirectoryInfo(ArtifactPath) { Attributes = FileAttributes.Normal };
diff --git a/Engine/Source/Programs/AutomationTool/Gauntlet/Unreal/Base/Gauntlet.UnrealSession.cs b/Engine/Source/Programs/AutomationTool/Gauntlet/Unreal/Base/Gauntlet.UnrealSession.cs
--- a/Engine/Source/Programs/AutomationTool/Gauntlet/Unreal/Base/Gauntlet.UnrealSession.cs
+++ b/Engine/Source/Programs/AutomationTool/Gauntlet/Unreal/Base/Gauntlet.UnrealSession.cs
@@ -127,6 +127,11 @@ namespace Gauntlet
 		/// </summary>
 		public bool InstallOnly { get; set; }
 
+		/// <summary>
+		/// Whether this role will launched by the test node at a later time, typically during TickTest(). By default, all roles are launched immediately.
+		/// </summary>
+		public bool DeferredLaunch { get; set; }
+
 		/// <summary>
 		/// Is this role Null? 
 		/// </summary>
@@ -206,6 +211,7 @@ namespace Gauntlet
 			}
 
 			InstallOnly = false;
+			DeferredLaunch = false;
 			Options = InOptions;
             FilesToCopy = new List<UnrealFileToCopy>();
 			CommandLineParams = new GauntletCommandLine();
@@ -261,10 +267,25 @@ namespace Gauntlet
 			}
 		};
 
+		/// <summary>
+		/// All roles
+		/// </summary>
+		public RoleInstance[] AllRoles { get; protected set; }
+		
 		/// <summary>
 		/// All running roles
 		/// </summary>
-		public RoleInstance[] RunningRoles { get; protected set; }
+		public IEnumerable<RoleInstance> RunningRoles { get { return AllRoles.Where( X => X.AppInstance != null ); } }
+
+		/// <summary>
+		/// All deferred roles
+		/// </summary>
+		public IEnumerable<RoleInstance> DeferredRoles { get { return DeferredRoleToAppInstall.Keys; } }
+
+		/// <summary>
+		/// All deferred roles and their associated IAppInstall
+		/// </summary>
+		public Dictionary<RoleInstance, IAppInstall> DeferredRoleToAppInstall { get; protected set; }
 
 		/// <summary>
 		/// Helper for accessing all client processes. May return an empty array if no clients are involved
@@ -333,12 +354,14 @@ namespace Gauntlet
 		}
 
 		/// <summary>
-		/// COnstructor. Roles must be passed in
+		/// Constructor. Roles must be passed in
 		/// </summary>
-		/// <param name="InRunningRoles"></param>
-		public UnrealSessionInstance(RoleInstance[] InRunningRoles)
+		/// <param name="InAllRoles"></param>
+		/// <param name="InDeferredRoleToAppInstall"></param>
+		public UnrealSessionInstance(RoleInstance[] InAllRoles, Dictionary<RoleInstance, IAppInstall> InDeferredRoleToAppInstall = null)
 		{
-			RunningRoles = InRunningRoles;
+			AllRoles = InAllRoles;
+			DeferredRoleToAppInstall = InDeferredRoleToAppInstall;
 		}
 
 		~UnrealSessionInstance()
@@ -371,6 +394,75 @@ namespace Gauntlet
 		}
 		#endregion
 
+
+		/// <summary>
+		/// Returns the app install for a given role, where the role was marked as 'DeferredLaunch'
+		/// </summary>
+		/// <param name="Role"></param>
+		public IAppInstall FindInstallForDeferredRole( UnrealSessionRole Role )
+		{
+			IAppInstall AppInstall = DeferredRoleToAppInstall
+				.Where(X => X.Key.Role == Role)
+				.Select(X => X.Value)
+				.FirstOrDefault();
+
+			if (AppInstall == null)
+			{
+				Log.Error("Cannot find derferred role {0}", Role.ToString());
+			}
+
+			return AppInstall;
+		}
+
+
+		/// <summary>
+		/// Launches a role that was previously flagged as 'DeferredLaunch'. Note that this does not handle device failure, marking problem devices etc.
+		/// </summary>
+		/// <param name="Role"></param>
+		/// <returns></returns>
+		public bool LaunchDeferredRole( UnrealSessionRole Role )
+		{
+			if (Role.DeferredLaunch == false)
+			{
+				Log.Error("Cannot start deferred role {0} because it is not marked 'DeferredLaunch'", Role.ToString());
+				return false;
+			}
+
+			int CurrentRoleInstanceIndex = AllRoles.FindIndex( X => X.Role == Role );
+			if (CurrentRoleInstanceIndex < 0)
+			{
+				Log.Error("Cannot start dereferred role {0} because it cannot be found among the 'all roles' list", Role.ToString());
+				return false;
+			}
+
+			IAppInstall CurrentInstall = FindInstallForDeferredRole(Role);
+
+			bool Success = false;
+
+			try
+			{
+				Log.Info("Starting deferred {0} on {1}", Role, CurrentInstall.Device);
+				IAppInstance Instance = CurrentInstall.Run();
+				IDeviceUsageReporter.RecordStart(Instance.Device.Name, Instance.Device.Platform, IDeviceUsageReporter.EventType.Test);
+
+				if (Instance != null || Globals.CancelSignalled)
+				{
+					// remove deferred role and update to a running role instance
+					DeferredRoleToAppInstall.Remove(AllRoles[CurrentRoleInstanceIndex]);
+					AllRoles[CurrentRoleInstanceIndex] = new RoleInstance(Role, Instance);
+				}
+
+				Success = true;
+			}
+			catch (DeviceException Ex)
+			{
+				Log.Error("Cannot start deferred role {0} because device {Name} threw an exception during launch. \nException={Exception}", Role.ToString(), CurrentInstall.Device, Ex.Message);
+				Success = false;
+			}
+
+			return Success;
+		}
+
 		/// <summary>
 		/// Shutdown the session by killing any remaining processes.
 		/// </summary>
@@ -845,14 +937,24 @@ namespace Gauntlet
 
 				if (InstallSuccess && Globals.CancelSignalled == false)
 				{
-					List<UnrealSessionInstance.RoleInstance> RunningRoles = new List<UnrealSessionInstance.RoleInstance>();
+					List<UnrealSessionInstance.RoleInstance> AllRoles = new List<UnrealSessionInstance.RoleInstance>();
+					Dictionary<UnrealSessionInstance.RoleInstance, IAppInstall> DeferredRoleToAppInstall = new Dictionary<UnrealSessionInstance.RoleInstance, IAppInstall>();
 
 					// Now try to run all installs on their devices
 					foreach (var InstallRoleKV in InstallsToRoles)
 					{
 						IAppInstall CurrentInstall = InstallRoleKV.Key;
 						if (InstallRoleKV.Value.InstallOnly)
-						{	RunningRoles.Add(new UnrealSessionInstance.RoleInstance(InstallRoleKV.Value, null));
+						{
+							AllRoles.Add(new UnrealSessionInstance.RoleInstance(InstallRoleKV.Value, null));
+							continue;
+						}
+						if (InstallRoleKV.Value.DeferredLaunch)
+						{
+							UnrealSessionInstance.RoleInstance DeferredRoleInstance = new UnrealSessionInstance.RoleInstance(InstallRoleKV.Value, null);
+
+							DeferredRoleToAppInstall.Add(DeferredRoleInstance, CurrentInstall);
+							AllRoles.Add(DeferredRoleInstance);
 							continue;
 						}
 						bool Success = false;
@@ -865,7 +967,7 @@ namespace Gauntlet
 
 							if (Instance != null || Globals.CancelSignalled)
 							{
-								RunningRoles.Add(new UnrealSessionInstance.RoleInstance(InstallRoleKV.Value, Instance));
+								AllRoles.Add(new UnrealSessionInstance.RoleInstance(InstallRoleKV.Value, Instance));
 							}
 
 							Success = true;
@@ -882,7 +984,7 @@ namespace Gauntlet
 							Log.Warning(KnownLogEvents.Gauntlet_DeviceEvent, "Failed to start build on {Name}. Marking as problem device and retrying with new set", CurrentInstall.Device);
 
 							// terminate anything that's running
-							foreach (UnrealSessionInstance.RoleInstance RunningRole in RunningRoles)
+							foreach (UnrealSessionInstance.RoleInstance RunningRole in AllRoles.Where(X => X.AppInstance != null) )
 							{
 								Log.Info("Shutting down {0}", RunningRole.AppInstance.Device);
 								RunningRole.AppInstance.Kill();
@@ -905,9 +1007,9 @@ namespace Gauntlet
 						}
 					}
 
-					if (RunningRoles.Count() == SessionRoles.Count())
+					if (AllRoles.Count() == SessionRoles.Count())
 					{
-						SessionInstance = new UnrealSessionInstance(RunningRoles.ToArray());
+						SessionInstance = new UnrealSessionInstance(AllRoles.ToArray(), DeferredRoleToAppInstall);
 					}
 				}
 			}
diff --git a/Engine/Source/Programs/AutomationTool/Gauntlet/Unreal/Base/Gauntlet.UnrealTestConfiguration.cs b/Engine/Source/Programs/AutomationTool/Gauntlet/Unreal/Base/Gauntlet.UnrealTestConfiguration.cs
--- a/Engine/Source/Programs/AutomationTool/Gauntlet/Unreal/Base/Gauntlet.UnrealTestConfiguration.cs
+++ b/Engine/Source/Programs/AutomationTool/Gauntlet/Unreal/Base/Gauntlet.UnrealTestConfiguration.cs
@@ -336,6 +336,7 @@ namespace Gauntlet
 			AdditionalArtifactDirectories = new List<EIntendedBaseCopyDirectory>();
 			RoleType = ERoleModifier.None;
 			InstallOnly = false;
+			DeferredLaunch = false;
 			CommandLineParams = new GauntletCommandLine();
 		}
 
@@ -346,6 +347,11 @@ namespace Gauntlet
 		/// </summary>
 		public bool InstallOnly { get; set; }
 
+		/// <summary>
+		/// Whether this role will launched by the test node at a later time, typically during TickTest(). By default, all roles are launched immediately.
+		/// </summary>
+		public bool DeferredLaunch { get; set; }
+
 		/// <summary>
 		/// Type of process this role represents
 		/// </summary>
diff --git a/Engine/Source/Programs/AutomationTool/Gauntlet/Unreal/Base/Gauntlet.UnrealTestNode.cs b/Engine/Source/Programs/AutomationTool/Gauntlet/Unreal/Base/Gauntlet.UnrealTestNode.cs
--- a/Engine/Source/Programs/AutomationTool/Gauntlet/Unreal/Base/Gauntlet.UnrealTestNode.cs
+++ b/Engine/Source/Programs/AutomationTool/Gauntlet/Unreal/Base/Gauntlet.UnrealTestNode.cs
@@ -505,7 +505,7 @@ namespace Gauntlet
 				// nothing missing, keep going.
 				return true;
 			}
-			
+
 			// if all roles are gone, we're done
 			if (MissingRoles.Count == TestInstance.RunningRoles.Count())
 			{
@@ -607,6 +607,7 @@ namespace Gauntlet
 
 					UnrealSessionRole SessionRole = new UnrealSessionRole(RoleContext.Type, SessionPlatform, RoleContext.Configuration, TestRole.CommandLine);
 					SessionRole.InstallOnly = TestRole.InstallOnly;
+					SessionRole.DeferredLaunch = TestRole.DeferredLaunch;
 					SessionRole.CommandLineParams = TestRole.CommandLineParams;
  					SessionRole.RoleModifier = TestRole.RoleType;
 					SessionRole.Constraint = UseContextConstraint ? Context.Constraint : new UnrealDeviceTargetConstraint(SessionPlatform);
@@ -1088,6 +1089,12 @@ namespace Gauntlet
 		{
 			base.StopTest(InReason);
 
+			// Warn if there are still deferred roles that have not been launched when the test is finished
+			foreach (UnrealSessionInstance.RoleInstance DeferredRole in TestInstance.DeferredRoles)
+			{
+				Log.Warning("Deferred role {Role} was not started before the test was stopped", DeferredRole);
+			}
+
 			// Shutdown the instance so we can access all files, but do not null it or shutdown the UnrealApp because we still need
 			// access to these objects and their resources! Final cleanup is done in CleanupTest()
 			TestInstance.Shutdown();
