UAT: Refactor AutomationUtils and BuildGraph to use structured logging.

diff --git a/Engine/Source/Programs/AutomationTool/AutomationUtils/Automation.cs b/Engine/Source/Programs/AutomationTool/AutomationUtils/Automation.cs
--- a/Engine/Source/Programs/AutomationTool/AutomationUtils/Automation.cs
+++ b/Engine/Source/Programs/AutomationTool/AutomationUtils/Automation.cs
@@ -11,6 +11,7 @@ using EpicGames.Core;
 using OpenTracing.Util;
 using UnrealBuildBase;
 using System.Threading.Tasks;
+using Microsoft.Extensions.Logging;
 
 namespace AutomationTool
 {
@@ -69,7 +70,7 @@ namespace AutomationTool
 					}
 				}
 
-				Log.TraceVerbose("IsBuildMachine={0}", IsBuildMachine);
+				Log.Logger.LogDebug("IsBuildMachine={IsBuildMachine}", IsBuildMachine);
 				Environment.SetEnvironmentVariable("IsBuildMachine", IsBuildMachine ? "1" : "0");
 
 				// Register all the log event matchers
@@ -85,7 +86,7 @@ namespace AutomationTool
 
 				// should we kill processes on exit
 				ShouldKillProcesses = !GlobalCommandLine.NoKill;
-				Log.TraceVerbose("ShouldKillProcesses={0}", ShouldKillProcesses);
+				Log.Logger.LogDebug("ShouldKillProcesses={ShouldKillProcesses}", ShouldKillProcesses);
 
 				if (AutomationToolCommandLine.CommandsToExecute.Count == 0 && GlobalCommandLine.Help)
 				{
@@ -107,13 +108,13 @@ namespace AutomationTool
 				LogUtils.AddLogFileListener(new DirectoryReference(CommandUtils.CmdEnv.LogFolder), new DirectoryReference(CommandUtils.CmdEnv.FinalLogFolder));
 				if (LogUtils.FinalLogFileName != LogUtils.LogFileName)
 				{
-					Log.TraceInformation($"Final log location: {LogUtils.FinalLogFileName}");	
+					Log.Logger.LogInformation("Final log location: {Location}", LogUtils.FinalLogFileName);
 				}
 				
 				// Initialize UBT
 				if (!UnrealBuildTool.PlatformExports.Initialize(Log.Logger))
 				{
-					Log.TraceInformation("Failed to initialize UBT");
+					Log.Logger.LogInformation("Failed to initialize UBT");
 					return ExitCode.Error_Unknown;
 				}
 				
@@ -185,8 +186,8 @@ namespace AutomationTool
 				}
 				else if (Ex.OutputFormat == AutomationExceptionOutputFormat.Minimal)
 				{
-					Log.TraceInformation("{0}", Ex.ToString().Replace("\n", "\n  "));
-					Log.TraceLog("{0}", ExceptionUtils.FormatExceptionDetails(Ex));
+					Log.Logger.LogInformation(Ex, "{Message}", Ex.ToString().Replace("\n", "\n  "));
+					Log.Logger.LogDebug(Ex, "{Details}", ExceptionUtils.FormatExceptionDetails(Ex));
 				}
 				else
 				{
@@ -231,7 +232,7 @@ namespace AutomationTool
             }
             catch (Exception Ex)
             {
-                Log.TraceError("Exception performing nothrow action \"{0}\": {1}", ActionDesc, ExceptionUtils.FormatException(Ex));
+                Log.Logger.LogError(Ex, "Exception performing nothrow action \"{ActionDesc}\": {Exception}", ActionDesc, ExceptionUtils.FormatException(Ex));
             }
         }
 
@@ -246,7 +247,7 @@ namespace AutomationTool
 			for (int CommandIndex = 0; CommandIndex < CommandsToExecute.Count; ++CommandIndex)
 			{
 				var CommandInfo = CommandsToExecute[CommandIndex];
-				Log.TraceVerbose("Attempting to execute {0}", CommandInfo.ToString());
+				Log.Logger.LogDebug("Attempting to execute {CommandInfo}", CommandInfo.ToString());
 				Type CommandType;
 				if (!Commands.TryGetValue(CommandInfo.CommandName, out CommandType))
 				{
@@ -293,7 +294,7 @@ namespace AutomationTool
 				Type CommandType;
 				if (Commands.TryGetValue(CommandInfo.CommandName, out CommandType) == false)
 				{
-					Log.TraceError("Help: Failed to find command {0}", CommandInfo.CommandName);
+					Log.Logger.LogError("Help: Failed to find command {CommandName}", CommandInfo.CommandName);
 				}
 				else
 				{
diff --git a/Engine/Source/Programs/AutomationTool/AutomationUtils/BuildCommand.cs b/Engine/Source/Programs/AutomationTool/AutomationUtils/BuildCommand.cs
--- a/Engine/Source/Programs/AutomationTool/AutomationUtils/BuildCommand.cs
+++ b/Engine/Source/Programs/AutomationTool/AutomationUtils/BuildCommand.cs
@@ -6,6 +6,7 @@ using System.Text;
 using System.Threading.Tasks;
 using EpicGames.Core;
 using UnrealBuildBase;
+using Microsoft.Extensions.Logging;
 
 namespace AutomationTool
 {
@@ -52,11 +53,11 @@ namespace AutomationTool
 				{
 					if (ParamValue == null)
 					{
-						Log.TraceWarning($"Param name \"{ObsoleteParam}\" is deprecated, use \"{Param}\" instead.");
+						Log.Logger.LogWarning("Param name \"{ObsoleteParam}\" is deprecated, use \"{Param}\" instead.", ObsoleteParam, Param);
 					}
 					else
 					{
-						Log.TraceWarning($"Deprecated param name \"{ObsoleteParam}\" was ignored because \"{Param}\" was set.");
+						Log.Logger.LogWarning("Deprecated param name \"{ObsoleteParam}\" was ignored because \"{Param}\" was set.", ObsoleteParam, Param);
 					}
 				}
 
@@ -368,4 +369,4 @@ namespace AutomationTool
 			return Command.Execute();
 		}
 	}
-}
\ No newline at end of file
+}
diff --git a/Engine/Source/Programs/AutomationTool/AutomationUtils/CommandEnvironment.cs b/Engine/Source/Programs/AutomationTool/AutomationUtils/CommandEnvironment.cs
--- a/Engine/Source/Programs/AutomationTool/AutomationUtils/CommandEnvironment.cs
+++ b/Engine/Source/Programs/AutomationTool/AutomationUtils/CommandEnvironment.cs
@@ -9,6 +9,7 @@ using System.Diagnostics;
 using EpicGames.Core;
 using System.Text.RegularExpressions;
 using UnrealBuildBase;
+using Microsoft.Extensions.Logging;
 
 namespace AutomationTool
 {
@@ -136,7 +137,7 @@ namespace AutomationTool
 
 			if (IsChildInstance)
 			{
-				Log.TraceInformation($"AutomationTool is running as a child instance ({EnvVarNames.IsChildInstance}={IsChildInstanceInt.ToString()})");
+				Log.Logger.LogInformation("AutomationTool is running as a child instance ({Name}={Value})", EnvVarNames.IsChildInstance, IsChildInstanceInt.ToString());
 			}
 
 			// Setup the timestamp string
@@ -181,18 +182,18 @@ namespace AutomationTool
 
 		void LogSettings()
 		{
-			Log.TraceVerbose("Command Environment settings:");
-			Log.TraceVerbose("CmdExe={0}", CmdExe);
-			Log.TraceVerbose("EngineSavedFolder={0}", EngineSavedFolder);
-			Log.TraceVerbose("HasCapabilityToCompile={0}", HasCapabilityToCompile);
-			Log.TraceVerbose("LocalRoot={0}", LocalRoot);
-			Log.TraceVerbose("LogFolder={0}", LogFolder);
-			Log.TraceVerbose("MountExe={0}", MountExe);
-			Log.TraceVerbose("FrameworkMsbuildExe={0}", FrameworkMsbuildPath);
-			Log.TraceVerbose("DotnetMsbuildExe={0}", DotnetMsbuildPath);
-			Log.TraceVerbose("RobocopyExe={0}", RobocopyExe);
-			Log.TraceVerbose("TimestampAsString={0}", TimestampAsString);
-			Log.TraceVerbose("AutomationToolDll={0}", AutomationToolDll);			
+			Log.Logger.LogDebug("Command Environment settings:");
+			Log.Logger.LogDebug("CmdExe={CmdExe}", CmdExe);
+			Log.Logger.LogDebug("EngineSavedFolder={EngineSavedFolder}", EngineSavedFolder);
+			Log.Logger.LogDebug("HasCapabilityToCompile={HasCapabilityToCompile}", HasCapabilityToCompile);
+			Log.Logger.LogDebug("LocalRoot={LocalRoot}", LocalRoot);
+			Log.Logger.LogDebug("LogFolder={LogFolder}", LogFolder);
+			Log.Logger.LogDebug("MountExe={MountExe}", MountExe);
+			Log.Logger.LogDebug("FrameworkMsbuildExe={FrameworkMsbuildPath}", FrameworkMsbuildPath);
+			Log.Logger.LogDebug("DotnetMsbuildExe={DotnetMsbuildPath}", DotnetMsbuildPath);
+			Log.Logger.LogDebug("RobocopyExe={RobocopyExe}", RobocopyExe);
+			Log.Logger.LogDebug("TimestampAsString={TimestampAsString}", TimestampAsString);
+			Log.Logger.LogDebug("AutomationToolDll={AutomationToolDll}", AutomationToolDll);
 		}
 
 		/// <summary>
@@ -214,11 +215,11 @@ namespace AutomationTool
 				HasCapabilityToCompile = false;
 				FrameworkMsbuildPath = "";
 			}
-			Log.TraceVerbose("CompilationEvironment.HasCapabilityToCompile={0}", HasCapabilityToCompile);
-			Log.TraceVerbose("CompilationEvironment.FrameworkMsbuildExe={0}", FrameworkMsbuildPath);
+			Log.Logger.LogDebug("CompilationEvironment.HasCapabilityToCompile={HasCapabilityToCompile}", HasCapabilityToCompile);
+			Log.Logger.LogDebug("CompilationEvironment.FrameworkMsbuildExe={FrameworkMsbuildPath}", FrameworkMsbuildPath);
 
 			DotnetMsbuildPath = Unreal.DotnetPath.FullName;
-			Log.TraceVerbose("CompilationEvironment.DotnetMsbuildExe={0}", DotnetMsbuildPath);
+			Log.Logger.LogDebug("CompilationEvironment.DotnetMsbuildExe={DotnetMsbuildPath}", DotnetMsbuildPath);
 		}
 
 		/// <summary>
@@ -237,7 +238,7 @@ namespace AutomationTool
 			while (true)
 			{
 				var PathToCheck = Path.GetFullPath(CommandUtils.CombinePaths(CurrentPath, KnownFilePathFromRoot));
-				Log.TraceVerbose("Checking for {0}", PathToCheck);
+				Log.Logger.LogDebug("Checking for {PathToCheck}", PathToCheck);
 				if (!File.Exists(PathToCheck))
 				{
 					var LastSeparatorIndex = CurrentPath.LastIndexOf('/');
diff --git a/Engine/Source/Programs/AutomationTool/AutomationUtils/CommandUtils.cs b/Engine/Source/Programs/AutomationTool/AutomationUtils/CommandUtils.cs
--- a/Engine/Source/Programs/AutomationTool/AutomationUtils/CommandUtils.cs
+++ b/Engine/Source/Programs/AutomationTool/AutomationUtils/CommandUtils.cs
@@ -261,7 +261,7 @@ namespace AutomationTool
 		{
 			if(bShowProgress)
 			{
-				LogInformation("[@progress push {0}/{1} skipline]", Numerator, Denominator);
+				Log.Logger.LogInformation("[@progress push {Numerator}/{Denominator} skipline]", Numerator, Denominator);
 			}
 		}
 
@@ -269,7 +269,7 @@ namespace AutomationTool
 		{
 			if(bShowProgress)
 			{
-				LogInformation("[@progress pop skipline]");
+				Log.Logger.LogInformation("[@progress pop skipline]");
 			}
 		}
 
@@ -277,7 +277,7 @@ namespace AutomationTool
 		{
 			if(bShowProgress)
 			{
-				LogInformation("[@progress increment {0}/{1} skipline]", Numerator, Denominator);
+				Log.Logger.LogInformation("[@progress increment {Numerator}/{Denominator} skipline]", Numerator, Denominator);
 			}
 		}
 
@@ -286,7 +286,7 @@ namespace AutomationTool
 		{
 			if(bShowProgress)
 			{
-				LogInformation("[@progress '{0}' skipline]", String.Format(Format, Args));
+				Log.Logger.LogInformation("[@progress '{Message}' skipline]", String.Format(Format, Args));
 			}
 		}
 
@@ -295,7 +295,7 @@ namespace AutomationTool
 		{
 			if(bShowProgress)
 			{
-				LogInformation("[@progress {0}/{1} '{2}' skipline]", Numerator, Denominator, String.Format(Format, Args));
+				Log.Logger.LogInformation("[@progress {Numerator}/{Denominator} '{Arg2}' skipline]", Numerator, Denominator, String.Format(Format, Args));
 			}
 		}
 
@@ -486,7 +486,7 @@ namespace AutomationTool
 			var NormalizedFilename = ConvertSeparators(PathSeparator.Default, FileName);
 			if (!InternalUtils.SafeDeleteFile(NormalizedFilename))
 			{
-				LogWarning("Failed to delete file '{0}'", NormalizedFilename);
+				Log.Logger.LogWarning("Failed to delete file '{NormalizedFilename}'", NormalizedFilename);
 				Result = false;
 			}
 			return Result;
@@ -559,7 +559,7 @@ namespace AutomationTool
             {
                 if (!InternalUtils.SafeDeleteDirectory(NormalizedDirectory, bQuiet))
                 {
-                    LogWarning("Failed to delete directory '{0}'", NormalizedDirectory);
+					Log.Logger.LogWarning("Failed to delete directory '{NormalizedDirectory}'", NormalizedDirectory);
                     Result = false;
                 }
             }
@@ -567,8 +567,8 @@ namespace AutomationTool
             {
 				if (!bQuiet)
 				{
-					LogWarning("Failed to delete directory, exception '{0}'", NormalizedDirectory);
-					LogWarning(Ex.Message);
+					Log.Logger.LogWarning("Failed to delete directory, exception '{NormalizedDirectory}'", NormalizedDirectory);
+					Log.Logger.LogWarning(Ex, "{Text}", Ex.Message);
 				}
                 Result = false;
             }
@@ -595,7 +595,7 @@ namespace AutomationTool
 		/// <param name="DirectoryName"></param>
 		public static void DeleteDirectoryContents(string DirectoryName)
 		{
-			LogVerbose("DeleteDirectoryContents({0})", DirectoryName);
+			Log.Logger.LogDebug("DeleteDirectoryContents({DirectoryName})", DirectoryName);
 			const bool bQuiet = true;
 			var Files = CommandUtils.FindFiles_NoExceptions(bQuiet, "*", false, DirectoryName);
 			foreach (var Filename in Files)
@@ -645,8 +645,8 @@ namespace AutomationTool
 			}
 			catch (Exception Ex)
 			{
-				LogWarning("Unable to check if directory exists: {0}", NormalizedDirectory);
-				LogWarning(Ex.Message);
+				Log.Logger.LogWarning("Unable to check if directory exists: {NormalizedDirectory}", NormalizedDirectory);
+				Log.Logger.LogWarning(Ex, "{Text}", Ex.Message);
 				return false;
 			}
 		}
@@ -733,7 +733,7 @@ namespace AutomationTool
 			var Result = InternalUtils.SafeRenameFile(OldNormalized, NewNormalized);
 			if (!Result)
 			{
-				LogWarning("Failed to rename/move file '{0}' to '{1}'", OldName, NewName);
+				Log.Logger.LogWarning("Failed to rename/move file '{OldName}' to '{NewName}'", OldName, NewName);
 			}
 			return Result;
 		}
@@ -822,7 +822,7 @@ namespace AutomationTool
 			}
 			catch
 			{
-				LogWarning("Unable to change current directory to {0}", WorkingDirectory);
+				Log.Logger.LogWarning("Unable to change current directory to {WorkingDirectory}", WorkingDirectory);
 				Result = false;
 			}
 			return Result;
@@ -855,7 +855,7 @@ namespace AutomationTool
 			}
 			else
 			{
-				LogWarning("Unable to PopDir. WorkingDirectoryStack is empty.");
+				Log.Logger.LogWarning("Unable to PopDir. WorkingDirectoryStack is empty.");
 				Result = false;
 			}
 			return Result;
@@ -903,7 +903,7 @@ namespace AutomationTool
 			}
 			catch
 			{
-				LogWarning("Unable to change current directory to {0}", WorkingDirectory);
+				Log.Logger.LogWarning("Unable to change current directory to {WorkingDirectory}", WorkingDirectory);
 				Result = false;
 			}
 			return Result;
@@ -981,7 +981,7 @@ namespace AutomationTool
 			Filename = ConvertSeparators(PathSeparator.Default, Filename);
 			if (!File.Exists(Filename))
 			{
-				LogWarning("Unable to set attributes for a non-exisiting file ({0})", Filename);
+				Log.Logger.LogWarning("Unable to set attributes for a non-exisiting file ({Filename})", Filename);
 				return false;
 			}
 
@@ -994,8 +994,8 @@ namespace AutomationTool
 			}
 			catch (Exception Ex)
 			{
-				LogWarning("Error trying to set file attributes for: {0}", Filename);
-				LogWarning(Ex.Message);
+				Log.Logger.LogWarning("Error trying to set file attributes for: {Filename}", Filename);
+				Log.Logger.LogWarning(Ex, "{Text}", Ex.Message);
 				Result = false;
 			}
 			return Result;
@@ -1386,13 +1386,13 @@ namespace AutomationTool
 			{
 				if (bAllowDifferingTimestamps == true)
 				{
-					LogVerbose("CopyFileIncremental Skipping {0}, already exists", Dest);
+					Log.Logger.LogDebug("CopyFileIncremental Skipping {Dest}, already exists", Dest);
 					return;
 				}
 				TimeSpan Diff = File.GetLastWriteTimeUtc(Dest.FullName) - File.GetLastWriteTimeUtc(Source.FullName);
 				if (Diff.TotalSeconds > -1 && Diff.TotalSeconds < 1)
 				{
-					LogVerbose("CopyFileIncremental Skipping {0}, up to date.", Dest);
+					Log.Logger.LogDebug("CopyFileIncremental Skipping {Dest}, up to date.", Dest);
 					return;
 				}
 				InternalUtils.SafeDeleteFile(Dest.FullName);
@@ -1766,7 +1766,7 @@ namespace AutomationTool
 		{
 			if(!bQuiet)
 			{
-				LogInformation("Copying {0} file(s) using max {1} thread(s)", Source.Count, MaxThreads);
+				Log.Logger.LogInformation("Copying {NumFiles} file(s) using max {MaxThreads} thread(s)", Source.Count, MaxThreads);
 			}
 
             if (Source.Count != Dest.Count)
@@ -1809,7 +1809,7 @@ namespace AutomationTool
 		public static List<string> ThreadedCopyFiles(string SourceDir, string TargetDir, FileFilter Filter, bool bIgnoreSymlinks, int MaxThreads = 64, bool bRetry = false)
 		{
 			// Filter all the relative paths
-			LogInformation("Applying filter to {0}...", SourceDir);
+			Log.Logger.LogInformation("Applying filter to {SourceDir}...", SourceDir);
 			DirectoryReference SourceDirRef = new DirectoryReference(SourceDir);
 			var RelativePaths = Filter.ApplyToDirectory(SourceDirRef, bIgnoreSymlinks).Select(x => x.MakeRelativeTo(SourceDirRef)).ToList();
 			return ThreadedCopyFiles(SourceDir, TargetDir, RelativePaths, MaxThreads, bRetry: bRetry);
@@ -1912,7 +1912,7 @@ namespace AutomationTool
 		{
 			try
 			{
-				LogLog("SetEnvVar {0}={1}", Name, Value);
+				Log.Logger.LogDebug("SetEnvVar {Name}={Value}", Name, Value);
 				Environment.SetEnvironmentVariable(Name, Value.ToString());
 			}
 			catch (Exception Ex)
@@ -2119,7 +2119,7 @@ namespace AutomationTool
         {
             if (!DirectoryExists_NoExceptions(Dir))
             {
-				LogLog("Directory {0} does not exist", Dir);
+				Log.Logger.LogDebug("Directory {Dir} does not exist", Dir);
 				return false;
 			}
 
@@ -2134,7 +2134,7 @@ namespace AutomationTool
 				if(File.Exists(NativeFilename))
 				{
 		            DeleteFile_NoExceptions(Filename, true);
-		            LogLog("Directory {0} is writable", Dir);
+					Log.Logger.LogDebug("Directory {Dir} is writable", Dir);
 					return true;
 				}
 			}
@@ -2142,7 +2142,7 @@ namespace AutomationTool
 			{
 			}
 
-			LogLog("Directory {0} is not writable", Dir);
+			Log.Logger.LogDebug("Directory {Dir} is not writable", Dir);
 			return false;
 		}
 
@@ -2155,25 +2155,25 @@ namespace AutomationTool
             try
             {
                 DirectoryInfo DirInfo = new DirectoryInfo(ParentDir);
-				LogInformation("Looking for directories to delete in {0}", ParentDir);
+				Log.Logger.LogInformation("Looking for directories to delete in {ParentDir}", ParentDir);
                 foreach (DirectoryInfo ThisDirInfo in DirInfo.EnumerateDirectories(SearchPattern))
                 {
 					double AgeDays = (DateTime.UtcNow - ThisDirInfo.CreationTimeUtc).TotalDays;
 					if (AgeDays > MaximumDaysToKeepTempStorage)
                     {
-                        LogInformation("Deleting formal build directory {0}, because it is {1} days old (maximum {2}).", ThisDirInfo.FullName, (int)AgeDays, MaximumDaysToKeepTempStorage);
+						Log.Logger.LogInformation("Deleting formal build directory {Directory}, because it is {Age} days old (maximum {MaximumDaysToKeepTempStorage}).", ThisDirInfo.FullName, (int)AgeDays, MaximumDaysToKeepTempStorage);
                         DeleteDirectory_NoExceptions(true, ThisDirInfo.FullName);
                     }
                     else
                     {
-						LogVerbose("Not deleting formal build directory {0}, because it is {1} days old (maximum {2}).", ThisDirInfo.FullName, (int)AgeDays, MaximumDaysToKeepTempStorage);
+						Log.Logger.LogDebug("Not deleting formal build directory {Directory}, because it is {Age} days old (maximum {MaximumDaysToKeepTempStorage}).", ThisDirInfo.FullName, (int)AgeDays, MaximumDaysToKeepTempStorage);
                     }
                 }
             }
             catch (Exception Ex)
             {
-                LogWarning("Unable to clean formal builds from directory: {0}", ParentDir);
-                LogWarning(" Exception was {0}", LogUtils.FormatException(Ex));
+				Log.Logger.LogWarning("Unable to clean formal builds from directory: {ParentDir}", ParentDir);
+				Log.Logger.LogWarning(Ex, " Exception was {Exception}", LogUtils.FormatException(Ex));
             }
         }
 
@@ -2478,11 +2478,11 @@ namespace AutomationTool
 
 					if (Iterations == 0)
 					{
-						LogInformation("Waiting for lock file '{0}' to be removed...", LockFilePath);
+						Log.Logger.LogInformation("Waiting for lock file '{LockFilePath}' to be removed...", LockFilePath);
 					}
 					else if ((Iterations % 30) == 0)
 					{
-						LogInformation("Still waiting for lock file '{0}' after {1} seconds.", LockFilePath, CurrentTime.Subtract(StartTime).TotalSeconds);
+						Log.Logger.LogInformation("Still waiting for lock file '{LockFilePath}' after {Seconds} seconds.", LockFilePath, CurrentTime.Subtract(StartTime).TotalSeconds);
 					}
 
 					// Wait for a while before retrying.
diff --git a/Engine/Source/Programs/AutomationTool/AutomationUtils/CommandletUtils.cs b/Engine/Source/Programs/AutomationTool/AutomationUtils/CommandletUtils.cs
--- a/Engine/Source/Programs/AutomationTool/AutomationUtils/CommandletUtils.cs
+++ b/Engine/Source/Programs/AutomationTool/AutomationUtils/CommandletUtils.cs
@@ -8,6 +8,7 @@ using System.Threading;
 using UnrealBuildTool;
 using EpicGames.Core;
 using UnrealBuildBase;
+using Microsoft.Extensions.Logging;
 
 namespace AutomationTool
 {
@@ -266,7 +267,7 @@ namespace AutomationTool
 
 		public static void StartRunCommandlet(FileReference ProjectName, string UnrealExe, string Commandlet, string Parameters, ERunOptions RunOptions, out string LocalLogFile, out IProcessResult RunResult, ProcessResult.SpewFilterCallbackType SpewFilterCallback=null)
 		{
-			LogInformation("Running UnrealEditor {0} for project {1}", Commandlet, ProjectName);
+			Log.Logger.LogInformation("Running UnrealEditor {Commandlet} for project {ProjectName}", Commandlet, ProjectName);
 
 			var CWD = Path.GetDirectoryName(UnrealExe);
 
@@ -281,7 +282,7 @@ namespace AutomationTool
 			PushDir(CWD);
 
 			LocalLogFile = LogUtils.GetUniqueLogName(CombinePaths(CmdEnv.EngineSavedFolder, Commandlet));
-			LogInformation("Commandlet log file is {0}", LocalLogFile);
+			Log.Logger.LogInformation("Commandlet log file is {LocalLogFile}", LocalLogFile);
 			string Args = String.Format(
 				"{0} -run={1} {2} -abslog={3} -stdout -CrashForUAT -unattended -NoLogTimes {5}{4}",
 				(ProjectName == null) ? "" : CommandUtils.MakePathSafeToUseWithCommandLine(ProjectName.FullName),
diff --git a/Engine/Source/Programs/AutomationTool/AutomationUtils/CopyUAT.cs b/Engine/Source/Programs/AutomationTool/AutomationUtils/CopyUAT.cs
--- a/Engine/Source/Programs/AutomationTool/AutomationUtils/CopyUAT.cs
+++ b/Engine/Source/Programs/AutomationTool/AutomationUtils/CopyUAT.cs
@@ -9,6 +9,7 @@ using AutomationTool;
 using UnrealBuildTool;
 using EpicGames.Core;
 using UnrealBuildBase;
+using Microsoft.Extensions.Logging;
 
 /// <summary>
 /// Copies all UAT and UBT build products to a directory
@@ -54,7 +55,7 @@ public class CopyUAT : BuildCommand
 		List<FileReference> SourceFiles = SourceFilesSet.OrderBy(x => x.FullName).ToList();
 		CommandUtils.ThreadedCopyFiles(SourceFiles, Unreal.RootDirectory, TargetDir);
 
-		LogInformation("Copied {0} files to {1}", SourceFiles.Count, TargetDir);
+		Logger.LogInformation("Copied {NumFiles} files to {TargetDir}", SourceFiles.Count, TargetDir);
 		File.WriteAllLines(Path.Combine(TargetDirParam, "CopiedFiles.txt"), SourceFiles.Select(F => F.FullName));
 	}
 }
diff --git a/Engine/Source/Programs/AutomationTool/AutomationUtils/HashCollection.cs b/Engine/Source/Programs/AutomationTool/AutomationUtils/HashCollection.cs
--- a/Engine/Source/Programs/AutomationTool/AutomationUtils/HashCollection.cs
+++ b/Engine/Source/Programs/AutomationTool/AutomationUtils/HashCollection.cs
@@ -9,6 +9,7 @@ using System.Xml;
 using System.Xml.Serialization;
 using System.Linq;
 using EpicGames.Core;
+using Microsoft.Extensions.Logging;
 
 namespace AutomationTool
 {
@@ -173,14 +174,14 @@ namespace AutomationTool
 			{
 				if (!RHS.HasEntryForName(Entry.Name))
 				{
-					Log.TraceInformation("RHS does not have an entry for {0}", Entry.Name);
+					Log.Logger.LogInformation("RHS does not have an entry for {Name}", Entry.Name);
 				}
 				else if (!RHS.Hashes.Contains(Entry))
 				{
 					HashEntry RHSEntry = RHS.GetEntryForName(Entry.Name);
-					Log.TraceInformation("RHS hash mismatch for {0}", Entry.Name);
-					Log.TraceInformation("LHS:\n\t{0}\n\t{1}\n\t{2}", Entry.Hash, Entry.Name, Entry.MetaData);
-					Log.TraceInformation("RHS:\n\t{0}\n\t{1}\n\t{2}", RHSEntry.Hash, RHSEntry.Name, RHSEntry.MetaData);
+					Log.Logger.LogInformation("RHS hash mismatch for {Name}", Entry.Name);
+					Log.Logger.LogInformation("LHS:\n\t{Hash}\n\t{Name}\n\t{MetaData}", Entry.Hash, Entry.Name, Entry.MetaData);
+					Log.Logger.LogInformation("RHS:\n\t{Hash}\n\t{Name}\n\t{MetaData}", RHSEntry.Hash, RHSEntry.Name, RHSEntry.MetaData);
 				}
 			}
 		}
@@ -211,7 +212,7 @@ namespace AutomationTool
 			}
 			catch (Exception Ex)
 			{
-				Log.TraceWarning("Failed to load HashCollection from {0}. {1}", InPath, Ex.Message);
+				Log.Logger.LogWarning(Ex, "Failed to load HashCollection from {InPath}. {Message}", InPath, Ex.Message);
 			}
 
 			return null;
diff --git a/Engine/Source/Programs/AutomationTool/AutomationUtils/HelpUtils.cs b/Engine/Source/Programs/AutomationTool/AutomationUtils/HelpUtils.cs
--- a/Engine/Source/Programs/AutomationTool/AutomationUtils/HelpUtils.cs
+++ b/Engine/Source/Programs/AutomationTool/AutomationUtils/HelpUtils.cs
@@ -15,6 +15,7 @@ using System.ComponentModel;
 using System.Reflection;
 using EpicGames.Core;
 using System.Linq;
+using Microsoft.Extensions.Logging;
 
 namespace AutomationTool
 {
@@ -311,12 +312,12 @@ namespace AutomationTool
 					}
 					else
 					{
-						LogWarning("Duplicated help parameter \"{0}\"", ParamName);
+						Log.Logger.LogWarning("Duplicated help parameter \"{ParamName}\"", ParamName);
 					}
 				}
 			}
 
-			Log.TraceInformation("");
+			Log.Logger.LogInformation("");
 			HelpUtils.PrintHelp(String.Format("{0} Help:", Command.Name), Description, ParamDict.ToList());
 		}
 
diff --git a/Engine/Source/Programs/AutomationTool/AutomationUtils/MacHostPlatform.cs b/Engine/Source/Programs/AutomationTool/AutomationUtils/MacHostPlatform.cs
--- a/Engine/Source/Programs/AutomationTool/AutomationUtils/MacHostPlatform.cs
+++ b/Engine/Source/Programs/AutomationTool/AutomationUtils/MacHostPlatform.cs
@@ -10,6 +10,7 @@ using UnrealBuildTool;
 using System.Text.RegularExpressions;
 using EpicGames.Core;
 using UnrealBuildBase;
+using Microsoft.Extensions.Logging;
 
 namespace AutomationTool
 {
@@ -26,13 +27,13 @@ namespace AutomationTool
 
 				if (CanUseMsBuild)
 				{
-					Log.TraceInformation($"using {CommandUtils.WhichApp("dotnet")}!");
+					Log.Logger.LogInformation("using {DotNet}!", CommandUtils.WhichApp("dotnet"));
 
 					CachedFrameworkMsbuildExe = "dotnet msbuild";
 				}
 				else
 				{
-					Log.TraceInformation("Using xbuild. Install Mono 5.0 or greater for faster builds!");
+					Log.Logger.LogInformation("Using xbuild. Install Mono 5.0 or greater for faster builds!");
 					CachedFrameworkMsbuildExe = "xbuild";
 				}
 			}
diff --git a/Engine/Source/Programs/AutomationTool/AutomationUtils/P4Environment.cs b/Engine/Source/Programs/AutomationTool/AutomationUtils/P4Environment.cs
--- a/Engine/Source/Programs/AutomationTool/AutomationUtils/P4Environment.cs
+++ b/Engine/Source/Programs/AutomationTool/AutomationUtils/P4Environment.cs
@@ -9,6 +9,7 @@ using EpicGames.Core;
 using EpicGames.Perforce;
 using System.Threading.Tasks;
 using System.Diagnostics;
+using Microsoft.Extensions.Logging;
 
 namespace AutomationTool
 {
@@ -157,7 +158,7 @@ namespace AutomationTool
 				{
 					string HostName = System.Net.Dns.GetHostName();
 					ThisClient = DetectClient(DefaultConnection, User, HostName, CmdEnv.AutomationToolDll);
-					Log.TraceInformation("Using user {0} clientspec {1} {2}", User, ThisClient.Name, ThisClient.RootPath);
+					Log.Logger.LogInformation("Using user {User} clientspec {ClientName} {ClientPath}", User, ThisClient.Name, ThisClient.RootPath);
 					Client = ThisClient.Name;
 					CommandUtils.SetEnvVar(EnvVarNames.Client, Client);
 				}
@@ -217,7 +218,7 @@ namespace AutomationTool
 				using PerforceConnection Connection = new PerforceConnection(new PerforceSettings(ServerAndPort, User) { ClientName = Client }, Log.Logger);
 				CodeChangelistString = DetectCurrentCodeCL(Connection, Changelist).Result.ToString();
 				CommandUtils.SetEnvVar(EnvVarNames.CodeChangelist, CodeChangelistString);
-				Log.TraceVerbose("Took {0}ms to query last code change", Timer.ElapsedMilliseconds);
+				Log.Logger.LogDebug("Took {ElapsedMs}ms to query last code change", Timer.ElapsedMilliseconds);
 			}
 			if(!String.IsNullOrEmpty(CodeChangelistString))
 			{
@@ -232,18 +233,18 @@ namespace AutomationTool
 			// Write a summary of the settings to the output window
 			if (!CommandUtils.CmdEnv.IsChildInstance)
 			{
-				Log.TraceInformation("Detected Perforce Settings:");
-				Log.TraceInformation("  Server: {0}", ServerAndPort);
-				Log.TraceInformation("  User: {0}", User);
-				Log.TraceInformation("  Client: {0}", Client);
-				Log.TraceInformation("  Branch: {0}", Branch);
+				Log.Logger.LogInformation("Detected Perforce Settings:");
+				Log.Logger.LogInformation("  Server: {ServerAndPort}", ServerAndPort);
+				Log.Logger.LogInformation("  User: {User}", User);
+				Log.Logger.LogInformation("  Client: {Client}", Client);
+				Log.Logger.LogInformation("  Branch: {Branch}", Branch);
 				if (ChangelistInternal != -1)
 				{
-					Log.TraceInformation("  Last Change: {0}", Changelist);
+					Log.Logger.LogInformation("  Last Change: {Changelist}", Changelist);
 				}
 				if (CodeChangelistInternal != -1)
 				{
-					Log.TraceInformation("  Last Code Change: {0}", CodeChangelist);
+					Log.Logger.LogInformation("  Last Code Change: {CodeChangelist}", CodeChangelist);
 				}
 			}
 
@@ -289,7 +290,7 @@ namespace AutomationTool
 			// Otherwise fallback to the uebp variables, or the default
 			if(String.IsNullOrEmpty(P4Port))
 			{
-				Log.TraceWarning("P4PORT is not set. Using perforce:1666");
+				Log.Logger.LogWarning("P4PORT is not set. Using perforce:1666");
 				P4Port = "perforce:1666";
 			}
 
@@ -315,7 +316,7 @@ namespace AutomationTool
 			{
 				if (!String.IsNullOrEmpty(UserName))
 				{
-					Log.TraceWarning("Unable to retrieve perforce user name. Trying to fall back to {0} which is set to {1}.", EnvVarNames.User, UserName);
+					Log.Logger.LogWarning("Unable to retrieve perforce user name. Trying to fall back to {UserNameEnvVar} which is set to {UserName}.", EnvVarNames.User, UserName);
 				}
 				else
 				{
@@ -360,7 +361,7 @@ namespace AutomationTool
 			{
 				if (!String.IsNullOrEmpty(Client.Host) && String.Compare(Client.Host, HostName, true) != 0)
 				{
-					Log.TraceInformation("Rejecting client because of different Host {0} \"{1}\" != \"{2}\"", Client.Name, Client.Host, HostName);
+					Log.Logger.LogInformation("Rejecting client because of different Host {ClientName} \"{ClientHost}\" != \"{HostName}\"", Client.Name, Client.Host, HostName);
 					continue;
 				}
 				
@@ -389,7 +390,7 @@ namespace AutomationTool
 				}
 				if (ClientToUse == null)
 				{
-					Log.TraceWarning("{0} clients found that match the current host and root path. The most recently accessed client will be used.", MatchingClients.Count);
+					Log.Logger.LogWarning("{NumClients} clients found that match the current host and root path. The most recently accessed client will be used.", MatchingClients.Count);
 					ClientToUse = GetMostRecentClient(MatchingClients);
 				}
 			}
@@ -404,7 +405,7 @@ namespace AutomationTool
 		/// <returns>The most recent client from the list.</returns>
 		private static P4ClientInfo GetMostRecentClient(List<P4ClientInfo> Clients)
 		{
-			Log.TraceVerbose("Detecting the most recent client.");
+			Log.Logger.LogDebug("Detecting the most recent client.");
 			P4ClientInfo MostRecentClient = null;
 			var MostRecentAccessTime = DateTime.MinValue;
 			foreach (var ClientInfo in Clients)
@@ -525,7 +526,7 @@ namespace AutomationTool
 				string KnownFileDepotMapping = P4Result.Output;
 
 				// Get the build root
-				Log.TraceVerbose("Looking for {0} in {1}", KnownFilePathFromRoot, KnownFileDepotMapping);
+				Log.Logger.LogDebug("Looking for {KnownFilePathFromRoot} in {KnownFileDepotMapping}", KnownFilePathFromRoot, KnownFileDepotMapping);
 				int EndIdx = KnownFileDepotMapping.IndexOf(KnownFilePathFromRoot, StringComparison.CurrentCultureIgnoreCase);
 				if (EndIdx < 0)
 				{
diff --git a/Engine/Source/Programs/AutomationTool/AutomationUtils/P4Utils.cs b/Engine/Source/Programs/AutomationTool/AutomationUtils/P4Utils.cs
--- a/Engine/Source/Programs/AutomationTool/AutomationUtils/P4Utils.cs
+++ b/Engine/Source/Programs/AutomationTool/AutomationUtils/P4Utils.cs
@@ -1,5 +1,6 @@
 // Copyright Epic Games, Inc. All Rights Reserved.
 
+using EpicGames.Core;
 using EpicGames.Perforce;
 using System;
 using System.Collections.Generic;
@@ -12,6 +13,7 @@ using System.Text.RegularExpressions;
 using System.Reflection;
 using System.Collections;
 using UnrealBuildBase;
+using Microsoft.Extensions.Logging;
 
 namespace AutomationTool
 {
@@ -1080,7 +1082,7 @@ namespace AutomationTool
 			// Init AllowSubmit
 			// If we do not specify on the commandline if submitting is allowed or not, this is 
 			// depending on whether we run locally or on a build machine.
-			LogVerbose("Initializing AllowSubmit.");
+			Log.Logger.LogDebug("Initializing AllowSubmit.");
 			if (GlobalCommandLine.Submit || GlobalCommandLine.NoSubmit)
 			{
 				AllowSubmit = GlobalCommandLine.Submit;
@@ -1089,10 +1091,10 @@ namespace AutomationTool
 			{
 				AllowSubmit = Automation.IsBuildMachine;
 			}
-			LogVerbose("AllowSubmit={0}", AllowSubmit);
+			Log.Logger.LogDebug("AllowSubmit={AllowSubmit}", AllowSubmit);
 
 			// Init P4Enabled
-			LogVerbose("Initializing P4Enabled.");
+			Log.Logger.LogDebug("Initializing P4Enabled.");
 			if (Automation.IsBuildMachine)
 			{
 				P4Enabled = !GlobalCommandLine.NoP4;
@@ -1107,8 +1109,8 @@ namespace AutomationTool
 				P4Enabled = GlobalCommandLine.P4 || bRequireP4;
 				P4CLRequired = GlobalCommandLine.P4 || bRequireCL;
 			}
-			LogVerbose("P4Enabled={0}", P4Enabled);
-			LogVerbose("P4CLRequired={0}", P4CLRequired);
+			Log.Logger.LogDebug("P4Enabled={P4Enabled}", P4Enabled);
+			Log.Logger.LogDebug("P4CLRequired={P4CLRequired}", P4CLRequired);
 		}
 
 		/// <summary>
@@ -1131,7 +1133,7 @@ namespace AutomationTool
 					{
 						if(!GlobalCommandLine.P4)
 						{
-							LogInformation("Command {0} requires P4 functionality.", Command.Name);
+							Log.Logger.LogInformation("Command {CommandName} requires P4 functionality.", Command.Name);
 						}
 						bRequireP4 = true;
 
@@ -4324,4 +4326,4 @@ namespace AutomationTool
 			return Text;
 		}
 	}
-}
\ No newline at end of file
+}
diff --git a/Engine/Source/Programs/AutomationTool/AutomationUtils/P4WriteConfig.cs b/Engine/Source/Programs/AutomationTool/AutomationUtils/P4WriteConfig.cs
--- a/Engine/Source/Programs/AutomationTool/AutomationUtils/P4WriteConfig.cs
+++ b/Engine/Source/Programs/AutomationTool/AutomationUtils/P4WriteConfig.cs
@@ -7,6 +7,7 @@ using EpicGames.Core;
 using System.Collections.Generic;
 using System.Linq;
 using UnrealBuildBase;
+using Microsoft.Extensions.Logging;
 
 namespace AutomationTool
 {
@@ -19,7 +20,7 @@ namespace AutomationTool
 	{
 		public override ExitCode Execute()
 		{
-			Log.TraceInformation("Setting up Perforce environment.");
+			Log.Logger.LogInformation("Setting up Perforce environment.");
 
 			// User can specify these to help auto detection
 			string Port = ParseParamValue("p4port", "");
@@ -51,7 +52,7 @@ namespace AutomationTool
 			{
 				Log.TraceError(
 					"Unable to find matching Perforce info. If the below does not help try P4WriteConfig -p4port=<server:port> and -p4user=<username> to supply more info");
-				Log.TraceError("{0}", Ex.Message);
+				Log.Logger.LogError(Ex, "{Message}", Ex.Message);
 				return ExitCode.Error_Arguments;
 			}
 
@@ -95,14 +96,14 @@ namespace AutomationTool
 
 				string OutputFile = Path.GetFileName(OutputPath);
 
-				Log.TraceInformation("Wrote P4 settings to {0}", OutputPath);
+				Log.Logger.LogInformation("Wrote P4 settings to {OutputPath}", OutputPath);
 
 				P4.P4(string.Format("set P4CONFIG={0}", OutputFile));
-				Log.TraceInformation("set P4CONFIG={0}", OutputFile);
+				Log.Logger.LogInformation("set P4CONFIG={OutputFile}", OutputFile);
 			}
 			else
 			{
-				Log.TraceInformation("Skipped write");
+				Log.Logger.LogInformation("Skipped write");
 			}
 
 			return ExitCode.Success;
diff --git a/Engine/Source/Programs/AutomationTool/AutomationUtils/Platform.cs b/Engine/Source/Programs/AutomationTool/AutomationUtils/Platform.cs
--- a/Engine/Source/Programs/AutomationTool/AutomationUtils/Platform.cs
+++ b/Engine/Source/Programs/AutomationTool/AutomationUtils/Platform.cs
@@ -9,6 +9,7 @@ using System.Reflection;
 using UnrealBuildTool;
 using EpicGames.Core;
 using UnrealBuildBase;
+using Microsoft.Extensions.Logging;
 
 namespace AutomationTool
 {
@@ -107,7 +108,7 @@ namespace AutomationTool
 		private static Dictionary<TargetPlatformDescriptor, Platform> AllPlatforms = new Dictionary<TargetPlatformDescriptor, Platform>();
 		internal static void InitializePlatforms(HashSet<Assembly> AssembliesWithPlatforms)
 		{
-			LogVerbose("Creating platforms.");
+			Log.Logger.LogDebug("Creating platforms.");
 
 			// Create all available platforms.
 			foreach (var ScriptAssembly in AssembliesWithPlatforms)
@@ -121,7 +122,7 @@ namespace AutomationTool
 				Platform ExistingInstance;
 				if (AllPlatforms.TryGetValue(TargetDesc, out ExistingInstance) == false)
 				{
-					LogVerbose("Creating placeholder platform for target: {0}", TargetDesc.Type);
+					Log.Logger.LogDebug("Creating placeholder platform for target: {TargetType}", TargetDesc.Type);
 					AllPlatforms.Add(TargetDesc, new Platform(TargetDesc.Type));
 				}
 			}
@@ -129,7 +130,7 @@ namespace AutomationTool
 
 		private static void CreatePlatformsFromAssembly(Assembly ScriptAssembly)
 		{
-			LogVerbose("Looking for platforms in {0}", ScriptAssembly.Location);
+			Log.Logger.LogDebug("Looking for platforms in {Location}", ScriptAssembly.Location);
 			Type[] AllTypes = null;
 			try
 			{
@@ -137,21 +138,21 @@ namespace AutomationTool
 			}
 			catch (Exception Ex)
 			{
-				LogError("Failed to get assembly types for {0}", ScriptAssembly.Location);
+				Log.Logger.LogError("Failed to get assembly types for {Location}", ScriptAssembly.Location);
 				if (Ex is ReflectionTypeLoadException)
 				{
 					var TypeLoadException = (ReflectionTypeLoadException)Ex;
 					if (!IsNullOrEmpty(TypeLoadException.LoaderExceptions))
 					{
-						LogError("Loader Exceptions:");
+						Log.Logger.LogError("Loader Exceptions:");
 						foreach (var LoaderException in TypeLoadException.LoaderExceptions)
 						{
-							LogError(LogUtils.FormatException(LoaderException));
+							Log.Logger.LogError(LoaderException, "{Text}", LogUtils.FormatException(LoaderException));
 						}
 					}
 					else
 					{
-						LogError("No Loader Exceptions available.");
+						Log.Logger.LogError("No Loader Exceptions available.");
 					}
 				}
 				// Re-throw, this is still a critical error!
@@ -161,7 +162,7 @@ namespace AutomationTool
 			{
 				if (PotentialPlatformType != typeof(Platform) && typeof(Platform).IsAssignableFrom(PotentialPlatformType) && !PotentialPlatformType.IsAbstract)
 				{
-					LogVerbose("Creating platform {0} from {1}.", PotentialPlatformType.Name, ScriptAssembly.Location);
+					Log.Logger.LogDebug("Creating platform {Platform} from {Location}.", PotentialPlatformType.Name, ScriptAssembly.Location);
 					var PlatformInstance = Activator.CreateInstance(PotentialPlatformType) as Platform;
 					var PlatformDesc = PlatformInstance.GetTargetPlatformDescriptor();
 
@@ -174,7 +175,7 @@ namespace AutomationTool
 					{
 						if (ExistingInstance.GetType() != PlatformInstance.GetType())
 						{
-							LogWarning("Platform {0} already exists", PotentialPlatformType.Name);
+							Log.Logger.LogWarning("Platform {Platform} already exists", PotentialPlatformType.Name);
 						}
 					}
 				}
@@ -361,7 +362,7 @@ namespace AutomationTool
 		public virtual void GetConnectedDevices(ProjectParams Params, out List<string> Devices)
 		{
 			Devices = null;
-			LogWarning("{0} does not implement GetConnectedDevices", PlatformType);
+			Logger.LogWarning("{PlatformType} does not implement GetConnectedDevices", PlatformType);
 		}
 
 		/// <summary>
@@ -382,7 +383,7 @@ namespace AutomationTool
 		/// <param name="SC"></param>
 		public virtual void Deploy(ProjectParams Params, DeploymentContext SC)
 		{
-			LogWarning("{0} does not implement Deploy...", PlatformType);
+			Logger.LogWarning("{PlatformType} does not implement Deploy...", PlatformType);
 		}
 
 		/// <summary>
@@ -887,12 +888,12 @@ namespace AutomationTool
 
 		public virtual void PrepareForDebugging(string SourcePackage, string ProjectFilePath, string ClientPlatform)
 		{
-			LogError("Not implemented for the %s platform.", ClientPlatform);
+			Logger.LogError("Not implemented for the {Platform} platform.", ClientPlatform);
 		}
 
 		public virtual void SetSecondaryRemoteMac(string ProjectFilePath, string ClientPlatform)
 		{
-			LogError("Not implemented for this platform.");
+			Logger.LogError("Not implemented for this platform.");
 		}
 
 		// let the platform set the exe extension if it chooses (otherwise, use
diff --git a/Engine/Source/Programs/AutomationTool/AutomationUtils/ProcessUtils.cs b/Engine/Source/Programs/AutomationTool/AutomationUtils/ProcessUtils.cs
--- a/Engine/Source/Programs/AutomationTool/AutomationUtils/ProcessUtils.cs
+++ b/Engine/Source/Programs/AutomationTool/AutomationUtils/ProcessUtils.cs
@@ -12,6 +12,7 @@ using System.Runtime.CompilerServices;
 using EpicGames.Core;
 using UnrealBuildTool;
 using UnrealBuildBase;
+using Microsoft.Extensions.Logging;
 
 namespace AutomationTool
 {
@@ -819,11 +820,11 @@ namespace AutomationTool
 
 			if (ResolvedPath != null)
 			{
-				Log.TraceVeryVerbose("Resolved {0} to {1}", App, ResolvedPath);
+				Log.Logger.LogTrace("Resolved {App} to {ResolvedPath}", App, ResolvedPath);
 			}
 			else
 			{
-				Log.TraceVerbose("Could not resolve app {0}", App);
+				Log.Logger.LogDebug("Could not resolve app {App}", App);
 			}
 
 			return ResolvedPath;
@@ -1133,7 +1134,7 @@ namespace AutomationTool
 		{
 			while (!FileExists(LogFilename) && !LogProcess.HasExited)
 			{
-				LogInformation("Waiting for logging process to start...");
+				Log.Logger.LogInformation("Waiting for logging process to start...");
 				Thread.Sleep(2000);
 			}
 			Thread.Sleep(1000);
diff --git a/Engine/Source/Programs/AutomationTool/AutomationUtils/ProjectParams.cs b/Engine/Source/Programs/AutomationTool/AutomationUtils/ProjectParams.cs
--- a/Engine/Source/Programs/AutomationTool/AutomationUtils/ProjectParams.cs
+++ b/Engine/Source/Programs/AutomationTool/AutomationUtils/ProjectParams.cs
@@ -8,6 +8,7 @@ using System.IO;
 using System.Reflection;
 using EpicGames.Core;
 using UnrealBuildBase;
+using Microsoft.Extensions.Logging;
 
 namespace AutomationTool
 {
@@ -129,12 +130,12 @@ namespace AutomationTool
 			{
 				if (SpecifiedValue == null)
 				{
-					Log.TraceWarning($"Value was provided for \"{ParamName}\" using obsolete name \"{ObsoleteParamName}\"");
+					Log.Logger.LogWarning("Value was provided for \"{ParamName}\" using obsolete name \"{ObsoleteParamName}\"", ParamName, ObsoleteParamName);
 					Result = SpecifiedValue;
 				}
 				else
 				{
-					Log.TraceWarning($"Value provided for obsolete name \"{ObsoleteParamName}\" will be ignored as \"{ParamName}\" was provided");
+					Log.Logger.LogWarning("Value provided for obsolete name \"{ObsoleteParamName}\" will be ignored as \"{ParamName}\" was provided", ObsoleteParamName, ParamName);
 				}
 			}
 			else if (SpecifiedValue != null)
@@ -153,14 +154,14 @@ namespace AutomationTool
 						// Didn't find the new name on the command line. If the obsolete name was found, use it, and warn.
 						if (ParsedObsolete != null)
 						{
-							Log.TraceWarning($"Obsolete argument \"{ObsoleteParamName}\" on command line - use \"{ParamName}\" instead");
+							Log.Logger.LogWarning("Obsolete argument \"{ObsoleteParamName}\" on command line - use \"{ParamName}\" instead", ObsoleteParamName, ParamName);
 							Parsed = ParsedObsolete;
 						}
 					}
 					else if (ParsedObsolete != null)
 					{
 						// Did find the new name on the command line - check for the obsolete name. If found, do not use it, and warn.
-						Log.TraceWarning($"Obsolete argument \"{ObsoleteParamName}\" will be ignored as \"{ParamName}\" was provided");
+						Log.Logger.LogWarning("Obsolete argument \"{ObsoleteParamName}\" will be ignored as \"{ParamName}\" was provided", ObsoleteParamName, ParamName);
 					}
 				}
 
diff --git a/Engine/Source/Programs/AutomationTool/AutomationUtils/ProjectUtils.cs b/Engine/Source/Programs/AutomationTool/AutomationUtils/ProjectUtils.cs
--- a/Engine/Source/Programs/AutomationTool/AutomationUtils/ProjectUtils.cs
+++ b/Engine/Source/Programs/AutomationTool/AutomationUtils/ProjectUtils.cs
@@ -12,6 +12,7 @@ using System.Reflection;
 using UnrealBuildBase;
 using System.Runtime.Serialization;
 using System.Collections;
+using Microsoft.Extensions.Logging;
 
 namespace AutomationTool
 {
@@ -240,7 +241,7 @@ namespace AutomationTool
 					string Reason;
 					if(RequiresTempTarget(RawProjectPath, bHasCode, Platform, Configuration, TargetType.Game, AssetNativizationRequested, true, out Reason))
 					{
-						Log.TraceInformation("{0} requires a temporary target.cs to be generated ({1})", RawProjectPath.GetFileName(), Reason);
+						Log.Logger.LogInformation("{Project} requires a temporary target.cs to be generated ({Reason})", RawProjectPath.GetFileName(), Reason);
 						return true;
 					}
 				}
@@ -813,7 +814,7 @@ namespace AutomationTool
 
 				if (!CommandUtils.IsBuildMachine && !CheckIfScriptAssemblyIsOutOfDate(TargetsDllFilename, TargetScripts))
 				{
-					Log.TraceVerbose("Targets DLL {0} is up to date.", TargetsDllFilename);
+					Log.Logger.LogDebug("Targets DLL {Filename} is up to date.", TargetsDllFilename);
 					DoNotCompile = true;
 				}
 				if (!DoNotCompile && CommandUtils.FileExists_NoExceptions(TargetsDllFilename.FullName))
diff --git a/Engine/Source/Programs/AutomationTool/AutomationUtils/ScriptManager.cs b/Engine/Source/Programs/AutomationTool/AutomationUtils/ScriptManager.cs
--- a/Engine/Source/Programs/AutomationTool/AutomationUtils/ScriptManager.cs
+++ b/Engine/Source/Programs/AutomationTool/AutomationUtils/ScriptManager.cs
@@ -9,6 +9,7 @@ using UnrealBuildTool;
 using EpicGames.Core;
 using System.Threading.Tasks;
 using UnrealBuildBase;
+using Microsoft.Extensions.Logging;
 
 namespace AutomationTool
 {
@@ -56,7 +57,7 @@ namespace AutomationTool
 				{
 					foreach (Exception SubEx in LoadEx.LoaderExceptions)
 					{
-						Log.TraceWarning("Got type loader exception: {0}", SubEx.ToString());
+						Log.Logger.LogWarning(SubEx, "Got type loader exception: {SubEx}", SubEx.ToString());
 					}
 					throw new AutomationException("Failed to add commands from {0}. {1}", CompiledScripts, LoadEx);
 				}
diff --git a/Engine/Source/Programs/AutomationTool/AutomationUtils/Utils.cs b/Engine/Source/Programs/AutomationTool/AutomationUtils/Utils.cs
--- a/Engine/Source/Programs/AutomationTool/AutomationUtils/Utils.cs
+++ b/Engine/Source/Programs/AutomationTool/AutomationUtils/Utils.cs
@@ -11,6 +11,7 @@ using UnrealBuildTool;
 using System.Text.RegularExpressions;
 using EpicGames.Core;
 using UnrealBuildBase;
+using Microsoft.Extensions.Logging;
 
 namespace AutomationTool
 {
@@ -85,8 +86,8 @@ namespace AutomationTool
 				}
 				else
 				{
-					Log.TraceWarning("Failed to create directory {0} in {1} attempts.", Path, MaxAttempts);
-					Log.TraceWarning(LogUtils.FormatException(LastException));
+					Log.Logger.LogWarning("Failed to create directory {Path} in {MaxAttempts} attempts.", Path, MaxAttempts);
+					Log.Logger.LogWarning("{Text}", LogUtils.FormatException(LastException));
 				}
 			}
 			return Result;
@@ -147,8 +148,8 @@ namespace AutomationTool
 				}
 				else
 				{
-					Log.TraceWarning("Failed to delete file {0} in {1} attempts.", Path, MaxAttempts);
-					Log.TraceWarning(LogUtils.FormatException(LastException));
+					Log.Logger.LogWarning("Failed to delete file {Path} in {MaxAttempts} attempts.", Path, MaxAttempts);
+					Log.Logger.LogWarning("{Text}", LogUtils.FormatException(LastException));
 				}
 			}
 
@@ -230,8 +231,8 @@ namespace AutomationTool
 
 			if (Result == false && LastException != null)
 			{
-				Log.TraceWarning("Failed to delete directory {0} in {1} attempts.", Path, MaxAttempts);
-				Log.TraceWarning(LogUtils.FormatException(LastException));
+				Log.Logger.LogWarning("Failed to delete directory {Path} in {MaxAttempts} attempts.", Path, MaxAttempts);
+				Log.Logger.LogWarning("{Text}", LogUtils.FormatException(LastException));
 			}
 
 			return Result;
@@ -289,8 +290,8 @@ namespace AutomationTool
 					{
 						if (!bQuiet)
 						{
-							Log.TraceWarning("Failed to rename {0} to {1}", OldName, NewName);
-							Log.TraceWarning(LogUtils.FormatException(Ex));
+							Log.Logger.LogWarning("Failed to rename {OldName} to {NewName}", OldName, NewName);
+							Log.Logger.LogWarning("{Text}", LogUtils.FormatException(Ex));
 						}
 						Result = false;
 					}
@@ -347,8 +348,8 @@ namespace AutomationTool
 				{
 					if (File.Exists(OldName) == true || File.Exists(NewName) == false)
 					{
-						Log.TraceWarning("Failed to rename {0} to {1}", OldName, NewName);
-						Log.TraceWarning(LogUtils.FormatException(Ex));
+						Log.Logger.LogWarning("Failed to rename {OldName} to {NewName}", OldName, NewName);
+						Log.Logger.LogWarning("{Text}", LogUtils.FormatException(Ex));
 						Result = false;
 					}
 				}
@@ -470,7 +471,7 @@ namespace AutomationTool
 						}
 						else
 						{
-							Log.TraceInformation("Skip copying file {0} because it doesn't exist.", SourceName);
+							Log.Logger.LogInformation("Skip copying file {SourceName} because it doesn't exist.", SourceName);
 						}
 					}
 					Retry = !File.Exists(TargetName);
@@ -480,7 +481,7 @@ namespace AutomationTool
 						FileInfo TargetInfo = new FileInfo(TargetName);
 						if (!bSkipSizeCheck && SourceInfo.Length != TargetInfo.Length)
 						{
-							Log.TraceInformation("Size mismatch {0} = {1} to {2} = {3}", SourceName, SourceInfo.Length, TargetName, TargetInfo.Length);
+							Log.Logger.LogInformation("Size mismatch {SourceName} = {SourceLength} to {TargetName} = {TargetLength}", SourceName, SourceInfo.Length, TargetName, TargetInfo.Length);
 							Retry = true;
 						}
 						// Timestamps should be no more than 2 seconds out - assuming this as exFAT filesystems store timestamps at 2 second intervals:
@@ -493,7 +494,7 @@ namespace AutomationTool
 							TargetInfo.Refresh();
 							if (!((SourceInfo.LastWriteTimeUtc - TargetInfo.LastWriteTimeUtc).TotalSeconds < 2 && (SourceInfo.LastWriteTimeUtc - TargetInfo.LastWriteTimeUtc).TotalSeconds > -2))
 							{
-								Log.TraceInformation("Date mismatch {0} = {1} to {2} = {3}", SourceName, SourceInfo.LastWriteTimeUtc, TargetName, TargetInfo.LastWriteTimeUtc);
+								Log.Logger.LogInformation("Date mismatch {SourceName} = {SourceTime} to {TargetName} = {TargetTime}", SourceName, SourceInfo.LastWriteTimeUtc, TargetName, TargetInfo.LastWriteTimeUtc);
 								Retry = true;
 							}
 						}
@@ -501,7 +502,7 @@ namespace AutomationTool
 				}
 				catch (Exception Ex)
 				{
-					Log.TraceInformation("SafeCopyFile Exception was {0}", LogUtils.FormatException(Ex));
+					Log.Logger.LogInformation("SafeCopyFile Exception was {Ex}", LogUtils.FormatException(Ex));
 					Retry = true;
 				}
 
@@ -509,7 +510,7 @@ namespace AutomationTool
 				{
 					if (Attempts + 1 < MaxAttempts)
 					{
-						Log.TraceInformation("Failed to copy {0} to {1}, deleting, waiting 10s and retrying.", SourceName, TargetName);
+						Log.Logger.LogInformation("Failed to copy {SourceName} to {TargetName}, deleting, waiting 10s and retrying.", SourceName, TargetName);
 						if (File.Exists(TargetName))
 						{
 							SafeDeleteFile(TargetName);
@@ -518,7 +519,7 @@ namespace AutomationTool
 					}
 					else
 					{
-						Log.TraceError("Failed to copy {0} to {1}", SourceName, TargetName);
+						Log.Logger.LogError("Failed to copy {SourceName} to {TargetName}", SourceName, TargetName);
 					}
 					Result = false;
 				}
@@ -544,8 +545,8 @@ namespace AutomationTool
 			}
 			catch (Exception Ex)
 			{
-				Log.TraceWarning("Failed to load {0}", Filename);
-				Log.TraceWarning(LogUtils.FormatException(Ex));
+				Log.Logger.LogWarning("Failed to load {Filename}", Filename);
+				Log.Logger.LogWarning(Ex, "{Text}", LogUtils.FormatException(Ex));
 			}
 			return Result;
 		}
@@ -565,8 +566,8 @@ namespace AutomationTool
 			}
 			catch (Exception Ex)
 			{
-				Log.TraceWarning("Failed to load {0}", Filename);
-				Log.TraceWarning(LogUtils.FormatException(Ex));
+				Log.Logger.LogWarning("Failed to load {Filename}", Filename);
+				Log.Logger.LogWarning(Ex, "{Text}", LogUtils.FormatException(Ex));
 			}
 			return Result;
 		}
@@ -604,7 +605,7 @@ namespace AutomationTool
 					{
 						if (!bQuiet)
 						{
-							Log.TraceWarning("Ignoring symlink {0}", File.FullName);
+							Log.Logger.LogWarning("Ignoring symlink {Path}", File.FullName);
 						}
 						continue;
 					}
@@ -652,8 +653,8 @@ namespace AutomationTool
 			}
 			catch (Exception Ex)
 			{
-				Log.TraceWarning("Unable to Find Files in {0}", Path);
-				Log.TraceWarning(LogUtils.FormatException(Ex));
+				Log.Logger.LogWarning("Unable to Find Files in {Path}", Path);
+				Log.Logger.LogWarning(Ex, "{Text}", LogUtils.FormatException(Ex));
 			}
 			return Files;
 		}
@@ -678,8 +679,8 @@ namespace AutomationTool
 			}
 			catch (Exception Ex)
 			{
-				Log.TraceWarning("Unable to Find Directories in {0}", Path);
-				Log.TraceWarning(LogUtils.FormatException(Ex));
+				Log.Logger.LogWarning("Unable to Find Directories in {Path}", Path);
+				Log.Logger.LogWarning(Ex, "{Text}", LogUtils.FormatException(Ex));
 			}
 			return Directories;
 		}
@@ -703,8 +704,8 @@ namespace AutomationTool
 			}
 			catch (Exception Ex)
 			{
-				Log.TraceWarning("Unable to check if file {0} exists.", Path);
-				Log.TraceWarning(LogUtils.FormatException(Ex));
+				Log.Logger.LogWarning("Unable to check if file {Path} exists.", Path);
+				Log.Logger.LogWarning(Ex, "{Text}", LogUtils.FormatException(Ex));
 			}
 			return Result;
 		}
@@ -729,8 +730,8 @@ namespace AutomationTool
 			}
 			catch (Exception Ex)
 			{
-				Log.TraceWarning("Unable to check if directory {0} exists.", Path);
-				Log.TraceWarning(LogUtils.FormatException(Ex));
+				Log.Logger.LogWarning("Unable to check if directory {Path} exists.", Path);
+				Log.Logger.LogWarning(Ex, "{Text}", LogUtils.FormatException(Ex));
 			}
 			return Result;
 		}
@@ -752,8 +753,8 @@ namespace AutomationTool
 			}
 			catch (Exception Ex)
 			{
-				Log.TraceWarning("Unable to write text to {0}", Path);
-				Log.TraceWarning(LogUtils.FormatException(Ex));
+				Log.Logger.LogWarning("Unable to write text to {Path}", Path);
+				Log.Logger.LogWarning(Ex, "{Text}", LogUtils.FormatException(Ex));
 			}
 			return Result;
 		}
@@ -775,8 +776,8 @@ namespace AutomationTool
 			}
 			catch (Exception Ex)
 			{
-				Log.TraceWarning("Unable to write text to {0}", Path);
-				Log.TraceWarning(LogUtils.FormatException(Ex));
+				Log.Logger.LogWarning("Unable to write text to {Path}", Path);
+				Log.Logger.LogWarning(Ex, "{Text}", LogUtils.FormatException(Ex));
 			}
 			return Result;
 		}
@@ -798,8 +799,8 @@ namespace AutomationTool
 			}
 			catch (Exception Ex)
 			{
-				Log.TraceWarning("Unable to write text to {0}", Path);
-				Log.TraceWarning(LogUtils.FormatException(Ex));
+				Log.Logger.LogWarning("Unable to write text to {Path}", Path);
+				Log.Logger.LogWarning(Ex, "{Text}", LogUtils.FormatException(Ex));
 			}
 			return Result;
 		}
diff --git a/Engine/Source/Programs/AutomationTool/BuildGraph/BuildGraph.cs b/Engine/Source/Programs/AutomationTool/BuildGraph/BuildGraph.cs
--- a/Engine/Source/Programs/AutomationTool/BuildGraph/BuildGraph.cs
+++ b/Engine/Source/Programs/AutomationTool/BuildGraph/BuildGraph.cs
@@ -278,7 +278,7 @@ namespace AutomationTool
 			// Look for overrides
 			if (!string.IsNullOrEmpty(BranchOverride))
 			{
-				LogInformation("Overriding default branch '{0}' with '{1}'", DefaultProperties["Branch"], BranchOverride);
+				Logger.LogInformation("Overriding default branch '{Branch}' with '{BranchOverride}'", DefaultProperties["Branch"], BranchOverride);
 				DefaultProperties["Branch"] = BranchOverride;
 				DefaultProperties["EscapedBranch"] = CommandUtils.EscapePath(DefaultProperties["Branch"]);
 			}
@@ -326,7 +326,7 @@ namespace AutomationTool
 					}
 					else
 					{
-						LogWarning("Missing value for '{0}'", Param.Substring(SetPrefix.Length));
+						Logger.LogWarning("Missing value for '{Arg0}'", Param.Substring(SetPrefix.Length));
 					}
 				}
 
@@ -349,7 +349,7 @@ namespace AutomationTool
 					}
 					else
 					{
-						LogWarning("Missing value for '{0}'", Param.Substring(AppendPrefix.Length));
+						Logger.LogWarning("Missing value for '{Arg0}'", Param.Substring(AppendPrefix.Length));
 					}
 				}
 			}
@@ -420,7 +420,7 @@ namespace AutomationTool
 					if (SchemaFileName != null)
 					{
 						FileReference FullSchemaFileName = new FileReference(SchemaFileName);
-						LogInformation("Writing schema to {0}...", FullSchemaFileName.FullName);
+						Logger.LogInformation("Writing schema to {Arg0}...", FullSchemaFileName.FullName);
 						Schema.Export(FullSchemaFileName);
 						if (ScriptFileName == null)
 						{
@@ -432,7 +432,7 @@ namespace AutomationTool
 				// Check there was a script specified
 				if (ScriptFileName == null)
 				{
-					LogError("Missing -Script= parameter for BuildGraph");
+					Logger.LogError("Missing -Script= parameter for BuildGraph");
 					return ExitCode.Error_Unknown;
 				}
 
@@ -471,7 +471,7 @@ namespace AutomationTool
 			{
 				if (!bListOnly && SingleNodeName == null)
 				{
-					LogError("Missing -Target= parameter for BuildGraph");
+					Logger.LogError("Missing -Target= parameter for BuildGraph");
 					return ExitCode.Error_Unknown;
 				}
 				TargetNodes.UnionWith(Graph.Agents.SelectMany(x => x.Nodes));
@@ -496,7 +496,7 @@ namespace AutomationTool
 					BgNodeDef[] Nodes;
 					if(!Graph.TryResolveReference(TargetName, out Nodes))
 					{
-						LogError("Target '{0}' is not in graph", TargetName);
+						Logger.LogError("Target '{TargetName}' is not in graph", TargetName);
 						return ExitCode.Error_Unknown;
 					}
 					TargetNodes.UnionWith(Nodes);
@@ -549,12 +549,12 @@ namespace AutomationTool
 						HashSet<BgNodeDef> SkipNodes = new HashSet<BgNodeDef>();
 						foreach(IGrouping<string, FileReference> MissingTokensForBuild in MissingTokens.GroupBy(x => x.Value, x => x.Key))
 						{
-							LogInformation("Skipping the following nodes due to {0}:", MissingTokensForBuild.Key);
+							Logger.LogInformation("Skipping the following nodes due to {Arg0}:", MissingTokensForBuild.Key);
 							foreach(FileReference MissingToken in MissingTokensForBuild)
 							{
 								foreach(BgNodeDef SkipNode in TargetNodes.Where(x => x.RequiredTokens.Contains(MissingToken) && SkipNodes.Add(x)))
 								{
-									LogInformation("    {0}", SkipNode);
+									Logger.LogInformation("    {SkipNode}", SkipNode);
 								}
 							}
 						}
@@ -563,14 +563,14 @@ namespace AutomationTool
 						if(SkipNodes.Count > 0)
 						{
 							TargetNodes.ExceptWith(SkipNodes);
-							LogInformation("Remaining target nodes:");
+							Logger.LogInformation("Remaining target nodes:");
 							foreach(BgNodeDef TargetNode in TargetNodes)
 							{
-								LogInformation("    {0}", TargetNode);
+								Logger.LogInformation("    {TargetNode}", TargetNode);
 							}
 							if(TargetNodes.Count == 0)
 							{
-								LogInformation("    None.");
+								Logger.LogInformation("    None.");
 							}
 						}
 					}
@@ -579,7 +579,7 @@ namespace AutomationTool
 						foreach(KeyValuePair<FileReference, string> Pair in MissingTokens)
 						{
 							List<BgNodeDef> SkipNodes = TargetNodes.Where(x => x.RequiredTokens.Contains(Pair.Key)).ToList();
-							LogError("Cannot run {0} due to previous build: {1}", String.Join(", ", SkipNodes), Pair.Value);
+							Logger.LogError("Cannot run {Arg0} due to previous build: {Arg1}", String.Join(", ", SkipNodes), Pair.Value);
 						}
 						foreach(FileReference CreatedToken in CreatedTokens)
 						{
@@ -605,7 +605,7 @@ namespace AutomationTool
 			if (PreprocessedFileName != null)
 			{
 				FileReference PreprocessedFileLocation = new FileReference(PreprocessedFileName);
-				LogInformation("Writing {0}...", PreprocessedFileLocation);
+				Logger.LogInformation("Writing {PreprocessedFileLocation}...", PreprocessedFileLocation);
 				Graph.Write(PreprocessedFileLocation, (SchemaFileName != null)? new FileReference(SchemaFileName) : null);
 				bListOnly = true;
 			}
@@ -614,7 +614,7 @@ namespace AutomationTool
 			BgNodeDef SingleNode = null;
 			if(SingleNodeName != null && !Graph.NameToNode.TryGetValue(SingleNodeName, out SingleNode))
 			{
-				LogError("Node '{0}' is not in the trimmed graph", SingleNodeName);
+				Logger.LogError("Node '{SingleNodeName}' is not in the trimmed graph", SingleNodeName);
 				return ExitCode.Error_Unknown;
 			}
 
@@ -741,7 +741,7 @@ namespace AutomationTool
 				}
 				catch (ReflectionTypeLoadException ex)
 				{
-					LogWarning("Exception {0} while trying to get types from assembly {1}. LoaderExceptions: {2}", ex, LoadedAssembly, string.Join("\n", ex.LoaderExceptions.Select(x => x.Message)));
+					Logger.LogWarning("Exception {ex} while trying to get types from assembly {LoadedAssembly}. LoaderExceptions: {Arg2}", ex, LoadedAssembly, string.Join("\n", ex.LoaderExceptions.Select(x => x.Message)));
 					continue;
 				}
 
@@ -777,7 +777,7 @@ namespace AutomationTool
 				}
 				catch (ReflectionTypeLoadException ex)
 				{
-					LogWarning("Exception {0} while trying to get types from assembly {1}. LoaderExceptions: {2}", ex, LoadedAssembly, string.Join("\n", ex.LoaderExceptions.Select(x => x.Message)));
+					Logger.LogWarning("Exception {ex} while trying to get types from assembly {LoadedAssembly}. LoaderExceptions: {Arg2}", ex, LoadedAssembly, string.Join("\n", ex.LoaderExceptions.Select(x => x.Message)));
 					continue;
 				}
 
@@ -938,15 +938,15 @@ namespace AutomationTool
 			int NodeIdx = 0;
 			foreach(BgNodeDef NodeToExecute in NodesToExecute)
 			{
-				LogInformation("****** [{0}/{1}] {2}", ++NodeIdx, NodesToExecute.Length, NodeToExecute.Name);
+				Logger.LogInformation("****** [{Arg0}/{Arg1}] {Arg2}", ++NodeIdx, NodesToExecute.Length, NodeToExecute.Name);
 				if(!Storage.IsComplete(NodeToExecute.Name))
 				{
-					LogInformation("");
+					Logger.LogInformation("");
 					if(!await BuildNodeAsync(Job, Graph, NodeToExecute, NodeToExecutor, Storage, bWithBanner: false))
 					{
 						return false;
 					} 
-					LogInformation("");
+					Logger.LogInformation("");
 				}
 			}
 			return true;
@@ -1001,7 +1001,7 @@ namespace AutomationTool
 					TempStorageBlock CurrentStorageBlock;
 					if(FileToStorageBlock.TryGetValue(File, out CurrentStorageBlock) && !TempStorage.IsDuplicateBuildProduct(File))
 					{
-						LogError("File '{0}' was produced by {1} and {2}", File, InputStorageBlock, CurrentStorageBlock);
+						Logger.LogError("File '{File}' was produced by {InputStorageBlock} and {CurrentStorageBlock}", File, InputStorageBlock, CurrentStorageBlock);
 					}
 					FileToStorageBlock[File] = InputStorageBlock;
 				}
@@ -1332,7 +1332,7 @@ namespace AutomationTool
 		/// <param name="OutputFile">The output file to write to</param>
 		static void WriteDocumentationHTML(Dictionary<string, ScriptTaskBinding> NameToTask, Dictionary<string, XmlElement> MemberNameToElement, FileReference OutputFile)
 		{
-			LogInformation("Writing {0}...", OutputFile);
+			Logger.LogInformation("Writing {OutputFile}...", OutputFile);
 			using (StreamWriter Writer = new StreamWriter(OutputFile.FullName))
 			{
 				Writer.WriteLine("<html>");
diff --git a/Engine/Source/Programs/AutomationTool/BuildGraph/Scripts/InstalledBuild.cs b/Engine/Source/Programs/AutomationTool/BuildGraph/Scripts/InstalledBuild.cs
--- a/Engine/Source/Programs/AutomationTool/BuildGraph/Scripts/InstalledBuild.cs
+++ b/Engine/Source/Programs/AutomationTool/BuildGraph/Scripts/InstalledBuild.cs
@@ -12,6 +12,7 @@ using UnrealBuildBase;
 using System.Threading.Tasks;
 
 using static AutomationTool.Tasks.StandardTasks;
+using Microsoft.Extensions.Logging;
 
 #nullable enable
 
@@ -247,7 +248,7 @@ namespace AutomationTool
 			if (State.Get(EmbedSrcSrvInfo))
 			{
 				// Embed source info into the PDB files. Should be done from this machine to ensure that paths are correct.
-				Log.TraceInformation("Embedding source file information into PDB files...");
+				Log.Logger.LogInformation("Embedding source file information into PDB files...");
 				FileSet SourceFiles = Workspace.Filter("Engine/Source/...;Engine/Plugins/...").Except("Engine/Source/ThirdParty/...").Filter("*.c;*.h;*.cpp;*.hpp;*.inl");
 //				State.SrcSrv(BinaryFiles: Full, SourceFiles: SourceFiles);
 			}
diff --git a/Engine/Source/Programs/AutomationTool/BuildGraph/Tasks/AwsEcsDeployTask.cs b/Engine/Source/Programs/AutomationTool/BuildGraph/Tasks/AwsEcsDeployTask.cs
--- a/Engine/Source/Programs/AutomationTool/BuildGraph/Tasks/AwsEcsDeployTask.cs
+++ b/Engine/Source/Programs/AutomationTool/BuildGraph/Tasks/AwsEcsDeployTask.cs
@@ -10,6 +10,7 @@ using System.Text.Json;
 using System.Threading.Tasks;
 using System.Xml;
 using UnrealBuildBase;
+using Microsoft.Extensions.Logging;
 
 namespace AutomationTool.Tasks
 {
@@ -115,7 +116,7 @@ namespace AutomationTool.Tasks
 			string Params = $"ecs update-service --cluster {Parameters.Cluster} --service {Parameters.Service} --task-definition {TaskDefFamily}:{TaskDefRevision}";
 			await SpawnTaskBase.ExecuteAsync("aws", Params, EnvVars: ParseEnvVars(Parameters.Environment, Parameters.EnvironmentFile), LogOutput: Parameters.LogOutput);
 
-			Log.TraceInformation($"Service {Parameters.Service} updated to use new task def {TaskDefFamily}:{TaskDefRevision}");
+			Log.Logger.LogInformation("Service {Service} updated to use new task def {TaskDefFamily}:{TaskDefRevision}", Parameters.Service, TaskDefFamily, TaskDefRevision);
 		}
 
 		/// <summary>
diff --git a/Engine/Source/Programs/AutomationTool/BuildGraph/Tasks/CommandTask.cs b/Engine/Source/Programs/AutomationTool/BuildGraph/Tasks/CommandTask.cs
--- a/Engine/Source/Programs/AutomationTool/BuildGraph/Tasks/CommandTask.cs
+++ b/Engine/Source/Programs/AutomationTool/BuildGraph/Tasks/CommandTask.cs
@@ -13,6 +13,7 @@ using EpicGames.Core;
 using UnrealBuildBase;
 using EpicGames.BuildGraph;
 using AutomationTool.Tasks;
+using Microsoft.Extensions.Logging;
 
 namespace AutomationTool.Tasks
 {
@@ -125,7 +126,7 @@ namespace AutomationTool.Tasks
 				}
 				else
 				{
-					Log.TraceWarning("Unable to read UAT telemetry file from {0}", TelemetryFile);
+					Log.Logger.LogWarning("Unable to read UAT telemetry file from {TelemetryFile}", TelemetryFile);
 				}
 			}
 			return Task.CompletedTask;
diff --git a/Engine/Source/Programs/AutomationTool/BuildGraph/Tasks/CopyTask.cs b/Engine/Source/Programs/AutomationTool/BuildGraph/Tasks/CopyTask.cs
--- a/Engine/Source/Programs/AutomationTool/BuildGraph/Tasks/CopyTask.cs
+++ b/Engine/Source/Programs/AutomationTool/BuildGraph/Tasks/CopyTask.cs
@@ -172,11 +172,11 @@ namespace AutomationTool.Tasks
 						DirectoryReference.CreateDirectory(FirstTargetDirectory);
 						if(NumRetries == 1)
 						{
-							Log.TraceInformation("Created target directory {0} after 1 retry.", FirstTargetDirectory);
+							Log.Logger.LogInformation("Created target directory {FirstTargetDirectory} after 1 retry.", FirstTargetDirectory);
 						}
 						else if(NumRetries > 1)
 						{
-							Log.TraceInformation("Created target directory {0} after {1} retries.", FirstTargetDirectory, NumRetries);
+							Log.Logger.LogInformation("Created target directory {FirstTargetDirectory} after {NumRetries} retries.", FirstTargetDirectory, NumRetries);
 						}
 						break;
 					}
@@ -184,7 +184,7 @@ namespace AutomationTool.Tasks
 					{
 						if(NumRetries == 0)
 						{
-							Log.TraceInformation("Unable to create directory '{0}' on first attempt. Retrying {1} times...", FirstTargetDirectory, MaxNumRetries);
+							Log.Logger.LogInformation("Unable to create directory '{FirstTargetDirectory}' on first attempt. Retrying {MaxNumRetries} times...", FirstTargetDirectory, MaxNumRetries);
 						}
 
 						Log.TraceLog("  {0}", Ex);
diff --git a/Engine/Source/Programs/AutomationTool/BuildGraph/Tasks/CsCompileTask.cs b/Engine/Source/Programs/AutomationTool/BuildGraph/Tasks/CsCompileTask.cs
--- a/Engine/Source/Programs/AutomationTool/BuildGraph/Tasks/CsCompileTask.cs
+++ b/Engine/Source/Programs/AutomationTool/BuildGraph/Tasks/CsCompileTask.cs
@@ -10,6 +10,7 @@ using System.Xml;
 using EpicGames.Core;
 using UnrealBuildBase;
 using UnrealBuildTool;
+using Microsoft.Extensions.Logging;
 
 namespace AutomationTool.Tasks
 {
@@ -141,7 +142,7 @@ namespace AutomationTool.Tasks
 						int EqualsIdx = Property.IndexOf('=');
 						if (EqualsIdx == -1)
 						{
-							Log.TraceWarning("Missing '=' in property assignment");
+							Log.Logger.LogWarning("Missing '=' in property assignment");
 						}
 						else
 						{
diff --git a/Engine/Source/Programs/AutomationTool/BuildGraph/Tasks/DeleteTask.cs b/Engine/Source/Programs/AutomationTool/BuildGraph/Tasks/DeleteTask.cs
--- a/Engine/Source/Programs/AutomationTool/BuildGraph/Tasks/DeleteTask.cs
+++ b/Engine/Source/Programs/AutomationTool/BuildGraph/Tasks/DeleteTask.cs
@@ -12,6 +12,7 @@ using System.Xml;
 using EpicGames.Core;
 using UnrealBuildTool;
 using UnrealBuildBase;
+using Microsoft.Extensions.Logging;
 
 namespace AutomationTool.Tasks
 {
@@ -81,7 +82,7 @@ namespace AutomationTool.Tasks
 				{
 					if (Parameters.Verbose)
 					{
-						Log.TraceInformation("Deleting {0}", File.FullName);
+						Log.Logger.LogInformation("Deleting {File}", File.FullName);
 					}
 					if (!InternalUtils.SafeDeleteFile(File.FullName))
 					{
@@ -120,7 +121,7 @@ namespace AutomationTool.Tasks
 					{
 						if (Parameters.Verbose)
 						{
-							Log.TraceInformation("Deleting {0}", Directory);
+							Log.Logger.LogInformation("Deleting {Directory}", Directory);
 						}
 						DirectoryReference FullDir = new DirectoryReference(Directory);
 						if (DirectoryReference.Exists(FullDir))
diff --git a/Engine/Source/Programs/AutomationTool/BuildGraph/Tasks/DeployToolTask.cs b/Engine/Source/Programs/AutomationTool/BuildGraph/Tasks/DeployToolTask.cs
--- a/Engine/Source/Programs/AutomationTool/BuildGraph/Tasks/DeployToolTask.cs
+++ b/Engine/Source/Programs/AutomationTool/BuildGraph/Tasks/DeployToolTask.cs
@@ -12,6 +12,7 @@ using System.Text;
 using System.Threading.Tasks;
 using System.Xml;
 using UnrealBuildBase;
+using Microsoft.Extensions.Logging;
 
 #nullable enable
 
@@ -141,7 +142,7 @@ namespace AutomationTool.Tasks
 						content.Add(streamContent, "file", file.GetFileName());
 						request.Content = content;
 
-						Log.TraceInformation("Uploading {0} to {1}", file, request.RequestUri);
+						Log.Logger.LogInformation("Uploading {File} to {Url}", file, request.RequestUri);
 						using (HttpResponseMessage response = await httpClient.SendAsync(request))
 						{
 							if (!response.IsSuccessStatusCode)
diff --git a/Engine/Source/Programs/AutomationTool/BuildGraph/Tasks/DockerBuildTask.cs b/Engine/Source/Programs/AutomationTool/BuildGraph/Tasks/DockerBuildTask.cs
--- a/Engine/Source/Programs/AutomationTool/BuildGraph/Tasks/DockerBuildTask.cs
+++ b/Engine/Source/Programs/AutomationTool/BuildGraph/Tasks/DockerBuildTask.cs
@@ -11,6 +11,7 @@ using AutomationTool;
 using UnrealBuildBase;
 using System.Threading.Tasks;
 using System.Text.RegularExpressions;
+using Microsoft.Extensions.Logging;
 
 namespace AutomationTool.Tasks
 {
@@ -114,7 +115,7 @@ namespace AutomationTool.Tasks
 		/// <param name="TagNameToFileSet">Mapping from tag names to the set of files they include</param>
 		public override async Task ExecuteAsync(JobContext Job, HashSet<FileReference> BuildProducts, Dictionary<string, HashSet<FileReference>> TagNameToFileSet)
 		{
-			Log.TraceInformation("Building Docker image");
+			Log.Logger.LogInformation("Building Docker image");
 			using (LogIndentScope Scope = new LogIndentScope("  "))
 			{
 				DirectoryReference BaseDir = ResolveDirectory(Parameters.BaseDir);
diff --git a/Engine/Source/Programs/AutomationTool/BuildGraph/Tasks/DockerComposeDownTask.cs b/Engine/Source/Programs/AutomationTool/BuildGraph/Tasks/DockerComposeDownTask.cs
--- a/Engine/Source/Programs/AutomationTool/BuildGraph/Tasks/DockerComposeDownTask.cs
+++ b/Engine/Source/Programs/AutomationTool/BuildGraph/Tasks/DockerComposeDownTask.cs
@@ -10,6 +10,7 @@ using System.Xml;
 using AutomationTool;
 using UnrealBuildBase;
 using System.Threading.Tasks;
+using Microsoft.Extensions.Logging;
 
 namespace AutomationTool.Tasks
 {
@@ -70,7 +71,7 @@ namespace AutomationTool.Tasks
 				Arguments.Append($" {Parameters.Arguments}");
 			}
 
-			Log.TraceInformation("Running docker compose {0}", Arguments.ToString());
+			Log.Logger.LogInformation("Running docker compose {Arguments}", Arguments.ToString());
 			using (LogIndentScope Scope = new LogIndentScope("  "))
 			{
 				await SpawnTaskBase.ExecuteAsync("docker-compose", Arguments.ToString());
diff --git a/Engine/Source/Programs/AutomationTool/BuildGraph/Tasks/DockerComposeUpTask.cs b/Engine/Source/Programs/AutomationTool/BuildGraph/Tasks/DockerComposeUpTask.cs
--- a/Engine/Source/Programs/AutomationTool/BuildGraph/Tasks/DockerComposeUpTask.cs
+++ b/Engine/Source/Programs/AutomationTool/BuildGraph/Tasks/DockerComposeUpTask.cs
@@ -10,6 +10,7 @@ using System.Xml;
 using AutomationTool;
 using UnrealBuildBase;
 using System.Threading.Tasks;
+using Microsoft.Extensions.Logging;
 
 namespace AutomationTool.Tasks
 {
@@ -70,7 +71,7 @@ namespace AutomationTool.Tasks
 				Arguments.Append($" {Parameters.Arguments}");
 			}
 
-			Log.TraceInformation("Running docker compose {0}", Arguments.ToString());
+			Log.Logger.LogInformation("Running docker compose {Arguments}", Arguments.ToString());
 			using (LogIndentScope Scope = new LogIndentScope("  "))
 			{
 				await SpawnTaskBase.ExecuteAsync("docker-compose", Arguments.ToString());
diff --git a/Engine/Source/Programs/AutomationTool/BuildGraph/Tasks/DockerPushTask.cs b/Engine/Source/Programs/AutomationTool/BuildGraph/Tasks/DockerPushTask.cs
--- a/Engine/Source/Programs/AutomationTool/BuildGraph/Tasks/DockerPushTask.cs
+++ b/Engine/Source/Programs/AutomationTool/BuildGraph/Tasks/DockerPushTask.cs
@@ -10,6 +10,7 @@ using System.Text;
 using System.Text.Json;
 using System.Threading.Tasks;
 using System.Xml;
+using Microsoft.Extensions.Logging;
 
 namespace AutomationTool.Tasks
 {
@@ -89,7 +90,7 @@ namespace AutomationTool.Tasks
 		/// <param name="TagNameToFileSet">Mapping from tag names to the set of files they include</param>
 		public override async Task ExecuteAsync(JobContext Job, HashSet<FileReference> BuildProducts, Dictionary<string, HashSet<FileReference>> TagNameToFileSet)
 		{
-			Log.TraceInformation("Pushing Docker image");
+			Log.Logger.LogInformation("Pushing Docker image");
 			using (LogIndentScope Scope = new LogIndentScope("  "))
 			{
 				string Exe = DockerTask.GetDockerExecutablePath();
diff --git a/Engine/Source/Programs/AutomationTool/BuildGraph/Tasks/FindModifiedFiles.cs b/Engine/Source/Programs/AutomationTool/BuildGraph/Tasks/FindModifiedFiles.cs
--- a/Engine/Source/Programs/AutomationTool/BuildGraph/Tasks/FindModifiedFiles.cs
+++ b/Engine/Source/Programs/AutomationTool/BuildGraph/Tasks/FindModifiedFiles.cs
@@ -117,20 +117,20 @@ namespace AutomationTool.Tasks
 				}
 				else
 				{
-					Log.TraceInformation("Unable to map {0} to workspace; skipping.", File.DepotFile);
+					Log.Logger.LogInformation("Unable to map {DepotFile} to workspace; skipping.", File.DepotFile);
 				}
 			}
 
-			Log.TraceInformation("Found {0} modified files matching {1}", LocalFiles.Count, Filter.ToString());
+			Log.Logger.LogInformation("Found {NumFiles} modified files matching {Filter}", LocalFiles.Count, Filter.ToString());
 			foreach (string LocalFile in LocalFiles)
 			{
-				Log.TraceInformation("  {0}", LocalFile);
+				Log.Logger.LogInformation("  {LocalFile}", LocalFile);
 			}
 
 			if (Parameters.Output != null)
 			{
 				await FileReference.WriteAllLinesAsync(Parameters.Output, LocalFiles);
-				Log.TraceInformation("Written {0}", Parameters.Output);
+				Log.Logger.LogInformation("Written {OutputFile}", Parameters.Output);
 			}
 		}
 
diff --git a/Engine/Source/Programs/AutomationTool/BuildGraph/Tasks/GitCloneTask.cs b/Engine/Source/Programs/AutomationTool/BuildGraph/Tasks/GitCloneTask.cs
--- a/Engine/Source/Programs/AutomationTool/BuildGraph/Tasks/GitCloneTask.cs
+++ b/Engine/Source/Programs/AutomationTool/BuildGraph/Tasks/GitCloneTask.cs
@@ -7,6 +7,7 @@ using System.Text;
 using System.Threading.Tasks;
 using System.Xml;
 using UnrealBuildBase;
+using Microsoft.Extensions.Logging;
 
 namespace AutomationTool.Tasks
 {
@@ -75,7 +76,7 @@ namespace AutomationTool.Tasks
 			}
 
 			DirectoryReference Dir = ResolveDirectory(Parameters.Dir);
-			Log.TraceInformation("Cloning Git repository into {0}", Parameters.Dir);
+			Log.Logger.LogInformation("Cloning Git repository into {Dir}", Parameters.Dir);
 			using (LogIndentScope Scope = new LogIndentScope("  "))
 			{
 				DirectoryReference GitDir = DirectoryReference.Combine(Dir, ".git");
diff --git a/Engine/Source/Programs/AutomationTool/BuildGraph/Tasks/SanitizeReceiptTask.cs b/Engine/Source/Programs/AutomationTool/BuildGraph/Tasks/SanitizeReceiptTask.cs
--- a/Engine/Source/Programs/AutomationTool/BuildGraph/Tasks/SanitizeReceiptTask.cs
+++ b/Engine/Source/Programs/AutomationTool/BuildGraph/Tasks/SanitizeReceiptTask.cs
@@ -10,6 +10,7 @@ using System.Xml;
 using EpicGames.Core;
 using UnrealBuildBase;
 using UnrealBuildTool;
+using Microsoft.Extensions.Logging;
 
 namespace AutomationTool.Tasks
 {
@@ -80,7 +81,7 @@ namespace AutomationTool.Tasks
 				}
 
 				// Print the name of the file being scanned
-				Log.TraceInformation("Sanitizing {0}", TargetFile);
+				Log.Logger.LogInformation("Sanitizing {TargetFile}", TargetFile);
 				using(new LogIndentScope("  "))
 				{
 					// Read the receipt
@@ -101,7 +102,7 @@ namespace AutomationTool.Tasks
 						}
 						else
 						{
-							Log.TraceInformation("Removing build product: {0}", BuildProduct.Path);
+							Log.Logger.LogInformation("Removing build product: {File}", BuildProduct.Path);
 						}
 					}
 					Receipt.BuildProducts = NewBuildProducts;
@@ -116,7 +117,7 @@ namespace AutomationTool.Tasks
 						}
 						else
 						{
-							Log.TraceInformation("Removing runtime dependency: {0}", RuntimeDependency.Path);
+							Log.Logger.LogInformation("Removing runtime dependency: {File}", RuntimeDependency.Path);
 						}
 					}
 					Receipt.RuntimeDependencies = NewRuntimeDependencies;
diff --git a/Engine/Source/Programs/AutomationTool/BuildGraph/Tasks/SubmitTask.cs b/Engine/Source/Programs/AutomationTool/BuildGraph/Tasks/SubmitTask.cs
--- a/Engine/Source/Programs/AutomationTool/BuildGraph/Tasks/SubmitTask.cs
+++ b/Engine/Source/Programs/AutomationTool/BuildGraph/Tasks/SubmitTask.cs
@@ -112,11 +112,11 @@ namespace AutomationTool.Tasks
 			HashSet<FileReference> Files = ResolveFilespec(Unreal.RootDirectory, Parameters.Files, TagNameToFileSet);
 			if (Files.Count == 0)
 			{
-				Log.TraceInformation("No files to submit.");
+				Log.Logger.LogInformation("No files to submit.");
 			}
 			else if (!CommandUtils.AllowSubmit)
 			{
-				Log.TraceWarning("Submitting to Perforce is disabled by default. Run with the -submit argument to allow.");
+				Log.Logger.LogWarning("Submitting to Perforce is disabled by default. Run with the -submit argument to allow.");
 			}
 			else
 			{
diff --git a/Engine/Source/Programs/AutomationTool/BuildGraph/Tasks/UploadTask.cs b/Engine/Source/Programs/AutomationTool/BuildGraph/Tasks/UploadTask.cs
--- a/Engine/Source/Programs/AutomationTool/BuildGraph/Tasks/UploadTask.cs
+++ b/Engine/Source/Programs/AutomationTool/BuildGraph/Tasks/UploadTask.cs
@@ -13,6 +13,7 @@ using System.Threading.Tasks;
 using System.Xml;
 using EpicGames.Core;
 using EpicGames.Jupiter;
+using Microsoft.Extensions.Logging;
 
 namespace AutomationTool.Tasks
 {
@@ -128,7 +129,7 @@ namespace AutomationTool.Tasks
 			}
 
 			// Create the jupiter tree
-			Log.TraceInformation("Uploading {0} files to {1}...", Files.Count, Parameters.JupiterUrl);
+			Log.Logger.LogInformation("Uploading {NumFiles} files to {Url}...", Files.Count, Parameters.JupiterUrl);
 
 			JupiterFileTree FileTree = new JupiterFileTree(Parameters.FromDir, Parameters.LimitMemoryUsage);
 			foreach (FileReference File in Files)
diff --git a/Engine/Source/Programs/AutomationTool/BuildGraph/Tasks/ZipTask.cs b/Engine/Source/Programs/AutomationTool/BuildGraph/Tasks/ZipTask.cs
--- a/Engine/Source/Programs/AutomationTool/BuildGraph/Tasks/ZipTask.cs
+++ b/Engine/Source/Programs/AutomationTool/BuildGraph/Tasks/ZipTask.cs
@@ -10,6 +10,7 @@ using System.Threading.Tasks;
 using System.Xml;
 using EpicGames.Core;
 using UnrealBuildTool;
+using Microsoft.Extensions.Logging;
 
 namespace AutomationTool.Tasks
 {
@@ -83,7 +84,7 @@ namespace AutomationTool.Tasks
 			}
 
 			// Create the zip file
-			Log.TraceInformation("Adding {0} files to {1}...", Files.Count, Parameters.ZipFile);
+			Log.Logger.LogInformation("Adding {NumFiles} files to {ZipFile}...", Files.Count, Parameters.ZipFile);
 			CommandUtils.ZipFiles(Parameters.ZipFile, Parameters.FromDir, Files);
 
 			// Apply the optional tag to the produced archive
diff --git a/Engine/Source/Programs/AutomationTool/BuildGraph/TempStorage.cs b/Engine/Source/Programs/AutomationTool/BuildGraph/TempStorage.cs
--- a/Engine/Source/Programs/AutomationTool/BuildGraph/TempStorage.cs
+++ b/Engine/Source/Programs/AutomationTool/BuildGraph/TempStorage.cs
@@ -16,6 +16,7 @@ using System.Threading.Tasks;
 using System.Xml;
 using System.Xml.Serialization;
 using UnrealBuildTool;
+using Microsoft.Extensions.Logging;
 
 namespace AutomationTool
 {
@@ -1156,7 +1157,7 @@ namespace AutomationTool
 													throw;
 												}
 
-												Log.TraceWarning("Failed to unzip '{0}' from '{1}' to '{2}', retrying.. (Error: {3})", Entry.FullName, LocalZipFile, ExtractedFilename, IOEx.Message);
+												Log.Logger.LogWarning(IOEx, "Failed to unzip '{File}' from '{LocalZipFile}' to '{ExtractedFilename}', retrying.. (Error: {Message})", Entry.FullName, LocalZipFile, ExtractedFilename, IOEx.Message);
 											}
 										}
 									}
@@ -1169,7 +1170,7 @@ namespace AutomationTool
 						{
 							if (UnzipFileAttempts == 0)
 							{
-								Log.TraceError("All retries exhausted attempting to unzip entries from '{0}'. Terminating.", LocalZipFile);
+								Log.Logger.LogError(Ex, "All retries exhausted attempting to unzip entries from '{LocalZipFile}'. Terminating.", LocalZipFile);
 								string LogPath = CommandUtils.CombinePaths(RootDir.FullName, "Engine/Programs/AutomationTool/Saved/Logs", $"Copy Manifest - {ZipFile.Name}.log");
 								CommandUtils.LogInformation("Saving copy log to {0}", LogPath);
 								File.WriteAllText(LogPath, string.Join(Environment.NewLine, CopyResults));
@@ -1179,7 +1180,7 @@ namespace AutomationTool
 							// Some exceptions may be caused by networking hiccups. We want to retry in those cases.
 							if ((Ex is IOException || Ex is InvalidDataException))
 							{
-								Log.TraceWarning("Failed to unzip entries from '{0}' to '{1}', retrying.. (Error: {2})", LocalZipFile, RootDir.FullName, Ex.Message);
+								Log.Logger.LogWarning(Ex, "Failed to unzip entries from '{LocalZipFile}' to '{TargetDir}', retrying.. (Error: {Message})", LocalZipFile, RootDir.FullName, Ex.Message);
 							}
 						}
 						finally
@@ -1528,13 +1529,13 @@ namespace AutomationTool
 						}
 						catch (Exception Ex)
 						{
-							Log.TraceError("Exception while trying to scan files under {0}: {1}", BuildDirectory, Ex);
+							Log.Logger.LogError(Ex, "Exception while trying to scan files under {BuildDirectory}: {Ex}", BuildDirectory, Ex);
 						}
 					}
 				}
 				catch (Exception Ex)
 				{
-					Log.TraceError("Exception while trying to scan {0}: {1}", StreamDirectory, Ex);
+					Log.Logger.LogError(Ex, "Exception while trying to scan {StreamDirectory}: {Ex}", StreamDirectory, Ex);
 				}
 			}
 			CommandUtils.LogInformation("Found {0} builds; {1} to delete.", NumBuilds, BuildsToDelete.Count);
@@ -1593,7 +1594,7 @@ namespace AutomationTool
 				}
 				catch (Exception Ex)
 				{
-					Log.TraceError("Exception while trying to delete {0}: {1}", StreamDirectory, Ex);
+					Log.Logger.LogError(Ex, "Exception while trying to delete {StreamDirectory}: {Ex}", StreamDirectory, Ex);
 				}
 			}
 		}
