Missing files.

diff --git a/Engine/Source/Programs/AutomationTool/BuildGraph/BuildGraph.cs b/Engine/Source/Programs/AutomationTool/BuildGraph/BuildGraph.cs
--- a/Engine/Source/Programs/AutomationTool/BuildGraph/BuildGraph.cs
+++ b/Engine/Source/Programs/AutomationTool/BuildGraph/BuildGraph.cs
@@ -741,7 +741,7 @@ namespace AutomationTool
 				}
 				catch (ReflectionTypeLoadException ex)
 				{
-					Logger.LogWarning("Exception {ex} while trying to get types from assembly {LoadedAssembly}. LoaderExceptions: {Arg2}", ex, LoadedAssembly, string.Join("\n", ex.LoaderExceptions.Select(x => x.Message)));
+					Log.Logger.LogWarning("Exception {ex} while trying to get types from assembly {LoadedAssembly}. LoaderExceptions: {Arg2}", ex, LoadedAssembly, string.Join("\n", ex.LoaderExceptions.Select(x => x.Message)));
 					continue;
 				}
 
@@ -777,7 +777,7 @@ namespace AutomationTool
 				}
 				catch (ReflectionTypeLoadException ex)
 				{
-					Logger.LogWarning("Exception {ex} while trying to get types from assembly {LoadedAssembly}. LoaderExceptions: {Arg2}", ex, LoadedAssembly, string.Join("\n", ex.LoaderExceptions.Select(x => x.Message)));
+					Log.Logger.LogWarning("Exception {ex} while trying to get types from assembly {LoadedAssembly}. LoaderExceptions: {Arg2}", ex, LoadedAssembly, string.Join("\n", ex.LoaderExceptions.Select(x => x.Message)));
 					continue;
 				}
 
@@ -1332,7 +1332,7 @@ namespace AutomationTool
 		/// <param name="OutputFile">The output file to write to</param>
 		static void WriteDocumentationHTML(Dictionary<string, ScriptTaskBinding> NameToTask, Dictionary<string, XmlElement> MemberNameToElement, FileReference OutputFile)
 		{
-			Logger.LogInformation("Writing {OutputFile}...", OutputFile);
+			Log.Logger.LogInformation("Writing {OutputFile}...", OutputFile);
 			using (StreamWriter Writer = new StreamWriter(OutputFile.FullName))
 			{
 				Writer.WriteLine("<html>");
diff --git a/Engine/Source/Programs/AutomationTool/BuildGraph/Tasks/FindModifiedFiles.cs b/Engine/Source/Programs/AutomationTool/BuildGraph/Tasks/FindModifiedFiles.cs
--- a/Engine/Source/Programs/AutomationTool/BuildGraph/Tasks/FindModifiedFiles.cs
+++ b/Engine/Source/Programs/AutomationTool/BuildGraph/Tasks/FindModifiedFiles.cs
@@ -12,6 +12,7 @@ using System.Xml;
 using EpicGames.Core;
 using UnrealBuildBase;
 using UnrealBuildTool;
+using Microsoft.Extensions.Logging;
 
 namespace AutomationTool.Tasks
 {
