Change to structured logging in other AutomationTool projects.

diff --git a/Engine/Platforms/Hololens/Source/Programs/AutomationTool/HoloLensPlatform.Automation.cs b/Engine/Platforms/Hololens/Source/Programs/AutomationTool/HoloLensPlatform.Automation.cs
--- a/Engine/Platforms/Hololens/Source/Programs/AutomationTool/HoloLensPlatform.Automation.cs
+++ b/Engine/Platforms/Hololens/Source/Programs/AutomationTool/HoloLensPlatform.Automation.cs
@@ -20,6 +20,7 @@ using Microsoft.Tools.WindowsDevicePortal;
 using UnrealBuildBase;
 using AutomationScripts;
 using System.Runtime.Versioning;
+using Microsoft.Extensions.Logging;
 
 namespace HoloLens.Automation
 {
@@ -388,7 +389,7 @@ namespace HoloLens.Automation
 				{
 					if (!IsLocalDevice(DeviceAddress))
 					{
-						LogWarning("Project will be packaged to support deployment to remote HoloLens device {0}.", DeviceAddress);
+						Logger.LogWarning("Project will be packaged to support deployment to remote HoloLens device {DeviceAddress}.", DeviceAddress);
 						ProjParams.Package = true;
 						break;
 					}
@@ -426,7 +427,7 @@ namespace HoloLens.Automation
 						}
 
 						ArchSet.Add(Arch);
-						LogInformation(String.Format("Project will be compiled for the architecture {0} of the HoloLens device {1}.", Arch, DeviceAddress));
+						Logger.LogInformation("Project will be compiled for the architecture {Arch} of the HoloLens device {Device}.", Arch, DeviceAddress);
 					}
 				}
 
@@ -436,18 +437,18 @@ namespace HoloLens.Automation
 					if (PlatformEngineConfig.GetBool("/Script/HoloLensPlatformEditor.HoloLensTargetSettings", "bBuildForEmulation", out bBuildForEmulation) && bBuildForEmulation)
 					{
 						ArchSet.Add(UnrealArch.X64);
-						LogInformation(String.Format("Project will be compiled for the architecture X64 of the HoloLens  device for packaging because bBuildForEmulation is true in project settings."));
+						Logger.LogInformation("Project will be compiled for the architecture X64 of the HoloLens  device for packaging because bBuildForEmulation is true in project settings.");
 					}
 					else
 					{
 						ArchSet.Add(UnrealArch.Arm64);
-						LogInformation(String.Format("Project will be compiled for the architecture Arm64 of the HoloLens device for packaging."));
+						Logger.LogInformation("Project will be compiled for the architecture Arm64 of the HoloLens device for packaging.");
 					}
 				}
 
 				if (ArchSet.Count > 0)
 				{
-					LogInformation(String.Format("Project will be compiled for the architectures {0} for HoloLens.", string.Join(", ", ArchSet)));
+					Logger.LogInformation("Project will be compiled for the architectures {Archs} for HoloLens.", string.Join(", ", ArchSet));
 					ProjParams.ClientArchitecture = new UnrealArchitectures(ArchSet);
 				}
 			}
@@ -463,7 +464,7 @@ namespace HoloLens.Automation
 				Windows10SDKVersion = "Latest";
 			}
 
-			if(!HoloLensExports.InitWindowsSdkToolPath(Windows10SDKVersion, Log.Logger))
+			if(!HoloLensExports.InitWindowsSdkToolPath(Windows10SDKVersion, Logger))
 			{
 				throw new AutomationException(ExitCode.Error_Arguments, "Wrong WinSDK toolchain selected on \'Platforms/HoloLens/Toolchain\' page. Please check.");
 			}
@@ -547,7 +548,7 @@ namespace HoloLens.Automation
 			{
 				// Stage all the build products
 				DirectoryReference ProjectBinariesFolder = Params.GetProjectBinariesPathForPlatform(PlatformType);
-				HoloLensExports DeployExports = new HoloLensExports(Log.Logger);
+				HoloLensExports DeployExports = new HoloLensExports(Logger);
 				foreach (StageTarget Target in SC.StageTargets)
 				{
 					SC.StageBuildProductsFromReceipt(Target.Receipt, Target.RequireFilesExist, Params.bTreatNonShippingBinariesAsDebugFiles);
@@ -678,7 +679,7 @@ namespace HoloLens.Automation
 		private void PackagePakFiles(ProjectParams Params, DeploymentContext SC, string OutputNameBase)
 		{
 			string IntermediateDirectory = Path.Combine(SC.ProjectRoot.FullName, "Intermediate", "Deploy", "neutral");
-			var ListResources = new HoloLensManifestGenerator(Log.Logger).CreateAssetsManifest(SC.StageTargetPlatform.PlatformType, SC.StageDirectory.FullName, IntermediateDirectory, SC.RawProjectPath, SC.ProjectRoot.FullName);
+			var ListResources = new HoloLensManifestGenerator(Logger).CreateAssetsManifest(SC.StageTargetPlatform.PlatformType, SC.StageDirectory.FullName, IntermediateDirectory, SC.RawProjectPath, SC.ProjectRoot.FullName);
 
 			string OutputName = OutputNameBase + "_pak";
 
@@ -981,7 +982,7 @@ namespace HoloLens.Automation
 			{
 				if (!IsBuildMachine && !Params.Unattended)
 				{
-					LogError("Certificate is required.  Please go to Project Settings > HoloLens > Create Signing Certificate");
+					Logger.LogError("Certificate is required.  Please go to Project Settings > HoloLens > Create Signing Certificate");
 				}
 			}
 		}
@@ -1037,7 +1038,7 @@ namespace HoloLens.Automation
 				if (Params.GameConfigs.TryGetValue(PlatformType, out PlatformGameConfig))
 				{
 					PlatformGameConfig.GetBool("/Script/EngineSettings.GeneralProjectSettings", "bStartInVR", out bStartInVR);
-					Log.TraceInformation("bStartInVR = {0}", bStartInVR.ToString());
+					Logger.LogInformation("bStartInVR = {StartInVr}", bStartInVR.ToString());
 				}
 
 				if (bStartInVR)
@@ -1058,7 +1059,7 @@ namespace HoloLens.Automation
 					{
 						// Update the uecommandline.txt
 						FileReference IntermediateCmdLineFile = FileReference.Combine(SC.StageDirectory, "UECommandLine.txt");
-						Log.TraceInformation("Writing cmd line to: " + IntermediateCmdLineFile.FullName);
+						Logger.LogInformation("Writing cmd line to: {File}", IntermediateCmdLineFile.FullName);
 						Project.WriteStageCommandline(IntermediateCmdLineFile, Params, SC);
 					}
 				}
@@ -1189,7 +1190,7 @@ namespace HoloLens.Automation
 			StartInfo.Arguments = String.Format("\"{0}\" \"{1}\" -p", SourceFile.FullName, TargetFile.FullName);
 			StartInfo.UseShellExecute = false;
 			StartInfo.CreateNoWindow = true;
-			Utils.RunLocalProcessAndLogOutput(StartInfo, Log.Logger);
+			Utils.RunLocalProcessAndLogOutput(StartInfo, Logger);
 
 			if (bStripInPlace)
 			{
@@ -1239,8 +1240,8 @@ namespace HoloLens.Automation
 			if (deploymentOperation.Status == AsyncStatus.Error)
 			{
 				DeploymentResult deploymentResult = deploymentOperation.GetResults();
-				LogInformation("Error code: {0}", deploymentOperation.ErrorCode);
-				LogInformation("Error text: {0}", deploymentResult.ErrorText);
+				Logger.LogInformation("Error code: {ErrorCode}", deploymentOperation.ErrorCode);
+				Logger.LogInformation("Error text: {ErrorText}", deploymentResult.ErrorText);
 				throw new AutomationException(ExitCode.Error_AppInstallFailed, deploymentResult.ErrorText, "");
 			}
 			else if (deploymentOperation.Status == AsyncStatus.Canceled)
@@ -1292,7 +1293,7 @@ namespace HoloLens.Automation
 					string PackagePath = Path.Combine(SC.StageDirectory.FullName, "AppxManifest.xml");
 
 
-					LogInformation(String.Format("Registring {0}", PackagePath));
+					Logger.LogInformation("Registring {Path}", PackagePath);
 					WaitFor(PackMgr.RegisterPackageAsync(new Uri(PackagePath), null, DeploymentOptions.DevelopmentMode));
 				}
 			}
@@ -1310,7 +1311,7 @@ namespace HoloLens.Automation
 			}
 			catch
 			{
-				LogWarning("Failed to apply a loopback exemption to the deployed app.  Connection to a local cook server will fail.");
+				Logger.LogWarning("Failed to apply a loopback exemption to the deployed app.  Connection to a local cook server will fail.");
 			}
 		}
 
@@ -1348,7 +1349,7 @@ namespace HoloLens.Automation
 			string[] osBlocks = OsVersionString.Split('.');
 			if (osBlocks.Length < 3)
 			{
-				LogError(String.Format("Wrong OS version string {0} from {1} device", OsVersionString, DeviceAddress));
+				Logger.LogError("Wrong OS version string {VersionString} from {DeviceAddress} device", OsVersionString, DeviceAddress);
 				throw new AutomationException(ExitCode.Error_AppInstallFailed, "Wrong OS version string");
 			}
 
@@ -1364,7 +1365,7 @@ namespace HoloLens.Automation
 			}
 			else
 			{
-				LogError(String.Format("Unsupported OS architecture {0} from {1} device", osArchBlock, DeviceAddress));
+				Logger.LogError("Unsupported OS architecture {Arch} from {DeviceAddress} device", osArchBlock, DeviceAddress);
 				throw new AutomationException(ExitCode.Error_AppInstallFailed, "Unsupported OS architecture");
 			}
 		}
@@ -1685,7 +1686,7 @@ namespace HoloLens.Automation
 			}
 			else
 			{
-				LogLog(args.Message);
+				Logger.LogDebug("{Text}", args.Message);
 			}
 		}
 
@@ -1697,7 +1698,7 @@ namespace HoloLens.Automation
 			}
 			else
 			{
-				LogLog(args.Message);
+				Logger.LogDebug("{Text}", args.Message);
 			}
 		}
 
@@ -1723,9 +1724,9 @@ namespace HoloLens.Automation
 			}
 
 
-			HoloLensExports.CreateManifestForDLC(Params.DLCFile, SC.StageDirectory, Log.Logger);
+			HoloLensExports.CreateManifestForDLC(Params.DLCFile, SC.StageDirectory, Logger);
 		}
 
 		private static List<string> AcceptThumbprints = new List<string>();
 	}
-}
\ No newline at end of file
+}
diff --git a/Engine/Source/Programs/AutomationTool/Android/AndroidFileClient.cs b/Engine/Source/Programs/AutomationTool/Android/AndroidFileClient.cs
--- a/Engine/Source/Programs/AutomationTool/Android/AndroidFileClient.cs
+++ b/Engine/Source/Programs/AutomationTool/Android/AndroidFileClient.cs
@@ -19,12 +19,15 @@ using System.Net;
 using System.Net.Sockets;
 using System.IO.Compression;
 using Newtonsoft.Json.Linq;
+using Microsoft.Extensions.Logging;
 
 
 namespace AutomationTool
 {
 	public class AndroidFileClient
 	{
+		protected static ILogger Logger => Log.Logger;
+
 		private const bool TimeBatching = false;
 
 		private const int DefaultPort = 57099;
@@ -137,7 +140,7 @@ namespace AutomationTool
 					IProcessResult Result = CommandUtils.Run(Executable, "start-server", null, CommandUtils.ERunOptions.NoLoggingOfRunCommand);
 					if (Result.ExitCode != 0)
 					{
-						Log.TraceError("Unable to execute adb to start adb server");
+						Logger.LogError("Unable to execute adb to start adb server");
 						return false;
 					}
 
@@ -159,11 +162,11 @@ namespace AutomationTool
 					}
 					catch (SocketException se)
 					{
-						Log.TraceWarning("SocketException: {0}", se.ToString());
+						Logger.LogWarning("SocketException: {Arg0}", se.ToString());
 					}
 					catch (Exception e)
 					{
-						Log.TraceError("Unexpected Exception: {0}", e.ToString());
+						Logger.LogError("Unexpected Exception: {Arg0}", e.ToString());
 					}
 					finally
 					{
@@ -741,17 +744,17 @@ namespace AutomationTool
 
 		public void Stats_Report()
 		{
-			Log.TraceInformation("Total bytes received:    {0}", Stats_TotalBytesReceived);
-			Log.TraceInformation("Payload bytes received:  {0}", Stats_PayloadBytesReceived);
-			Log.TraceInformation("Total files read:        {0}", Stats_FilesRead);
-			Log.TraceInformation(".");
-			Log.TraceInformation("Total bytes sent:        {0}", Stats_TotalBytesSent);
-			Log.TraceInformation("Payload bytes sent:      {0}", Stats_PayloadBytesSent);
-			Log.TraceInformation("Total files written:     {0}", Stats_FilesWrite);
-			Log.TraceInformation(".");
-			Log.TraceInformation("Uncompressed bytes recv: {0}", Stats_CompressedBytesReceived);
-			Log.TraceInformation("Uncompressed bytes sent: {0}", Stats_CompressedBytesSent);
-			Log.TraceInformation(".");
+			Logger.LogInformation("Total bytes received:    {Stats_TotalBytesReceived}", Stats_TotalBytesReceived);
+			Logger.LogInformation("Payload bytes received:  {Stats_PayloadBytesReceived}", Stats_PayloadBytesReceived);
+			Logger.LogInformation("Total files read:        {Stats_FilesRead}", Stats_FilesRead);
+			Logger.LogInformation(".");
+			Logger.LogInformation("Total bytes sent:        {Stats_TotalBytesSent}", Stats_TotalBytesSent);
+			Logger.LogInformation("Payload bytes sent:      {Stats_PayloadBytesSent}", Stats_PayloadBytesSent);
+			Logger.LogInformation("Total files written:     {Stats_FilesWrite}", Stats_FilesWrite);
+			Logger.LogInformation(".");
+			Logger.LogInformation("Uncompressed bytes recv: {Stats_CompressedBytesReceived}", Stats_CompressedBytesReceived);
+			Logger.LogInformation("Uncompressed bytes sent: {Stats_CompressedBytesSent}", Stats_CompressedBytesSent);
+			Logger.LogInformation(".");
 
 			if (Stats_Stopwatch != null)
 			{
@@ -759,9 +762,9 @@ namespace AutomationTool
 
 				var ElapsedMs = Stats_Stopwatch.ElapsedMilliseconds;
 				double MBPerSec = (double)Stats_TotalBytesSent / ((double)ElapsedMs / 1000.0) / 1024.0 / 1024.0;
-				Log.TraceInformation("Actual sent:  {0:N1} MB/s ({1} in {2:N3}s)", MBPerSec, Stats_TotalBytesSent, (float)(ElapsedMs / 1000.0));
+				Logger.LogInformation("Actual sent:  {0:N1} MB/s ({Stats_TotalBytesSent} in {2:N3}s)", MBPerSec, Stats_TotalBytesSent, (float)(ElapsedMs / 1000.0));
 				MBPerSec = (double)Stats_CompressedBytesSent / ((double)ElapsedMs / 1000.0) / 1024.0 / 1024.0;
-				Log.TraceInformation("Uncompressed: {0:N1} MB/s ({1} in {2:N3}s)", MBPerSec, Stats_CompressedBytesSent, (float)(ElapsedMs / 1000.0));
+				Logger.LogInformation("Uncompressed: {0:N1} MB/s ({Stats_CompressedBytesSent} in {2:N3}s)", MBPerSec, Stats_CompressedBytesSent, (float)(ElapsedMs / 1000.0));
 			}
 		}
 
@@ -839,7 +842,7 @@ namespace AutomationTool
 				}
 				catch (Exception e)
 				{
-					Log.TraceError("Unexpected Exception: {0}", e.ToString());
+					Logger.LogError("Unexpected Exception: {Arg0}", e.ToString());
 					CloseConnection();
 				}
 			}
@@ -1037,7 +1040,7 @@ namespace AutomationTool
 			}
 			catch (Exception e)
 			{
-				Log.TraceError("Unexpected Exception: {0}", e.ToString());
+				Logger.LogError("Unexpected Exception: {Arg0}", e.ToString());
 				CloseConnection();
 			}
 			return Result;
@@ -1079,7 +1082,7 @@ namespace AutomationTool
 			}
 			catch (Exception e)
 			{
-				Log.TraceError("Unexpected Exception: {0}", e.ToString());
+				Logger.LogError("Unexpected Exception: {Arg0}", e.ToString());
 				CloseConnection();
 			}
 			return Result;
@@ -1259,7 +1262,7 @@ namespace AutomationTool
 						int version = message[0] + (message[1] << 8);
 						if (version == ServerVersion)
 						{
-							Log.TraceInformation("Connected to RemoteFileManager");
+							Logger.LogInformation("Connected to RemoteFileManager");
 							return true;
 						}
 					}
@@ -1274,7 +1277,7 @@ namespace AutomationTool
 			}
 			catch (ArgumentNullException ane)
 			{
-				Log.TraceWarning("ArgumentNullException: {0}", ane.ToString());
+				Logger.LogWarning("ArgumentNullException: {Arg0}", ane.ToString());
 			}
 			catch (SocketException se)
 			{
@@ -1285,12 +1288,12 @@ namespace AutomationTool
 				}
 				else
 				{
-					Log.TraceWarning("SocketException: {0}", message);
+					Logger.LogWarning("SocketException: {message}", message);
 				}
 			}
 			catch (Exception e)
 			{
-				Log.TraceWarning("OpenConnection failed: {0}", e.ToString());
+				Logger.LogWarning("OpenConnection failed: {Arg0}", e.ToString());
 			}
 
 			if (ClientSocket != null)
@@ -1636,7 +1639,7 @@ namespace AutomationTool
 				}
 				if (!bFound)
 				{
-					Log.TraceInformation("Did not find package with receiver");
+					Logger.LogInformation("Did not find package with receiver");
 					return false;
 				}
 			}
@@ -1662,7 +1665,7 @@ namespace AutomationTool
 				}
 				if (!bFound)
 				{
-					Log.TraceInformation("Did not find package with activity");
+					Logger.LogInformation("Did not find package with activity");
 					return false;
 				}
 			}
@@ -1701,7 +1704,7 @@ namespace AutomationTool
 					GetListenStatus(Device, ServerPort, out bUSB, out bWifi, out WifiAddress);
 				}
 
-				Log.TraceInformation("Did not find a bind listener");
+				Logger.LogInformation("Did not find a bind listener");
 				return false;
 			}
 
@@ -1720,7 +1723,7 @@ namespace AutomationTool
 				adb.Shell(Device, StartCommand);
 			}
 
-			Log.TraceInformation("Timed out on connection attempts");
+			Logger.LogInformation("Timed out on connection attempts");
 			return false;
 		}
 
@@ -1749,7 +1752,7 @@ namespace AutomationTool
 				}
 				catch (SocketException se)
 				{
-					Log.TraceWarning("SocketException: {0}", se.ToString());
+					Logger.LogWarning("SocketException: {Arg0}", se.ToString());
 				}
 				catch (ObjectDisposedException)
 				{
@@ -1757,7 +1760,7 @@ namespace AutomationTool
 				}
 				catch (Exception e)
 				{
-					Log.TraceError("Unexpected Exception: {0}", e.ToString());
+					Logger.LogError("Unexpected Exception: {Arg0}", e.ToString());
 				}
 			}
 		}
@@ -1779,7 +1782,7 @@ namespace AutomationTool
 				}
 				catch (SocketException se)
 				{
-					Log.TraceWarning("SocketException: {0}", se.ToString());
+					Logger.LogWarning("SocketException: {Arg0}", se.ToString());
 				}
 				catch (ObjectDisposedException)
 				{
@@ -1787,7 +1790,7 @@ namespace AutomationTool
 				}
 				catch (Exception e)
 				{
-					Log.TraceError("Unexpected Exception: {0}", e.ToString());
+					Logger.LogError("Unexpected Exception: {Arg0}", e.ToString());
 				}
 			}
 
@@ -1877,7 +1880,7 @@ namespace AutomationTool
 		{
 			Boolean Result = false;
 
-			Log.TraceInformation("FileRead {0}", SourcePath);
+			Logger.LogInformation("FileRead {SourcePath}", SourcePath);
 
 			if (ClientSocket == null)
 			{
@@ -1930,7 +1933,7 @@ namespace AutomationTool
 					}
 					catch (Exception e)
 					{
-						Log.TraceError("Unexpected Exception: {0}", e.ToString());
+						Logger.LogError("Unexpected Exception: {Arg0}", e.ToString());
 						// ignore received data
 						long remainingSkip = resultSize;
 						while (remainingSkip > 0)
@@ -1947,7 +1950,7 @@ namespace AutomationTool
 			}
 			catch (Exception e)
 			{
-				Log.TraceError("Unexpected Exception: {0}", e.ToString());
+				Logger.LogError("Unexpected Exception: {Arg0}", e.ToString());
 				CloseConnection();
 			}
 			return Result;
@@ -1970,7 +1973,7 @@ namespace AutomationTool
 					DestPath = OptimizePath(DestPath);
 					if (bLog > 0)
 					{
-						Log.TraceInformation("{0}> Writing {1}", bLog, DestPath);
+						Logger.LogInformation("{bLog}> Writing {DestPath}", bLog, DestPath);
 					}
 
 					long fileSize = fileStream.Length;
@@ -2050,7 +2053,7 @@ namespace AutomationTool
 			}
 			catch (Exception e)
 			{
-				Log.TraceError("{0}> Unexpected Exception: {1}", bLog, e.ToString());
+				Logger.LogError("{bLog}> Unexpected Exception: {Arg1}", bLog, e.ToString());
 				CloseConnection();
 			}
 			return Result;
@@ -2080,7 +2083,7 @@ namespace AutomationTool
 					DestPath = OptimizePath(DestPath);
 					if (bLog > 0)
 					{
-						Log.TraceInformation("{0}> Writing Compressed {1}", bLog, DestPath);
+						Logger.LogInformation("{bLog}> Writing Compressed {DestPath}", bLog, DestPath);
 					}
 
 					Stats_FilesWrite++;
@@ -2184,7 +2187,7 @@ namespace AutomationTool
 			}
 			catch (Exception e)
 			{
-				Log.TraceError("{0}> Unexpected Exception: {1}", bLog, e.ToString());
+				Logger.LogError("{bLog}> Unexpected Exception: {Arg1}", bLog, e.ToString());
 				CloseConnection();
 			}
 			return Result;
@@ -2231,7 +2234,7 @@ namespace AutomationTool
 			}
 			catch (Exception e)
 			{
-				Log.TraceError("Unexpected Exception: {0}", e.ToString());
+				Logger.LogError("Unexpected Exception: {Arg0}", e.ToString());
 				CloseConnection();
 			}
 			return Result;
@@ -2253,14 +2256,14 @@ namespace AutomationTool
 			}
 			if (!Result && bLog > 0)
 			{
-				Log.TraceInformation("{0}> Failed to copy {1} to {2}", bLog, InSource, InDest);
+				Logger.LogInformation("{bLog}> Failed to copy {InSource} to {InDest}", bLog, InSource, InDest);
 				return false;
 			}
 
 			/*
 			if (bLog > 0)
 			{
-				Log.TraceInformation("{0}> Copied: {1}", bLog, InSource);
+				Logger.LogInformation("{bLog}> Copied: {InSource}", bLog, InSource);
 			}
 			*/
 			return true;
@@ -2399,7 +2402,7 @@ namespace AutomationTool
 			var ElapsedMs = stopwatch.ElapsedMilliseconds;
 			if (bReportStats)
 			{
-				Log.TraceInformation("Time to create directories:  {0:N3}s", (float)(ElapsedMs / 1000.0));
+				Logger.LogInformation("Time to create directories:  {0:N3}s", (float)(ElapsedMs / 1000.0));
 			}
 
 			Stats_Clear();
@@ -2511,7 +2514,7 @@ namespace AutomationTool
 				if (bReportStats)
 				{
 					ElapsedMs = stopwatch.ElapsedMilliseconds;
-					Log.TraceInformation("Time to build batch:  {0:N3}s", (float)(ElapsedMs / 1000.0));
+					Logger.LogInformation("Time to build batch:  {0:N3}s", (float)(ElapsedMs / 1000.0));
 				}
 				Stats_Stopwatch.Stop();
 				Stats_StartStopwatch();
@@ -2530,7 +2533,7 @@ namespace AutomationTool
 				Stats_Report();
 
 				ElapsedMs = totalwatch.ElapsedMilliseconds;
-				Log.TraceInformation("Total time to Deploy AFS:  {0:N3}s", (float)(ElapsedMs / 1000.0));
+				Logger.LogInformation("Total time to Deploy AFS:  {0:N3}s", (float)(ElapsedMs / 1000.0));
 			}
 
 			return Result;
@@ -2539,29 +2542,29 @@ namespace AutomationTool
 		/*
 		public void TestOptimizePath()
 		{
-			Log.TraceInformation("{0}", OptimizePath("/storage/emulated/0/Android/data/com.epicgames.SaveTest/files/UE4Game/SaveTest/Manifest_DebugFiles_Android.txt"));
-			Log.TraceInformation("{0}", OptimizePath("/storage/emulated/0/Android/data/com.epicgames.SaveTest/files/UE4Game/SaveTest/DataDrivenPlatformInfo.ini"));
-			Log.TraceInformation("{0}", OptimizePath("/storage/emulated/0/Android/data/com.epicgames.SaveTest/files/UE4Game/SaveTest/Engine/Config/Layouts/DefaultLayout.ini"));
-			Log.TraceInformation("{0}", OptimizePath("/storage/emulated/0/Android/data/com.epicgames.SaveTest/files/UE4Game/SaveTest/Engine/Content/ArtTools/RenderToTexture/Materials/Debug/M_Emissive_Color.uasset"));
-
-			Log.TraceInformation("{0}", OptimizePath("/storage/0/emulated/commandfile.txt"));
-			Log.TraceInformation("{0}", OptimizePath("/storage/0/emulated/123.txt"));
-			Log.TraceInformation("{0}", OptimizePath("/storage/0/emulated/game/123.txt"));
-			Log.TraceInformation("{0}", OptimizePath("/storage/0/emulated/every.txt"));
-
-			Log.TraceInformation("{0}", OptimizePath("^commandfile"));
-			Log.TraceInformation("{0}", OptimizePath("^commandfile"));
-
-			Log.TraceInformation("{0}", OptimizePath("^ext/z.txt"));
-			Log.TraceInformation("{0}", OptimizePath("^ext/abc.txt"));
-			Log.TraceInformation("{0}", OptimizePath("^ext/cde.txt"));
-			Log.TraceInformation("{0}", OptimizePath("^ext/1/abc.txt"));
-			Log.TraceInformation("{0}", OptimizePath("^ext/1/2/abc.txt"));
-			Log.TraceInformation("{0}", OptimizePath("^ext/xyz.txt"));
-			Log.TraceInformation("{0}", OptimizePath("^ext/1/2/stu.txt"));
-			Log.TraceInformation("{0}", OptimizePath("^ext/1/wer.txt"));
-			Log.TraceInformation("{0}", OptimizePath("^int/1/2/stu.txt"));
-			Log.TraceInformation("{0}", OptimizePath("^int/1/abc.txt"));
+			Logger.LogInformation("{Arg0}", OptimizePath("/storage/emulated/0/Android/data/com.epicgames.SaveTest/files/UE4Game/SaveTest/Manifest_DebugFiles_Android.txt"));
+			Logger.LogInformation("{Arg0}", OptimizePath("/storage/emulated/0/Android/data/com.epicgames.SaveTest/files/UE4Game/SaveTest/DataDrivenPlatformInfo.ini"));
+			Logger.LogInformation("{Arg0}", OptimizePath("/storage/emulated/0/Android/data/com.epicgames.SaveTest/files/UE4Game/SaveTest/Engine/Config/Layouts/DefaultLayout.ini"));
+			Logger.LogInformation("{Arg0}", OptimizePath("/storage/emulated/0/Android/data/com.epicgames.SaveTest/files/UE4Game/SaveTest/Engine/Content/ArtTools/RenderToTexture/Materials/Debug/M_Emissive_Color.uasset"));
+
+			Logger.LogInformation("{Arg0}", OptimizePath("/storage/0/emulated/commandfile.txt"));
+			Logger.LogInformation("{Arg0}", OptimizePath("/storage/0/emulated/123.txt"));
+			Logger.LogInformation("{Arg0}", OptimizePath("/storage/0/emulated/game/123.txt"));
+			Logger.LogInformation("{Arg0}", OptimizePath("/storage/0/emulated/every.txt"));
+
+			Logger.LogInformation("{Arg0}", OptimizePath("^commandfile"));
+			Logger.LogInformation("{Arg0}", OptimizePath("^commandfile"));
+
+			Logger.LogInformation("{Arg0}", OptimizePath("^ext/z.txt"));
+			Logger.LogInformation("{Arg0}", OptimizePath("^ext/abc.txt"));
+			Logger.LogInformation("{Arg0}", OptimizePath("^ext/cde.txt"));
+			Logger.LogInformation("{Arg0}", OptimizePath("^ext/1/abc.txt"));
+			Logger.LogInformation("{Arg0}", OptimizePath("^ext/1/2/abc.txt"));
+			Logger.LogInformation("{Arg0}", OptimizePath("^ext/xyz.txt"));
+			Logger.LogInformation("{Arg0}", OptimizePath("^ext/1/2/stu.txt"));
+			Logger.LogInformation("{Arg0}", OptimizePath("^ext/1/wer.txt"));
+			Logger.LogInformation("{Arg0}", OptimizePath("^int/1/2/stu.txt"));
+			Logger.LogInformation("{Arg0}", OptimizePath("^int/1/abc.txt"));
 		}
 		*/
 	}
diff --git a/Engine/Source/Programs/AutomationTool/Android/AndroidPlatform.Automation.cs b/Engine/Source/Programs/AutomationTool/Android/AndroidPlatform.Automation.cs
--- a/Engine/Source/Programs/AutomationTool/Android/AndroidPlatform.Automation.cs
+++ b/Engine/Source/Programs/AutomationTool/Android/AndroidPlatform.Automation.cs
@@ -17,6 +17,7 @@ using System.Text.RegularExpressions;
 using AutomationScripts;
 using System.Drawing;
 using System.Security.Policy;
+using Microsoft.Extensions.Logging;
 
 public class AndroidPlatform : Platform
 {
@@ -804,7 +805,7 @@ public class AndroidPlatform : Platform
 				if (PluginExtras.FirstOrDefault(x => x == PluginPath) == null)
 				{
 					PluginExtras.Add(PluginPath);
-					LogInformation("AndroidPlugin: {0}", PluginPath);
+					Logger.LogInformation("AndroidPlugin: {PluginPath}", PluginPath);
 				}
 			}
 		}
@@ -903,7 +904,7 @@ public class AndroidPlatform : Platform
 		}
 		if (!bAllowNetworkConnection && ConnectionType != EConnectionType.USBOnly)
 		{
-			Log.TraceWarning("AFS will only use USB connection due to network connection disabled");
+			Logger.LogWarning("AFS will only use USB connection due to network connection disabled");
 			ConnectionType = EConnectionType.USBOnly;
 		}
 
@@ -967,7 +968,7 @@ public class AndroidPlatform : Platform
 
 	private bool CreateOBBFile(DeploymentContext SC, string OutputFilename, List<FileReference> FilesForObb)
 	{
-		LogInformation("Creating {0} from {1}", OutputFilename, SC.StageDirectory);
+		Logger.LogInformation("Creating {OutputFilename} from {Arg1}", OutputFilename, SC.StageDirectory);
 		using (ZipFile ObbFile = new ZipFile(OutputFilename))
 		{
 			ObbFile.CompressionMethod = CompressionMethod.None;
@@ -1067,7 +1068,7 @@ public class AndroidPlatform : Platform
 		}
 
 		string BaseApkName = GetFinalApkName(Params, SC.StageExecutables[0], true, Architecture:null);
-		LogInformation("BaseApkName = {0}", BaseApkName);
+		Logger.LogInformation("BaseApkName = {BaseApkName}", BaseApkName);
 
 		// Create main OBB with entire contents of staging dir. This
 		// includes any PAK files, movie files, etc.
@@ -1124,7 +1125,7 @@ public class AndroidPlatform : Platform
 
 		if (!OBBNeedsUpdate)
 		{
-			LogInformation("OBB is up to date: " + LocalObbName);
+			Logger.LogInformation("{Text}", "OBB is up to date: " + LocalObbName);
 		}
 		else
 		{
@@ -1205,14 +1206,14 @@ public class AndroidPlatform : Platform
 						else
 						{
 							// no room in either file
-							LogInformation("Failed to build OBB: " + LocalObbName);
+							Logger.LogInformation("{Text}", "Failed to build OBB: " + LocalObbName);
 							throw new AutomationException(ExitCode.Error_AndroidOBBError, "Stage Failed. Could not build OBB {0}. The file may be too big to fit in an OBB ({1} limit)", LocalObbName, LimitString);
 						}
 					}
 					else
 					{
 						// no room in either file
-						LogInformation("Failed to build OBB: " + LocalObbName);
+						Logger.LogInformation("{Text}", "Failed to build OBB: " + LocalObbName);
 						throw new AutomationException(ExitCode.Error_AndroidOBBError, "Stage Failed. Could not build OBB {0}. The file may be too big to fit in an OBB ({1} limit)", LocalObbName, LimitString);
 					}
 				}
@@ -1221,7 +1222,7 @@ public class AndroidPlatform : Platform
 			// Now create the main OBB as a ZIP archive.
 			if (!CreateOBBFile(SC, LocalObbName, FilesToObb))
 			{
-				LogInformation("Failed to build OBB: " + LocalObbName);
+				Logger.LogInformation("{Text}", "Failed to build OBB: " + LocalObbName);
 				throw new AutomationException(ExitCode.Error_AndroidOBBError, "Stage Failed. Could not build OBB {0}. The file may be too big to fit in an OBB ({1} limit)", LocalObbName, LimitString);
 			}
 
@@ -1230,7 +1231,7 @@ public class AndroidPlatform : Platform
 			{
 				if (!CreateOBBFile(SC, LocalPatchName, FilesToPatch))
 				{
-					LogInformation("Failed to build OBB: " + LocalPatchName);
+					Logger.LogInformation("{Text}", "Failed to build OBB: " + LocalPatchName);
 					throw new AutomationException(ExitCode.Error_AndroidOBBError, "Stage Failed. Could not build OBB {0}. The file may be too big to fit in an OBB ({1} limit)", LocalPatchName, LimitString);
 				}
 			}
@@ -1240,7 +1241,7 @@ public class AndroidPlatform : Platform
 			{
 				if (!CreateOBBFile(SC, LocalOverflow1Name, FilesToOverflow1))
 				{
-					LogInformation("Failed to build OBB: " + LocalOverflow1Name);
+					Logger.LogInformation("{Text}", "Failed to build OBB: " + LocalOverflow1Name);
 					throw new AutomationException(ExitCode.Error_AndroidOBBError, "Stage Failed. Could not build OBB {0}. The file may be too big to fit in an OBB ({1} limit)", LocalPatchName, LimitString);
 				}
 			}
@@ -1250,7 +1251,7 @@ public class AndroidPlatform : Platform
 			{
 				if (!CreateOBBFile(SC, LocalOverflow2Name, FilesToOverflow2))
 				{
-					LogInformation("Failed to build OBB: " + LocalOverflow2Name);
+					Logger.LogInformation("{Text}", "Failed to build OBB: " + LocalOverflow2Name);
 					throw new AutomationException(ExitCode.Error_AndroidOBBError, "Stage Failed. Could not build OBB {0}. The file may be too big to fit in an OBB ({1} limit)", LocalPatchName, LimitString);
 				}
 			}
@@ -1261,7 +1262,7 @@ public class AndroidPlatform : Platform
 		Int64 ObbFileLength = OBBFileInfo.Length;
 		if (ObbFileLength > OBBSizeAllowed)
 		{
-			LogInformation("OBB exceeds " + LimitString + " limit: " + ObbFileLength + " bytes");
+			Logger.LogInformation("OBB exceeds " + LimitString + " limit: " + ObbFileLength + " bytes");
 			throw new AutomationException(ExitCode.Error_AndroidOBBError, "Stage Failed. OBB {0} exceeds {1} limit)", LocalObbName, LimitString);
 		}
 
@@ -1273,7 +1274,7 @@ public class AndroidPlatform : Platform
 		Ini.GetInt32("/Script/AndroidRuntimeSettings.AndroidRuntimeSettings", "MinSDKVersion", out MinSDKVersion);
 		int TargetSDKVersion = MinSDKVersion;
 		Ini.GetInt32("/Script/AndroidRuntimeSettings.AndroidRuntimeSettings", "TargetSDKVersion", out TargetSDKVersion);
-		LogInformation("Target SDK Version " + TargetSDKVersion);
+		Logger.LogInformation("{Text}", "Target SDK Version " + TargetSDKVersion);
 		bool bDisablePerfHarden = false;
         if (TargetConfiguration != UnrealTargetConfiguration.Shipping)
         {
@@ -1291,7 +1292,7 @@ public class AndroidPlatform : Platform
                 UnrealSOName = UnrealSOName.Replace(".apk", ".so");
                 if (FileExists_NoExceptions(UnrealSOName) == false)
 				{
-					LogInformation("Failed to find game .so " + UnrealSOName);
+					Logger.LogInformation("{Text}", "Failed to find game .so " + UnrealSOName);
                     throw new AutomationException(ExitCode.Error_MissingExecutable, "Stage Failed. Could not find .so {0}. You may need to build the UE project with your target configuration and platform.", UnrealSOName);
 				}
 			}
@@ -1320,7 +1321,7 @@ public class AndroidPlatform : Platform
 					{
 						Type = TargetType.Server;
 					}
-					LogInformation("SavePackageInfo");
+					Logger.LogInformation("SavePackageInfo");
 					Deploy.SavePackageInfo(Params.ShortProjectName, SC.ProjectRoot.FullName, Type, true);
 				}
 				Deploy.PrepForUATPackageOrDeploy(Params.RawProjectPath, Params.ShortProjectName, SC.ProjectRoot, SOName, SC.LocalRoot + "/Engine", Params.Distribution, CookFlavor, SC.StageTargets[0].Receipt.Configuration, false, bShouldCompileAsDll);
@@ -1508,7 +1509,7 @@ public class AndroidPlatform : Platform
 		{
 			return "linux-x64/UnrealAndroidFileTool";
 		}
-		Log.TraceWarning("GetAFSExecutable unsupported target, assuming Win64");
+		Logger.LogWarning("GetAFSExecutable unsupported target, assuming Win64");
 		return "win-x64/UnrealAndroidFileTool.exe";
 	}
 
@@ -1570,7 +1571,7 @@ public class AndroidPlatform : Platform
 				Overflow2InstallCommand = bNoOverflow2Install ? AFSCommand + " deletefile '^overflow2obb'" : AFSCommand + " push " + Path.GetFileName(Overflow2Name) + " '^overflow2obb'";
 			}
 
-			LogInformation("Writing shell script for install with {0}", bPackageDataInsideApk ? "data in APK" : "separate obb");
+			Logger.LogInformation("Writing shell script for install with {Arg0}", bPackageDataInsideApk ? "data in APK" : "separate obb");
             BatchLines = new string[] {
 						"#!/bin/sh",
 						"cd \"`dirname \"$0\"`\"",
@@ -1655,7 +1656,7 @@ public class AndroidPlatform : Platform
 				Overflow2InstallCommand = bNoOverflow2Install ? AFSCommand + " deletefile \"^overflow2obb\"" : AFSCommand + " push " + Path.GetFileName(Overflow2Name) + " \"^overflow2obb\"";
 			}
 
-			LogInformation("Writing bat for install with {0}", bPackageDataInsideApk ? "data in APK" : "separate OBB");
+			Logger.LogInformation("Writing bat for install with {Arg0}", bPackageDataInsideApk ? "data in APK" : "separate OBB");
             BatchLines = new string[] {
 						"setlocal",
 						"if NOT \"%UE_SDKS_ROOT%\"==\"\" (call %UE_SDKS_ROOT%\\HostWin64\\Android\\SetupEnvironmentVars.bat)",
@@ -1723,7 +1724,7 @@ public class AndroidPlatform : Platform
 
 		if (!bIsPC)
 		{
-			LogInformation("Writing shell script for uninstall with {0}", bPackageDataInsideApk ? "data in APK" : "separate obb");
+			Logger.LogInformation("Writing shell script for uninstall with {Arg0}", bPackageDataInsideApk ? "data in APK" : "separate obb");
 			BatchLines = new string[] {
 						"#!/bin/sh",
 						"cd \"`dirname \"$0\"`\"",
@@ -1747,7 +1748,7 @@ public class AndroidPlatform : Platform
 		}
 		else
 		{
-			LogInformation("Writing bat for uninstall with {0}", bPackageDataInsideApk ? "data in APK" : "separate OBB");
+			Logger.LogInformation("Writing bat for uninstall with {Arg0}", bPackageDataInsideApk ? "data in APK" : "separate OBB");
 			BatchLines = new string[] {
 						"setlocal",
 						"if NOT \"%UE_SDKS_ROOT%\"==\"\" (call %UE_SDKS_ROOT%\\HostWin64\\Android\\SetupEnvironmentVars.bat)",
@@ -1779,7 +1780,7 @@ public class AndroidPlatform : Platform
 
 		if (!bIsPC)
 		{
-			LogInformation("Writing shell script for symbolize with {0}", "data in APK" );
+			Logger.LogInformation("Writing shell script for symbolize with {Arg0}", "data in APK" );
 			BatchLines = new string[] {
 				"#!/bin/sh",
 				"if [ $? -ne 0]; then",
@@ -1795,7 +1796,7 @@ public class AndroidPlatform : Platform
 		}
 		else
 		{
-			LogInformation("Writing bat for symbolize");
+			Logger.LogInformation("Writing bat for symbolize");
 			BatchLines = new string[] {
 						"@echo off",
 						"IF %1.==. GOTO NoArgs",
@@ -2200,13 +2201,13 @@ public class AndroidPlatform : Platform
 		{
 			if (PackageName == null || PackageName == "")
 			{
-				Log.TraceInformation("Retrieve Manifests: Unable to start file server without package name, ignoring manifests");
+				Logger.LogInformation("Retrieve Manifests: Unable to start file server without package name, ignoring manifests");
 				return false;
 			}
-			Log.TraceInformation("Retrieve Manifests: Trying to start file server {0}", PackageName);
+			Logger.LogInformation("Retrieve Manifests: Trying to start file server {PackageName}", PackageName);
 			if (!client.StartServer(PackageName, AFSToken))
 			{
-				Log.TraceInformation("Retrieve Manifests: Failed to start server, ignoring manifests");
+				Logger.LogInformation("Retrieve Manifests: Failed to start server, ignoring manifests");
 				return false;
 			}
 		}
@@ -2217,18 +2218,18 @@ public class AndroidPlatform : Platform
 		{
 			if (PackageName == null || PackageName == "")
 			{
-				Log.TraceInformation("Retrieve Manifests: Unable to start file server without package name, ignoring manifests");
+				Logger.LogInformation("Retrieve Manifests: Unable to start file server without package name, ignoring manifests");
 				client.CloseConnection();
 				return false;
 			}
 
-			Log.TraceInformation("Connected to wrong server {0}, trying again", DevicePackageName);
+			Logger.LogInformation("Connected to wrong server {DevicePackageName}, trying again", DevicePackageName);
 			client.TerminateServer();
 
-			Log.TraceInformation("Retrieve Manifests: Trying to start file server {0}", PackageName);
+			Logger.LogInformation("Retrieve Manifests: Trying to start file server {PackageName}", PackageName);
 			if (!client.StartServer(PackageName, AFSToken))
 			{
-				Log.TraceInformation("Retrieve Manifests: Failed to start server, ignoring manifests");
+				Logger.LogInformation("Retrieve Manifests: Failed to start server, ignoring manifests");
 				return false;
 			}
 		}
@@ -2255,7 +2256,7 @@ public class AndroidPlatform : Platform
 		NonUFSManifests = new List<string>();
 		NonUFSManifests.Add(NonUFSManifestFileName);
 
-		Log.TraceInformation("Retrieve Manifests: Success!!");
+		Logger.LogInformation("Retrieve Manifests: Success!!");
 
 		return true;
 	}
@@ -2286,7 +2287,7 @@ public class AndroidPlatform : Platform
 		IProcessResult UFSResult = RunAdbCommand(Params, DeviceName, " pull " + RemoteDir + "/" + SC.GetUFSDeployedManifestFileName(null) + " \"" + RetrievedUFSManifestFileName + "\"", null, ERunOptions.AppMustExist);
 		if (!(UFSResult.Output.Contains("bytes") || UFSResult.Output.Contains("[100%]")))
 		{
-			LogWarning("Failed retrieving UFS Manifest: {0}", UFSResult.Output);
+			Logger.LogWarning("Failed retrieving UFS Manifest: {Arg0}", UFSResult.Output);
 			return false;
 		}
 
@@ -2295,7 +2296,7 @@ public class AndroidPlatform : Platform
 		IProcessResult NonUFSResult = RunAdbCommand(Params, DeviceName, " pull " + RemoteDir + "/" + SC.GetNonUFSDeployedManifestFileName(null) + " \"" + RetrievedNonUFSManifestFileName + "\"", null, ERunOptions.AppMustExist);
 		if (!(NonUFSResult.Output.Contains("bytes") || NonUFSResult.Output.Contains("[100%]")))
 		{
-			LogWarning("Failed retrieving NonUFS Manifest: {0}", NonUFSResult.Output);
+			Logger.LogWarning("Failed retrieving NonUFS Manifest: {Arg0}", NonUFSResult.Output);
 			// Did not retrieve both so delete one we did retrieve
 			File.Delete(RetrievedUFSManifestFileName);
 			return false;
@@ -2425,7 +2426,7 @@ public class AndroidPlatform : Platform
 				String InstalledResult = adb.Shell(DeviceName, "pm list packages " + PackageName);
 				if (InstalledResult.Contains(PackageName))
 				{
-					Log.TraceInformation("{0} already installed!", PackageName);
+					Logger.LogInformation("{PackageName} already installed!", PackageName);
 					// already installed so enable --fast-deploy option if need to update apk
 					bFastDeploy = true;
 
@@ -2433,13 +2434,13 @@ public class AndroidPlatform : Platform
 					InstalledResult = adb.Shell(DeviceName, "cat " + ExtFiles + "/APKFileStamp.txt");
 					if (InstalledResult.StartsWith("APK: "))
 					{
-						Log.TraceInformation("Found APKFileStamp.txt! {0}", InstalledResult);
+						Logger.LogInformation("Found APKFileStamp.txt! {InstalledResult}", InstalledResult);
 						if (InstalledResult.Substring(5).Trim() == APKLastUpdateTime)
 							bNeedAPKInstall = false;
 
 						if (InstalledResult.Substring(5).Trim() != APKLastUpdateTime)
 						{
-							Log.TraceInformation("{0} != {1}", InstalledResult.Substring(5).Trim(), APKLastUpdateTime);
+							Logger.LogInformation("{Arg0} != {APKLastUpdateTime}", InstalledResult.Substring(5).Trim(), APKLastUpdateTime);
 						}
 
 						// Stop the previously running copy (uninstall/install did this before)
@@ -2486,7 +2487,7 @@ public class AndroidPlatform : Platform
 					}
 					if (ErrorMessage.Contains("OLDER_SDK"))
 					{
-						LogError("minSdkVersion is higher than Android version installed on device, possibly due to NDK API Level");
+						Logger.LogError("minSdkVersion is higher than Android version installed on device, possibly due to NDK API Level");
 						throw new AutomationException(ExitCode.Error_AppInstallFailed, ErrorMessage);
 					}
 
@@ -2515,7 +2516,7 @@ public class AndroidPlatform : Platform
 						}
 						if (ErrorMessage.Contains("OLDER_SDK"))
 						{
-							LogError("minSdkVersion is higher than Android version installed on device, possibly due to NDK API Level");
+							Logger.LogError("minSdkVersion is higher than Android version installed on device, possibly due to NDK API Level");
 						}
 						throw new AutomationException(ExitCode.Error_AppInstallFailed, ErrorMessage);
 					}
@@ -2534,18 +2535,18 @@ public class AndroidPlatform : Platform
 
 			// reopen file server connection (either USB or Network)
 			IPAddress = (ConnectionType == EConnectionType.NetworkOnly) ? (bUseManualIPAddress ? ManualIPAddress : (DeviceIPAddress != null ? DeviceIPAddress : "127.0.0.1")) : "127.0.0.1";
-			Log.TraceInformation("Attempting to connect to file server [{0}]", (IPAddress == "127.0.0.1" ? "USB" : IPAddress));
+			Logger.LogInformation("Attempting to connect to file server [{Arg0}]", (IPAddress == "127.0.0.1" ? "USB" : IPAddress));
 			if (!client.OpenConnection(IPAddress))
 			{
-				Log.TraceInformation("Not connected, attempting to start file server");
+				Logger.LogInformation("Not connected, attempting to start file server");
 				if (!client.StartServer(PackageName, AFSToken, IPAddress))
 				{
 					// try one more time with longer delay
-					Log.TraceInformation("Trying again");
+					Logger.LogInformation("Trying again");
 					Thread.Sleep(1000);
 					if (!client.StartServer(PackageName, AFSToken, IPAddress))
 					{
-						Log.TraceWarning("Failed to start Android file server for {0}, skipping deploy for {1}", PackageName, DeviceName);
+						Logger.LogWarning("Failed to start Android file server for {PackageName}, skipping deploy for {DeviceName}", PackageName, DeviceName);
 						continue;
 					}
 				}
@@ -2555,18 +2556,18 @@ public class AndroidPlatform : Platform
 			string DevicePackageName = client.Query("^packagename");
 			if (DevicePackageName != PackageName)
 			{
-				Log.TraceInformation("Connected to wrong server {0}, trying again", DevicePackageName);
+				Logger.LogInformation("Connected to wrong server {DevicePackageName}, trying again", DevicePackageName);
 				client.TerminateServer();
 
-				Log.TraceInformation("Trying to start file server {0}", PackageName);
+				Logger.LogInformation("Trying to start file server {PackageName}", PackageName);
 				if (!client.StartServer(PackageName, AFSToken, IPAddress))
 				{
 					// try one more time with longer delay
-					Log.TraceInformation("Trying again");
+					Logger.LogInformation("Trying again");
 					Thread.Sleep(1000);
 					if (!client.StartServer(PackageName, AFSToken, IPAddress))
 					{
-						Log.TraceWarning("Failed to start Android file server for {0}, skipping deploy for {1}", PackageName, DeviceName);
+						Logger.LogWarning("Failed to start Android file server for {PackageName}, skipping deploy for {DeviceName}", PackageName, DeviceName);
 						continue;
 					}
 				}
@@ -2577,13 +2578,13 @@ public class AndroidPlatform : Platform
 				IPAddress = (bUseManualIPAddress ? ManualIPAddress : client.Query("^ip"));
 				client2 = new AndroidFileClient(DeviceName);
 
-				Log.TraceInformation("Attempting to connect to file server [{0}]", IPAddress);
+				Logger.LogInformation("Attempting to connect to file server [{IPAddress}]", IPAddress);
 				if (!client2.OpenConnection(IPAddress))
 				{
-					Log.TraceInformation("Not connected, attempting to start file server");
+					Logger.LogInformation("Not connected, attempting to start file server");
 					if (!client2.StartServer(PackageName, AFSToken, IPAddress, false))
 					{
-						Log.TraceWarning("Failed to start Android file server for {0}, only using one connection for {1}", PackageName, DeviceName);
+						Logger.LogWarning("Failed to start Android file server for {PackageName}, only using one connection for {DeviceName}", PackageName, DeviceName);
 						client2 = null;
 					}
 				}
@@ -2592,7 +2593,7 @@ public class AndroidPlatform : Platform
 				DevicePackageName = client2.Query("^packagename");
 				if (DevicePackageName != PackageName)
 				{
-					Log.TraceInformation("Connected to wrong server {0} for [{1}], not using network", DevicePackageName, IPAddress, DeviceName);
+					Logger.LogInformation("Connected to wrong server {DevicePackageName} for [{IPAddress}], not using network", DevicePackageName, IPAddress);
 					client2.CloseConnection();
 					client2 = null;
 				}
@@ -2671,12 +2672,12 @@ public class AndroidPlatform : Platform
 
 						if (CommandResult == null)
 						{
-							Log.TraceWarning("Failed to read remote dir: {0}", RemoteDir);
+							Logger.LogWarning("Failed to read remote dir: {RemoteDir}", RemoteDir);
 							RemoteDir = client.Query("^project", true);
 							CommandResult = client.DirListFlat(RemoteDir);
 							if (CommandResult == null)
 							{
-								Log.TraceWarning("Failed to read remote dir again: {0}", RemoteDir);
+								Logger.LogWarning("Failed to read remote dir again: {RemoteDir}", RemoteDir);
 							}
 						}
 
@@ -2701,7 +2702,7 @@ public class AndroidPlatform : Platform
 								string ProjectConfig = Params.ShortProjectName + "/Config";
 								const string EngineSaved = "Engine/Saved"; // is this safe to use, or should we use SC.EngineRoot.GetDirectoryName()?
 								const string EngineConfig = "Engine/Config";
-								LogWarning("Excluding {0} {1} {2} {3} from clean during deployment.", ProjectSaved, ProjectConfig, EngineSaved, EngineConfig);
+								Logger.LogWarning("Excluding {ProjectSaved} {ProjectConfig} {EngineSaved} {EngineConfig} from clean during deployment.", ProjectSaved, ProjectConfig, EngineSaved, EngineConfig);
 
 								string CurrentDir = "";
 								bool SkipFiles = false;
@@ -2761,7 +2762,7 @@ public class AndroidPlatform : Platform
 											{
 												// delete file from device
 												string FilePath = CurrentDir.Length == 0 ? FileName : (CurrentDir + "/" + FileName); // use / for Android target, no matter the development system
-												LogWarning("Deleting {0} from device; not found in staging area", FilePath);
+												Logger.LogWarning("Deleting {FilePath} from device; not found in staging area", FilePath);
 												FilesToDeleteFromDevice.Add(FilePath);
 											}
 										}
@@ -2779,7 +2780,7 @@ public class AndroidPlatform : Platform
 								bool ParentMarkedForDeletion = DirsToDeleteFromDevice.Contains(ParentDir);
 								if (!ParentMarkedForDeletion)
 								{
-									LogWarning("Deleting {0} and its contents from device; not found in staging area", DirToDelete);
+									Logger.LogWarning("Deleting {DirToDelete} and its contents from device; not found in staging area", DirToDelete);
 									client.DirDeleteRecurse(RemoteDir + "/" + DirToDelete);
 								}
 							}
@@ -2873,7 +2874,7 @@ public class AndroidPlatform : Platform
 				// to deploy. Files we deploy will get individually copied
 				// and dirs will get the tree copies by default (that's
 				// what ADB does, too).
-				Log.TraceInformation("Deploying files using AFS");
+				Logger.LogInformation("Deploying files using AFS");
 				string SourceDir = SC.StageDirectory.FullName;
 				client.Deploy(EntriesToDeploy, SourceDir, RemoteDir, bUseCompression, bLogFiles, bReportStats, client2);
 			}
@@ -3037,7 +3038,7 @@ public class AndroidPlatform : Platform
                     }
                     if (ErrorMessage.Contains("OLDER_SDK"))
                     {
-                        LogError("minSdkVersion is higher than Android version installed on device, possibly due to NDK API Level");
+                        Logger.LogError("minSdkVersion is higher than Android version installed on device, possibly due to NDK API Level");
                     }
                     throw new AutomationException(ExitCode.Error_AppInstallFailed, ErrorMessage);
                 }
@@ -3152,7 +3153,7 @@ public class AndroidPlatform : Platform
 							string ProjectConfig = Params.ShortProjectName + "/Config";
 							const string EngineSaved = "Engine/Saved"; // is this safe to use, or should we use SC.EngineRoot.GetDirectoryName()?
 							const string EngineConfig = "Engine/Config";
-							LogWarning("Excluding {0} {1} {2} {3} from clean during deployment.", ProjectSaved, ProjectConfig, EngineSaved, EngineConfig);
+							Logger.LogWarning("Excluding {ProjectSaved} {ProjectConfig} {EngineSaved} {EngineConfig} from clean during deployment.", ProjectSaved, ProjectConfig, EngineSaved, EngineConfig);
 
 							string CurrentDir = "";
 							bool SkipFiles = false;
@@ -3212,7 +3213,7 @@ public class AndroidPlatform : Platform
 										{
 											// delete file from device
 											string FilePath = CurrentDir.Length == 0 ? FileName : (CurrentDir + "/" + FileName); // use / for Android target, no matter the development system
-											LogWarning("Deleting {0} from device; not found in staging area", FilePath);
+											Logger.LogWarning("Deleting {FilePath} from device; not found in staging area", FilePath);
 											FilesToDeleteFromDevice.Add(FilePath);
 										}
 									}
@@ -3230,7 +3231,7 @@ public class AndroidPlatform : Platform
 							bool ParentMarkedForDeletion = DirsToDeleteFromDevice.Contains(ParentDir);
 							if (!ParentMarkedForDeletion)
 							{
-								LogWarning("Deleting {0} and its contents from device; not found in staging area", DirToDelete);
+								Logger.LogWarning("Deleting {DirToDelete} and its contents from device; not found in staging area", DirToDelete);
 								RunAdbCommand(Params, DeviceName, "shell rm -r " + RemoteDir + "/" + DirToDelete);
 							}
 						}
@@ -3495,7 +3496,7 @@ public class AndroidPlatform : Platform
 					ReturnValue = Tokens[TokenIndex];
 				}
 			}
-			LogInformation("GetPackageInfo ReturnValue: {0}", ReturnValue);
+			Logger.LogInformation("GetPackageInfo ReturnValue: {ReturnValue}", ReturnValue);
 		}
 
 		return ReturnValue;
@@ -3524,7 +3525,7 @@ public class AndroidPlatform : Platform
 		{
 			string[] Lines = File.ReadAllLines(PackageInfoPath);
 			int LineIndex = bRetrieveVersionCode ? 1 : 0;
-			LogInformation("packageInfo line index: {0}", LineIndex);
+			Logger.LogInformation("packageInfo line index: {LineIndex}", LineIndex);
 			if (Lines.Length >= 2)
 			{
 				ReturnValue = Lines[LineIndex];
@@ -3552,9 +3553,9 @@ public class AndroidPlatform : Platform
 			ReturnValue = StoreVersion.ToString("0");
 		}
 
-		LogInformation("packageInfo.txt file exists: {0}", fileExists);
-		LogInformation("packageInfo return MetaAppTypeLine: {0}", MetaAppTypeLine);
-		LogInformation("packageInfo return value: {0}", ReturnValue);
+		Logger.LogInformation("packageInfo.txt file exists: {fileExists}", fileExists);
+		Logger.LogInformation("packageInfo return MetaAppTypeLine: {MetaAppTypeLine}", MetaAppTypeLine);
+		Logger.LogInformation("packageInfo return value: {ReturnValue}", ReturnValue);
 
 		return ReturnValue;
 	}
@@ -3649,12 +3650,12 @@ public class AndroidPlatform : Platform
 					}
 					catch (Exception ex)
 					{
-						LogWarning(@"Ignoring duplicate package metadata entry '"+Line+"\'.\n" + ex.ToString());
+						Logger.LogWarning("{Text}", @"Ignoring duplicate package metadata entry '"+Line+"\'.\n" + ex.ToString());
 					}
 				}
 				else
 				{
-					LogWarning("Unexpected layout of package metadata: " + Line);
+					Logger.LogWarning("{Text}", "Unexpected layout of package metadata: " + Line);
 				}
 			}
 		}
@@ -3740,7 +3741,7 @@ public class AndroidPlatform : Platform
 		CachedAaptPath = BestToolPath;
 		LastAndroidHomePath = HomePath;
 
-		LogInformation("Using this aapt: {0}", CachedAaptPath);
+		Logger.LogInformation("Using this aapt: {CachedAaptPath}", CachedAaptPath);
 
 		return CachedAaptPath;
 	}
@@ -3787,10 +3788,10 @@ public class AndroidPlatform : Platform
 		AndroidFileClient client = new AndroidFileClient(DeviceName);
 		if (!client.OpenConnection())
 		{
-			Log.TraceInformation("DeployClientCmdLine: Trying to start file server {0}", PackageName);
+			Logger.LogInformation("DeployClientCmdLine: Trying to start file server {PackageName}", PackageName);
 			if (!client.StartServer(PackageName, AFSToken))
 			{
-				Log.TraceInformation("DeployClientCmdLine: Failed to start server {0}, ignoring client command line", PackageName);
+				Logger.LogInformation("DeployClientCmdLine: Failed to start server {PackageName}, ignoring client command line", PackageName);
 				return false;
 			}
 		}
@@ -3799,18 +3800,18 @@ public class AndroidPlatform : Platform
 		string DevicePackageName = client.Query("^packagename");
 		if (DevicePackageName != PackageName)
 		{
-			Log.TraceInformation("DeployClientCmdLine: Connected to wrong server {0}, trying again", DevicePackageName);
+			Logger.LogInformation("DeployClientCmdLine: Connected to wrong server {DevicePackageName}, trying again", DevicePackageName);
 			client.TerminateServer();
 
-			Log.TraceInformation("DeployClientCmdLine: Trying to start file server {0}", PackageName);
+			Logger.LogInformation("DeployClientCmdLine: Trying to start file server {PackageName}", PackageName);
 			if (!client.StartServer(PackageName, AFSToken))
 			{
-				Log.TraceInformation("DeployClientCmdLine: Failed to start server {0}, ignoring client command line", PackageName);
+				Logger.LogInformation("DeployClientCmdLine: Failed to start server {PackageName}, ignoring client command line", PackageName);
 				return false;
 			}
 		}
 
-		Log.TraceInformation("Writing ClientCmdLine to remote ^commandfile: {0}", ClientCmdLine);
+		Logger.LogInformation("Writing ClientCmdLine to remote ^commandfile: {ClientCmdLine}", ClientCmdLine);
 		client.FileWriteString(ClientCmdLine, "^commandfile");
 		client.TerminateServer();
 		client.CloseConnection();
@@ -3826,7 +3827,7 @@ public class AndroidPlatform : Platform
 		string ClientCmdLineTmpFile = Path.GetTempFileName();
 		string ClientCmdLineRemoteFile = RemoteDir + "/UECommandLine.txt";
 		File.WriteAllText(ClientCmdLineTmpFile, ClientCmdLine);
-		Log.TraceInformation("Pushing ClientCmdLine to remote file {0}: {1}", ClientCmdLineRemoteFile, ClientCmdLine);
+		Logger.LogInformation("Pushing ClientCmdLine to remote file {ClientCmdLineRemoteFile}: {ClientCmdLine}", ClientCmdLineRemoteFile, ClientCmdLine);
 		RunAdbCommand(DeviceName, String.Format("push {0} {1}", ClientCmdLineTmpFile, ClientCmdLineRemoteFile));
 		File.Delete(ClientCmdLineTmpFile);
 		return true;
diff --git a/Engine/Source/Programs/AutomationTool/AutomationUtils/Automation.cs b/Engine/Source/Programs/AutomationTool/AutomationUtils/Automation.cs
--- a/Engine/Source/Programs/AutomationTool/AutomationUtils/Automation.cs
+++ b/Engine/Source/Programs/AutomationTool/AutomationUtils/Automation.cs
@@ -18,6 +18,8 @@ namespace AutomationTool
 
     public static class Automation
 	{
+		static ILogger Logger => Log.Logger;
+
 		/// <summary>
 		/// Keep a persistent reference to the delegate for handling Ctrl-C events. Since it's passed to non-managed code, we have to prevent it from being garbage collected.
 		/// </summary>
@@ -70,7 +72,7 @@ namespace AutomationTool
 					}
 				}
 
-				Log.Logger.LogDebug("IsBuildMachine={IsBuildMachine}", IsBuildMachine);
+				Logger.LogDebug("IsBuildMachine={IsBuildMachine}", IsBuildMachine);
 				Environment.SetEnvironmentVariable("IsBuildMachine", IsBuildMachine ? "1" : "0");
 
 				// Register all the log event matchers
@@ -86,7 +88,7 @@ namespace AutomationTool
 
 				// should we kill processes on exit
 				ShouldKillProcesses = !GlobalCommandLine.NoKill;
-				Log.Logger.LogDebug("ShouldKillProcesses={ShouldKillProcesses}", ShouldKillProcesses);
+				Logger.LogDebug("ShouldKillProcesses={ShouldKillProcesses}", ShouldKillProcesses);
 
 				if (AutomationToolCommandLine.CommandsToExecute.Count == 0 && GlobalCommandLine.Help)
 				{
@@ -108,13 +110,13 @@ namespace AutomationTool
 				LogUtils.AddLogFileListener(new DirectoryReference(CommandUtils.CmdEnv.LogFolder), new DirectoryReference(CommandUtils.CmdEnv.FinalLogFolder));
 				if (LogUtils.FinalLogFileName != LogUtils.LogFileName)
 				{
-					Log.Logger.LogInformation("Final log location: {Location}", LogUtils.FinalLogFileName);
+					Logger.LogInformation("Final log location: {Location}", LogUtils.FinalLogFileName);
 				}
 				
 				// Initialize UBT
 				if (!UnrealBuildTool.PlatformExports.Initialize(Log.Logger))
 				{
-					Log.Logger.LogInformation("Failed to initialize UBT");
+					Logger.LogInformation("Failed to initialize UBT");
 					return ExitCode.Error_Unknown;
 				}
 				
@@ -127,7 +129,7 @@ namespace AutomationTool
 				// Compile scripts.
 				using (GlobalTracer.Instance.BuildSpan("ScriptLoad").StartActive())
 				{
-					ScriptManager.LoadScriptAssemblies(ScriptModuleAssemblies);
+					ScriptManager.LoadScriptAssemblies(ScriptModuleAssemblies, Logger);
 				}
 
 				if (GlobalCommandLine.List)
@@ -186,8 +188,8 @@ namespace AutomationTool
 				}
 				else if (Ex.OutputFormat == AutomationExceptionOutputFormat.Minimal)
 				{
-					Log.Logger.LogInformation(Ex, "{Message}", Ex.ToString().Replace("\n", "\n  "));
-					Log.Logger.LogDebug(Ex, "{Details}", ExceptionUtils.FormatExceptionDetails(Ex));
+					Logger.LogInformation(Ex, "{Message}", Ex.ToString().Replace("\n", "\n  "));
+					Logger.LogDebug(Ex, "{Details}", ExceptionUtils.FormatExceptionDetails(Ex));
 				}
 				else
 				{
@@ -232,7 +234,7 @@ namespace AutomationTool
             }
             catch (Exception Ex)
             {
-                Log.Logger.LogError(Ex, "Exception performing nothrow action \"{ActionDesc}\": {Exception}", ActionDesc, ExceptionUtils.FormatException(Ex));
+                Logger.LogError(Ex, "Exception performing nothrow action \"{ActionDesc}\": {Exception}", ActionDesc, ExceptionUtils.FormatException(Ex));
             }
         }
 
@@ -247,7 +249,7 @@ namespace AutomationTool
 			for (int CommandIndex = 0; CommandIndex < CommandsToExecute.Count; ++CommandIndex)
 			{
 				var CommandInfo = CommandsToExecute[CommandIndex];
-				Log.Logger.LogDebug("Attempting to execute {CommandInfo}", CommandInfo.ToString());
+				Logger.LogDebug("Attempting to execute {CommandInfo}", CommandInfo.ToString());
 				Type CommandType;
 				if (!Commands.TryGetValue(CommandInfo.CommandName, out CommandType))
 				{
@@ -294,7 +296,7 @@ namespace AutomationTool
 				Type CommandType;
 				if (Commands.TryGetValue(CommandInfo.CommandName, out CommandType) == false)
 				{
-					Log.Logger.LogError("Help: Failed to find command {CommandName}", CommandInfo.CommandName);
+					Logger.LogError("Help: Failed to find command {CommandName}", CommandInfo.CommandName);
 				}
 				else
 				{
diff --git a/Engine/Source/Programs/AutomationTool/AutomationUtils/BuildCommand.cs b/Engine/Source/Programs/AutomationTool/AutomationUtils/BuildCommand.cs
--- a/Engine/Source/Programs/AutomationTool/AutomationUtils/BuildCommand.cs
+++ b/Engine/Source/Programs/AutomationTool/AutomationUtils/BuildCommand.cs
@@ -53,11 +53,11 @@ namespace AutomationTool
 				{
 					if (ParamValue == null)
 					{
-						Log.Logger.LogWarning("Param name \"{ObsoleteParam}\" is deprecated, use \"{Param}\" instead.", ObsoleteParam, Param);
+						Logger.LogWarning("Param name \"{ObsoleteParam}\" is deprecated, use \"{Param}\" instead.", ObsoleteParam, Param);
 					}
 					else
 					{
-						Log.Logger.LogWarning("Deprecated param name \"{ObsoleteParam}\" was ignored because \"{Param}\" was set.", ObsoleteParam, Param);
+						Logger.LogWarning("Deprecated param name \"{ObsoleteParam}\" was ignored because \"{Param}\" was set.", ObsoleteParam, Param);
 					}
 				}
 
diff --git a/Engine/Source/Programs/AutomationTool/AutomationUtils/CommandEnvironment.cs b/Engine/Source/Programs/AutomationTool/AutomationUtils/CommandEnvironment.cs
--- a/Engine/Source/Programs/AutomationTool/AutomationUtils/CommandEnvironment.cs
+++ b/Engine/Source/Programs/AutomationTool/AutomationUtils/CommandEnvironment.cs
@@ -46,6 +46,8 @@ namespace AutomationTool
 	/// </summary>
 	public class CommandEnvironment
 	{
+		static ILogger Logger => Log.Logger;
+
 		/// <summary>
 		/// Path to a file we know to always exist under the Unreal root directory.
 		/// </summary>
@@ -137,7 +139,7 @@ namespace AutomationTool
 
 			if (IsChildInstance)
 			{
-				Log.Logger.LogInformation("AutomationTool is running as a child instance ({Name}={Value})", EnvVarNames.IsChildInstance, IsChildInstanceInt.ToString());
+				Logger.LogInformation("AutomationTool is running as a child instance ({Name}={Value})", EnvVarNames.IsChildInstance, IsChildInstanceInt.ToString());
 			}
 
 			// Setup the timestamp string
@@ -182,18 +184,18 @@ namespace AutomationTool
 
 		void LogSettings()
 		{
-			Log.Logger.LogDebug("Command Environment settings:");
-			Log.Logger.LogDebug("CmdExe={CmdExe}", CmdExe);
-			Log.Logger.LogDebug("EngineSavedFolder={EngineSavedFolder}", EngineSavedFolder);
-			Log.Logger.LogDebug("HasCapabilityToCompile={HasCapabilityToCompile}", HasCapabilityToCompile);
-			Log.Logger.LogDebug("LocalRoot={LocalRoot}", LocalRoot);
-			Log.Logger.LogDebug("LogFolder={LogFolder}", LogFolder);
-			Log.Logger.LogDebug("MountExe={MountExe}", MountExe);
-			Log.Logger.LogDebug("FrameworkMsbuildExe={FrameworkMsbuildPath}", FrameworkMsbuildPath);
-			Log.Logger.LogDebug("DotnetMsbuildExe={DotnetMsbuildPath}", DotnetMsbuildPath);
-			Log.Logger.LogDebug("RobocopyExe={RobocopyExe}", RobocopyExe);
-			Log.Logger.LogDebug("TimestampAsString={TimestampAsString}", TimestampAsString);
-			Log.Logger.LogDebug("AutomationToolDll={AutomationToolDll}", AutomationToolDll);
+			Logger.LogDebug("Command Environment settings:");
+			Logger.LogDebug("CmdExe={CmdExe}", CmdExe);
+			Logger.LogDebug("EngineSavedFolder={EngineSavedFolder}", EngineSavedFolder);
+			Logger.LogDebug("HasCapabilityToCompile={HasCapabilityToCompile}", HasCapabilityToCompile);
+			Logger.LogDebug("LocalRoot={LocalRoot}", LocalRoot);
+			Logger.LogDebug("LogFolder={LogFolder}", LogFolder);
+			Logger.LogDebug("MountExe={MountExe}", MountExe);
+			Logger.LogDebug("FrameworkMsbuildExe={FrameworkMsbuildPath}", FrameworkMsbuildPath);
+			Logger.LogDebug("DotnetMsbuildExe={DotnetMsbuildPath}", DotnetMsbuildPath);
+			Logger.LogDebug("RobocopyExe={RobocopyExe}", RobocopyExe);
+			Logger.LogDebug("TimestampAsString={TimestampAsString}", TimestampAsString);
+			Logger.LogDebug("AutomationToolDll={AutomationToolDll}", AutomationToolDll);
 		}
 
 		/// <summary>
@@ -215,11 +217,11 @@ namespace AutomationTool
 				HasCapabilityToCompile = false;
 				FrameworkMsbuildPath = "";
 			}
-			Log.Logger.LogDebug("CompilationEvironment.HasCapabilityToCompile={HasCapabilityToCompile}", HasCapabilityToCompile);
-			Log.Logger.LogDebug("CompilationEvironment.FrameworkMsbuildExe={FrameworkMsbuildPath}", FrameworkMsbuildPath);
+			Logger.LogDebug("CompilationEvironment.HasCapabilityToCompile={HasCapabilityToCompile}", HasCapabilityToCompile);
+			Logger.LogDebug("CompilationEvironment.FrameworkMsbuildExe={FrameworkMsbuildPath}", FrameworkMsbuildPath);
 
 			DotnetMsbuildPath = Unreal.DotnetPath.FullName;
-			Log.Logger.LogDebug("CompilationEvironment.DotnetMsbuildExe={DotnetMsbuildPath}", DotnetMsbuildPath);
+			Logger.LogDebug("CompilationEvironment.DotnetMsbuildExe={DotnetMsbuildPath}", DotnetMsbuildPath);
 		}
 
 		/// <summary>
@@ -238,7 +240,7 @@ namespace AutomationTool
 			while (true)
 			{
 				var PathToCheck = Path.GetFullPath(CommandUtils.CombinePaths(CurrentPath, KnownFilePathFromRoot));
-				Log.Logger.LogDebug("Checking for {PathToCheck}", PathToCheck);
+				Logger.LogDebug("Checking for {PathToCheck}", PathToCheck);
 				if (!File.Exists(PathToCheck))
 				{
 					var LastSeparatorIndex = CurrentPath.LastIndexOf('/');
diff --git a/Engine/Source/Programs/AutomationTool/AutomationUtils/CommandUtils.cs b/Engine/Source/Programs/AutomationTool/AutomationUtils/CommandUtils.cs
--- a/Engine/Source/Programs/AutomationTool/AutomationUtils/CommandUtils.cs
+++ b/Engine/Source/Programs/AutomationTool/AutomationUtils/CommandUtils.cs
@@ -103,7 +103,7 @@ namespace AutomationTool
 		/// <summary>
 		/// Provides access to the structured logging interface
 		/// </summary>
-		public ILogger Logger => Log.Logger;
+		public static ILogger Logger => Log.Logger;
 
 		/// <summary>
 		/// Writes formatted text to log (with LogEventType.Console).
@@ -261,7 +261,7 @@ namespace AutomationTool
 		{
 			if(bShowProgress)
 			{
-				Log.Logger.LogInformation("[@progress push {Numerator}/{Denominator} skipline]", Numerator, Denominator);
+				Logger.LogInformation("[@progress push {Numerator}/{Denominator} skipline]", Numerator, Denominator);
 			}
 		}
 
@@ -269,7 +269,7 @@ namespace AutomationTool
 		{
 			if(bShowProgress)
 			{
-				Log.Logger.LogInformation("[@progress pop skipline]");
+				Logger.LogInformation("[@progress pop skipline]");
 			}
 		}
 
@@ -277,7 +277,7 @@ namespace AutomationTool
 		{
 			if(bShowProgress)
 			{
-				Log.Logger.LogInformation("[@progress increment {Numerator}/{Denominator} skipline]", Numerator, Denominator);
+				Logger.LogInformation("[@progress increment {Numerator}/{Denominator} skipline]", Numerator, Denominator);
 			}
 		}
 
@@ -286,7 +286,7 @@ namespace AutomationTool
 		{
 			if(bShowProgress)
 			{
-				Log.Logger.LogInformation("[@progress '{Message}' skipline]", String.Format(Format, Args));
+				Logger.LogInformation("[@progress '{Message}' skipline]", String.Format(Format, Args));
 			}
 		}
 
@@ -295,7 +295,7 @@ namespace AutomationTool
 		{
 			if(bShowProgress)
 			{
-				Log.Logger.LogInformation("[@progress {Numerator}/{Denominator} '{Arg2}' skipline]", Numerator, Denominator, String.Format(Format, Args));
+				Logger.LogInformation("[@progress {Numerator}/{Denominator} '{Arg2}' skipline]", Numerator, Denominator, String.Format(Format, Args));
 			}
 		}
 
@@ -486,7 +486,7 @@ namespace AutomationTool
 			var NormalizedFilename = ConvertSeparators(PathSeparator.Default, FileName);
 			if (!InternalUtils.SafeDeleteFile(NormalizedFilename))
 			{
-				Log.Logger.LogWarning("Failed to delete file '{NormalizedFilename}'", NormalizedFilename);
+				Logger.LogWarning("Failed to delete file '{NormalizedFilename}'", NormalizedFilename);
 				Result = false;
 			}
 			return Result;
@@ -559,7 +559,7 @@ namespace AutomationTool
             {
                 if (!InternalUtils.SafeDeleteDirectory(NormalizedDirectory, bQuiet))
                 {
-					Log.Logger.LogWarning("Failed to delete directory '{NormalizedDirectory}'", NormalizedDirectory);
+					Logger.LogWarning("Failed to delete directory '{NormalizedDirectory}'", NormalizedDirectory);
                     Result = false;
                 }
             }
@@ -567,8 +567,8 @@ namespace AutomationTool
             {
 				if (!bQuiet)
 				{
-					Log.Logger.LogWarning("Failed to delete directory, exception '{NormalizedDirectory}'", NormalizedDirectory);
-					Log.Logger.LogWarning(Ex, "{Text}", Ex.Message);
+					Logger.LogWarning("Failed to delete directory, exception '{NormalizedDirectory}'", NormalizedDirectory);
+					Logger.LogWarning(Ex, "{Text}", Ex.Message);
 				}
                 Result = false;
             }
@@ -595,7 +595,7 @@ namespace AutomationTool
 		/// <param name="DirectoryName"></param>
 		public static void DeleteDirectoryContents(string DirectoryName)
 		{
-			Log.Logger.LogDebug("DeleteDirectoryContents({DirectoryName})", DirectoryName);
+			Logger.LogDebug("DeleteDirectoryContents({DirectoryName})", DirectoryName);
 			const bool bQuiet = true;
 			var Files = CommandUtils.FindFiles_NoExceptions(bQuiet, "*", false, DirectoryName);
 			foreach (var Filename in Files)
@@ -645,8 +645,8 @@ namespace AutomationTool
 			}
 			catch (Exception Ex)
 			{
-				Log.Logger.LogWarning("Unable to check if directory exists: {NormalizedDirectory}", NormalizedDirectory);
-				Log.Logger.LogWarning(Ex, "{Text}", Ex.Message);
+				Logger.LogWarning("Unable to check if directory exists: {NormalizedDirectory}", NormalizedDirectory);
+				Logger.LogWarning(Ex, "{Text}", Ex.Message);
 				return false;
 			}
 		}
@@ -733,7 +733,7 @@ namespace AutomationTool
 			var Result = InternalUtils.SafeRenameFile(OldNormalized, NewNormalized);
 			if (!Result)
 			{
-				Log.Logger.LogWarning("Failed to rename/move file '{OldName}' to '{NewName}'", OldName, NewName);
+				Logger.LogWarning("Failed to rename/move file '{OldName}' to '{NewName}'", OldName, NewName);
 			}
 			return Result;
 		}
@@ -822,7 +822,7 @@ namespace AutomationTool
 			}
 			catch
 			{
-				Log.Logger.LogWarning("Unable to change current directory to {WorkingDirectory}", WorkingDirectory);
+				Logger.LogWarning("Unable to change current directory to {WorkingDirectory}", WorkingDirectory);
 				Result = false;
 			}
 			return Result;
@@ -855,7 +855,7 @@ namespace AutomationTool
 			}
 			else
 			{
-				Log.Logger.LogWarning("Unable to PopDir. WorkingDirectoryStack is empty.");
+				Logger.LogWarning("Unable to PopDir. WorkingDirectoryStack is empty.");
 				Result = false;
 			}
 			return Result;
@@ -903,7 +903,7 @@ namespace AutomationTool
 			}
 			catch
 			{
-				Log.Logger.LogWarning("Unable to change current directory to {WorkingDirectory}", WorkingDirectory);
+				Logger.LogWarning("Unable to change current directory to {WorkingDirectory}", WorkingDirectory);
 				Result = false;
 			}
 			return Result;
@@ -981,7 +981,7 @@ namespace AutomationTool
 			Filename = ConvertSeparators(PathSeparator.Default, Filename);
 			if (!File.Exists(Filename))
 			{
-				Log.Logger.LogWarning("Unable to set attributes for a non-exisiting file ({Filename})", Filename);
+				Logger.LogWarning("Unable to set attributes for a non-exisiting file ({Filename})", Filename);
 				return false;
 			}
 
@@ -994,8 +994,8 @@ namespace AutomationTool
 			}
 			catch (Exception Ex)
 			{
-				Log.Logger.LogWarning("Error trying to set file attributes for: {Filename}", Filename);
-				Log.Logger.LogWarning(Ex, "{Text}", Ex.Message);
+				Logger.LogWarning("Error trying to set file attributes for: {Filename}", Filename);
+				Logger.LogWarning(Ex, "{Text}", Ex.Message);
 				Result = false;
 			}
 			return Result;
@@ -1386,13 +1386,13 @@ namespace AutomationTool
 			{
 				if (bAllowDifferingTimestamps == true)
 				{
-					Log.Logger.LogDebug("CopyFileIncremental Skipping {Dest}, already exists", Dest);
+					Logger.LogDebug("CopyFileIncremental Skipping {Dest}, already exists", Dest);
 					return;
 				}
 				TimeSpan Diff = File.GetLastWriteTimeUtc(Dest.FullName) - File.GetLastWriteTimeUtc(Source.FullName);
 				if (Diff.TotalSeconds > -1 && Diff.TotalSeconds < 1)
 				{
-					Log.Logger.LogDebug("CopyFileIncremental Skipping {Dest}, up to date.", Dest);
+					Logger.LogDebug("CopyFileIncremental Skipping {Dest}, up to date.", Dest);
 					return;
 				}
 				InternalUtils.SafeDeleteFile(Dest.FullName);
@@ -1766,7 +1766,7 @@ namespace AutomationTool
 		{
 			if(!bQuiet)
 			{
-				Log.Logger.LogInformation("Copying {NumFiles} file(s) using max {MaxThreads} thread(s)", Source.Count, MaxThreads);
+				Logger.LogInformation("Copying {NumFiles} file(s) using max {MaxThreads} thread(s)", Source.Count, MaxThreads);
 			}
 
             if (Source.Count != Dest.Count)
@@ -1809,7 +1809,7 @@ namespace AutomationTool
 		public static List<string> ThreadedCopyFiles(string SourceDir, string TargetDir, FileFilter Filter, bool bIgnoreSymlinks, int MaxThreads = 64, bool bRetry = false)
 		{
 			// Filter all the relative paths
-			Log.Logger.LogInformation("Applying filter to {SourceDir}...", SourceDir);
+			Logger.LogInformation("Applying filter to {SourceDir}...", SourceDir);
 			DirectoryReference SourceDirRef = new DirectoryReference(SourceDir);
 			var RelativePaths = Filter.ApplyToDirectory(SourceDirRef, bIgnoreSymlinks).Select(x => x.MakeRelativeTo(SourceDirRef)).ToList();
 			return ThreadedCopyFiles(SourceDir, TargetDir, RelativePaths, MaxThreads, bRetry: bRetry);
@@ -1912,7 +1912,7 @@ namespace AutomationTool
 		{
 			try
 			{
-				Log.Logger.LogDebug("SetEnvVar {Name}={Value}", Name, Value);
+				Logger.LogDebug("SetEnvVar {Name}={Value}", Name, Value);
 				Environment.SetEnvironmentVariable(Name, Value.ToString());
 			}
 			catch (Exception Ex)
@@ -2119,7 +2119,7 @@ namespace AutomationTool
         {
             if (!DirectoryExists_NoExceptions(Dir))
             {
-				Log.Logger.LogDebug("Directory {Dir} does not exist", Dir);
+				Logger.LogDebug("Directory {Dir} does not exist", Dir);
 				return false;
 			}
 
@@ -2134,7 +2134,7 @@ namespace AutomationTool
 				if(File.Exists(NativeFilename))
 				{
 		            DeleteFile_NoExceptions(Filename, true);
-					Log.Logger.LogDebug("Directory {Dir} is writable", Dir);
+					Logger.LogDebug("Directory {Dir} is writable", Dir);
 					return true;
 				}
 			}
@@ -2142,7 +2142,7 @@ namespace AutomationTool
 			{
 			}
 
-			Log.Logger.LogDebug("Directory {Dir} is not writable", Dir);
+			Logger.LogDebug("Directory {Dir} is not writable", Dir);
 			return false;
 		}
 
@@ -2155,25 +2155,25 @@ namespace AutomationTool
             try
             {
                 DirectoryInfo DirInfo = new DirectoryInfo(ParentDir);
-				Log.Logger.LogInformation("Looking for directories to delete in {ParentDir}", ParentDir);
+				Logger.LogInformation("Looking for directories to delete in {ParentDir}", ParentDir);
                 foreach (DirectoryInfo ThisDirInfo in DirInfo.EnumerateDirectories(SearchPattern))
                 {
 					double AgeDays = (DateTime.UtcNow - ThisDirInfo.CreationTimeUtc).TotalDays;
 					if (AgeDays > MaximumDaysToKeepTempStorage)
                     {
-						Log.Logger.LogInformation("Deleting formal build directory {Directory}, because it is {Age} days old (maximum {MaximumDaysToKeepTempStorage}).", ThisDirInfo.FullName, (int)AgeDays, MaximumDaysToKeepTempStorage);
+						Logger.LogInformation("Deleting formal build directory {Directory}, because it is {Age} days old (maximum {MaximumDaysToKeepTempStorage}).", ThisDirInfo.FullName, (int)AgeDays, MaximumDaysToKeepTempStorage);
                         DeleteDirectory_NoExceptions(true, ThisDirInfo.FullName);
                     }
                     else
                     {
-						Log.Logger.LogDebug("Not deleting formal build directory {Directory}, because it is {Age} days old (maximum {MaximumDaysToKeepTempStorage}).", ThisDirInfo.FullName, (int)AgeDays, MaximumDaysToKeepTempStorage);
+						Logger.LogDebug("Not deleting formal build directory {Directory}, because it is {Age} days old (maximum {MaximumDaysToKeepTempStorage}).", ThisDirInfo.FullName, (int)AgeDays, MaximumDaysToKeepTempStorage);
                     }
                 }
             }
             catch (Exception Ex)
             {
-				Log.Logger.LogWarning("Unable to clean formal builds from directory: {ParentDir}", ParentDir);
-				Log.Logger.LogWarning(Ex, " Exception was {Exception}", LogUtils.FormatException(Ex));
+				Logger.LogWarning("Unable to clean formal builds from directory: {ParentDir}", ParentDir);
+				Logger.LogWarning(Ex, " Exception was {Exception}", LogUtils.FormatException(Ex));
             }
         }
 
@@ -2478,11 +2478,11 @@ namespace AutomationTool
 
 					if (Iterations == 0)
 					{
-						Log.Logger.LogInformation("Waiting for lock file '{LockFilePath}' to be removed...", LockFilePath);
+						Logger.LogInformation("Waiting for lock file '{LockFilePath}' to be removed...", LockFilePath);
 					}
 					else if ((Iterations % 30) == 0)
 					{
-						Log.Logger.LogInformation("Still waiting for lock file '{LockFilePath}' after {Seconds} seconds.", LockFilePath, CurrentTime.Subtract(StartTime).TotalSeconds);
+						Logger.LogInformation("Still waiting for lock file '{LockFilePath}' after {Seconds} seconds.", LockFilePath, CurrentTime.Subtract(StartTime).TotalSeconds);
 					}
 
 					// Wait for a while before retrying.
diff --git a/Engine/Source/Programs/AutomationTool/AutomationUtils/CommandletUtils.cs b/Engine/Source/Programs/AutomationTool/AutomationUtils/CommandletUtils.cs
--- a/Engine/Source/Programs/AutomationTool/AutomationUtils/CommandletUtils.cs
+++ b/Engine/Source/Programs/AutomationTool/AutomationUtils/CommandletUtils.cs
@@ -267,7 +267,7 @@ namespace AutomationTool
 
 		public static void StartRunCommandlet(FileReference ProjectName, string UnrealExe, string Commandlet, string Parameters, ERunOptions RunOptions, out string LocalLogFile, out IProcessResult RunResult, ProcessResult.SpewFilterCallbackType SpewFilterCallback=null)
 		{
-			Log.Logger.LogInformation("Running UnrealEditor {Commandlet} for project {ProjectName}", Commandlet, ProjectName);
+			Logger.LogInformation("Running UnrealEditor {Commandlet} for project {ProjectName}", Commandlet, ProjectName);
 
 			var CWD = Path.GetDirectoryName(UnrealExe);
 
@@ -282,7 +282,7 @@ namespace AutomationTool
 			PushDir(CWD);
 
 			LocalLogFile = LogUtils.GetUniqueLogName(CombinePaths(CmdEnv.EngineSavedFolder, Commandlet));
-			Log.Logger.LogInformation("Commandlet log file is {LocalLogFile}", LocalLogFile);
+			Logger.LogInformation("Commandlet log file is {LocalLogFile}", LocalLogFile);
 			string Args = String.Format(
 				"{0} -run={1} {2} -abslog={3} -stdout -CrashForUAT -unattended -NoLogTimes {5}{4}",
 				(ProjectName == null) ? "" : CommandUtils.MakePathSafeToUseWithCommandLine(ProjectName.FullName),
diff --git a/Engine/Source/Programs/AutomationTool/AutomationUtils/HashCollection.cs b/Engine/Source/Programs/AutomationTool/AutomationUtils/HashCollection.cs
--- a/Engine/Source/Programs/AutomationTool/AutomationUtils/HashCollection.cs
+++ b/Engine/Source/Programs/AutomationTool/AutomationUtils/HashCollection.cs
@@ -170,18 +170,19 @@ namespace AutomationTool
 		/// <param name="RHS"></param>
 		public void LogDifferences(HashCollection RHS)
 		{
+			ILogger Logger = Log.Logger;
 			foreach (var Entry in Hashes)
 			{
 				if (!RHS.HasEntryForName(Entry.Name))
 				{
-					Log.Logger.LogInformation("RHS does not have an entry for {Name}", Entry.Name);
+					Logger.LogInformation("RHS does not have an entry for {Name}", Entry.Name);
 				}
 				else if (!RHS.Hashes.Contains(Entry))
 				{
 					HashEntry RHSEntry = RHS.GetEntryForName(Entry.Name);
-					Log.Logger.LogInformation("RHS hash mismatch for {Name}", Entry.Name);
-					Log.Logger.LogInformation("LHS:\n\t{Hash}\n\t{Name}\n\t{MetaData}", Entry.Hash, Entry.Name, Entry.MetaData);
-					Log.Logger.LogInformation("RHS:\n\t{Hash}\n\t{Name}\n\t{MetaData}", RHSEntry.Hash, RHSEntry.Name, RHSEntry.MetaData);
+					Logger.LogInformation("RHS hash mismatch for {Name}", Entry.Name);
+					Logger.LogInformation("LHS:\n\t{Hash}\n\t{Name}\n\t{MetaData}", Entry.Hash, Entry.Name, Entry.MetaData);
+					Logger.LogInformation("RHS:\n\t{Hash}\n\t{Name}\n\t{MetaData}", RHSEntry.Hash, RHSEntry.Name, RHSEntry.MetaData);
 				}
 			}
 		}
diff --git a/Engine/Source/Programs/AutomationTool/AutomationUtils/HelpUtils.cs b/Engine/Source/Programs/AutomationTool/AutomationUtils/HelpUtils.cs
--- a/Engine/Source/Programs/AutomationTool/AutomationUtils/HelpUtils.cs
+++ b/Engine/Source/Programs/AutomationTool/AutomationUtils/HelpUtils.cs
@@ -312,12 +312,12 @@ namespace AutomationTool
 					}
 					else
 					{
-						Log.Logger.LogWarning("Duplicated help parameter \"{ParamName}\"", ParamName);
+						Logger.LogWarning("Duplicated help parameter \"{ParamName}\"", ParamName);
 					}
 				}
 			}
 
-			Log.Logger.LogInformation("");
+			Logger.LogInformation("");
 			HelpUtils.PrintHelp(String.Format("{0} Help:", Command.Name), Description, ParamDict.ToList());
 		}
 
diff --git a/Engine/Source/Programs/AutomationTool/AutomationUtils/HostPlatform.cs b/Engine/Source/Programs/AutomationTool/AutomationUtils/HostPlatform.cs
--- a/Engine/Source/Programs/AutomationTool/AutomationUtils/HostPlatform.cs
+++ b/Engine/Source/Programs/AutomationTool/AutomationUtils/HostPlatform.cs
@@ -7,6 +7,7 @@ using System.Linq;
 using System.Text;
 using System.Diagnostics;
 using EpicGames.Core;
+using Microsoft.Extensions.Logging;
 
 namespace AutomationTool
 {
@@ -15,6 +16,8 @@ namespace AutomationTool
 	/// </summary>
 	public abstract class HostPlatform
 	{
+		protected static ILogger Logger => Log.Logger;
+
 		/// <summary>
 		/// Current running host platform.
 		/// </summary>
diff --git a/Engine/Source/Programs/AutomationTool/AutomationUtils/MacHostPlatform.cs b/Engine/Source/Programs/AutomationTool/AutomationUtils/MacHostPlatform.cs
--- a/Engine/Source/Programs/AutomationTool/AutomationUtils/MacHostPlatform.cs
+++ b/Engine/Source/Programs/AutomationTool/AutomationUtils/MacHostPlatform.cs
@@ -27,13 +27,13 @@ namespace AutomationTool
 
 				if (CanUseMsBuild)
 				{
-					Log.Logger.LogInformation("using {DotNet}!", CommandUtils.WhichApp("dotnet"));
+					Logger.LogInformation("using {DotNet}!", CommandUtils.WhichApp("dotnet"));
 
 					CachedFrameworkMsbuildExe = "dotnet msbuild";
 				}
 				else
 				{
-					Log.Logger.LogInformation("Using xbuild. Install Mono 5.0 or greater for faster builds!");
+					Logger.LogInformation("Using xbuild. Install Mono 5.0 or greater for faster builds!");
 					CachedFrameworkMsbuildExe = "xbuild";
 				}
 			}
diff --git a/Engine/Source/Programs/AutomationTool/AutomationUtils/P4Environment.cs b/Engine/Source/Programs/AutomationTool/AutomationUtils/P4Environment.cs
--- a/Engine/Source/Programs/AutomationTool/AutomationUtils/P4Environment.cs
+++ b/Engine/Source/Programs/AutomationTool/AutomationUtils/P4Environment.cs
@@ -111,6 +111,8 @@ namespace AutomationTool
 		/// </summary>
 		private int CodeChangelistInternal = -1;
 
+		static ILogger Logger => Log.Logger;
+
 		/// <summary>
 		/// Constructor. Derives the Perforce environment settings.
 		/// </summary>
@@ -158,7 +160,7 @@ namespace AutomationTool
 				{
 					string HostName = System.Net.Dns.GetHostName();
 					ThisClient = DetectClient(DefaultConnection, User, HostName, CmdEnv.AutomationToolDll);
-					Log.Logger.LogInformation("Using user {User} clientspec {ClientName} {ClientPath}", User, ThisClient.Name, ThisClient.RootPath);
+					Logger.LogInformation("Using user {User} clientspec {ClientName} {ClientPath}", User, ThisClient.Name, ThisClient.RootPath);
 					Client = ThisClient.Name;
 					CommandUtils.SetEnvVar(EnvVarNames.Client, Client);
 				}
@@ -218,7 +220,7 @@ namespace AutomationTool
 				using PerforceConnection Connection = new PerforceConnection(new PerforceSettings(ServerAndPort, User) { ClientName = Client }, Log.Logger);
 				CodeChangelistString = DetectCurrentCodeCL(Connection, Changelist).Result.ToString();
 				CommandUtils.SetEnvVar(EnvVarNames.CodeChangelist, CodeChangelistString);
-				Log.Logger.LogDebug("Took {ElapsedMs}ms to query last code change", Timer.ElapsedMilliseconds);
+				Logger.LogDebug("Took {ElapsedMs}ms to query last code change", Timer.ElapsedMilliseconds);
 			}
 			if(!String.IsNullOrEmpty(CodeChangelistString))
 			{
@@ -233,18 +235,18 @@ namespace AutomationTool
 			// Write a summary of the settings to the output window
 			if (!CommandUtils.CmdEnv.IsChildInstance)
 			{
-				Log.Logger.LogInformation("Detected Perforce Settings:");
-				Log.Logger.LogInformation("  Server: {ServerAndPort}", ServerAndPort);
-				Log.Logger.LogInformation("  User: {User}", User);
-				Log.Logger.LogInformation("  Client: {Client}", Client);
-				Log.Logger.LogInformation("  Branch: {Branch}", Branch);
+				Logger.LogInformation("Detected Perforce Settings:");
+				Logger.LogInformation("  Server: {ServerAndPort}", ServerAndPort);
+				Logger.LogInformation("  User: {User}", User);
+				Logger.LogInformation("  Client: {Client}", Client);
+				Logger.LogInformation("  Branch: {Branch}", Branch);
 				if (ChangelistInternal != -1)
 				{
-					Log.Logger.LogInformation("  Last Change: {Changelist}", Changelist);
+					Logger.LogInformation("  Last Change: {Changelist}", Changelist);
 				}
 				if (CodeChangelistInternal != -1)
 				{
-					Log.Logger.LogInformation("  Last Code Change: {CodeChangelist}", CodeChangelist);
+					Logger.LogInformation("  Last Code Change: {CodeChangelist}", CodeChangelist);
 				}
 			}
 
@@ -290,7 +292,7 @@ namespace AutomationTool
 			// Otherwise fallback to the uebp variables, or the default
 			if(String.IsNullOrEmpty(P4Port))
 			{
-				Log.Logger.LogWarning("P4PORT is not set. Using perforce:1666");
+				Logger.LogWarning("P4PORT is not set. Using perforce:1666");
 				P4Port = "perforce:1666";
 			}
 
@@ -316,7 +318,7 @@ namespace AutomationTool
 			{
 				if (!String.IsNullOrEmpty(UserName))
 				{
-					Log.Logger.LogWarning("Unable to retrieve perforce user name. Trying to fall back to {UserNameEnvVar} which is set to {UserName}.", EnvVarNames.User, UserName);
+					Logger.LogWarning("Unable to retrieve perforce user name. Trying to fall back to {UserNameEnvVar} which is set to {UserName}.", EnvVarNames.User, UserName);
 				}
 				else
 				{
@@ -361,7 +363,7 @@ namespace AutomationTool
 			{
 				if (!String.IsNullOrEmpty(Client.Host) && String.Compare(Client.Host, HostName, true) != 0)
 				{
-					Log.Logger.LogInformation("Rejecting client because of different Host {ClientName} \"{ClientHost}\" != \"{HostName}\"", Client.Name, Client.Host, HostName);
+					Logger.LogInformation("Rejecting client because of different Host {ClientName} \"{ClientHost}\" != \"{HostName}\"", Client.Name, Client.Host, HostName);
 					continue;
 				}
 				
@@ -390,7 +392,7 @@ namespace AutomationTool
 				}
 				if (ClientToUse == null)
 				{
-					Log.Logger.LogWarning("{NumClients} clients found that match the current host and root path. The most recently accessed client will be used.", MatchingClients.Count);
+					Logger.LogWarning("{NumClients} clients found that match the current host and root path. The most recently accessed client will be used.", MatchingClients.Count);
 					ClientToUse = GetMostRecentClient(MatchingClients);
 				}
 			}
@@ -405,7 +407,7 @@ namespace AutomationTool
 		/// <returns>The most recent client from the list.</returns>
 		private static P4ClientInfo GetMostRecentClient(List<P4ClientInfo> Clients)
 		{
-			Log.Logger.LogDebug("Detecting the most recent client.");
+			Logger.LogDebug("Detecting the most recent client.");
 			P4ClientInfo MostRecentClient = null;
 			var MostRecentAccessTime = DateTime.MinValue;
 			foreach (var ClientInfo in Clients)
@@ -526,7 +528,7 @@ namespace AutomationTool
 				string KnownFileDepotMapping = P4Result.Output;
 
 				// Get the build root
-				Log.Logger.LogDebug("Looking for {KnownFilePathFromRoot} in {KnownFileDepotMapping}", KnownFilePathFromRoot, KnownFileDepotMapping);
+				Logger.LogDebug("Looking for {KnownFilePathFromRoot} in {KnownFileDepotMapping}", KnownFilePathFromRoot, KnownFileDepotMapping);
 				int EndIdx = KnownFileDepotMapping.IndexOf(KnownFilePathFromRoot, StringComparison.CurrentCultureIgnoreCase);
 				if (EndIdx < 0)
 				{
diff --git a/Engine/Source/Programs/AutomationTool/AutomationUtils/P4Utils.cs b/Engine/Source/Programs/AutomationTool/AutomationUtils/P4Utils.cs
--- a/Engine/Source/Programs/AutomationTool/AutomationUtils/P4Utils.cs
+++ b/Engine/Source/Programs/AutomationTool/AutomationUtils/P4Utils.cs
@@ -1082,7 +1082,7 @@ namespace AutomationTool
 			// Init AllowSubmit
 			// If we do not specify on the commandline if submitting is allowed or not, this is 
 			// depending on whether we run locally or on a build machine.
-			Log.Logger.LogDebug("Initializing AllowSubmit.");
+			Logger.LogDebug("Initializing AllowSubmit.");
 			if (GlobalCommandLine.Submit || GlobalCommandLine.NoSubmit)
 			{
 				AllowSubmit = GlobalCommandLine.Submit;
@@ -1091,10 +1091,10 @@ namespace AutomationTool
 			{
 				AllowSubmit = Automation.IsBuildMachine;
 			}
-			Log.Logger.LogDebug("AllowSubmit={AllowSubmit}", AllowSubmit);
+			Logger.LogDebug("AllowSubmit={AllowSubmit}", AllowSubmit);
 
 			// Init P4Enabled
-			Log.Logger.LogDebug("Initializing P4Enabled.");
+			Logger.LogDebug("Initializing P4Enabled.");
 			if (Automation.IsBuildMachine)
 			{
 				P4Enabled = !GlobalCommandLine.NoP4;
@@ -1109,8 +1109,8 @@ namespace AutomationTool
 				P4Enabled = GlobalCommandLine.P4 || bRequireP4;
 				P4CLRequired = GlobalCommandLine.P4 || bRequireCL;
 			}
-			Log.Logger.LogDebug("P4Enabled={P4Enabled}", P4Enabled);
-			Log.Logger.LogDebug("P4CLRequired={P4CLRequired}", P4CLRequired);
+			Logger.LogDebug("P4Enabled={P4Enabled}", P4Enabled);
+			Logger.LogDebug("P4CLRequired={P4CLRequired}", P4CLRequired);
 		}
 
 		/// <summary>
@@ -1133,7 +1133,7 @@ namespace AutomationTool
 					{
 						if(!GlobalCommandLine.P4)
 						{
-							Log.Logger.LogInformation("Command {CommandName} requires P4 functionality.", Command.Name);
+							Logger.LogInformation("Command {CommandName} requires P4 functionality.", Command.Name);
 						}
 						bRequireP4 = true;
 
diff --git a/Engine/Source/Programs/AutomationTool/AutomationUtils/P4WriteConfig.cs b/Engine/Source/Programs/AutomationTool/AutomationUtils/P4WriteConfig.cs
--- a/Engine/Source/Programs/AutomationTool/AutomationUtils/P4WriteConfig.cs
+++ b/Engine/Source/Programs/AutomationTool/AutomationUtils/P4WriteConfig.cs
@@ -20,7 +20,7 @@ namespace AutomationTool
 	{
 		public override ExitCode Execute()
 		{
-			Log.Logger.LogInformation("Setting up Perforce environment.");
+			Logger.LogInformation("Setting up Perforce environment.");
 
 			// User can specify these to help auto detection
 			string Port = ParseParamValue("p4port", "");
@@ -52,7 +52,7 @@ namespace AutomationTool
 			{
 				Log.TraceError(
 					"Unable to find matching Perforce info. If the below does not help try P4WriteConfig -p4port=<server:port> and -p4user=<username> to supply more info");
-				Log.Logger.LogError(Ex, "{Message}", Ex.Message);
+				Logger.LogError(Ex, "{Message}", Ex.Message);
 				return ExitCode.Error_Arguments;
 			}
 
@@ -96,14 +96,14 @@ namespace AutomationTool
 
 				string OutputFile = Path.GetFileName(OutputPath);
 
-				Log.Logger.LogInformation("Wrote P4 settings to {OutputPath}", OutputPath);
+				Logger.LogInformation("Wrote P4 settings to {OutputPath}", OutputPath);
 
 				P4.P4(string.Format("set P4CONFIG={0}", OutputFile));
-				Log.Logger.LogInformation("set P4CONFIG={OutputFile}", OutputFile);
+				Logger.LogInformation("set P4CONFIG={OutputFile}", OutputFile);
 			}
 			else
 			{
-				Log.Logger.LogInformation("Skipped write");
+				Logger.LogInformation("Skipped write");
 			}
 
 			return ExitCode.Success;
diff --git a/Engine/Source/Programs/AutomationTool/AutomationUtils/Platform.cs b/Engine/Source/Programs/AutomationTool/AutomationUtils/Platform.cs
--- a/Engine/Source/Programs/AutomationTool/AutomationUtils/Platform.cs
+++ b/Engine/Source/Programs/AutomationTool/AutomationUtils/Platform.cs
@@ -108,7 +108,7 @@ namespace AutomationTool
 		private static Dictionary<TargetPlatformDescriptor, Platform> AllPlatforms = new Dictionary<TargetPlatformDescriptor, Platform>();
 		internal static void InitializePlatforms(HashSet<Assembly> AssembliesWithPlatforms)
 		{
-			Log.Logger.LogDebug("Creating platforms.");
+			Logger.LogDebug("Creating platforms.");
 
 			// Create all available platforms.
 			foreach (var ScriptAssembly in AssembliesWithPlatforms)
@@ -122,7 +122,7 @@ namespace AutomationTool
 				Platform ExistingInstance;
 				if (AllPlatforms.TryGetValue(TargetDesc, out ExistingInstance) == false)
 				{
-					Log.Logger.LogDebug("Creating placeholder platform for target: {TargetType}", TargetDesc.Type);
+					Logger.LogDebug("Creating placeholder platform for target: {TargetType}", TargetDesc.Type);
 					AllPlatforms.Add(TargetDesc, new Platform(TargetDesc.Type));
 				}
 			}
@@ -130,7 +130,7 @@ namespace AutomationTool
 
 		private static void CreatePlatformsFromAssembly(Assembly ScriptAssembly)
 		{
-			Log.Logger.LogDebug("Looking for platforms in {Location}", ScriptAssembly.Location);
+			Logger.LogDebug("Looking for platforms in {Location}", ScriptAssembly.Location);
 			Type[] AllTypes = null;
 			try
 			{
@@ -138,21 +138,21 @@ namespace AutomationTool
 			}
 			catch (Exception Ex)
 			{
-				Log.Logger.LogError("Failed to get assembly types for {Location}", ScriptAssembly.Location);
+				Logger.LogError("Failed to get assembly types for {Location}", ScriptAssembly.Location);
 				if (Ex is ReflectionTypeLoadException)
 				{
 					var TypeLoadException = (ReflectionTypeLoadException)Ex;
 					if (!IsNullOrEmpty(TypeLoadException.LoaderExceptions))
 					{
-						Log.Logger.LogError("Loader Exceptions:");
+						Logger.LogError("Loader Exceptions:");
 						foreach (var LoaderException in TypeLoadException.LoaderExceptions)
 						{
-							Log.Logger.LogError(LoaderException, "{Text}", LogUtils.FormatException(LoaderException));
+							Logger.LogError(LoaderException, "{Text}", LogUtils.FormatException(LoaderException));
 						}
 					}
 					else
 					{
-						Log.Logger.LogError("No Loader Exceptions available.");
+						Logger.LogError("No Loader Exceptions available.");
 					}
 				}
 				// Re-throw, this is still a critical error!
@@ -162,7 +162,7 @@ namespace AutomationTool
 			{
 				if (PotentialPlatformType != typeof(Platform) && typeof(Platform).IsAssignableFrom(PotentialPlatformType) && !PotentialPlatformType.IsAbstract)
 				{
-					Log.Logger.LogDebug("Creating platform {Platform} from {Location}.", PotentialPlatformType.Name, ScriptAssembly.Location);
+					Logger.LogDebug("Creating platform {Platform} from {Location}.", PotentialPlatformType.Name, ScriptAssembly.Location);
 					var PlatformInstance = Activator.CreateInstance(PotentialPlatformType) as Platform;
 					var PlatformDesc = PlatformInstance.GetTargetPlatformDescriptor();
 
@@ -175,7 +175,7 @@ namespace AutomationTool
 					{
 						if (ExistingInstance.GetType() != PlatformInstance.GetType())
 						{
-							Log.Logger.LogWarning("Platform {Platform} already exists", PotentialPlatformType.Name);
+							Logger.LogWarning("Platform {Platform} already exists", PotentialPlatformType.Name);
 						}
 					}
 				}
@@ -319,7 +319,7 @@ namespace AutomationTool
 
 		public virtual bool SetDeviceAutoSoftwareUpdateMode(DeviceInfo Device, bool bEnableAutoSoftwareUpdates)
 		{
-			LogWarning("{0} does not implement SetDeviceAutoSoftwareUpdateMode", PlatformType);
+			Logger.LogWarning("{PlatformType} does not implement SetDeviceAutoSoftwareUpdateMode", PlatformType);
 			return false;
 		}
 
diff --git a/Engine/Source/Programs/AutomationTool/AutomationUtils/ProcessUtils.cs b/Engine/Source/Programs/AutomationTool/AutomationUtils/ProcessUtils.cs
--- a/Engine/Source/Programs/AutomationTool/AutomationUtils/ProcessUtils.cs
+++ b/Engine/Source/Programs/AutomationTool/AutomationUtils/ProcessUtils.cs
@@ -820,11 +820,11 @@ namespace AutomationTool
 
 			if (ResolvedPath != null)
 			{
-				Log.Logger.LogTrace("Resolved {App} to {ResolvedPath}", App, ResolvedPath);
+				Logger.LogTrace("Resolved {App} to {ResolvedPath}", App, ResolvedPath);
 			}
 			else
 			{
-				Log.Logger.LogDebug("Could not resolve app {App}", App);
+				Logger.LogDebug("Could not resolve app {App}", App);
 			}
 
 			return ResolvedPath;
@@ -1134,7 +1134,7 @@ namespace AutomationTool
 		{
 			while (!FileExists(LogFilename) && !LogProcess.HasExited)
 			{
-				Log.Logger.LogInformation("Waiting for logging process to start...");
+				Logger.LogInformation("Waiting for logging process to start...");
 				Thread.Sleep(2000);
 			}
 			Thread.Sleep(1000);
diff --git a/Engine/Source/Programs/AutomationTool/AutomationUtils/ProjectParams.cs b/Engine/Source/Programs/AutomationTool/AutomationUtils/ProjectParams.cs
--- a/Engine/Source/Programs/AutomationTool/AutomationUtils/ProjectParams.cs
+++ b/Engine/Source/Programs/AutomationTool/AutomationUtils/ProjectParams.cs
@@ -34,6 +34,8 @@ namespace AutomationTool
 	[Help("servertargetplatform=PlatformName", "target platform for building, cooking and deployment of the dedicated server (also -ServerPlatform)")]
 	public class ProjectParams
 	{
+		static ILogger Logger => Log.Logger;
+
 		/// <summary>
 		/// Gets a parameter from the command line if it hasn't been specified in the constructor. 
 		/// If the command line is not available, default value will be used.
@@ -130,12 +132,12 @@ namespace AutomationTool
 			{
 				if (SpecifiedValue == null)
 				{
-					Log.Logger.LogWarning("Value was provided for \"{ParamName}\" using obsolete name \"{ObsoleteParamName}\"", ParamName, ObsoleteParamName);
+					Logger.LogWarning("Value was provided for \"{ParamName}\" using obsolete name \"{ObsoleteParamName}\"", ParamName, ObsoleteParamName);
 					Result = SpecifiedValue;
 				}
 				else
 				{
-					Log.Logger.LogWarning("Value provided for obsolete name \"{ObsoleteParamName}\" will be ignored as \"{ParamName}\" was provided", ObsoleteParamName, ParamName);
+					Logger.LogWarning("Value provided for obsolete name \"{ObsoleteParamName}\" will be ignored as \"{ParamName}\" was provided", ObsoleteParamName, ParamName);
 				}
 			}
 			else if (SpecifiedValue != null)
@@ -154,14 +156,14 @@ namespace AutomationTool
 						// Didn't find the new name on the command line. If the obsolete name was found, use it, and warn.
 						if (ParsedObsolete != null)
 						{
-							Log.Logger.LogWarning("Obsolete argument \"{ObsoleteParamName}\" on command line - use \"{ParamName}\" instead", ObsoleteParamName, ParamName);
+							Logger.LogWarning("Obsolete argument \"{ObsoleteParamName}\" on command line - use \"{ParamName}\" instead", ObsoleteParamName, ParamName);
 							Parsed = ParsedObsolete;
 						}
 					}
 					else if (ParsedObsolete != null)
 					{
 						// Did find the new name on the command line - check for the obsolete name. If found, do not use it, and warn.
-						Log.Logger.LogWarning("Obsolete argument \"{ObsoleteParamName}\" will be ignored as \"{ParamName}\" was provided", ObsoleteParamName, ParamName);
+						Logger.LogWarning("Obsolete argument \"{ObsoleteParamName}\" will be ignored as \"{ParamName}\" was provided", ObsoleteParamName, ParamName);
 					}
 				}
 
