Change to structured logging in other AutomationTool projects.

diff --git a/Engine/Source/Programs/AutomationTool/AutomationUtils/ProjectUtils.cs b/Engine/Source/Programs/AutomationTool/AutomationUtils/ProjectUtils.cs
--- a/Engine/Source/Programs/AutomationTool/AutomationUtils/ProjectUtils.cs
+++ b/Engine/Source/Programs/AutomationTool/AutomationUtils/ProjectUtils.cs
@@ -130,7 +130,7 @@ namespace AutomationTool
 			}
 		}
 
-
+		private static ILogger Logger => Log.Logger;
 		private static Dictionary<PropertyCacheKey, ProjectProperties> PropertiesCache = new Dictionary<PropertyCacheKey, ProjectProperties>();
 
 		/// <summary>
@@ -241,7 +241,7 @@ namespace AutomationTool
 					string Reason;
 					if(RequiresTempTarget(RawProjectPath, bHasCode, Platform, Configuration, TargetType.Game, AssetNativizationRequested, true, out Reason))
 					{
-						Log.Logger.LogInformation("{Project} requires a temporary target.cs to be generated ({Reason})", RawProjectPath.GetFileName(), Reason);
+						Logger.LogInformation("{Project} requires a temporary target.cs to be generated ({Reason})", RawProjectPath.GetFileName(), Reason);
 						return true;
 					}
 				}
@@ -814,7 +814,7 @@ namespace AutomationTool
 
 				if (!CommandUtils.IsBuildMachine && !CheckIfScriptAssemblyIsOutOfDate(TargetsDllFilename, TargetScripts))
 				{
-					Log.Logger.LogDebug("Targets DLL {Filename} is up to date.", TargetsDllFilename);
+					Logger.LogDebug("Targets DLL {Filename} is up to date.", TargetsDllFilename);
 					DoNotCompile = true;
 				}
 				if (!DoNotCompile && CommandUtils.FileExists_NoExceptions(TargetsDllFilename.FullName))
diff --git a/Engine/Source/Programs/AutomationTool/AutomationUtils/ScriptManager.cs b/Engine/Source/Programs/AutomationTool/AutomationUtils/ScriptManager.cs
--- a/Engine/Source/Programs/AutomationTool/AutomationUtils/ScriptManager.cs
+++ b/Engine/Source/Programs/AutomationTool/AutomationUtils/ScriptManager.cs
@@ -25,7 +25,7 @@ namespace AutomationTool
 		/// <summary>
 		/// Populates ScriptCommands
 		/// </summary>
-		static void EnumerateScriptCommands()
+		static void EnumerateScriptCommands(ILogger Logger)
 		{
 			ScriptCommands = new Dictionary<string, Type>(StringComparer.InvariantCultureIgnoreCase);
 			foreach (Assembly CompiledScripts in AllScriptAssemblies)
@@ -46,7 +46,7 @@ namespace AutomationTool
 
 								if (IsSame == false)
 								{
-									Log.TraceWarning("Unable to add command {0} twice. Previous: {1}, Current: {2}", ClassType.Name,
+									Logger.LogWarning("Unable to add command {ClassName} twice. Previous: {OldQualifiedName}, Current: {NewQualifiedName}", ClassType.Name,
 										ClassType.AssemblyQualifiedName, ScriptCommands[ClassType.Name].AssemblyQualifiedName);
 								}
 							}
@@ -57,7 +57,7 @@ namespace AutomationTool
 				{
 					foreach (Exception SubEx in LoadEx.LoaderExceptions)
 					{
-						Log.Logger.LogWarning(SubEx, "Got type loader exception: {SubEx}", SubEx.ToString());
+						Logger.LogWarning(SubEx, "Got type loader exception: {SubEx}", SubEx.ToString());
 					}
 					throw new AutomationException("Failed to add commands from {0}. {1}", CompiledScripts, LoadEx);
 				}
@@ -103,9 +103,10 @@ namespace AutomationTool
 		/// <summary>
 		/// Loads all precompiled assemblies (DLLs that end with *Scripts.dll).
 		/// </summary>
-		/// <param name="Projects">Projects to load</param>
+		/// <param name="AssemblyPaths"></param>
+		/// <param name="Logger"></param>
 		/// <returns>List of compiled assemblies</returns>
-		public static void LoadScriptAssemblies(IEnumerable<FileReference> AssemblyPaths)
+		public static void LoadScriptAssemblies(IEnumerable<FileReference> AssemblyPaths, ILogger Logger)
 		{
 			foreach (FileReference AssemblyLocation in AssemblyPaths)
 			{
@@ -127,7 +128,7 @@ namespace AutomationTool
 
 			Platform.InitializePlatforms(AllScriptAssemblies);
 
-			EnumerateScriptCommands();
+			EnumerateScriptCommands(Logger);
 
 			EnumerateBuildProducts();
 		}
diff --git a/Engine/Source/Programs/AutomationTool/AutomationUtils/Utils.cs b/Engine/Source/Programs/AutomationTool/AutomationUtils/Utils.cs
--- a/Engine/Source/Programs/AutomationTool/AutomationUtils/Utils.cs
+++ b/Engine/Source/Programs/AutomationTool/AutomationUtils/Utils.cs
@@ -20,6 +20,8 @@ namespace AutomationTool
 	/// </summary>
 	public static class InternalUtils
 	{
+		static ILogger Logger => Log.Logger;
+
 		/// <summary>
 		/// Gets environment variable value.
 		/// </summary>
@@ -86,8 +88,8 @@ namespace AutomationTool
 				}
 				else
 				{
-					Log.Logger.LogWarning("Failed to create directory {Path} in {MaxAttempts} attempts.", Path, MaxAttempts);
-					Log.Logger.LogWarning("{Text}", LogUtils.FormatException(LastException));
+					Logger.LogWarning("Failed to create directory {Path} in {MaxAttempts} attempts.", Path, MaxAttempts);
+					Logger.LogWarning("{Text}", LogUtils.FormatException(LastException));
 				}
 			}
 			return Result;
@@ -148,8 +150,8 @@ namespace AutomationTool
 				}
 				else
 				{
-					Log.Logger.LogWarning("Failed to delete file {Path} in {MaxAttempts} attempts.", Path, MaxAttempts);
-					Log.Logger.LogWarning("{Text}", LogUtils.FormatException(LastException));
+					Logger.LogWarning("Failed to delete file {Path} in {MaxAttempts} attempts.", Path, MaxAttempts);
+					Logger.LogWarning("{Text}", LogUtils.FormatException(LastException));
 				}
 			}
 
@@ -231,8 +233,8 @@ namespace AutomationTool
 
 			if (Result == false && LastException != null)
 			{
-				Log.Logger.LogWarning("Failed to delete directory {Path} in {MaxAttempts} attempts.", Path, MaxAttempts);
-				Log.Logger.LogWarning("{Text}", LogUtils.FormatException(LastException));
+				Logger.LogWarning("Failed to delete directory {Path} in {MaxAttempts} attempts.", Path, MaxAttempts);
+				Logger.LogWarning("{Text}", LogUtils.FormatException(LastException));
 			}
 
 			return Result;
@@ -290,8 +292,8 @@ namespace AutomationTool
 					{
 						if (!bQuiet)
 						{
-							Log.Logger.LogWarning("Failed to rename {OldName} to {NewName}", OldName, NewName);
-							Log.Logger.LogWarning("{Text}", LogUtils.FormatException(Ex));
+							Logger.LogWarning("Failed to rename {OldName} to {NewName}", OldName, NewName);
+							Logger.LogWarning("{Text}", LogUtils.FormatException(Ex));
 						}
 						Result = false;
 					}
@@ -348,8 +350,8 @@ namespace AutomationTool
 				{
 					if (File.Exists(OldName) == true || File.Exists(NewName) == false)
 					{
-						Log.Logger.LogWarning("Failed to rename {OldName} to {NewName}", OldName, NewName);
-						Log.Logger.LogWarning("{Text}", LogUtils.FormatException(Ex));
+						Logger.LogWarning("Failed to rename {OldName} to {NewName}", OldName, NewName);
+						Logger.LogWarning("{Text}", LogUtils.FormatException(Ex));
 						Result = false;
 					}
 				}
@@ -471,7 +473,7 @@ namespace AutomationTool
 						}
 						else
 						{
-							Log.Logger.LogInformation("Skip copying file {SourceName} because it doesn't exist.", SourceName);
+							Logger.LogInformation("Skip copying file {SourceName} because it doesn't exist.", SourceName);
 						}
 					}
 					Retry = !File.Exists(TargetName);
@@ -481,7 +483,7 @@ namespace AutomationTool
 						FileInfo TargetInfo = new FileInfo(TargetName);
 						if (!bSkipSizeCheck && SourceInfo.Length != TargetInfo.Length)
 						{
-							Log.Logger.LogInformation("Size mismatch {SourceName} = {SourceLength} to {TargetName} = {TargetLength}", SourceName, SourceInfo.Length, TargetName, TargetInfo.Length);
+							Logger.LogInformation("Size mismatch {SourceName} = {SourceLength} to {TargetName} = {TargetLength}", SourceName, SourceInfo.Length, TargetName, TargetInfo.Length);
 							Retry = true;
 						}
 						// Timestamps should be no more than 2 seconds out - assuming this as exFAT filesystems store timestamps at 2 second intervals:
@@ -494,7 +496,7 @@ namespace AutomationTool
 							TargetInfo.Refresh();
 							if (!((SourceInfo.LastWriteTimeUtc - TargetInfo.LastWriteTimeUtc).TotalSeconds < 2 && (SourceInfo.LastWriteTimeUtc - TargetInfo.LastWriteTimeUtc).TotalSeconds > -2))
 							{
-								Log.Logger.LogInformation("Date mismatch {SourceName} = {SourceTime} to {TargetName} = {TargetTime}", SourceName, SourceInfo.LastWriteTimeUtc, TargetName, TargetInfo.LastWriteTimeUtc);
+								Logger.LogInformation("Date mismatch {SourceName} = {SourceTime} to {TargetName} = {TargetTime}", SourceName, SourceInfo.LastWriteTimeUtc, TargetName, TargetInfo.LastWriteTimeUtc);
 								Retry = true;
 							}
 						}
@@ -502,7 +504,7 @@ namespace AutomationTool
 				}
 				catch (Exception Ex)
 				{
-					Log.Logger.LogInformation("SafeCopyFile Exception was {Ex}", LogUtils.FormatException(Ex));
+					Logger.LogInformation("SafeCopyFile Exception was {Ex}", LogUtils.FormatException(Ex));
 					Retry = true;
 				}
 
@@ -510,7 +512,7 @@ namespace AutomationTool
 				{
 					if (Attempts + 1 < MaxAttempts)
 					{
-						Log.Logger.LogInformation("Failed to copy {SourceName} to {TargetName}, deleting, waiting 10s and retrying.", SourceName, TargetName);
+						Logger.LogInformation("Failed to copy {SourceName} to {TargetName}, deleting, waiting 10s and retrying.", SourceName, TargetName);
 						if (File.Exists(TargetName))
 						{
 							SafeDeleteFile(TargetName);
@@ -519,7 +521,7 @@ namespace AutomationTool
 					}
 					else
 					{
-						Log.Logger.LogError("Failed to copy {SourceName} to {TargetName}", SourceName, TargetName);
+						Logger.LogError("Failed to copy {SourceName} to {TargetName}", SourceName, TargetName);
 					}
 					Result = false;
 				}
@@ -545,8 +547,8 @@ namespace AutomationTool
 			}
 			catch (Exception Ex)
 			{
-				Log.Logger.LogWarning("Failed to load {Filename}", Filename);
-				Log.Logger.LogWarning(Ex, "{Text}", LogUtils.FormatException(Ex));
+				Logger.LogWarning("Failed to load {Filename}", Filename);
+				Logger.LogWarning(Ex, "{Text}", LogUtils.FormatException(Ex));
 			}
 			return Result;
 		}
@@ -566,8 +568,8 @@ namespace AutomationTool
 			}
 			catch (Exception Ex)
 			{
-				Log.Logger.LogWarning("Failed to load {Filename}", Filename);
-				Log.Logger.LogWarning(Ex, "{Text}", LogUtils.FormatException(Ex));
+				Logger.LogWarning("Failed to load {Filename}", Filename);
+				Logger.LogWarning(Ex, "{Text}", LogUtils.FormatException(Ex));
 			}
 			return Result;
 		}
@@ -605,7 +607,7 @@ namespace AutomationTool
 					{
 						if (!bQuiet)
 						{
-							Log.Logger.LogWarning("Ignoring symlink {Path}", File.FullName);
+							Logger.LogWarning("Ignoring symlink {Path}", File.FullName);
 						}
 						continue;
 					}
@@ -653,8 +655,8 @@ namespace AutomationTool
 			}
 			catch (Exception Ex)
 			{
-				Log.Logger.LogWarning("Unable to Find Files in {Path}", Path);
-				Log.Logger.LogWarning(Ex, "{Text}", LogUtils.FormatException(Ex));
+				Logger.LogWarning("Unable to Find Files in {Path}", Path);
+				Logger.LogWarning(Ex, "{Text}", LogUtils.FormatException(Ex));
 			}
 			return Files;
 		}
@@ -679,8 +681,8 @@ namespace AutomationTool
 			}
 			catch (Exception Ex)
 			{
-				Log.Logger.LogWarning("Unable to Find Directories in {Path}", Path);
-				Log.Logger.LogWarning(Ex, "{Text}", LogUtils.FormatException(Ex));
+				Logger.LogWarning("Unable to Find Directories in {Path}", Path);
+				Logger.LogWarning(Ex, "{Text}", LogUtils.FormatException(Ex));
 			}
 			return Directories;
 		}
@@ -704,8 +706,8 @@ namespace AutomationTool
 			}
 			catch (Exception Ex)
 			{
-				Log.Logger.LogWarning("Unable to check if file {Path} exists.", Path);
-				Log.Logger.LogWarning(Ex, "{Text}", LogUtils.FormatException(Ex));
+				Logger.LogWarning("Unable to check if file {Path} exists.", Path);
+				Logger.LogWarning(Ex, "{Text}", LogUtils.FormatException(Ex));
 			}
 			return Result;
 		}
@@ -730,8 +732,8 @@ namespace AutomationTool
 			}
 			catch (Exception Ex)
 			{
-				Log.Logger.LogWarning("Unable to check if directory {Path} exists.", Path);
-				Log.Logger.LogWarning(Ex, "{Text}", LogUtils.FormatException(Ex));
+				Logger.LogWarning("Unable to check if directory {Path} exists.", Path);
+				Logger.LogWarning(Ex, "{Text}", LogUtils.FormatException(Ex));
 			}
 			return Result;
 		}
@@ -753,8 +755,8 @@ namespace AutomationTool
 			}
 			catch (Exception Ex)
 			{
-				Log.Logger.LogWarning("Unable to write text to {Path}", Path);
-				Log.Logger.LogWarning(Ex, "{Text}", LogUtils.FormatException(Ex));
+				Logger.LogWarning("Unable to write text to {Path}", Path);
+				Logger.LogWarning(Ex, "{Text}", LogUtils.FormatException(Ex));
 			}
 			return Result;
 		}
@@ -776,8 +778,8 @@ namespace AutomationTool
 			}
 			catch (Exception Ex)
 			{
-				Log.Logger.LogWarning("Unable to write text to {Path}", Path);
-				Log.Logger.LogWarning(Ex, "{Text}", LogUtils.FormatException(Ex));
+				Logger.LogWarning("Unable to write text to {Path}", Path);
+				Logger.LogWarning(Ex, "{Text}", LogUtils.FormatException(Ex));
 			}
 			return Result;
 		}
@@ -799,8 +801,8 @@ namespace AutomationTool
 			}
 			catch (Exception Ex)
 			{
-				Log.Logger.LogWarning("Unable to write text to {Path}", Path);
-				Log.Logger.LogWarning(Ex, "{Text}", LogUtils.FormatException(Ex));
+				Logger.LogWarning("Unable to write text to {Path}", Path);
+				Logger.LogWarning(Ex, "{Text}", LogUtils.FormatException(Ex));
 			}
 			return Result;
 		}
diff --git a/Engine/Source/Programs/AutomationTool/BuildGraph/BgGraphBuilder.cs b/Engine/Source/Programs/AutomationTool/BuildGraph/BgGraphBuilder.cs
--- a/Engine/Source/Programs/AutomationTool/BuildGraph/BgGraphBuilder.cs
+++ b/Engine/Source/Programs/AutomationTool/BuildGraph/BgGraphBuilder.cs
@@ -2,6 +2,8 @@
 
 using EpicGames.BuildGraph;
 using EpicGames.BuildGraph.Expressions;
+using EpicGames.Core;
+using Microsoft.Extensions.Logging;
 using System;
 using System.Collections.Generic;
 using System.Text;
@@ -13,6 +15,11 @@ namespace AutomationTool
 	/// </summary>
 	public abstract class BgGraphBuilder
 	{
+		/// <summary>
+		/// Accessor for default logger instance
+		/// </summary>
+		protected static ILogger Logger => Log.Logger;
+
 		/// <summary>
 		/// Callback used to instantiate the graph
 		/// </summary>
diff --git a/Engine/Source/Programs/AutomationTool/BuildGraph/BuildGraph.cs b/Engine/Source/Programs/AutomationTool/BuildGraph/BuildGraph.cs
--- a/Engine/Source/Programs/AutomationTool/BuildGraph/BuildGraph.cs
+++ b/Engine/Source/Programs/AutomationTool/BuildGraph/BuildGraph.cs
@@ -741,7 +741,7 @@ namespace AutomationTool
 				}
 				catch (ReflectionTypeLoadException ex)
 				{
-					Log.Logger.LogWarning("Exception {ex} while trying to get types from assembly {LoadedAssembly}. LoaderExceptions: {Arg2}", ex, LoadedAssembly, string.Join("\n", ex.LoaderExceptions.Select(x => x.Message)));
+					Logger.LogWarning("Exception {ex} while trying to get types from assembly {LoadedAssembly}. LoaderExceptions: {Arg2}", ex, LoadedAssembly, string.Join("\n", ex.LoaderExceptions.Select(x => x.Message)));
 					continue;
 				}
 
@@ -777,7 +777,7 @@ namespace AutomationTool
 				}
 				catch (ReflectionTypeLoadException ex)
 				{
-					Log.Logger.LogWarning("Exception {ex} while trying to get types from assembly {LoadedAssembly}. LoaderExceptions: {Arg2}", ex, LoadedAssembly, string.Join("\n", ex.LoaderExceptions.Select(x => x.Message)));
+					Logger.LogWarning("Exception {ex} while trying to get types from assembly {LoadedAssembly}. LoaderExceptions: {Arg2}", ex, LoadedAssembly, string.Join("\n", ex.LoaderExceptions.Select(x => x.Message)));
 					continue;
 				}
 
@@ -1332,7 +1332,7 @@ namespace AutomationTool
 		/// <param name="OutputFile">The output file to write to</param>
 		static void WriteDocumentationHTML(Dictionary<string, ScriptTaskBinding> NameToTask, Dictionary<string, XmlElement> MemberNameToElement, FileReference OutputFile)
 		{
-			Log.Logger.LogInformation("Writing {OutputFile}...", OutputFile);
+			Logger.LogInformation("Writing {OutputFile}...", OutputFile);
 			using (StreamWriter Writer = new StreamWriter(OutputFile.FullName))
 			{
 				Writer.WriteLine("<html>");
diff --git a/Engine/Source/Programs/AutomationTool/BuildGraph/Scripts/InstalledBuild.cs b/Engine/Source/Programs/AutomationTool/BuildGraph/Scripts/InstalledBuild.cs
--- a/Engine/Source/Programs/AutomationTool/BuildGraph/Scripts/InstalledBuild.cs
+++ b/Engine/Source/Programs/AutomationTool/BuildGraph/Scripts/InstalledBuild.cs
@@ -248,7 +248,7 @@ namespace AutomationTool
 			if (State.Get(EmbedSrcSrvInfo))
 			{
 				// Embed source info into the PDB files. Should be done from this machine to ensure that paths are correct.
-				Log.Logger.LogInformation("Embedding source file information into PDB files...");
+				Logger.LogInformation("Embedding source file information into PDB files...");
 				FileSet SourceFiles = Workspace.Filter("Engine/Source/...;Engine/Plugins/...").Except("Engine/Source/ThirdParty/...").Filter("*.c;*.h;*.cpp;*.hpp;*.inl");
 //				State.SrcSrv(BinaryFiles: Full, SourceFiles: SourceFiles);
 			}
diff --git a/Engine/Source/Programs/AutomationTool/BuildGraph/Task.cs b/Engine/Source/Programs/AutomationTool/BuildGraph/Task.cs
--- a/Engine/Source/Programs/AutomationTool/BuildGraph/Task.cs
+++ b/Engine/Source/Programs/AutomationTool/BuildGraph/Task.cs
@@ -13,6 +13,7 @@ using EpicGames.Core;
 using OpenTracing;
 using UnrealBuildBase;
 using UnrealBuildTool;
+using Microsoft.Extensions.Logging;
 
 namespace AutomationTool
 {
@@ -97,6 +98,11 @@ namespace AutomationTool
 	/// </summary>
 	public abstract class BgTaskImpl
 	{
+		/// <summary>
+		/// Accessor for the default log interface
+		/// </summary>
+		protected static ILogger Logger => Log.Logger;
+
 		/// <summary>
 		/// Line number in a source file that this task was declared. Optional; used for log messages.
 		/// </summary>
diff --git a/Engine/Source/Programs/AutomationTool/BuildGraph/Tasks/AwsEcsDeployTask.cs b/Engine/Source/Programs/AutomationTool/BuildGraph/Tasks/AwsEcsDeployTask.cs
--- a/Engine/Source/Programs/AutomationTool/BuildGraph/Tasks/AwsEcsDeployTask.cs
+++ b/Engine/Source/Programs/AutomationTool/BuildGraph/Tasks/AwsEcsDeployTask.cs
@@ -116,7 +116,7 @@ namespace AutomationTool.Tasks
 			string Params = $"ecs update-service --cluster {Parameters.Cluster} --service {Parameters.Service} --task-definition {TaskDefFamily}:{TaskDefRevision}";
 			await SpawnTaskBase.ExecuteAsync("aws", Params, EnvVars: ParseEnvVars(Parameters.Environment, Parameters.EnvironmentFile), LogOutput: Parameters.LogOutput);
 
-			Log.Logger.LogInformation("Service {Service} updated to use new task def {TaskDefFamily}:{TaskDefRevision}", Parameters.Service, TaskDefFamily, TaskDefRevision);
+			Logger.LogInformation("Service {Service} updated to use new task def {TaskDefFamily}:{TaskDefRevision}", Parameters.Service, TaskDefFamily, TaskDefRevision);
 		}
 
 		/// <summary>
diff --git a/Engine/Source/Programs/AutomationTool/BuildGraph/Tasks/CommandTask.cs b/Engine/Source/Programs/AutomationTool/BuildGraph/Tasks/CommandTask.cs
--- a/Engine/Source/Programs/AutomationTool/BuildGraph/Tasks/CommandTask.cs
+++ b/Engine/Source/Programs/AutomationTool/BuildGraph/Tasks/CommandTask.cs
@@ -126,7 +126,7 @@ namespace AutomationTool.Tasks
 				}
 				else
 				{
-					Log.Logger.LogWarning("Unable to read UAT telemetry file from {TelemetryFile}", TelemetryFile);
+					Logger.LogWarning("Unable to read UAT telemetry file from {TelemetryFile}", TelemetryFile);
 				}
 			}
 			return Task.CompletedTask;
diff --git a/Engine/Source/Programs/AutomationTool/BuildGraph/Tasks/CopyTask.cs b/Engine/Source/Programs/AutomationTool/BuildGraph/Tasks/CopyTask.cs
--- a/Engine/Source/Programs/AutomationTool/BuildGraph/Tasks/CopyTask.cs
+++ b/Engine/Source/Programs/AutomationTool/BuildGraph/Tasks/CopyTask.cs
@@ -172,11 +172,11 @@ namespace AutomationTool.Tasks
 						DirectoryReference.CreateDirectory(FirstTargetDirectory);
 						if(NumRetries == 1)
 						{
-							Log.Logger.LogInformation("Created target directory {FirstTargetDirectory} after 1 retry.", FirstTargetDirectory);
+							Logger.LogInformation("Created target directory {FirstTargetDirectory} after 1 retry.", FirstTargetDirectory);
 						}
 						else if(NumRetries > 1)
 						{
-							Log.Logger.LogInformation("Created target directory {FirstTargetDirectory} after {NumRetries} retries.", FirstTargetDirectory, NumRetries);
+							Logger.LogInformation("Created target directory {FirstTargetDirectory} after {NumRetries} retries.", FirstTargetDirectory, NumRetries);
 						}
 						break;
 					}
@@ -184,7 +184,7 @@ namespace AutomationTool.Tasks
 					{
 						if(NumRetries == 0)
 						{
-							Log.Logger.LogInformation("Unable to create directory '{FirstTargetDirectory}' on first attempt. Retrying {MaxNumRetries} times...", FirstTargetDirectory, MaxNumRetries);
+							Logger.LogInformation("Unable to create directory '{FirstTargetDirectory}' on first attempt. Retrying {MaxNumRetries} times...", FirstTargetDirectory, MaxNumRetries);
 						}
 
 						Log.TraceLog("  {0}", Ex);
diff --git a/Engine/Source/Programs/AutomationTool/BuildGraph/Tasks/CsCompileTask.cs b/Engine/Source/Programs/AutomationTool/BuildGraph/Tasks/CsCompileTask.cs
--- a/Engine/Source/Programs/AutomationTool/BuildGraph/Tasks/CsCompileTask.cs
+++ b/Engine/Source/Programs/AutomationTool/BuildGraph/Tasks/CsCompileTask.cs
@@ -142,7 +142,7 @@ namespace AutomationTool.Tasks
 						int EqualsIdx = Property.IndexOf('=');
 						if (EqualsIdx == -1)
 						{
-							Log.Logger.LogWarning("Missing '=' in property assignment");
+							Logger.LogWarning("Missing '=' in property assignment");
 						}
 						else
 						{
diff --git a/Engine/Source/Programs/AutomationTool/BuildGraph/Tasks/DeleteTask.cs b/Engine/Source/Programs/AutomationTool/BuildGraph/Tasks/DeleteTask.cs
--- a/Engine/Source/Programs/AutomationTool/BuildGraph/Tasks/DeleteTask.cs
+++ b/Engine/Source/Programs/AutomationTool/BuildGraph/Tasks/DeleteTask.cs
@@ -82,7 +82,7 @@ namespace AutomationTool.Tasks
 				{
 					if (Parameters.Verbose)
 					{
-						Log.Logger.LogInformation("Deleting {File}", File.FullName);
+						Logger.LogInformation("Deleting {File}", File.FullName);
 					}
 					if (!InternalUtils.SafeDeleteFile(File.FullName))
 					{
@@ -121,7 +121,7 @@ namespace AutomationTool.Tasks
 					{
 						if (Parameters.Verbose)
 						{
-							Log.Logger.LogInformation("Deleting {Directory}", Directory);
+							Logger.LogInformation("Deleting {Directory}", Directory);
 						}
 						DirectoryReference FullDir = new DirectoryReference(Directory);
 						if (DirectoryReference.Exists(FullDir))
diff --git a/Engine/Source/Programs/AutomationTool/BuildGraph/Tasks/DeployToolTask.cs b/Engine/Source/Programs/AutomationTool/BuildGraph/Tasks/DeployToolTask.cs
--- a/Engine/Source/Programs/AutomationTool/BuildGraph/Tasks/DeployToolTask.cs
+++ b/Engine/Source/Programs/AutomationTool/BuildGraph/Tasks/DeployToolTask.cs
@@ -142,7 +142,7 @@ namespace AutomationTool.Tasks
 						content.Add(streamContent, "file", file.GetFileName());
 						request.Content = content;
 
-						Log.Logger.LogInformation("Uploading {File} to {Url}", file, request.RequestUri);
+						Logger.LogInformation("Uploading {File} to {Url}", file, request.RequestUri);
 						using (HttpResponseMessage response = await httpClient.SendAsync(request))
 						{
 							if (!response.IsSuccessStatusCode)
diff --git a/Engine/Source/Programs/AutomationTool/BuildGraph/Tasks/DockerBuildTask.cs b/Engine/Source/Programs/AutomationTool/BuildGraph/Tasks/DockerBuildTask.cs
--- a/Engine/Source/Programs/AutomationTool/BuildGraph/Tasks/DockerBuildTask.cs
+++ b/Engine/Source/Programs/AutomationTool/BuildGraph/Tasks/DockerBuildTask.cs
@@ -115,7 +115,7 @@ namespace AutomationTool.Tasks
 		/// <param name="TagNameToFileSet">Mapping from tag names to the set of files they include</param>
 		public override async Task ExecuteAsync(JobContext Job, HashSet<FileReference> BuildProducts, Dictionary<string, HashSet<FileReference>> TagNameToFileSet)
 		{
-			Log.Logger.LogInformation("Building Docker image");
+			Logger.LogInformation("Building Docker image");
 			using (LogIndentScope Scope = new LogIndentScope("  "))
 			{
 				DirectoryReference BaseDir = ResolveDirectory(Parameters.BaseDir);
diff --git a/Engine/Source/Programs/AutomationTool/BuildGraph/Tasks/DockerComposeDownTask.cs b/Engine/Source/Programs/AutomationTool/BuildGraph/Tasks/DockerComposeDownTask.cs
--- a/Engine/Source/Programs/AutomationTool/BuildGraph/Tasks/DockerComposeDownTask.cs
+++ b/Engine/Source/Programs/AutomationTool/BuildGraph/Tasks/DockerComposeDownTask.cs
@@ -71,7 +71,7 @@ namespace AutomationTool.Tasks
 				Arguments.Append($" {Parameters.Arguments}");
 			}
 
-			Log.Logger.LogInformation("Running docker compose {Arguments}", Arguments.ToString());
+			Logger.LogInformation("Running docker compose {Arguments}", Arguments.ToString());
 			using (LogIndentScope Scope = new LogIndentScope("  "))
 			{
 				await SpawnTaskBase.ExecuteAsync("docker-compose", Arguments.ToString());
diff --git a/Engine/Source/Programs/AutomationTool/BuildGraph/Tasks/DockerComposeUpTask.cs b/Engine/Source/Programs/AutomationTool/BuildGraph/Tasks/DockerComposeUpTask.cs
--- a/Engine/Source/Programs/AutomationTool/BuildGraph/Tasks/DockerComposeUpTask.cs
+++ b/Engine/Source/Programs/AutomationTool/BuildGraph/Tasks/DockerComposeUpTask.cs
@@ -71,7 +71,7 @@ namespace AutomationTool.Tasks
 				Arguments.Append($" {Parameters.Arguments}");
 			}
 
-			Log.Logger.LogInformation("Running docker compose {Arguments}", Arguments.ToString());
+			Logger.LogInformation("Running docker compose {Arguments}", Arguments.ToString());
 			using (LogIndentScope Scope = new LogIndentScope("  "))
 			{
 				await SpawnTaskBase.ExecuteAsync("docker-compose", Arguments.ToString());
diff --git a/Engine/Source/Programs/AutomationTool/BuildGraph/Tasks/DockerPushTask.cs b/Engine/Source/Programs/AutomationTool/BuildGraph/Tasks/DockerPushTask.cs
--- a/Engine/Source/Programs/AutomationTool/BuildGraph/Tasks/DockerPushTask.cs
+++ b/Engine/Source/Programs/AutomationTool/BuildGraph/Tasks/DockerPushTask.cs
@@ -90,7 +90,7 @@ namespace AutomationTool.Tasks
 		/// <param name="TagNameToFileSet">Mapping from tag names to the set of files they include</param>
 		public override async Task ExecuteAsync(JobContext Job, HashSet<FileReference> BuildProducts, Dictionary<string, HashSet<FileReference>> TagNameToFileSet)
 		{
-			Log.Logger.LogInformation("Pushing Docker image");
+			Logger.LogInformation("Pushing Docker image");
 			using (LogIndentScope Scope = new LogIndentScope("  "))
 			{
 				string Exe = DockerTask.GetDockerExecutablePath();
diff --git a/Engine/Source/Programs/AutomationTool/BuildGraph/Tasks/FindModifiedFiles.cs b/Engine/Source/Programs/AutomationTool/BuildGraph/Tasks/FindModifiedFiles.cs
--- a/Engine/Source/Programs/AutomationTool/BuildGraph/Tasks/FindModifiedFiles.cs
+++ b/Engine/Source/Programs/AutomationTool/BuildGraph/Tasks/FindModifiedFiles.cs
@@ -118,20 +118,20 @@ namespace AutomationTool.Tasks
 				}
 				else
 				{
-					Log.Logger.LogInformation("Unable to map {DepotFile} to workspace; skipping.", File.DepotFile);
+					Logger.LogInformation("Unable to map {DepotFile} to workspace; skipping.", File.DepotFile);
 				}
 			}
 
-			Log.Logger.LogInformation("Found {NumFiles} modified files matching {Filter}", LocalFiles.Count, Filter.ToString());
+			Logger.LogInformation("Found {NumFiles} modified files matching {Filter}", LocalFiles.Count, Filter.ToString());
 			foreach (string LocalFile in LocalFiles)
 			{
-				Log.Logger.LogInformation("  {LocalFile}", LocalFile);
+				Logger.LogInformation("  {LocalFile}", LocalFile);
 			}
 
 			if (Parameters.Output != null)
 			{
 				await FileReference.WriteAllLinesAsync(Parameters.Output, LocalFiles);
-				Log.Logger.LogInformation("Written {OutputFile}", Parameters.Output);
+				Logger.LogInformation("Written {OutputFile}", Parameters.Output);
 			}
 		}
 
diff --git a/Engine/Source/Programs/AutomationTool/BuildGraph/Tasks/GitCloneTask.cs b/Engine/Source/Programs/AutomationTool/BuildGraph/Tasks/GitCloneTask.cs
--- a/Engine/Source/Programs/AutomationTool/BuildGraph/Tasks/GitCloneTask.cs
+++ b/Engine/Source/Programs/AutomationTool/BuildGraph/Tasks/GitCloneTask.cs
@@ -76,7 +76,7 @@ namespace AutomationTool.Tasks
 			}
 
 			DirectoryReference Dir = ResolveDirectory(Parameters.Dir);
-			Log.Logger.LogInformation("Cloning Git repository into {Dir}", Parameters.Dir);
+			Logger.LogInformation("Cloning Git repository into {Dir}", Parameters.Dir);
 			using (LogIndentScope Scope = new LogIndentScope("  "))
 			{
 				DirectoryReference GitDir = DirectoryReference.Combine(Dir, ".git");
diff --git a/Engine/Source/Programs/AutomationTool/BuildGraph/Tasks/RandomDataTask.cs b/Engine/Source/Programs/AutomationTool/BuildGraph/Tasks/RandomDataTask.cs
--- a/Engine/Source/Programs/AutomationTool/BuildGraph/Tasks/RandomDataTask.cs
+++ b/Engine/Source/Programs/AutomationTool/BuildGraph/Tasks/RandomDataTask.cs
@@ -96,7 +96,7 @@ namespace AutomationTool.Tasks
 				BuildProducts.Add(file);
 			}
 
-			Job.OwnerCommand.Logger.LogInformation("Created {NumFiles:n0} files of {Size:n0} bytes in {OutputDir} (Different={Different})", Parameters.Count, Parameters.Size, OutputDir, Parameters.Different);
+			Logger.LogInformation("Created {NumFiles:n0} files of {Size:n0} bytes in {OutputDir} (Different={Different})", Parameters.Count, Parameters.Size, OutputDir, Parameters.Different);
 
 			// Apply the optional output tag to them
 			foreach (string TagName in FindTagNamesFromList(Parameters.Tag))
diff --git a/Engine/Source/Programs/AutomationTool/BuildGraph/Tasks/SanitizeReceiptTask.cs b/Engine/Source/Programs/AutomationTool/BuildGraph/Tasks/SanitizeReceiptTask.cs
--- a/Engine/Source/Programs/AutomationTool/BuildGraph/Tasks/SanitizeReceiptTask.cs
+++ b/Engine/Source/Programs/AutomationTool/BuildGraph/Tasks/SanitizeReceiptTask.cs
@@ -81,7 +81,7 @@ namespace AutomationTool.Tasks
 				}
 
 				// Print the name of the file being scanned
-				Log.Logger.LogInformation("Sanitizing {TargetFile}", TargetFile);
+				Logger.LogInformation("Sanitizing {TargetFile}", TargetFile);
 				using(new LogIndentScope("  "))
 				{
 					// Read the receipt
@@ -102,7 +102,7 @@ namespace AutomationTool.Tasks
 						}
 						else
 						{
-							Log.Logger.LogInformation("Removing build product: {File}", BuildProduct.Path);
+							Logger.LogInformation("Removing build product: {File}", BuildProduct.Path);
 						}
 					}
 					Receipt.BuildProducts = NewBuildProducts;
@@ -117,7 +117,7 @@ namespace AutomationTool.Tasks
 						}
 						else
 						{
-							Log.Logger.LogInformation("Removing runtime dependency: {File}", RuntimeDependency.Path);
+							Logger.LogInformation("Removing runtime dependency: {File}", RuntimeDependency.Path);
 						}
 					}
 					Receipt.RuntimeDependencies = NewRuntimeDependencies;
diff --git a/Engine/Source/Programs/AutomationTool/BuildGraph/Tasks/SubmitTask.cs b/Engine/Source/Programs/AutomationTool/BuildGraph/Tasks/SubmitTask.cs
--- a/Engine/Source/Programs/AutomationTool/BuildGraph/Tasks/SubmitTask.cs
+++ b/Engine/Source/Programs/AutomationTool/BuildGraph/Tasks/SubmitTask.cs
@@ -112,11 +112,11 @@ namespace AutomationTool.Tasks
 			HashSet<FileReference> Files = ResolveFilespec(Unreal.RootDirectory, Parameters.Files, TagNameToFileSet);
 			if (Files.Count == 0)
 			{
-				Log.Logger.LogInformation("No files to submit.");
+				Logger.LogInformation("No files to submit.");
 			}
 			else if (!CommandUtils.AllowSubmit)
 			{
-				Log.Logger.LogWarning("Submitting to Perforce is disabled by default. Run with the -submit argument to allow.");
+				Logger.LogWarning("Submitting to Perforce is disabled by default. Run with the -submit argument to allow.");
 			}
 			else
 			{
@@ -185,7 +185,7 @@ namespace AutomationTool.Tasks
 				}
 				catch (P4Exception Ex)
 				{
-					Job.OwnerCommand.Logger.LogError(KnownLogEvents.Systemic_Perforce, "{Message}", Ex.Message);
+					Logger.LogError(KnownLogEvents.Systemic_Perforce, "{Message}", Ex.Message);
 					throw new AutomationException(Ex.ErrorCode, Ex, "{0}", Ex.Message) { OutputFormat = AutomationExceptionOutputFormat.Silent };
 				}
 			}
diff --git a/Engine/Source/Programs/AutomationTool/BuildGraph/Tasks/UploadTask.cs b/Engine/Source/Programs/AutomationTool/BuildGraph/Tasks/UploadTask.cs
--- a/Engine/Source/Programs/AutomationTool/BuildGraph/Tasks/UploadTask.cs
+++ b/Engine/Source/Programs/AutomationTool/BuildGraph/Tasks/UploadTask.cs
@@ -129,7 +129,7 @@ namespace AutomationTool.Tasks
 			}
 
 			// Create the jupiter tree
-			Log.Logger.LogInformation("Uploading {NumFiles} files to {Url}...", Files.Count, Parameters.JupiterUrl);
+			Logger.LogInformation("Uploading {NumFiles} files to {Url}...", Files.Count, Parameters.JupiterUrl);
 
 			JupiterFileTree FileTree = new JupiterFileTree(Parameters.FromDir, Parameters.LimitMemoryUsage);
 			foreach (FileReference File in Files)
diff --git a/Engine/Source/Programs/AutomationTool/BuildGraph/Tasks/ZipTask.cs b/Engine/Source/Programs/AutomationTool/BuildGraph/Tasks/ZipTask.cs
--- a/Engine/Source/Programs/AutomationTool/BuildGraph/Tasks/ZipTask.cs
+++ b/Engine/Source/Programs/AutomationTool/BuildGraph/Tasks/ZipTask.cs
@@ -84,7 +84,7 @@ namespace AutomationTool.Tasks
 			}
 
 			// Create the zip file
-			Log.Logger.LogInformation("Adding {NumFiles} files to {ZipFile}...", Files.Count, Parameters.ZipFile);
+			Logger.LogInformation("Adding {NumFiles} files to {ZipFile}...", Files.Count, Parameters.ZipFile);
 			CommandUtils.ZipFiles(Parameters.ZipFile, Parameters.FromDir, Files);
 
 			// Apply the optional tag to the produced archive
diff --git a/Engine/Source/Programs/AutomationTool/BuildGraph/TempStorage.cs b/Engine/Source/Programs/AutomationTool/BuildGraph/TempStorage.cs
--- a/Engine/Source/Programs/AutomationTool/BuildGraph/TempStorage.cs
+++ b/Engine/Source/Programs/AutomationTool/BuildGraph/TempStorage.cs
@@ -1529,13 +1529,13 @@ namespace AutomationTool
 						}
 						catch (Exception Ex)
 						{
-							Log.Logger.LogError(Ex, "Exception while trying to scan files under {BuildDirectory}: {Ex}", BuildDirectory, Ex);
+							Logger.LogError(Ex, "Exception while trying to scan files under {BuildDirectory}: {Ex}", BuildDirectory, Ex);
 						}
 					}
 				}
 				catch (Exception Ex)
 				{
-					Log.Logger.LogError(Ex, "Exception while trying to scan {StreamDirectory}: {Ex}", StreamDirectory, Ex);
+					Logger.LogError(Ex, "Exception while trying to scan {StreamDirectory}: {Ex}", StreamDirectory, Ex);
 				}
 			}
 			CommandUtils.LogInformation("Found {0} builds; {1} to delete.", NumBuilds, BuildsToDelete.Count);
@@ -1594,7 +1594,7 @@ namespace AutomationTool
 				}
 				catch (Exception Ex)
 				{
-					Log.Logger.LogError(Ex, "Exception while trying to delete {StreamDirectory}: {Ex}", StreamDirectory, Ex);
+					Logger.LogError(Ex, "Exception while trying to delete {StreamDirectory}: {Ex}", StreamDirectory, Ex);
 				}
 			}
 		}
diff --git a/Engine/Source/Programs/AutomationTool/CookedEditor/CookedEditor.Automation.cs b/Engine/Source/Programs/AutomationTool/CookedEditor/CookedEditor.Automation.cs
--- a/Engine/Source/Programs/AutomationTool/CookedEditor/CookedEditor.Automation.cs
+++ b/Engine/Source/Programs/AutomationTool/CookedEditor/CookedEditor.Automation.cs
@@ -9,6 +9,7 @@ using AutomationScripts;
 using EpicGames.Core;
 using UnrealBuildBase;
 using System.Text.Json;
+using Microsoft.Extensions.Logging;
 
 public class ConfigHelper
 {
@@ -353,7 +354,7 @@ public class MakeCookedEditor : BuildCommand
 
 	public override void ExecuteBuild()
 	{
-		LogInformation("************************* MakeCookedEditor");
+		Logger.LogInformation("************************* MakeCookedEditor");
 
 		bIsCookedCooker = ParseParam("cookedcooker");
 		ProjectFile = ParseProjectParam();
@@ -402,11 +403,11 @@ public class MakeCookedEditor : BuildCommand
 			{
 				if (CookedEditorStageDirectory == null || ReleaseStageDirectory == null)
 				{
-					LogError("Combining Release and CookedEditor together currently requires that both are staged this run (-stage -makerelease=stage)");
+					Logger.LogError("Combining Release and CookedEditor together currently requires that both are staged this run (-stage -makerelease=stage)");
 					return;
 				}
 
-				LogInformation($"Combinging {ReleaseStageDirectory} + {CookedEditorStageDirectory} -> {CombinedPath}");
+				Logger.LogInformation("Combinging {ReleaseStageDirectory} + {CookedEditorStageDirectory} -> {CombinedPath}", ReleaseStageDirectory, CookedEditorStageDirectory, CombinedPath);
 				DirectoryReference Combined = new DirectoryReference(CombinedPath);
 				CopyDirectory_NoExceptions(ReleaseStageDirectory.FullName, Combined.FullName, CopyDirectoryOptions.Default);
 				CopyDirectory_NoExceptions(CookedEditorStageDirectory.FullName, Combined.FullName, CopyDirectoryOptions.Merge);
@@ -510,10 +511,10 @@ public class MakeCookedEditor : BuildCommand
 			// these could be just copied, but StageFiles handles copying easily
 			SC.StageFiles(StagedFileType.NonUFS, ReleaseOptionalFileStageDirectory, StageFilesSearch.AllDirectories, new StagedDirectoryReference($"{Context.ProjectName}/Content/Paks"));
 
-			Log.TraceInformation($"Staging optional files from {ReleaseOptionalFileStageDirectory.FullName}:");
+			Logger.LogInformation("Staging optional files from {Arg0}:", ReleaseOptionalFileStageDirectory.FullName);
 			foreach (var OptionalFile in DirectoryReference.EnumerateFiles(ReleaseOptionalFileStageDirectory, "*"))
 			{
-				Log.TraceInformation($"  '{OptionalFile.FullName}'");
+				Logger.LogInformation("  '{Arg0}'", OptionalFile.FullName);
 			}
 		}
 	}
@@ -720,7 +721,7 @@ public class MakeCookedEditor : BuildCommand
 
 		if (ParseParamValue("MakeRelease", null) != null && !bBuildAgainstRelease)
 		{
-			LogWarning("-makerelease is meant for projects that have bBuildAgainstRelease set. Will force it on, with default settings for [DLCPluginName, ReleaseName, ReleaseTargetType]");
+			Logger.LogWarning("-makerelease is meant for projects that have bBuildAgainstRelease set. Will force it on, with default settings for [DLCPluginName, ReleaseName, ReleaseTargetType]");
 			bBuildAgainstRelease = true;
 		}
 
diff --git a/Engine/Source/Programs/AutomationTool/Gauntlet/Unreal/RunUnreal.cs b/Engine/Source/Programs/AutomationTool/Gauntlet/Unreal/RunUnreal.cs
--- a/Engine/Source/Programs/AutomationTool/Gauntlet/Unreal/RunUnreal.cs
+++ b/Engine/Source/Programs/AutomationTool/Gauntlet/Unreal/RunUnreal.cs
@@ -10,6 +10,7 @@ using UnrealBuildTool;
 using System.IO;
 using EpicGames.Core;
 using UnrealBuildBase;
+using Microsoft.Extensions.Logging;
 
 namespace Gauntlet
 {
@@ -339,7 +340,7 @@ namespace Gauntlet
 
 				if (ParseParam("clean"))
 				{
-					LogInformation("Deleting temp dir {0}", Options.TempDir);
+					Logger.LogInformation("Deleting temp dir {Arg0}", Options.TempDir);
 					DirectoryInfo Di = new DirectoryInfo(Options.TempDir);
 					if (Di.Exists)
 					{
diff --git a/Engine/Source/Programs/AutomationTool/IOS/ExportIPAFromArchive.cs b/Engine/Source/Programs/AutomationTool/IOS/ExportIPAFromArchive.cs
--- a/Engine/Source/Programs/AutomationTool/IOS/ExportIPAFromArchive.cs
+++ b/Engine/Source/Programs/AutomationTool/IOS/ExportIPAFromArchive.cs
@@ -9,6 +9,7 @@ using System.Threading.Tasks;
 using System.Xml;
 using AutomationTool;
 using UnrealBuildBase;
+using Microsoft.Extensions.Logging;
 
 [Help(@"Creates an IPA from an xarchive file")]
 [Help("method=<method>", @"Purpose of the IPA. (Development, Adhoc, Store)")]
@@ -131,7 +132,7 @@ public class ExportIPAFromArchive : BuildCommand
 
 		string ExportArgs = string.Format("-exportArchive -archivePath \"{0}\" -exportOptionsPlist {1} -exportPath \"{2}\" {3}", Archive, PlistPath, OutputPath, ProvisioningArgs);
 
-		LogInformation("Creating IPA at at {0} with {1} profile.", OutputPath, Method);
+		Logger.LogInformation("Creating IPA at at {OutputPath} with {Method} profile.", OutputPath, Method);
 
 
 		DateTime ExportStartTime = DateTime.Now;
@@ -140,7 +141,7 @@ public class ExportIPAFromArchive : BuildCommand
 
 		if (Result != 0)
 		{
-			LogError("XcodeBuild failed to export archive. {0} might help diagose the issue..", PlistPath);
+			Logger.LogError("XcodeBuild failed to export archive. {PlistPath} might help diagose the issue..", PlistPath);
 			return ExitCode.Error_FailedToCreateIPA;
 		}
 		
diff --git a/Engine/Source/Programs/AutomationTool/IOS/IOSPlatform.Automation.cs b/Engine/Source/Programs/AutomationTool/IOS/IOSPlatform.Automation.cs
--- a/Engine/Source/Programs/AutomationTool/IOS/IOSPlatform.Automation.cs
+++ b/Engine/Source/Programs/AutomationTool/IOS/IOSPlatform.Automation.cs
@@ -15,6 +15,7 @@ using System.Diagnostics;
 using EpicGames.Core;
 using System.Xml;
 using UnrealBuildBase;
+using Microsoft.Extensions.Logging;
 
 static class IOSEnvVarNames
 {
@@ -712,7 +713,7 @@ public class IOSPlatform : ApplePlatform
 		}
 		else
 		{
-			LogWarning("Use Custom Launch Screen Storyboard is checked but not compiled storyboard could be found. Have you compiled on Mac first ? Falling back to Standard Storyboard");
+			Logger.LogWarning("Use Custom Launch Screen Storyboard is checked but not compiled storyboard could be found. Have you compiled on Mac first ? Falling back to Standard Storyboard");
 			StageStandardLaunchScreenStoryboard(Params, SC);
 		}
 	}
@@ -747,7 +748,7 @@ public class IOSPlatform : ApplePlatform
 
 	public override void Package(ProjectParams Params, DeploymentContext SC, int WorkingCL)
 	{
-		LogInformation("Package {0}", Params.RawProjectPath);
+		Logger.LogInformation("Package {Arg0}", Params.RawProjectPath);
 
 		bool bIsBuiltAsFramework = IsBuiltAsFramework(Params, SC);
 
@@ -757,7 +758,7 @@ public class IOSPlatform : ApplePlatform
 		string FullExePath = CombinePaths(Path.GetDirectoryName(ProjectGameExeFilename), SC.StageExecutables[0] + (UnrealBuildTool.BuildHostPlatform.Current.Platform != UnrealTargetPlatform.Mac ? ".stub" : ""));
 		if (!SC.IsCodeBasedProject && !FileExists_NoExceptions(FullExePath) && !bIsBuiltAsFramework)
 		{
-			LogError("Failed to find game binary " + FullExePath);
+			Logger.LogError("{Text}", "Failed to find game binary " + FullExePath);
 			throw new AutomationException(ExitCode.Error_MissingExecutable, "Stage Failed. Could not find binary {0}. You may need to build the Unreal Engine project with your target configuration and platform.", FullExePath);
 		}
 #endif // PLATFORM_MAC
@@ -908,11 +909,11 @@ public class IOSPlatform : ApplePlatform
 			// package a .ipa from the now staged directory
 			var IPPExe = CombinePaths(CmdEnv.LocalRoot, "Engine/Binaries/DotNET/IOS/IPhonePackager.exe");
 
-			LogLog("ProjectName={0}", Params.ShortProjectName);
-			LogLog("ProjectStub={0}", ProjectStub);
-			LogLog("ProjectIPA={0}", ProjectIPA);
-			LogLog("IPPProjectIPA={0}", IPPProjectIPA);
-			LogLog("IPPExe={0}", IPPExe);
+			Logger.LogDebug("ProjectName={Arg0}", Params.ShortProjectName);
+			Logger.LogDebug("ProjectStub={ProjectStub}", ProjectStub);
+			Logger.LogDebug("ProjectIPA={ProjectIPA}", ProjectIPA);
+			Logger.LogDebug("IPPProjectIPA={IPPProjectIPA}", IPPProjectIPA);
+			Logger.LogDebug("IPPExe={IPPExe}", IPPExe);
 
 			bool cookonthefly = Params.CookOnTheFly || Params.SkipCookOnTheFly;
 
@@ -1166,7 +1167,7 @@ public class IOSPlatform : ApplePlatform
 							UUID = AllText.Substring(idx, AllText.IndexOf("</string>", idx) - idx);
 							Arguments += " PROVISIONING_PROFILE_SPECIFIER=" + UUID;
 
-							LogInformation("Extracted Provision UUID {0} from {1}", UUID, Provision);
+							Logger.LogInformation("Extracted Provision UUID {UUID} from {Provision}", UUID, Provision);
 						}
 					}
 				}
@@ -1344,7 +1345,7 @@ public class IOSPlatform : ApplePlatform
 		string DestSymlink = Path.Combine(DestOuterApp, "WrappedBundle");
 		string DestInnerApp = Path.Combine(DestOuterApp, "Wrapper", $"{GameName}.app");
 
-		LogInformation($"Creating IOSOnMac app {DestOuterApp}");
+		Logger.LogInformation("Creating IOSOnMac app {DestOuterApp}", DestOuterApp);
 
 		DeleteDirectory_NoExceptions(DestOuterApp);
 		CopyDirectory_NoExceptions(SourceApp, DestInnerApp);
@@ -1682,7 +1683,7 @@ public class IOSPlatform : ApplePlatform
 		}
 		else if (bXCArchive && RuntimePlatform.IsWindows)
 		{
-			LogWarning("Can not produce an XCArchive on windows");
+			Logger.LogWarning("Can not produce an XCArchive on windows");
 		}
 		SC.ArchiveFiles(Path.GetDirectoryName(ProjectIPA), Path.GetFileName(ProjectIPA));
 	}
@@ -1714,7 +1715,7 @@ public class IOSPlatform : ApplePlatform
 			IdeviceInstallerArgs = GetLibimobileDeviceNetworkedArgument(IdeviceInstallerArgs, Params.DeviceNames[0]);
 
 			var DeviceInstaller = GetPathToLibiMobileDeviceTool("ideviceinstaller");
-			LogInformation("Checking if bundle {0} is installed", BundleIdentifier);
+			Logger.LogInformation("Checking if bundle {BundleIdentifier} is installed", BundleIdentifier);
 
 			string Output = CommandUtils.RunAndLog(DeviceInstaller, IdeviceInstallerArgs);
 			bool bBundleIsInstalled = Output.Contains(string.Format("CFBundleIdentifier -> {0}{1}", BundleIdentifier, Environment.NewLine));
@@ -1722,7 +1723,7 @@ public class IOSPlatform : ApplePlatform
 
 			if (bBundleIsInstalled)
 			{
-				LogInformation("Bundle {0} found, retrieving deployed manifests...", BundleIdentifier);
+				Logger.LogInformation("Bundle {BundleIdentifier} found, retrieving deployed manifests...", BundleIdentifier);
 
 				var DeviceFS = GetPathToLibiMobileDeviceTool("idevicefs");
 
@@ -1747,7 +1748,7 @@ public class IOSPlatform : ApplePlatform
 			}
 			else
 			{
-				LogInformation("Bundle {0} not found, skipping retrieving deployed manifests", BundleIdentifier);
+				Logger.LogInformation("Bundle {BundleIdentifier} not found, skipping retrieving deployed manifests", BundleIdentifier);
 			}
 		}
 		catch (System.Exception)
@@ -2357,17 +2358,17 @@ public class IOSPlatform : ApplePlatform
 			if (!File.Exists(SourcePackage))
             {
 				IPAFiles = Directory.GetFiles(PackagePath, "*.ipa");
-				LogWarning("Source package not found : {0}, trying to find another IPA file in the same folder.", SourcePackage);
+				Logger.LogWarning("Source package not found : {SourcePackage}, trying to find another IPA file in the same folder.", SourcePackage);
 				if (IPAFiles.Length == 0)
                 {
-					LogError("No IPA file found in : {0}. Aborting.", PackagePath);
+					Logger.LogError("No IPA file found in : {PackagePath}. Aborting.", PackagePath);
 					throw new AutomationException(ExitCode.Error_MissingExecutable, "No IPA file found in {0}.", PackagePath);
 				}
 				else
                 {
 					if (IPAFiles.Length > 1)
                     {
-						LogWarning("More than one IPA file found. Taking the first one found, {0}", IPAFiles[0]);
+						Logger.LogWarning("More than one IPA file found. Taking the first one found, {Arg0}", IPAFiles[0]);
 					}
 					SourcePackage = IPAFiles[0];
 				}					
@@ -2379,16 +2380,16 @@ public class IOSPlatform : ApplePlatform
 			PayloadPath = PayloadPath.Substring(0, PayloadPath.LastIndexOf('\\'));
 			string CookedDataDirectory = Path.Combine(Path.GetDirectoryName(PayloadPath), ClientPlatform, "Payload", PackageName + ".app", "cookeddata");
 
-			LogInformation("ClientPlatform : {0}", ClientPlatform);
-			LogInformation("ProjectFilePath : {0}", ProjectFilePath);
-			LogInformation("Source : {0}", SourcePackage);
-			LogInformation("ZipFile {0}", ZipFile);
-			LogInformation("PackageName {0}", PackageName);
-			LogInformation("PayloadPath {0}", PayloadPath);
+			Logger.LogInformation("ClientPlatform : {ClientPlatform}", ClientPlatform);
+			Logger.LogInformation("ProjectFilePath : {ProjectFilePath}", ProjectFilePath);
+			Logger.LogInformation("Source : {SourcePackage}", SourcePackage);
+			Logger.LogInformation("ZipFile {ZipFile}", ZipFile);
+			Logger.LogInformation("PackageName {PackageName}", PackageName);
+			Logger.LogInformation("PayloadPath {PayloadPath}", PayloadPath);
 
 			if (File.Exists(ZipFile))
 			{
-				LogInformation("Deleting previously present ZIP file created from IPA");
+				Logger.LogInformation("Deleting previously present ZIP file created from IPA");
 				File.Delete(ZipFile);
 			}
 
@@ -2414,13 +2415,13 @@ public class IOSPlatform : ApplePlatform
 				}
 			}
 			//cleanup
-			LogInformation("Deleting temp files ...");
+			Logger.LogInformation("Deleting temp files ...");
 			File.Delete(ZipFile);
-			LogInformation("{0} deleted", ZipFile);
+			Logger.LogInformation("{ZipFile} deleted", ZipFile);
 		}
 		else
 		{
-			LogInformation("Wrangling data for debug for an iOS/tvOS app for XCode is a Mac and Windows (Remote) only feature. Aborting command.");
+			Logger.LogInformation("Wrangling data for debug for an iOS/tvOS app for XCode is a Mac and Windows (Remote) only feature. Aborting command.");
 			return;
 		}
 	}
@@ -2429,7 +2430,7 @@ public class IOSPlatform : ApplePlatform
 	{
 		string UnzipPath = PackageToUnzip;
 		UnzipPath = UnzipPath.Substring(0, UnzipPath.LastIndexOf('\\'));
-		LogInformation("Unzipping to {0}", UnzipPath);
+		Logger.LogInformation("Unzipping to {UnzipPath}", UnzipPath);
 
 		using (Ionic.Zip.ZipFile Zip = new Ionic.Zip.ZipFile(PackageToUnzip))
 		{
@@ -2441,7 +2442,7 @@ public class IOSPlatform : ApplePlatform
 				{
 					Entry.Extract(OutputStream);
 				}
-				LogInformation("Extracted {0}", OutputFileName);
+				Logger.LogInformation("Extracted {OutputFileName}", OutputFileName);
 			}
 		}
 	}
diff --git a/Engine/Source/Programs/AutomationTool/IOS/MakeIPA.cs b/Engine/Source/Programs/AutomationTool/IOS/MakeIPA.cs
--- a/Engine/Source/Programs/AutomationTool/IOS/MakeIPA.cs
+++ b/Engine/Source/Programs/AutomationTool/IOS/MakeIPA.cs
@@ -9,6 +9,7 @@ using System.Threading.Tasks;
 using System.Xml;
 using AutomationTool;
 using UnrealBuildBase;
+using Microsoft.Extensions.Logging;
 
 [Help(@"Creates an IPA from an xarchive file")]
 [Help("method=<method>", @"Purpose of the IPA. (Development, Adhoc, Store)")]
@@ -118,13 +119,13 @@ public class MakeIPA : BuildCommand
 
 		string ExportArgs = string.Format("-exportArchive -archivePath \"{0}\" -exportOptionsPlist {1} -exportPath \"{2}\" {3}", Archive, PlistPath, Output, ProvisioningArgs);
 
-		LogInformation("Creating IPA at at {0} with {1} profile.", Output, Method);
+		Logger.LogInformation("Creating IPA at at {Output} with {Method} profile.", Output, Method);
 
 		int Result = XcodeBuild(ExportArgs);
 
 		if (Result != 0)
 		{
-			LogError("XcodeBuild failed to export archive. {0} might help diagose the issue..", PlistPath);
+			Logger.LogError("XcodeBuild failed to export archive. {PlistPath} might help diagose the issue..", PlistPath);
 			return ExitCode.Error_FailedToCreateIPA;
 		}
 		
diff --git a/Engine/Source/Programs/AutomationTool/IOS/WriteIniValueToPlist.cs b/Engine/Source/Programs/AutomationTool/IOS/WriteIniValueToPlist.cs
--- a/Engine/Source/Programs/AutomationTool/IOS/WriteIniValueToPlist.cs
+++ b/Engine/Source/Programs/AutomationTool/IOS/WriteIniValueToPlist.cs
@@ -10,6 +10,7 @@ using System.Threading.Tasks;
 using System.Xml;
 using AutomationTool;
 using UnrealBuildBase;
+using Microsoft.Extensions.Logging;
 
 [Help(@"Pulls a value from an ini file and inserts it into a plist.")]
 [Help(@"Note currently only looks at values irrespective of sections!")]
@@ -88,7 +89,7 @@ public class WriteIniValueToPlist : BuildCommand
 
 			string IniValue = IniMatch.Groups[1].ToString().Trim();
 
-			LogInformation("Found ini value {0} for property {1}", IniValue, IniProperty);
+			Logger.LogInformation("Found ini value {IniValue} for property {IniProperty}", IniValue, IniProperty);
 
 			string PlistFileContents = File.ReadAllText(PlistFile);
 
@@ -105,25 +106,25 @@ public class WriteIniValueToPlist : BuildCommand
 			string PListEntry = PlistMatch.Groups[0].ToString();
 			string PListValue = PlistMatch.Groups[1].ToString();
 
-			LogInformation("Found existing value {0} for plist entry {1}", PListValue, PlistProperty);
+			Logger.LogInformation("Found existing value {PListValue} for plist entry {PlistProperty}", PListValue, PlistProperty);
 
 			string NewPlistEntry = PListEntry.Replace(PListValue, IniValue);
 
-			LogInformation("Set entry {0} to {1}", PlistProperty, IniValue);
+			Logger.LogInformation("Set entry {PlistProperty} to {IniValue}", PlistProperty, IniValue);
 
 			PlistFileContents = PlistFileContents.Replace(PListEntry, NewPlistEntry);
 
-			LogInformation("Saving new plist to {0}", PlistFile);
+			Logger.LogInformation("Saving new plist to {PlistFile}", PlistFile);
 
 			File.WriteAllText(PlistFile, PlistFileContents);
 
-			LogInformation("Done");
+			Logger.LogInformation("Done");
 		}
 		catch (Exception Ex)
 		{
-			LogWarning("Failed to update plist {0}. {1}", PlistFile, Ex);
+			Logger.LogWarning("Failed to update plist {PlistFile}. {Ex}", PlistFile, Ex);
 		}
 
 		return ExitCode.Success;
 	}
-}
\ No newline at end of file
+}
diff --git a/Engine/Source/Programs/AutomationTool/LowLevelTests/RunLowLevelTests.cs b/Engine/Source/Programs/AutomationTool/LowLevelTests/RunLowLevelTests.cs
--- a/Engine/Source/Programs/AutomationTool/LowLevelTests/RunLowLevelTests.cs
+++ b/Engine/Source/Programs/AutomationTool/LowLevelTests/RunLowLevelTests.cs
@@ -10,9 +10,12 @@ using UnrealBuildBase;
 using Gauntlet;
 using System.Text.RegularExpressions;
 using AutomationTool.DeviceReservation;
+using Microsoft.Extensions.Logging;
 
 namespace LowLevelTests
 {
+	using LogLevel = Gauntlet.LogLevel;
+
 	public class RunLowLevelTests : BuildCommand
 	{
 		public override ExitCode Execute()
@@ -114,7 +117,7 @@ namespace LowLevelTests
 
 				if (ParseParam("clean"))
 				{
-					LogInformation("Deleting temp dir {0}", Options.TempDir);
+					Logger.LogInformation("Deleting temp dir {Arg0}", Options.TempDir);
 					DirectoryInfo TempDirInfo = new DirectoryInfo(Options.TempDir);
 					if (TempDirInfo.Exists)
 					{
diff --git a/Engine/Source/Programs/AutomationTool/Mac/MacPlatform.Automation.cs b/Engine/Source/Programs/AutomationTool/Mac/MacPlatform.Automation.cs
--- a/Engine/Source/Programs/AutomationTool/Mac/MacPlatform.Automation.cs
+++ b/Engine/Source/Programs/AutomationTool/Mac/MacPlatform.Automation.cs
@@ -1,4 +1,4 @@
-﻿// Copyright Epic Games, Inc. All Rights Reserved.
+// Copyright Epic Games, Inc. All Rights Reserved.
 using System;
 using System.Collections.Generic;
 using System.Linq;
@@ -9,6 +9,7 @@ using AutomationTool;
 using UnrealBuildTool;
 using UnrealBuildBase;
 using EpicGames.Core;
+using Microsoft.Extensions.Logging;
 
 public class MacPlatform : ApplePlatform
 {
@@ -513,7 +514,7 @@ public class MacPlatform : ApplePlatform
 
 							RenameFile(FoundDebugFile, MovedDebugFile);
 
-							Log.TraceInformation("Moving debug file: '{0}')", FoundDebugFile);
+							Logger.LogInformation("Moving debug file: '{FoundDebugFile}')", FoundDebugFile);
 							ManifestLines.RemoveAt(ManifestLineIndex);
 							ModifyManifest = true;
 						}
@@ -592,29 +593,29 @@ public class MacPlatform : ApplePlatform
 			if (Params.Archive)
 			{
 				// Remove extra RPATHs if we will be archiving the project
-				LogInformation("Removing extraneous rpath entries");
+				Logger.LogInformation("Removing extraneous rpath entries");
 				RemoveExtraRPaths(Params, SC);
 			}
 
 			// Sign everything we built
 			List<FileReference> FilesToSign = GetExecutableNames(SC);
-			LogInformation("RuntimeProjectRootDir: " + SC.RuntimeProjectRootDir);
+			Logger.LogInformation("{Text}", "RuntimeProjectRootDir: " + SC.RuntimeProjectRootDir);
 			foreach (var Exe in FilesToSign)
 			{
-				LogInformation("Signing: " + Exe);
+				Logger.LogInformation("{Text}", "Signing: " + Exe);
 				string AppBundlePath = "";
 				if (Exe.IsUnderDirectory(DirectoryReference.Combine(SC.RuntimeProjectRootDir, "Binaries", SC.PlatformDir)))
 				{
-					LogInformation("Starts with Binaries");
+					Logger.LogInformation("Starts with Binaries");
 					AppBundlePath = CombinePaths(SC.RuntimeProjectRootDir.FullName, "Binaries", SC.PlatformDir, Path.GetFileNameWithoutExtension(Exe.FullName) + ".app");
 				}
 				else if (Exe.IsUnderDirectory(DirectoryReference.Combine(SC.RuntimeRootDir, "Engine/Binaries", SC.PlatformDir)))
 				{
-					LogInformation("Starts with Engine/Binaries");
+					Logger.LogInformation("Starts with Engine/Binaries");
 					AppBundlePath = CombinePaths("Engine/Binaries", SC.PlatformDir, Path.GetFileNameWithoutExtension(Exe.FullName) + ".app");
 				}
 
-				LogInformation("Signing: " + AppBundlePath);
+				Logger.LogInformation("{Text}", "Signing: " + AppBundlePath);
 				CodeSign.SignMacFileOrFolder(AppBundlePath);
 			}
 
diff --git a/Engine/Source/Programs/AutomationTool/Scripts/AnalyzeThirdPartyLibs.Automation.cs b/Engine/Source/Programs/AutomationTool/Scripts/AnalyzeThirdPartyLibs.Automation.cs
--- a/Engine/Source/Programs/AutomationTool/Scripts/AnalyzeThirdPartyLibs.Automation.cs
+++ b/Engine/Source/Programs/AutomationTool/Scripts/AnalyzeThirdPartyLibs.Automation.cs
@@ -6,6 +6,7 @@ using System.IO;
 using AutomationTool;
 using UnrealBuildTool;
 using System.Linq;
+using Microsoft.Extensions.Logging;
 
 /*
  - You can also use the full program to test compiling all or a subset of libs:
@@ -127,7 +128,7 @@ class AnalyzeThirdPartyLibs : BuildCommand
 
 	public override void ExecuteBuild()
 	{
-		LogInformation("************************* Analyze Third Party Libs");
+		Logger.LogInformation("************************* Analyze Third Party Libs");
 
 		// figure out what batch/script to run
 		if (UnrealBuildTool.BuildHostPlatform.Current.Platform != UnrealTargetPlatform.Win64 &&
@@ -188,19 +189,19 @@ class AnalyzeThirdPartyLibs : BuildCommand
 
 			long Size = Info.GetSize(Platforms);
 
-			LogInformation("Library {0} is {1}", Lib, ToMegabytes(Size));
+			Logger.LogInformation("Library {Lib} is {Arg1}", Lib, ToMegabytes(Size));
 
 			long Total = 0;
 			for (int Index = 0; Index < Platforms.Count; ++Index)
 			{
 				PlatformLibraryInfo Platform = Platforms[Index];
 				long Growth = Platform.TotalSize - LastSizes[Index];
-				LogInformation("  {0} is {1}", Platform.PlatformName, ToMegabytes(Growth));
+				Logger.LogInformation("  {Arg0} is {Arg1}", Platform.PlatformName, ToMegabytes(Growth));
 
 				LastSizes[Index] = Platform.TotalSize;
 				Total += Growth;
 			}
-			LogInformation("  Platform neutral is probably {0} (specific sum {1})", ToMegabytes(Size - Total), ToMegabytes(Total));
+			Logger.LogInformation("  Platform neutral is probably {Arg0} (specific sum {Arg1})", ToMegabytes(Size - Total), ToMegabytes(Total));
 
 			TotalSize += Size;
 		}
@@ -211,19 +212,19 @@ class AnalyzeThirdPartyLibs : BuildCommand
 		LargeFileExtensions.AddRange(new string[] { ".pdb", ".a", ".lib", ".dll", ".dylib", ".bc", ".so" });
 
 		// Hackery, look for big files (re-traverses everything)
-		LogInformation("----");
+		Logger.LogInformation("----");
 		foreach (string Lib in LibsToEvaluate)
 		{
 			ThirdPartyLibraryInfo Info = new ThirdPartyLibraryInfo(Lib);
 			Info.FindLargeFiles(LargeFileExtensions, 1024 * 1024);
 		}
 
-		LogInformation("----");
+		Logger.LogInformation("----");
 		foreach (var Platform in Platforms)
 		{
-			LogInformation("  {0} is {1} (estimate)", Platform.PlatformName, ToMegabytes(Platform.TotalSize));
+			Logger.LogInformation("  {Arg0} is {Arg1} (estimate)", Platform.PlatformName, ToMegabytes(Platform.TotalSize));
 		}
-		LogInformation("  OVERALL is {0} (accurate)", ToMegabytes(TotalSize));
+		Logger.LogInformation("  OVERALL is {Arg0} (accurate)", ToMegabytes(TotalSize));
 
 		// undo the LibDir push
 		CommandUtils.PopDir();
diff --git a/Engine/Source/Programs/AutomationTool/Scripts/ArchiveCommand.Automation.cs b/Engine/Source/Programs/AutomationTool/Scripts/ArchiveCommand.Automation.cs
--- a/Engine/Source/Programs/AutomationTool/Scripts/ArchiveCommand.Automation.cs
+++ b/Engine/Source/Programs/AutomationTool/Scripts/ArchiveCommand.Automation.cs
@@ -7,6 +7,7 @@ using System.Reflection;
 using AutomationTool;
 using UnrealBuildTool;
 using EpicGames.Core;
+using Microsoft.Extensions.Logging;
 
 namespace AutomationScripts
 {
@@ -53,10 +54,10 @@ namespace AutomationScripts
 				return;
 			}
 
-			LogInformation("********** ARCHIVE COMMAND STARTED **********");
+			Logger.LogInformation("********** ARCHIVE COMMAND STARTED **********");
 			var StartTime = DateTime.UtcNow;
 
-			LogInformation("Archiving to {0}", Params.ArchiveDirectoryParam);
+			Logger.LogInformation("Archiving to {Arg0}", Params.ArchiveDirectoryParam);
 
 			if (!Params.NoClient)
 			{
@@ -80,8 +81,8 @@ namespace AutomationScripts
 					SC.StageTargetPlatform.ProcessArchivedProject(Params, SC);
 				}
 			}
-			LogInformation("Archive command time: {0:0.00} s", (DateTime.UtcNow - StartTime).TotalMilliseconds / 1000);
-			LogInformation("********** ARCHIVE COMMAND COMPLETED **********");
+			Logger.LogInformation("Archive command time: {0:0.00} s", (DateTime.UtcNow - StartTime).TotalMilliseconds / 1000);
+			Logger.LogInformation("********** ARCHIVE COMMAND COMPLETED **********");
 		}
 	}
 }
diff --git a/Engine/Source/Programs/AutomationTool/Scripts/BenchmarkBuild/BenchmarkBuild.Automation.cs b/Engine/Source/Programs/AutomationTool/Scripts/BenchmarkBuild/BenchmarkBuild.Automation.cs
--- a/Engine/Source/Programs/AutomationTool/Scripts/BenchmarkBuild/BenchmarkBuild.Automation.cs
+++ b/Engine/Source/Programs/AutomationTool/Scripts/BenchmarkBuild/BenchmarkBuild.Automation.cs
@@ -11,6 +11,7 @@ using System.Threading.Tasks;
 using EpicGames.Core;
 using UnrealBuildTool;
 using UnrealBuildBase;
+using Microsoft.Extensions.Logging;
 
 namespace AutomationTool.Benchmark
 {
@@ -294,7 +295,7 @@ namespace AutomationTool.Benchmark
 					/*&& (PlatformsToTest.Count() > 1 || ProjectsToTest.Count() > 1)*/)
 				{
 					BuildOptions |= UBTBuildOptions.PostClean;
-					Log.TraceInformation("Building multiple platforms. Will clean each platform after build step to save space. (use -nopostclean to prevent this)");
+					Logger.LogInformation("Building multiple platforms. Will clean each platform after build step to save space. (use -nopostclean to prevent this)");
 				}
 
 				// parse processor args
@@ -383,7 +384,7 @@ namespace AutomationTool.Benchmark
 					if (XGEOptions.Contains(XGETaskOptions.WithXGE) 
 						|| XGEOptions.Contains(XGETaskOptions.WithEditorXGE))
 					{
-						Log.TraceWarning("XGE requested but is not available. Removing XGE options");
+						Logger.LogWarning("XGE requested but is not available. Removing XGE options");
 						XGEOptions.Remove(XGETaskOptions.WithXGE);
 						XGEOptions.Remove(XGETaskOptions.WithEditorXGE);
 					}					
@@ -499,11 +500,11 @@ namespace AutomationTool.Benchmark
 				Tasks.AddRange(EditorTasks);
 			}
 
-			Log.TraceInformation("Will execute tasks:");
+			Logger.LogInformation("Will execute tasks:");
 
 			foreach (var Task in Tasks)
 			{
-				Log.TraceInformation("{0}", Task.FullName);
+				Logger.LogInformation("{Arg0}", Task.FullName);
 			}
 
 			if (!Options.Preview)
@@ -520,20 +521,20 @@ namespace AutomationTool.Benchmark
 				{
 					foreach (var Task in Tasks)
 					{
-						Log.TraceInformation("Starting task {0} (Pass {1})", Task.FullName, i + 1);
+						Logger.LogInformation("Starting task {Arg0} (Pass {Arg1})", Task.FullName, i + 1);
 
 						Task.Run();
 
-						Log.TraceInformation("Task {0} took {1}", Task.FullName, Task.TaskTime.ToString(@"hh\:mm\:ss"));
+						Logger.LogInformation("Task {Arg0} took {Arg1}", Task.FullName, Task.TaskTime.ToString(@"hh\:mm\:ss"));
 
 						if (Task.Failed)
 						{
-							Log.TraceError("Task failed! Benchmark time may be inaccurate.");
+							Logger.LogError("Task failed! Benchmark time may be inaccurate.");
 						}
 
 						if (Task.SkipReport)
 						{
-							Log.TraceInformation("Skipping reporting of {0}", Task.FullName);
+							Logger.LogInformation("Skipping reporting of {Arg0}", Task.FullName);
 						}
 						else
 						{
@@ -547,13 +548,13 @@ namespace AutomationTool.Benchmark
 							WriteCSVResults(Options.FileName, Tasks, Results);
 						}
 
-						Log.TraceInformation("Waiting {0} secs until next task", Options.TimeBetweenTasks);
+						Logger.LogInformation("Waiting {Arg0} secs until next task", Options.TimeBetweenTasks);
 						Thread.Sleep(Options.TimeBetweenTasks * 1000);
 					}
 				}
 
-				Log.TraceInformation("**********************************************************************");
-				Log.TraceInformation("Test Results:");
+				Logger.LogInformation("**********************************************************************");
+				Logger.LogInformation("Test Results:");
 
 				foreach (var Task in Tasks)
 				{
@@ -592,13 +593,13 @@ namespace AutomationTool.Benchmark
 						AvgTimeString = string.Format(" (Avg: {0})", AvgTime.ToString(@"hh\:mm\:ss"));
 					}
 
-					Log.TraceInformation("Task {0}:\t\t{1}{2}", Task.FullName, TimeString, AvgTimeString);
+					Logger.LogInformation("Task {Arg0}:\t\t{TimeString}{AvgTimeString}", Task.FullName, TimeString, AvgTimeString);
 				}
-				Log.TraceInformation("**********************************************************************");
+				Logger.LogInformation("**********************************************************************");
 
 				TimeSpan Elapsed = DateTime.Now - StartTime;
 
-				Log.TraceInformation("Total benchmark time: {0}", Elapsed.ToString(@"hh\:mm\:ss"));
+				Logger.LogInformation("Total benchmark time: {Arg0}", Elapsed.ToString(@"hh\:mm\:ss"));
 
 				WriteCSVResults(Options.FileName, Tasks, Results);
 			}
@@ -730,7 +731,7 @@ namespace AutomationTool.Benchmark
 		/// </summary>
 		void WriteCSVResults(string InFileName, IEnumerable<BenchmarkTaskBase> InTasks, Dictionary<BenchmarkTaskBase, List<BenchmarkResult>> InResults)
 		{
-			Log.TraceInformation("Writing results to {0}", InFileName);
+			Logger.LogInformation("Writing results to {InFileName}", InFileName);
 
 			try
 			{
@@ -804,7 +805,7 @@ namespace AutomationTool.Benchmark
 			}
 			catch (Exception Ex)
 			{
-				Log.TraceError("Failed to write CSV to {0}. {1}", InFileName, Ex);
+				Logger.LogError("Failed to write CSV to {InFileName}. {Ex}", InFileName, Ex);
 			}
 		}
 
diff --git a/Engine/Source/Programs/AutomationTool/Scripts/BenchmarkBuild/BenchmarkCookTask.cs b/Engine/Source/Programs/AutomationTool/Scripts/BenchmarkBuild/BenchmarkCookTask.cs
--- a/Engine/Source/Programs/AutomationTool/Scripts/BenchmarkBuild/BenchmarkCookTask.cs
+++ b/Engine/Source/Programs/AutomationTool/Scripts/BenchmarkBuild/BenchmarkCookTask.cs
@@ -10,6 +10,7 @@ using System.Threading.Tasks;
 using EpicGames.Core;
 using UnrealBuildTool;
 using UnrealBuildBase;
+using Microsoft.Extensions.Logging;
 
 namespace AutomationTool.Benchmark
 {
@@ -107,7 +108,7 @@ namespace AutomationTool.Benchmark
 
 			if (AssetFile == null)
 			{
-				Log.TraceError("Could not find asset file to touch under {0}", ContentDir);
+				Logger.LogError("Could not find asset file to touch under {ContentDir}", ContentDir);
 				return false;
 			}
 
diff --git a/Engine/Source/Programs/AutomationTool/Scripts/BenchmarkBuild/BenchmarkEditorTaskBase.cs b/Engine/Source/Programs/AutomationTool/Scripts/BenchmarkBuild/BenchmarkEditorTaskBase.cs
--- a/Engine/Source/Programs/AutomationTool/Scripts/BenchmarkBuild/BenchmarkEditorTaskBase.cs
+++ b/Engine/Source/Programs/AutomationTool/Scripts/BenchmarkBuild/BenchmarkEditorTaskBase.cs
@@ -10,6 +10,7 @@ using System.Threading.Tasks;
 using EpicGames.Core;
 using UnrealBuildBase;
 using UnrealBuildTool;
+using Microsoft.Extensions.Logging;
 
 namespace AutomationTool.Benchmark
 {
@@ -193,7 +194,7 @@ namespace AutomationTool.Benchmark
 				{
 					if (Message.Contains("TEST COMPLETE"))
 					{
-						Log.TraceInformation("Automation test reported as complete.");
+						Logger.LogInformation("Automation test reported as complete.");
 						//TestCompleted = true;
 					}
 
@@ -228,7 +229,7 @@ namespace AutomationTool.Benchmark
 				{
 					if ((DateTime.Now - LastStdOutTime).TotalMinutes >= TimeoutMins)
 					{
-						Log.TraceError("Gave up waiting for task after {0} minutes of no output", TimeoutMins);
+						Logger.LogError("Gave up waiting for task after {TimeoutMins} minutes of no output", TimeoutMins);
 						CurrentProcess.ProcessObject.Kill();
 					}
 				}
@@ -262,7 +263,7 @@ namespace AutomationTool.Benchmark
 
 				if (Command.Execute() != ExitCode.Success)
 				{
-					Log.TraceError("Failed to build editor!");
+					Logger.LogError("Failed to build editor!");
 					return false;
 				}
 				else
@@ -273,7 +274,7 @@ namespace AutomationTool.Benchmark
 			}
 			else if (BuiltEditors.Contains(ProjectName) == false)
 			{
-				Log.TraceError("SkipBuild = true but no editor has been built for {0}", ProjectName);
+				Logger.LogError("SkipBuild = true but no editor has been built for {ProjectName}", ProjectName);
 				return false;
 			}	
 
diff --git a/Engine/Source/Programs/AutomationTool/Scripts/BenchmarkBuild/BenchmarkSingleCompileTask.cs b/Engine/Source/Programs/AutomationTool/Scripts/BenchmarkBuild/BenchmarkSingleCompileTask.cs
--- a/Engine/Source/Programs/AutomationTool/Scripts/BenchmarkBuild/BenchmarkSingleCompileTask.cs
+++ b/Engine/Source/Programs/AutomationTool/Scripts/BenchmarkBuild/BenchmarkSingleCompileTask.cs
@@ -10,6 +10,7 @@ using System.Threading.Tasks;
 using EpicGames.Core;
 using UnrealBuildTool;
 using UnrealBuildBase;
+using Microsoft.Extensions.Logging;
 
 namespace AutomationTool.Benchmark
 {
@@ -57,7 +58,7 @@ namespace AutomationTool.Benchmark
 				SourceFile = FileReference.Combine(Unreal.EngineDirectory, "Source/Runtime/Engine/Private/UnrealEngine.cpp");
 			}
 
-			Log.TraceVerbose("Will compile {0} for single-file compilation test for {1}", SourceFile, ProjectName);
+			Logger.LogDebug("Will compile {SourceFile} for single-file compilation test for {ProjectName}", SourceFile, ProjectName);
 		}
 
 		protected override bool PerformPrequisites()
diff --git a/Engine/Source/Programs/AutomationTool/Scripts/BenchmarkBuild/BenchmarkTaskBase.cs b/Engine/Source/Programs/AutomationTool/Scripts/BenchmarkBuild/BenchmarkTaskBase.cs
--- a/Engine/Source/Programs/AutomationTool/Scripts/BenchmarkBuild/BenchmarkTaskBase.cs
+++ b/Engine/Source/Programs/AutomationTool/Scripts/BenchmarkBuild/BenchmarkTaskBase.cs
@@ -9,6 +9,7 @@ using System.Threading.Tasks;
 using EpicGames.Core;
 using UnrealBuildBase;
 using UnrealBuildTool;
+using Microsoft.Extensions.Logging;
 
 namespace AutomationTool.Benchmark
 {
@@ -67,6 +68,11 @@ namespace AutomationTool.Benchmark
 	/// </summary>
 	abstract class BenchmarkTaskBase
 	{
+		/// <summary>
+		/// Default logger instance
+		/// </summary>
+		protected static ILogger Logger => Log.Logger;
+
 		/// <summary>
 		/// True or false based on whether the task failed
 		/// </summary>
@@ -224,7 +230,7 @@ namespace AutomationTool.Benchmark
 			
 			if (Failed)
 			{
-				Log.TraceError("{0} failed. {1}", FullName, FailureString);
+				Logger.LogError("{FullName} failed. {FailureString}", FullName, FailureString);
 			}
 
 			try
@@ -233,7 +239,7 @@ namespace AutomationTool.Benchmark
 			}
 			catch (Exception Ex)
 			{
-				Log.TraceError("Cleanup of {0} failed. {1}", FullName, Ex);
+				Logger.LogError("Cleanup of {FullName} failed. {Ex}", FullName, Ex);
 			}
 		}
 
@@ -244,11 +250,11 @@ namespace AutomationTool.Benchmark
 		{
 			if (!Failed)
 			{
-				Log.TraceInformation("Task {0}:\t\t\t\t{1}", FullName, TaskTime.ToString(@"hh\:mm\:ss"));
+				Logger.LogInformation("Task {FullName}:\t\t\t\t{Arg1}", FullName, TaskTime.ToString(@"hh\:mm\:ss"));
 			}
 			else
 			{
-				Log.TraceInformation("Task {0}::\t\t\t\t{1} Failed. {2}", FullName, TaskTime.ToString(@"hh\:mm\:ss"), FailureString);
+				Logger.LogInformation("Task {FullName}::\t\t\t\t{Arg1} Failed. {FailureString}", FullName, TaskTime.ToString(@"hh\:mm\:ss"), FailureString);
 			}
 		}
 
diff --git a/Engine/Source/Programs/AutomationTool/Scripts/BuildCMakeLib.Automation.cs b/Engine/Source/Programs/AutomationTool/Scripts/BuildCMakeLib.Automation.cs
--- a/Engine/Source/Programs/AutomationTool/Scripts/BuildCMakeLib.Automation.cs
+++ b/Engine/Source/Programs/AutomationTool/Scripts/BuildCMakeLib.Automation.cs
@@ -9,6 +9,7 @@ using System.Text;
 using EpicGames.Core;
 using UnrealBuildTool;
 using UnrealBuildBase;
+using Microsoft.Extensions.Logging;
 
 // TODO Currently this only supports one lib and one platform at a time.
 // The reason for this is that the the library version and additional arguments (which are per platform) is passed on the command line.
@@ -359,12 +360,12 @@ public sealed class BuildCMakeLib : BuildCommand
 
 		public virtual void SetupTargetLib(TargetLib TargetLib, string TargetConfiguration)
 		{
-			LogInformation("Building {0} for {1} ({2})...", TargetLib.Name, FriendlyName, TargetConfiguration ?? "");
+			Logger.LogInformation("Building {Arg0} for {FriendlyName} ({Arg2})...", TargetLib.Name, FriendlyName, TargetConfiguration ?? "");
 
 			DirectoryReference CMakeTargetDirectory = GetProjectsDirectory(TargetLib, TargetConfiguration);
 			MakeFreshDirectoryIfRequired(CMakeTargetDirectory);
 
-			LogInformation("Generating projects for {0} for {1}", TargetLib.Name, FriendlyName);
+			Logger.LogInformation("Generating projects for {Arg0} for {FriendlyName}", TargetLib.Name, FriendlyName);
 
 			string CMakeArgs = GetCMakeSetupArguments(TargetLib, TargetConfiguration);
 			if (Run(CMakeCommand, CMakeArgs).ExitCode != 0)
@@ -568,7 +569,7 @@ public sealed class BuildCMakeLib : BuildCommand
 			{
 				if (File.Exists(PathComponent + "/make.exe") || File.Exists(PathComponent + "make.exe") || File.Exists(PathComponent + "/cygwin1.dll"))
 				{
-					LogInformation("Removing {0} from PATH since it contains possibly colliding make.exe", PathComponent);
+					Logger.LogInformation("Removing {PathComponent} from PATH since it contains possibly colliding make.exe", PathComponent);
 					continue;
 				}
 			}
@@ -594,7 +595,7 @@ public sealed class BuildCMakeLib : BuildCommand
 			string PathWithoutCygwin = RemoveOtherMakeAndCygwinFromPath(PrevPath);
 			Environment.SetEnvironmentVariable("PATH", CMakePath + ";" + MakePath + ";" + PathWithoutCygwin);
 			Environment.SetEnvironmentVariable("PATH", CMakePath + ";" + MakePath + ";" + Environment.GetEnvironmentVariable("PATH"));
-			LogInformation("set {0}={1}", "PATH", Environment.GetEnvironmentVariable("PATH"));
+			Logger.LogInformation("set {Arg0}={Arg1}", "PATH", Environment.GetEnvironmentVariable("PATH"));
 		}
 	}
 
@@ -715,13 +716,13 @@ public sealed class BuildCMakeLib : BuildCommand
 			{
 				if (!P4.TryDeleteEmptyChange(P4ChangeList))
 				{
-					LogInformation("Submitting changelist " + P4ChangeList.ToString());
+					Logger.LogInformation("{Text}", "Submitting changelist " + P4ChangeList.ToString());
 					int SubmittedChangeList = InvalidChangeList;
 					P4.Submit(P4ChangeList, out SubmittedChangeList);
 				}
 				else
 				{
-					LogInformation("Nothing to submit!");
+					Logger.LogInformation("Nothing to submit!");
 				}
 			}
 		}
diff --git a/Engine/Source/Programs/AutomationTool/Scripts/BuildCommonTools.Automation.cs b/Engine/Source/Programs/AutomationTool/Scripts/BuildCommonTools.Automation.cs
--- a/Engine/Source/Programs/AutomationTool/Scripts/BuildCommonTools.Automation.cs
+++ b/Engine/Source/Programs/AutomationTool/Scripts/BuildCommonTools.Automation.cs
@@ -8,6 +8,7 @@ using System.Text;
 using AutomationTool;
 using UnrealBuildTool;
 using EpicGames.Core;
+using Microsoft.Extensions.Logging;
 
 [Help("Builds common tools used by the engine which are not part of typical editor or game builds. Useful when syncing source-only on GitHub.")]
 [Help("platforms=<X>+<Y>+...", "Specifies on or more platforms to build for (defaults to the current host platform)")]
@@ -16,7 +17,7 @@ public class BuildCommonTools : BuildCommand
 {
 	public override void ExecuteBuild()
 	{
-		LogInformation("************************* BuildCommonTools");
+		Logger.LogInformation("************************* BuildCommonTools");
 
 		List<UnrealTargetPlatform> Platforms = new List<UnrealTargetPlatform>();
 
@@ -146,10 +147,10 @@ public class ZipProjectUp : BuildCommand
 		string InstallDirectory = ParseParamValue("install", "");
 		ProjectDirectory = Path.GetDirectoryName(ProjectDirectory);
 
-		LogInformation("Started zipping project up");
-		LogInformation("Project directory: {0}", ProjectDirectory);
-		LogInformation("Install directory: {0}", InstallDirectory);
-		LogInformation("Packaging up the project...");
+		Logger.LogInformation("Started zipping project up");
+		Logger.LogInformation("Project directory: {ProjectDirectory}", ProjectDirectory);
+		Logger.LogInformation("Install directory: {InstallDirectory}", InstallDirectory);
+		Logger.LogInformation("Packaging up the project...");
 
 		// Setup filters
 		FileFilter Filter = new FileFilter();
@@ -163,6 +164,6 @@ public class ZipProjectUp : BuildCommand
 
 		ZipFiles(new FileReference(InstallDirectory), new DirectoryReference(ProjectDirectory), Filter);
 
-		LogInformation("Completed zipping project up");
+		Logger.LogInformation("Completed zipping project up");
 	}
 }
diff --git a/Engine/Source/Programs/AutomationTool/Scripts/BuildCookRun.Automation.cs b/Engine/Source/Programs/AutomationTool/Scripts/BuildCookRun.Automation.cs
--- a/Engine/Source/Programs/AutomationTool/Scripts/BuildCookRun.Automation.cs
+++ b/Engine/Source/Programs/AutomationTool/Scripts/BuildCookRun.Automation.cs
@@ -8,6 +8,7 @@ using AutomationTool;
 using AutomationScripts;
 using UnrealBuildTool;
 using EpicGames.Core;
+using Microsoft.Extensions.Logging;
 
 [Help(@"Builds/Cooks/Runs a project.
 
@@ -41,12 +42,12 @@ public class BuildCookRun : BuildCommand, IProjectParamsHelpers
 
 		DoBuildCookRun(Params);
 
-		LogInformation("BuildCookRun time: {0:0.00} s", (DateTime.UtcNow - StartTime).TotalMilliseconds / 1000);
+		Logger.LogInformation("BuildCookRun time: {0:0.00} s", (DateTime.UtcNow - StartTime).TotalMilliseconds / 1000);
 	}
 
 	protected ProjectParams SetupParams()
 	{
-		LogInformation("Setting up ProjectParams for {0}", ProjectPath);
+		Logger.LogInformation("Setting up ProjectParams for {ProjectPath}", ProjectPath);
 
 		var Params = new ProjectParams
 		(
@@ -144,13 +145,13 @@ public class BuildCookRun : BuildCommand, IProjectParamsHelpers
 	private string GetDefaultMap(ProjectParams Params)
 	{
 		const string EngineEntryMap = "/Engine/Maps/Entry";
-		LogInformation("Trying to find DefaultMap in ini files");
+		Logger.LogInformation("Trying to find DefaultMap in ini files");
 		string DefaultMap = null;
 		var ProjectFolder = GetDirectoryName(Params.RawProjectPath.FullName);
 		var DefaultGameEngineConfig = CombinePaths(ProjectFolder, "Config", "DefaultEngine.ini");
 		if (FileExists(DefaultGameEngineConfig))
 		{
-			LogInformation("Looking for DefaultMap in {0}", DefaultGameEngineConfig);
+			Logger.LogInformation("Looking for DefaultMap in {DefaultGameEngineConfig}", DefaultGameEngineConfig);
 			DefaultMap = GetDefaultMapFromIni(DefaultGameEngineConfig, Params.DedicatedServer);
 			if (DefaultMap == null && Params.DedicatedServer)
 			{
@@ -162,7 +163,7 @@ public class BuildCookRun : BuildCommand, IProjectParamsHelpers
 			var BaseEngineConfig = CombinePaths(CmdEnv.LocalRoot, "Config", "BaseEngine.ini");
 			if (FileExists(BaseEngineConfig))
 			{
-				LogInformation("Looking for DefaultMap in {0}", BaseEngineConfig);
+				Logger.LogInformation("Looking for DefaultMap in {BaseEngineConfig}", BaseEngineConfig);
 				DefaultMap = GetDefaultMapFromIni(BaseEngineConfig, Params.DedicatedServer);
 				if (DefaultMap == null && Params.DedicatedServer)
 				{
@@ -173,12 +174,12 @@ public class BuildCookRun : BuildCommand, IProjectParamsHelpers
 		// We check for null here becase null == not found
 		if (DefaultMap == null)
 		{
-			LogInformation("No DefaultMap found, assuming: {0}", EngineEntryMap);
+			Logger.LogInformation("No DefaultMap found, assuming: {EngineEntryMap}", EngineEntryMap);
 			DefaultMap = EngineEntryMap;
 		}
 		else
 		{
-			LogInformation("Found DefaultMap={0}", DefaultMap);
+			Logger.LogInformation("Found DefaultMap={DefaultMap}", DefaultMap);
 		}
 		return DefaultMap;
 	}
@@ -246,7 +247,7 @@ public class BuildCookRun : BuildCommand, IProjectParamsHelpers
 		}
 
 		var Dest = ParseParamValue("ForeignDest", CombinePaths(@"C:\testue\foreign\", DestSample + "_ _Dir"));
-		LogInformation("Make a foreign sample {0} -> {1}", Src, Dest);
+		Logger.LogInformation("Make a foreign sample {Src} -> {Dest}", Src, Dest);
 
 		CloneDirectory(Src, Dest);
 
@@ -272,7 +273,7 @@ public class BuildCookRun : BuildCommand, IProjectParamsHelpers
 		}
 
 		var Dest = ParseParamValue("ForeignDest", CombinePaths(@"C:\testue\foreign\", DestSample + "_ _Dir"));
-		LogInformation("Make a foreign sample {0} -> {1}", Src, Dest);
+		Logger.LogInformation("Make a foreign sample {Src} -> {Dest}", Src, Dest);
 
 		CloneDirectory(Src, Dest);
 		DeleteDirectory_NoExceptions(CombinePaths(Dest, "Intermediate"));
diff --git a/Engine/Source/Programs/AutomationTool/Scripts/BuildDerivedDataCache.Automation.cs b/Engine/Source/Programs/AutomationTool/Scripts/BuildDerivedDataCache.Automation.cs
--- a/Engine/Source/Programs/AutomationTool/Scripts/BuildDerivedDataCache.Automation.cs
+++ b/Engine/Source/Programs/AutomationTool/Scripts/BuildDerivedDataCache.Automation.cs
@@ -10,6 +10,7 @@ using System.Threading.Tasks;
 using System.IO;
 using EpicGames.Core;
 using UnrealBuildBase;
+using Microsoft.Extensions.Logging;
 
 public class BuildDerivedDataCache : BuildCommand
 {
@@ -77,7 +78,7 @@ public class BuildDerivedDataCache : BuildCommand
 				}
 				if(FilteredPlatforms.Count == 0)
 				{
-					LogInformation("Did not find any project specific platforms for FeaturePack {0} out of supplied TargetPlatforms {1}, skipping it!", GameName, ProjectSpecificPlatforms);
+					Logger.LogInformation("Did not find any project specific platforms for FeaturePack {GameName} out of supplied TargetPlatforms {ProjectSpecificPlatforms}, skipping it!", GameName, ProjectSpecificPlatforms);
 					continue;
 				}
 				ProjectSpecificPlatforms = CommandUtils.CombineCommandletParams(FilteredPlatforms.Distinct().ToArray());
diff --git a/Engine/Source/Programs/AutomationTool/Scripts/BuildHlslcc.Automation.cs b/Engine/Source/Programs/AutomationTool/Scripts/BuildHlslcc.Automation.cs
--- a/Engine/Source/Programs/AutomationTool/Scripts/BuildHlslcc.Automation.cs
+++ b/Engine/Source/Programs/AutomationTool/Scripts/BuildHlslcc.Automation.cs
@@ -12,6 +12,7 @@ using AutomationTool;
 using UnrealBuildTool;
 using EpicGames.Core;
 using UnrealBuildBase;
+using Microsoft.Extensions.Logging;
 
 [Help("Builds Hlslcc using CMake build system.")]
 [Help("TargetPlatforms", "Specify a list of target platforms to build, separated by '+' characters (eg. -TargetPlatforms=Win64+Linux+Mac). Architectures are specified with '-'. Default is Win64+Linux.")]
@@ -72,7 +73,7 @@ class BuildHlslcc : BuildCommand
 		}
 
 		Environment.SetEnvironmentVariable("CMAKE_ROOT", DirectoryReference.Combine(CMakeRootDirectory, "share").ToString());
-		LogInformation("set {0}={1}", "CMAKE_ROOT", Environment.GetEnvironmentVariable("CMAKE_ROOT"));
+		Logger.LogInformation("set {Arg0}={Arg1}", "CMAKE_ROOT", Environment.GetEnvironmentVariable("CMAKE_ROOT"));
 
 		if (TargetData.Platform == UnrealTargetPlatform.Mac)
 		{
@@ -137,7 +138,7 @@ class BuildHlslcc : BuildCommand
 			{
 				// try to find bundled toolchain
 				DirectoryReference ToolchainDir  = DirectoryReference.Combine(ThirdPartySourceDirectory, "..", "..", "Extras", "ThirdPartyNotUE", "SDKs", "HostLinux", "Linux_x64");
-				LogInformation("LINUX_MULTIARCH_ROOT not defined. Looking for Linux toolchain in {0}...", ToolchainDir);
+				Logger.LogInformation("LINUX_MULTIARCH_ROOT not defined. Looking for Linux toolchain in {ToolchainDir}...", ToolchainDir);
 
 				IEnumerable<DirectoryReference> AvailableToolchains = DirectoryReference.EnumerateDirectories(ToolchainDir);
 				if (AvailableToolchains.Count() > 0)
@@ -149,7 +150,7 @@ class BuildHlslcc : BuildCommand
 
 			if (string.IsNullOrEmpty(ToolchainPath))
 			{
-				LogInformation("Bundled toolchain not found. Using system clang.");
+				Logger.LogInformation("Bundled toolchain not found. Using system clang.");
 			}
 			else
 			{
@@ -157,7 +158,7 @@ class BuildHlslcc : BuildCommand
 				ToolchainPath += "/bin/";
 			}
 
-			LogInformation("Using toolchain: {0}", ToolchainPath);
+			Logger.LogInformation("Using toolchain: {ToolchainPath}", ToolchainPath);
 
 			return string.Format(" -DCMAKE_C_COMPILER={0}clang -DCMAKE_CXX_COMPILER={0}clang++ ", ToolchainPath) + ExtraSettings;
 		}
@@ -406,7 +407,7 @@ class BuildHlslcc : BuildCommand
 
 				if (!bCleanOnly)
 				{
-					LogInformation("Generating projects for lib " + TargetData.ToString());
+					Logger.LogInformation("{Text}", "Generating projects for lib " + TargetData.ToString());
 
 					ProcessStartInfo StartInfo = new ProcessStartInfo();
 					StartInfo.FileName = CMakeName;
@@ -430,7 +431,7 @@ class BuildHlslcc : BuildCommand
 
 				if (!bCleanOnly)
 				{
-					LogInformation("Generating projects for lib " + TargetData.ToString());
+					Logger.LogInformation("{Text}", "Generating projects for lib " + TargetData.ToString());
 
 					ProcessStartInfo StartInfo = new ProcessStartInfo();
 					StartInfo.FileName = CMakeName;
@@ -438,10 +439,10 @@ class BuildHlslcc : BuildCommand
 					StartInfo.Arguments = GetCMakeArguments(TargetData, BuildConfig);
 
 					System.Console.WriteLine("Working in '{0}'", StartInfo.WorkingDirectory);
-					LogInformation("Working in '{0}'", StartInfo.WorkingDirectory);
+					Logger.LogInformation("Working in '{Arg0}'", StartInfo.WorkingDirectory);
 
 					System.Console.WriteLine("{0} {1}", StartInfo.FileName, StartInfo.Arguments);
-					LogInformation("{0} {1}", StartInfo.FileName, StartInfo.Arguments);
+					Logger.LogInformation("{Arg0} {Arg1}", StartInfo.FileName, StartInfo.Arguments);
 
 					if (RunLocalProcessAndLogOutput(StartInfo) != 0)
 					{
@@ -457,7 +458,7 @@ class BuildHlslcc : BuildCommand
 
 			if (!bCleanOnly)
 			{
-				LogInformation("Generating projects for lib " + TargetData.ToString());
+				Logger.LogInformation("{Text}", "Generating projects for lib " + TargetData.ToString());
 
 				ProcessStartInfo StartInfo = new ProcessStartInfo();
 				StartInfo.FileName = CMakeName;
@@ -508,7 +509,7 @@ class BuildHlslcc : BuildCommand
 				if (File.Exists(PathComponent + "/make.exe") || File.Exists(PathComponent + "make.exe") || File.Exists(PathComponent + "/cygwin1.dll"))
 				{
 					// gotcha!
-					LogInformation("Removing {0} from PATH since it contains possibly colliding make.exe", PathComponent);
+					Logger.LogInformation("Removing {PathComponent} from PATH since it contains possibly colliding make.exe", PathComponent);
 					continue;
 				}
 			}
@@ -543,7 +544,7 @@ class BuildHlslcc : BuildCommand
 				string PathWithoutCygwin = RemoveOtherMakeAndCygwinFromPath(PrevPath);
 				Environment.SetEnvironmentVariable("PATH", CMakePath + ";" + MakePath + ";" + PathWithoutCygwin);
 				Environment.SetEnvironmentVariable("PATH", CMakePath + ";" + MakePath + ";" + Environment.GetEnvironmentVariable("PATH"));
-				LogInformation("set {0}={1}", "PATH", Environment.GetEnvironmentVariable("PATH"));
+				Logger.LogInformation("set {Arg0}={Arg1}", "PATH", Environment.GetEnvironmentVariable("PATH"));
 			}
 		}
 	}
@@ -601,8 +602,8 @@ class BuildHlslcc : BuildCommand
 			StartInfo.WorkingDirectory = ConfigDirectory.ToString();
 			StartInfo.Arguments = MakeOptions;
 
-			LogInformation("Working in: {0}", StartInfo.WorkingDirectory);
-			LogInformation("{0} {1}", StartInfo.FileName, StartInfo.Arguments);
+			Logger.LogInformation("Working in: {Arg0}", StartInfo.WorkingDirectory);
+			Logger.LogInformation("{Arg0} {Arg1}", StartInfo.FileName, StartInfo.Arguments);
 
 			if (RunLocalProcessAndLogOutput(StartInfo) != 0)
 			{
@@ -938,13 +939,13 @@ class BuildHlslcc : BuildCommand
 		{
 			if (!P4.TryDeleteEmptyChange(P4ChangeList))
 			{
-				LogInformation("Submitting changelist " + P4ChangeList.ToString());
+				Logger.LogInformation("{Text}", "Submitting changelist " + P4ChangeList.ToString());
 				int SubmittedChangeList = InvalidChangeList;
 				P4.Submit(P4ChangeList, out SubmittedChangeList);
 			}
 			else
 			{
-				LogInformation("Nothing to submit!");
+				Logger.LogInformation("Nothing to submit!");
 			}
 		}
 	}
diff --git a/Engine/Source/Programs/AutomationTool/Scripts/BuildPhysX.Automation.cs b/Engine/Source/Programs/AutomationTool/Scripts/BuildPhysX.Automation.cs
--- a/Engine/Source/Programs/AutomationTool/Scripts/BuildPhysX.Automation.cs
+++ b/Engine/Source/Programs/AutomationTool/Scripts/BuildPhysX.Automation.cs
@@ -11,6 +11,7 @@ using System.Text.RegularExpressions;
 using EpicGames.Core;
 using UnrealBuildTool;
 using UnrealBuildBase;
+using Microsoft.Extensions.Logging;
 
 [Help("Builds PhysX/APEX libraries using CMake build system.")]
 [Help("TargetLibs", "Specify a list of target libraries to build, separated by '+' characters (eg. -TargetLibs=PhysX+APEX). Default is PhysX+APEX+NvCloth.")]
@@ -262,20 +263,20 @@ public sealed class BuildPhysX : BuildCommand
 		{
 			// make sure we set up the environment variable specifying where the root of the PhysX SDK is
 			Environment.SetEnvironmentVariable("GW_DEPS_ROOT", PhysX3RootDirectory.FullName.Replace('\\', '/'));
-			LogInformation("set {0}={1}", "GW_DEPS_ROOT", Environment.GetEnvironmentVariable("GW_DEPS_ROOT"));
+			Logger.LogInformation("set {Arg0}={Arg1}", "GW_DEPS_ROOT", Environment.GetEnvironmentVariable("GW_DEPS_ROOT"));
 			Environment.SetEnvironmentVariable("CMAKE_MODULE_PATH", DirectoryReference.Combine(PhysX3RootDirectory, "Externals", "CMakeModules").FullName.Replace('\\', '/'));
-			LogInformation("set {0}={1}", "CMAKE_MODULE_PATH", Environment.GetEnvironmentVariable("CMAKE_MODULE_PATH"));
+			Logger.LogInformation("set {Arg0}={Arg1}", "CMAKE_MODULE_PATH", Environment.GetEnvironmentVariable("CMAKE_MODULE_PATH"));
 
 			if (BuildHostPlatform.Current.Platform.IsInGroup(UnrealPlatformGroup.Unix))
 			{
 				Environment.SetEnvironmentVariable("CMAKE_ROOT", DirectoryReference.Combine(CMakeRootDirectory, "share").FullName);
-				LogInformation("set {0}={1}", "CMAKE_ROOT", Environment.GetEnvironmentVariable("CMAKE_ROOT"));
+				Logger.LogInformation("set {Arg0}={Arg1}", "CMAKE_ROOT", Environment.GetEnvironmentVariable("CMAKE_ROOT"));
 			}
 
 			DirectoryReference CMakeTargetDirectory = GetProjectsDirectory(TargetLib, TargetConfiguration);
 			MakeFreshDirectoryIfRequired(CMakeTargetDirectory);
 
-			LogInformation("Generating projects for lib " + TargetLib.ToString() + ", " + FriendlyName);
+			Logger.LogInformation("{Text}", "Generating projects for lib " + TargetLib.ToString() + ", " + FriendlyName);
 
 			ProcessStartInfo StartInfo = new ProcessStartInfo();
 			StartInfo.FileName = CMakeCommand;
@@ -396,8 +397,8 @@ public sealed class BuildPhysX : BuildCommand
 				? string.Format("{1} \"MAKE={0} {1}\"", MakeCommand, MakeOptions)
 				: MakeOptions;
 
-			LogInformation("Working in: {0}", StartInfo.WorkingDirectory);
-			LogInformation("{0} {1}", StartInfo.FileName, StartInfo.Arguments);
+			Logger.LogInformation("Working in: {Arg0}", StartInfo.WorkingDirectory);
+			Logger.LogInformation("{Arg0} {Arg1}", StartInfo.FileName, StartInfo.Arguments);
 
 			if (Utils.RunLocalProcessAndLogOutput(StartInfo, Log.Logger) != 0)
 			{
@@ -584,7 +585,7 @@ public sealed class BuildPhysX : BuildCommand
 				if (File.Exists(PathComponent + "/make.exe") || File.Exists(PathComponent + "make.exe") || File.Exists(PathComponent + "/cygwin1.dll"))
 				{
 					// gotcha!
-					LogInformation("Removing {0} from PATH since it contains possibly colliding make.exe", PathComponent);
+					Logger.LogInformation("Removing {PathComponent} from PATH since it contains possibly colliding make.exe", PathComponent);
 					continue;
 				}
 			}
@@ -611,7 +612,7 @@ public sealed class BuildPhysX : BuildCommand
 			string PathWithoutCygwin = RemoveOtherMakeAndCygwinFromPath(PrevPath);
 			Environment.SetEnvironmentVariable("PATH", CMakePath + ";" + MakePath + ";" + PathWithoutCygwin);
 			Environment.SetEnvironmentVariable("PATH", CMakePath + ";" + MakePath + ";" + Environment.GetEnvironmentVariable("PATH"));
-			LogInformation("set {0}={1}", "PATH", Environment.GetEnvironmentVariable("PATH"));
+			Logger.LogInformation("set {Arg0}={Arg1}", "PATH", Environment.GetEnvironmentVariable("PATH"));
 		}
 	}
 
@@ -808,13 +809,13 @@ public sealed class BuildPhysX : BuildCommand
 			{
 				if (!P4.TryDeleteEmptyChange(P4ChangeList))
 				{
-					LogInformation("Submitting changelist " + P4ChangeList.ToString());
+					Logger.LogInformation("{Text}", "Submitting changelist " + P4ChangeList.ToString());
 					int SubmittedChangeList = InvalidChangeList;
 					P4.Submit(P4ChangeList, out SubmittedChangeList);
 				}
 				else
 				{
-					LogInformation("Nothing to submit!");
+					Logger.LogInformation("Nothing to submit!");
 				}
 			}
 		}
@@ -1036,12 +1037,12 @@ class BuildPhysX_Linux : BuildPhysX.MakefileTargetPlatform
 		if (!string.IsNullOrEmpty(OriginalToolchainPath))
 		{
 			string ToolchainPathToUse = OriginalToolchainPath.Replace("v16_clang-9.0.1-centos7", "v12_clang-6.0.1-centos7");
-			LogInformation("Working around problems with newer clangs: {0} -> {1}", OriginalToolchainPath, ToolchainPathToUse);
+			Logger.LogInformation("Working around problems with newer clangs: {OriginalToolchainPath} -> {ToolchainPathToUse}", OriginalToolchainPath, ToolchainPathToUse);
 			Environment.SetEnvironmentVariable("LINUX_MULTIARCH_ROOT", ToolchainPathToUse);
 		}
 		else
 		{
-			LogWarning("LINUX_MULTIARCH_ROOT is not set!");
+			Logger.LogWarning("LINUX_MULTIARCH_ROOT is not set!");
 		}
 
 		base.SetupTargetLib(TargetLib, TargetConfiguration);
@@ -1082,14 +1083,14 @@ class BuildPhysX_Linux : BuildPhysX.MakefileTargetPlatform
 			StartInfo.Arguments = SOFile.FullName + " " + PSymbolFile.FullName;
 			StartInfo.RedirectStandardError = true;
 
-			LogInformation("Running: '{0} {1}'", StartInfo.FileName, StartInfo.Arguments);
+			Logger.LogInformation("Running: '{Arg0} {Arg1}'", StartInfo.FileName, StartInfo.Arguments);
 			Utils.RunLocalProcessAndLogOutput(StartInfo, Log.Logger);
 
 			// BreakpadSymbolEncoder
 			StartInfo.FileName = BreakpadSymbolEncoderPath.FullName + ExeSuffix;
 			StartInfo.Arguments = PSymbolFile.FullName + " " + SymbolFile.FullName;
 
-			LogInformation("Running: '{0} {1}'", StartInfo.FileName, StartInfo.Arguments);
+			Logger.LogInformation("Running: '{Arg0} {Arg1}'", StartInfo.FileName, StartInfo.Arguments);
 			Utils.RunLocalProcessAndLogOutput(StartInfo, Log.Logger);
 
 			// Clean up the Temp *.psym file, as they are no longer needed
@@ -1103,7 +1104,7 @@ class BuildPhysX_Linux : BuildPhysX.MakefileTargetPlatform
 					SOFile.FullName + " " +
 					StrippedFile.FullName;
 
-				LogInformation("Running: '{0} {1}'", StartInfo.FileName, StartInfo.Arguments);
+				Logger.LogInformation("Running: '{Arg0} {Arg1}'", StartInfo.FileName, StartInfo.Arguments);
 				Utils.RunLocalProcessAndLogOutput(StartInfo, Log.Logger);
 
 				// objcopy --only-keep-debug sofile.so sofile.debug
@@ -1112,7 +1113,7 @@ class BuildPhysX_Linux : BuildPhysX.MakefileTargetPlatform
 					SOFile.FullName + " " +
 					DebugFile.FullName;
 
-				LogInformation("Running: '{0} {1}'", StartInfo.FileName, StartInfo.Arguments);
+				Logger.LogInformation("Running: '{Arg0} {Arg1}'", StartInfo.FileName, StartInfo.Arguments);
 				Utils.RunLocalProcessAndLogOutput(StartInfo, Log.Logger);
 
 				// objcopy --add-gnu-debuglink=sofile.debug sofile_stripped sofile.so
@@ -1122,7 +1123,7 @@ class BuildPhysX_Linux : BuildPhysX.MakefileTargetPlatform
 					StrippedFile.FullName + " " +
 					SOFile.FullName;
 
-				LogInformation("Running: '{0} {1}'", StartInfo.FileName, StartInfo.Arguments);
+				Logger.LogInformation("Running: '{Arg0} {Arg1}'", StartInfo.FileName, StartInfo.Arguments);
 				Utils.RunLocalProcessAndLogOutput(StartInfo, Log.Logger);
 
 				GeneratedDebugSymbols.Add(SOFile.FullName, true);
@@ -1229,7 +1230,7 @@ class BuildPhysX_Mac : BuildPhysX.TargetPlatform
 	{
 		// build for x86
 		x86Build.SetupTargetLib(TargetLib, TargetConfiguration);
-		LogInformation("Building x86_64 lib slice");
+		Logger.LogInformation("Building x86_64 lib slice");
 		x86Build.BuildTargetLib(TargetLib, TargetConfiguration);
 
 		IEnumerable<FileReference> x86Libs = x86Build.EnumerateOutputFiles(TargetLib, TargetConfiguration).Distinct();
@@ -1239,7 +1240,7 @@ class BuildPhysX_Mac : BuildPhysX.TargetPlatform
 		{
 			string Extension = LibFile.GetExtension();
 			FileReference x86File = LibFile.ChangeExtension(Extension + "_x86_64");
-			LogInformation("Moving {0} to {1}", LibFile, x86File);
+			Logger.LogInformation("Moving {LibFile} to {x86File}", LibFile, x86File);
 			FileReference.Delete(x86File);
 			FileReference.Move(LibFile, x86File);
 
@@ -1248,7 +1249,7 @@ class BuildPhysX_Mac : BuildPhysX.TargetPlatform
 
 		// build for arm
 		ArmBuild.SetupTargetLib(TargetLib, TargetConfiguration);
-		LogInformation("Building arm64 lib slice");
+		Logger.LogInformation("Building arm64 lib slice");
 		ArmBuild.BuildTargetLib(TargetLib, TargetConfiguration);
 
 		IEnumerable<FileReference> ArmLibs = ArmBuild.EnumerateOutputFiles(TargetLib, TargetConfiguration).Distinct();
@@ -1259,7 +1260,7 @@ class BuildPhysX_Mac : BuildPhysX.TargetPlatform
 			string Extension = LibFile.GetExtension();
 			FileReference x86File = LibFile.ChangeExtension(Extension + "_x86_64");
 			FileReference ArmFile = LibFile.ChangeExtension(Extension + "_arm");
-			LogInformation("Moving {0} to {1}", LibFile, ArmFile);
+			Logger.LogInformation("Moving {LibFile} to {ArmFile}", LibFile, ArmFile);
 			FileReference.Delete(ArmFile);
 			FileReference.Move(LibFile, ArmFile);
 
@@ -1272,8 +1273,8 @@ class BuildPhysX_Mac : BuildPhysX.TargetPlatform
 		x86Slices = x86Slices.Distinct().ToList();
 		ArmSlices = ArmSlices.Distinct().ToList();
 
-		LogInformation("x86_64 slices generated: {0}", string.Join(", ", x86Slices));
-		LogInformation("arm64 slices generated: {0}", string.Join(", ", ArmSlices));
+		Logger.LogInformation("x86_64 slices generated: {Arg0}", string.Join(", ", x86Slices));
+		Logger.LogInformation("arm64 slices generated: {Arg0}", string.Join(", ", ArmSlices));
 
 		foreach (FileReference LibFile in x86Slices)
 		{ 
@@ -1291,10 +1292,10 @@ class BuildPhysX_Mac : BuildPhysX.TargetPlatform
 			StartInfo.Arguments = string.Format("-create {0} {1} -output {2}", ArmFile, x86File, OutputFile);
 			StartInfo.RedirectStandardError = true;
 
-			LogInformation("Running: 'lipo {0}'", StartInfo.Arguments);
+			Logger.LogInformation("Running: 'lipo {Arg0}'", StartInfo.Arguments);
 			if (Utils.RunLocalProcessAndLogOutput(StartInfo, Log.Logger) != 0)
 			{
-				LogError("Failed to create universal binary for {0}", LibFile);
+				Logger.LogError("Failed to create universal binary for {LibFile}", LibFile);
 			}
 			else
 			{
