[Automation] Provides Gauntlet the ability to more robustly secure dev kit account logins for relevant platforms.
Implements the feature on ps5 and xbox

diff --git a/Engine/Source/Programs/AutomationTool/Gauntlet/Framework/Base/Gauntlet.OnlineServiceLogin.cs b/Engine/Source/Programs/AutomationTool/Gauntlet/Framework/Base/Gauntlet.OnlineServiceLogin.cs
new file mode 100644
--- /dev/null
+++ b/Engine/Source/Programs/AutomationTool/Gauntlet/Framework/Base/Gauntlet.OnlineServiceLogin.cs
@@ -0,0 +1,17 @@
+﻿// Copyright Epic Games, Inc. All Rights Reserved.
+
+namespace Gauntlet
+{
+	/// <summary>
+	/// TargetDevice extension for platforms that contain an online service.
+	/// Enables more robust user sign in verification.
+	/// </summary>
+	internal interface IOnlineServiceLogin
+	{
+		/// <summary>
+		/// Verifies the device has connected to its relevant platform network.
+		/// </summary>
+		/// <returns>True if the device is logged into a platform network account</returns>
+		bool VerifyLogin();
+	}
+}
\ No newline at end of file
diff --git a/Engine/Source/Programs/AutomationTool/Gauntlet/Unreal/Base/Gauntlet.UnrealSession.cs b/Engine/Source/Programs/AutomationTool/Gauntlet/Unreal/Base/Gauntlet.UnrealSession.cs
--- a/Engine/Source/Programs/AutomationTool/Gauntlet/Unreal/Base/Gauntlet.UnrealSession.cs
+++ b/Engine/Source/Programs/AutomationTool/Gauntlet/Unreal/Base/Gauntlet.UnrealSession.cs
@@ -766,7 +766,7 @@ namespace Gauntlet
 				bool InstallSuccess = true;
 
 				// sort by constraints, so that we pick constrained devices first
-				List<UnrealSessionRole> SortedRoles = SessionRoles.OrderBy(R => R.Constraint.IsIdentity() ? 1 : 0).ToList();				
+				List<UnrealSessionRole> SortedRoles = SessionRoles.OrderBy(R => R.Constraint.IsIdentity() ? 1 : 0).ToList();
 
 				// first install all roles on these devices
 				foreach (var Role in SortedRoles)
@@ -852,7 +852,7 @@ namespace Gauntlet
 
 					// todo - should this be elsewhere?
 					AppConfig.Sandbox = Sandbox;
-					
+
 					IAppInstall Install = null;
 					bool bReinstallPerPass = Globals.Params.ParseParam("ReinstallPerPass");
 					if (RolesToInstalls == null || !RolesToInstalls.ContainsKey(Role) || bReinstallPerPass)
@@ -867,13 +867,23 @@ namespace Gauntlet
 						IDeviceUsageReporter.RecordStart(Device.Name, Device.Platform, IDeviceUsageReporter.EventType.Install, IDeviceUsageReporter.EventState.Success, BuildSource.BuildName);
 						try
 						{
+							if (Globals.Params.ParseParam("VerifyLogin") && Device is IOnlineServiceLogin)
+							{
+								Log.Info("\nVerifying device login...");
+								if (!(Device as IOnlineServiceLogin).VerifyLogin())
+								{
+									throw new AutomationException("Unable to secure login to an online platform account!");
+								}
+								Log.Info("Success! User signed-in.\n");
+							}
+
 							Install = Device.InstallApplication(AppConfig);
 							IDeviceUsageReporter.RecordEnd(Device.Name, Device.Platform, IDeviceUsageReporter.EventType.Install, IDeviceUsageReporter.EventState.Success);
 						}
 						catch (System.Exception Ex)
 						{
 							// Warn, ignore the device, and do not continue
-							string ErrorMessage = string.Format("Failed to install app onto device {0} for role {1}. {2}. Will retry with new device", Device, Role, Ex);
+							string ErrorMessage = string.Format("Encountered error setting up device {0} for role {1}. {2}. Will retry with new device", Device, Role, Ex);
 							if (ErrorMessage.Contains("not enough space"))
 							{
 								Log.Error(KnownLogEvents.Gauntlet_DeviceEvent, ErrorMessage);
@@ -933,7 +943,6 @@ namespace Gauntlet
 
 					Log.Info("\nUnable to install application for session (retries left={0})\n", --SessionRetries);
 				}
-				
 
 				if (InstallSuccess && Globals.CancelSignalled == false)
 				{
