P4WriteConfig: Prefer ignore file in the root folder if it exists

diff --git a/Engine/Source/Programs/AutomationTool/AutomationUtils/P4WriteConfig.cs b/Engine/Source/Programs/AutomationTool/AutomationUtils/P4WriteConfig.cs
--- a/Engine/Source/Programs/AutomationTool/AutomationUtils/P4WriteConfig.cs
+++ b/Engine/Source/Programs/AutomationTool/AutomationUtils/P4WriteConfig.cs
@@ -12,7 +12,7 @@ using Microsoft.Extensions.Logging;
 namespace AutomationTool
 {
 	[Help("Auto-detects P4 settings based on the current path and creates a p4config file with the relevant settings.")]
-	[ParamHelp("SetIgnore", "Adds a P4IGNORE to the default file (Engine/Extras/Perforce/p4ignore)", ParamType = typeof(bool), Flag = "-SetIgnore")]
+	[ParamHelp("SetIgnore", "Adds a P4IGNORE to the default file (p4ignore in the root, or Engine/Extras/Perforce/p4ignore)", ParamType = typeof(bool), Flag = "-SetIgnore")]
 	[ParamHelp("Path", "Write to a path other than the current directory")]
 	[ParamHelp("p4port=<server:port>", "Optional hint/override of the server to use during lookup")]
 	[ParamHelp("p4user=<username>", "Optional hint/override of the username to use during lookup")]
@@ -64,7 +64,12 @@ namespace AutomationTool
 
 			if (SetIgnore)
 			{
-				string IgnorePath = Path.Combine(Unreal.EngineDirectory.ToString(), "Extras", "Perforce", "p4ignore");
+				// If a p4 ignore file is in the root, prefer it
+				string IgnorePath = Directory.EnumerateFiles(Unreal.RootDirectory.FullName, "*p4ignore*", SearchOption.TopDirectoryOnly).FirstOrDefault();
+				if (string.IsNullOrEmpty(IgnorePath) )
+				{
+					IgnorePath = Path.Combine(Unreal.EngineDirectory.ToString(), "Extras", "Perforce", "p4ignore");
+				}
 				P4Config.AppendLine($"P4IGNORE={IgnorePath}");
 			}
 
