UAT: Support specifying a duration and paused flag for new tool deployments.

diff --git a/Engine/Source/Programs/AutomationTool/BuildGraph/Tasks/DeployToolTask.cs b/Engine/Source/Programs/AutomationTool/BuildGraph/Tasks/DeployToolTask.cs
--- a/Engine/Source/Programs/AutomationTool/BuildGraph/Tasks/DeployToolTask.cs
+++ b/Engine/Source/Programs/AutomationTool/BuildGraph/Tasks/DeployToolTask.cs
@@ -41,6 +41,18 @@ namespace AutomationTool.Tasks
 		[TaskParameter]
 		public string Version = String.Empty;
 
+		/// <summary>
+		/// Duration over which to roll out the tool, in minutes.
+		/// </summary>
+		[TaskParameter(Optional = true)]
+		public int Duration = 0;
+
+		/// <summary>
+		/// Whether to create the deployment as paused
+		/// </summary>
+		[TaskParameter(Optional = true)]
+		public bool Paused = false;
+
 		/// <summary>
 		/// Zip file containing the data to upload
 		/// </summary>
@@ -139,6 +151,14 @@ namespace AutomationTool.Tasks
 
 						MultipartFormDataContent content = new MultipartFormDataContent();
 						content.Add(new StringContent(Parameters.Version), "version");
+						if (Parameters.Duration != 0)
+						{
+							content.Add(new StringContent(TimeSpan.FromMinutes(Parameters.Duration).ToString()), "duration");
+						}
+						if (Parameters.Paused)
+						{
+							content.Add(new StringContent("true"), "paused");
+						}
 						content.Add(streamContent, "file", file.GetFileName());
 						request.Content = content;
 
