UAT warning formatting fix
#rnx

diff --git a/Engine/Source/Programs/AutomationTool/Scripts/CopyBuildToStagingDirectory.Automation.cs b/Engine/Source/Programs/AutomationTool/Scripts/CopyBuildToStagingDirectory.Automation.cs
--- a/Engine/Source/Programs/AutomationTool/Scripts/CopyBuildToStagingDirectory.Automation.cs
+++ b/Engine/Source/Programs/AutomationTool/Scripts/CopyBuildToStagingDirectory.Automation.cs
@@ -1559,7 +1559,9 @@ namespace AutomationScripts
 				Nullable<bool> ShouldStage = ShouldStageConfigFile(SC, ConfigDir, ConfigFile, PlatformExtensionName);
 				if (ShouldStage == null)
 				{
-					Logger.LogWarning("The config file '{Arg0}' will be staged, but is not explicitly allowed or denied. Add +AllowedConfigFiles={Unknown1} or +DisallowedConfigFiles={Unknown2} to the [Staging] section of DefaultGame.ini", SC.GetStagedFileLocation(ConfigFile));
+					StagedFileReference StagedFile = SC.GetStagedFileLocation(ConfigFile);
+					Logger.LogWarning("The config file '{StagedFile}' will be staged, but is not explicitly allowed or denied. Add +AllowedConfigFiles={AllowedFile} or +DisallowedConfigFiles={DisallowedFile} to the [Staging] section of DefaultGame.ini"
+						, StagedFile, StagedFile, StagedFile);
 				}
 
 				if (ShouldStage ?? true)
