Fixed deprecation warning
#rnx

diff --git a/Engine/Source/Programs/AutomationTool/Scripts/CopyBuildToStagingDirectory.Automation.cs b/Engine/Source/Programs/AutomationTool/Scripts/CopyBuildToStagingDirectory.Automation.cs
--- a/Engine/Source/Programs/AutomationTool/Scripts/CopyBuildToStagingDirectory.Automation.cs
+++ b/Engine/Source/Programs/AutomationTool/Scripts/CopyBuildToStagingDirectory.Automation.cs
@@ -766,7 +766,7 @@ namespace AutomationScripts
 				CbObject ZenServerObject = ManifestObject["zenserver"].AsObject();
 				if (ZenServerObject != CbObject.Empty)
 				{
-					CommandUtils.LogInformation("Reading oplog from Zen...");
+					Logger.LogInformation("Reading oplog from Zen...");
 					string Host = "localhost";
 					int Port = 1337;
 					string ProjectId = ZenServerObject["projectid"].AsString();
