Fixed another deprecation warning
#rnx

diff --git a/Engine/Source/Programs/AutomationTool/Scripts/CopyBuildToStagingDirectory.Automation.cs b/Engine/Source/Programs/AutomationTool/Scripts/CopyBuildToStagingDirectory.Automation.cs
--- a/Engine/Source/Programs/AutomationTool/Scripts/CopyBuildToStagingDirectory.Automation.cs
+++ b/Engine/Source/Programs/AutomationTool/Scripts/CopyBuildToStagingDirectory.Automation.cs
@@ -813,7 +813,7 @@ namespace AutomationScripts
 									Arguments += " -unattended";
 								}
 
-								LogInformation("Running UnrealPak with arguments: {0}", Arguments);
+								Logger.LogInformation("Running UnrealPak with arguments: {UnrealPakParams}", Arguments);
 								RunAndLog(CmdEnv, GetUnrealPakLocation().FullName, Arguments, Options: ERunOptions.Default | ERunOptions.UTF8Output);
 								bAttemptAutoLaunchOnFailure = false;
 							}
