BuildGraph: Use .NET API for copying temp storage dirs and files when running under Wine

diff --git a/Engine/Source/Programs/AutomationTool/BuildGraph/TempStorage.cs b/Engine/Source/Programs/AutomationTool/BuildGraph/TempStorage.cs
--- a/Engine/Source/Programs/AutomationTool/BuildGraph/TempStorage.cs
+++ b/Engine/Source/Programs/AutomationTool/BuildGraph/TempStorage.cs
@@ -1298,10 +1298,24 @@ namespace AutomationTool
 			}
 			return false;
 		}
-
+		
+		/// <summary>
+		/// Copy a temp storage .zip file to directory.
+		/// Uses Robocopy on Windows, rsync on Linux/macOS and native .NET API when under Wine.
+		/// </summary>
+		/// <param name="ZipFile">.zip file to copy</param>
+		/// <param name="RootDir">Destination dir</param>
+		/// <returns>Output from the copy operation</returns>
 		static string CopyFile(FileInfo ZipFile, DirectoryReference RootDir)
 		{
-			if (HostPlatform.Current.HostEditorPlatform == UnrealTargetPlatform.Win64)
+			if (BuildHostPlatform.Current.IsRunningOnWine())
+			{
+				string sourceFile = ZipFile.FullName;
+				string destFile = Path.Join(RootDir.FullName, ZipFile.Name);
+				File.Copy(sourceFile, destFile, true);
+				return $".NET file copy. Source='{sourceFile}' Dest='{destFile}'";
+			}
+			else if (HostPlatform.Current.HostEditorPlatform == UnrealTargetPlatform.Win64)
 			{
 				return CommandUtils.RunAndLog(GetRoboCopyExe(), $"\"{ZipFile.DirectoryName}\" \"{RootDir}\" \"{ZipFile.Name}\" /w:5 /r:10", MaxSuccessCode: 3, Options: CommandUtils.ERunOptions.AppMustExist | CommandUtils.ERunOptions.NoLoggingOfRunCommand);
 			}
@@ -1311,7 +1325,34 @@ namespace AutomationTool
 			}
 		}
 
+		/// <summary>
+		/// Copy a directory using the most suitable option
+		/// </summary>
+		/// <param name="SourceDir">Source directory</param>
+		/// <param name="DestinationDir">Directory directory</param>
+		/// <returns>Output from directory copy operation</returns>
+		/// <exception cref="DirectoryNotFoundException">If source directory wasn't found</exception>
 		static string CopyDirectory(DirectoryReference SourceDir, DirectoryReference DestinationDir)
+		{
+			if (BuildHostPlatform.Current.IsRunningOnWine())
+			{
+				return CopyDirectoryDotNet(SourceDir, DestinationDir);
+			}
+			else
+			{
+				return CopyDirectoryExternalTool(SourceDir, DestinationDir);
+			}
+		}
+
+		/// <summary>
+		/// Copy a directory using an external tool native to OS
+		/// Robocopy Windows and rsync for Linux and macOS.
+		/// </summary>
+		/// <param name="SourceDir">Source directory</param>
+		/// <param name="DestinationDir">Directory directory</param>
+		/// <returns>Output from directory copy operation</returns>
+		/// <exception cref="DirectoryNotFoundException">If source directory wasn't found</exception>
+		static string CopyDirectoryExternalTool(DirectoryReference SourceDir, DirectoryReference DestinationDir)
 		{
 			if (HostPlatform.Current.HostEditorPlatform == UnrealTargetPlatform.Win64)
 			{
@@ -1322,6 +1363,40 @@ namespace AutomationTool
 				return CommandUtils.RunAndLog("rsync", $"-vam --include=\"**\" \"{SourceDir}/\" \"{DestinationDir}/\"", Options: CommandUtils.ERunOptions.AppMustExist | CommandUtils.ERunOptions.NoLoggingOfRunCommand);
 			}
 		}
+		
+		/// <summary>
+		/// Copy a directory using .NET SDK directory and file copy API
+		/// </summary>
+		/// <param name="SourceDir">Source directory</param>
+		/// <param name="DestinationDir">Directory directory</param>
+		/// <returns>Output from directory copy operation</returns>
+		/// <exception cref="DirectoryNotFoundException">If source directory wasn't found</exception>
+		static string CopyDirectoryDotNet(DirectoryReference SourceDir, DirectoryReference DestinationDir)
+		{
+			DirectoryInfo dir = new (SourceDir.FullName);
+
+			if (!dir.Exists)
+			{
+				throw new DirectoryNotFoundException($"Source directory not found: {dir.FullName}");
+			}
+
+			DirectoryInfo[] dirs = dir.GetDirectories();
+			Directory.CreateDirectory(DestinationDir.FullName);
+
+			foreach (FileInfo file in dir.GetFiles())
+			{
+				string targetFilePath = Path.Combine(DestinationDir.FullName, file.Name);
+				file.CopyTo(targetFilePath);
+			}
+
+			foreach (DirectoryInfo subDir in dirs)
+			{
+				string newDestinationDir = Path.Combine(DestinationDir.FullName, subDir.Name);
+				CopyDirectoryDotNet(new DirectoryReference(subDir.FullName), new DirectoryReference(newDestinationDir));
+			}
+
+			return ".NET directory copy";
+		}
 
 		static string GetRoboCopyExe()
 		{
