Provide a UAT command to deploy an archived build

diff --git a/Engine/Source/Programs/AutomationTool/Gauntlet/Unreal/InstallUnrealBuild.cs b/Engine/Source/Programs/AutomationTool/Gauntlet/Unreal/InstallUnrealBuild.cs
new file mode 100644
--- /dev/null
+++ b/Engine/Source/Programs/AutomationTool/Gauntlet/Unreal/InstallUnrealBuild.cs
@@ -0,0 +1,58 @@
+// Copyright Epic Games, Inc. All Rights Reserved.
+
+using System;
+using AutomationTool;
+using UnrealBuildBase;
+
+
+namespace Gauntlet
+{
+	/// <summary>
+	/// Main class for Unreal Installer script
+	/// </summary>
+	[Help("Install Unreal Build onto one or more connected devices using Gauntlet")]
+	[ParamHelp("platform=<platform>", "Platform to use")]
+	[ParamHelp("path=<path>", "Path to a build to install")]
+	[ParamHelp("device=<device>", "Device to install a build on")]
+	[ParamHelp("devices=<devices>", "List of devices, separated with commas, to install a build on")]
+	[ParamHelp("commandline=\"-arg1 -arg2 -etc\"", "Commandline to write to UECommandLine.txt on the device")]
+	[ParamHelp("parallel=[1:4]", "Number of installs to run in parallel",
+		ParamType = typeof(string), Choices = new string[] { "1", "2", "3", "4" }, DefaultValue = "1")]
+	[ParamHelp("project=<project>", "Name of the project", Required = true)]
+
+	public class InstallUnrealBuild : BuildCommand
+	{
+
+		/// <summary>
+		/// Entrance point for Unreal Installer script
+		/// </summary>
+		/// <returns></returns>
+		public override ExitCode Execute()
+		{
+			string PlatformParam = ParseParamValue("platform", string.Empty);
+
+			string BuildPath = ParseParamValue("build", string.Empty);
+			BuildPath = ParseParamValue("path", BuildPath);
+
+			string ProjectName = ParseParamValue("project", string.Empty);
+			string CommandLine = ParseParamValue("commandline", string.Empty);
+			string DevicesArg = ParseParamValue("device", string.Empty);
+			DevicesArg = ParseParamValue("devices", DevicesArg);
+
+			int ParallelTasks = Convert.ToInt32(ParseParamValue("parallel", "1"));
+
+			bool Success = false;
+
+			try
+			{
+				Success = InstallUnreal.RunInstall(PlatformParam, BuildPath, ProjectName, CommandLine, DevicesArg, ParallelTasks);
+			}
+			catch (Exception Ex)
+			{
+				Log.Error("Failed to install build from {Path} to {Device}. {Exception}", BuildPath, DevicesArg, Ex.Message);
+			}
+
+			return Success ? ExitCode.Success : ExitCode.Error_Unknown;
+		}
+	}
+}
diff --git a/Engine/Source/Programs/AutomationTool/Gauntlet/Unreal/Utils/Gauntlet.InstallUnreal.cs b/Engine/Source/Programs/AutomationTool/Gauntlet/Unreal/Utils/Gauntlet.InstallUnreal.cs
new file mode 100644
--- /dev/null
+++ b/Engine/Source/Programs/AutomationTool/Gauntlet/Unreal/Utils/Gauntlet.InstallUnreal.cs
@@ -0,0 +1,120 @@
+// Copyright Epic Games, Inc. All Rights Reserved.
+
+using System;
+using System.Collections.Generic;
+using System.IO;
+using System.Threading;
+using System.Linq;
+using AutomationTool;
+using UnrealBuildTool;
+using System.Threading.Tasks;
+
+
+namespace Gauntlet
+{
+	/// <summary>
+	/// Base class implementing Installer functionality
+	/// </summary>
+	public class InstallUnreal
+	{
+		/// <summary>
+		/// A method which is responsible for the core Installer functionality.
+		/// Iterates through devices and installs the project to each device.
+		/// </summary>
+		/// <returns></returns>
+		public static bool RunInstall(
+			string PlatformParam,
+			string BuildPath,
+			string ProjectName,
+			string CommandLine,
+			string DevicesArg,
+			int ParallelTasks )
+		{
+			if (string.IsNullOrEmpty(PlatformParam) || string.IsNullOrEmpty(BuildPath))
+			{
+				throw new Exception("need -platform=<platform> and -path=\"path to build\"");
+			}
+
+			if (Directory.Exists(BuildPath) == false)
+			{
+				throw new AutomationException("Path {0} does not exist", BuildPath);
+			}
+
+			UnrealTargetPlatform Platform = UnrealTargetPlatform.Parse(PlatformParam);
+
+			// find all build sources that can be created a folder path
+			IEnumerable<IFolderBuildSource> BuildSources = Gauntlet.Utils.InterfaceHelpers.FindImplementations<IFolderBuildSource>();
+
+			if (BuildSources.Count() == 0)
+			{
+				throw new AutomationException("No BuildSources found for platform {0}", Platform);
+			}
+
+			IEnumerable<IBuild> Builds = BuildSources.Where(S => S.CanSupportPlatform(Platform)).SelectMany(S => S.GetBuildsAtPath(ProjectName, BuildPath));
+
+			if (Builds.Count() == 0)
+			{
+				throw new AutomationException("No builds for {0} found at {1}", Platform, BuildPath);
+			}
+
+			IEnumerable<ITargetDevice> DeviceList = null;
+
+			if (string.IsNullOrEmpty(DevicesArg))
+			{
+				// find all build sources that can be created a folder path
+				IEnumerable<IDefaultDeviceSource> DeviceSources = Gauntlet.Utils.InterfaceHelpers.FindImplementations<IDefaultDeviceSource>();
+
+				DeviceList = DeviceSources.Where(S => S.CanSupportPlatform(Platform)).SelectMany(S => S.GetDefaultDevices());
+			}
+			else
+			{
+				IEnumerable<string> Devices = DevicesArg.Split(new[] { "," }, StringSplitOptions.RemoveEmptyEntries).Select(S => S.Trim());
+
+				IDeviceFactory Factory = Gauntlet.Utils.InterfaceHelpers.FindImplementations<IDeviceFactory>()
+					.Where(F => F.CanSupportPlatform(Platform))
+					.FirstOrDefault();
+
+				if (Factory == null)
+				{
+					throw new AutomationException("No IDeviceFactory implmenetation that supports {0}", Platform);
+				}
+
+				DeviceList = Devices.Select(D => Factory.CreateDevice(D, string.Empty)).ToArray();
+			}
+
+			if (DeviceList.Count() == 0)
+			{
+				throw new AutomationException("No devices found for {0}", Platform);
+			}
+
+			var POptions = new ParallelOptions { MaxDegreeOfParallelism = ParallelTasks };
+
+			// now copy it four builds at a time
+			Parallel.ForEach(DeviceList, POptions, Device =>
+			{
+				DateTime StartTime = DateTime.Now;
+
+				UnrealAppConfig Config = new UnrealAppConfig();
+
+				Config.CommandLine = CommandLine;
+				Config.Build = Builds.First();
+				Config.ProjectName = ProjectName;
+
+				Log.Info("Installing build on device {0}", Device.Name);
+
+				IAppInstall Install = Device.InstallApplication(Config);
+				Device.Run(Install);
+
+				TimeSpan Elapsed = (DateTime.Now - StartTime);
+				Log.Info("Installed on device {0} in {1:D2}m:{2:D2}s", Device.Name, Elapsed.Minutes, Elapsed.Seconds);
+			});
+
+			DeviceList = null;
+
+			// wait for all the adb tasks to start up or UAT will kill them on exit...
+			Thread.Sleep(5000);
+
+			return true;
+		}
+	}
+}
