UAT: Skip detecting msbuild path under Wine

Right now, it's expected that .NET Core dependencies are compiled outside of Wine.

diff --git a/Engine/Source/Programs/AutomationTool/AutomationUtils/CommandEnvironment.cs b/Engine/Source/Programs/AutomationTool/AutomationUtils/CommandEnvironment.cs
--- a/Engine/Source/Programs/AutomationTool/AutomationUtils/CommandEnvironment.cs
+++ b/Engine/Source/Programs/AutomationTool/AutomationUtils/CommandEnvironment.cs
@@ -10,6 +10,7 @@ using EpicGames.Core;
 using System.Text.RegularExpressions;
 using UnrealBuildBase;
 using Microsoft.Extensions.Logging;
+using UnrealBuildTool;
 
 namespace AutomationTool
 {
@@ -206,6 +207,13 @@ namespace AutomationTool
 			// Assume we have the capability co compile.
 			HasCapabilityToCompile = true;
 
+			if (BuildHostPlatform.Current.IsRunningOnWine())
+			{
+				// Just set an empty path as we currently compile .NET Framework/Core dependencies outside Wine
+				FrameworkMsbuildPath = "";
+				return;
+			}
+
 			try
 			{
 				FrameworkMsbuildPath = HostPlatform.Current.GetFrameworkMsbuildExe();
