Turnkey:  Fix for determining supported device software versions

- issue
  - turnkey verifysdk command is not returning the supported device software version ranges for platforms without specific version ranges per device type
    - this also causes it to not report the current software version on the device
  - the editor uses this turnkey command so it is currently unable to report current or supported device software ranges for some platforms

- fix
  - turnkey verifysdk command now queries for supported device software range for the specific device type
    - if no supported version range is found, it then queries for device type independent versions for that platform

- testing
  - edtior now correctly reports software versions for Windows as well as some NDA'd platform devices

diff --git a/Engine/Source/Programs/AutomationTool/Turnkey/Commands/VerifySdk.cs b/Engine/Source/Programs/AutomationTool/Turnkey/Commands/VerifySdk.cs
--- a/Engine/Source/Programs/AutomationTool/Turnkey/Commands/VerifySdk.cs
+++ b/Engine/Source/Programs/AutomationTool/Turnkey/Commands/VerifySdk.cs
@@ -240,6 +240,11 @@ namespace Turnkey.Commands
 
 						// get the min/max versions from PlatformSDK
 						SDKCollection SoftwareInfo = PlatformSDK.GetAllSoftwareInfo(Device.Type, Device.SoftwareVersion);
+						if (!SoftwareInfo.Sdks.Any())
+						{
+							// if we found no supported software versions for the specific device type, query for device type independent versions
+							SoftwareInfo = PlatformSDK.GetAllSoftwareInfo(null, Device.SoftwareVersion);
+						}
 
 						SdkUtils.LocalAvailability DeviceState = SdkUtils.LocalAvailability.None;
 						if (!bArePrerequisitesValid)
