- Fixing staging programs where the .target file was being looked for in wrong spot

diff --git a/Engine/Source/Programs/AutomationTool/Scripts/CopyBuildToStagingDirectory.Automation.cs b/Engine/Source/Programs/AutomationTool/Scripts/CopyBuildToStagingDirectory.Automation.cs
--- a/Engine/Source/Programs/AutomationTool/Scripts/CopyBuildToStagingDirectory.Automation.cs
+++ b/Engine/Source/Programs/AutomationTool/Scripts/CopyBuildToStagingDirectory.Automation.cs
@@ -4793,7 +4793,7 @@ namespace AutomationScripts
 					string Target = TargetAndConfig.Item1;
 					UnrealTargetConfiguration Config = TargetAndConfig.Item2;
 					UnrealArchitectures Architecture = TargetAndConfig.Item3;
-					DirectoryReference ReceiptBaseDir = Params.IsCodeBasedProject ? ProjectDir : EngineDir;
+					DirectoryReference ReceiptBaseDir = (Params.IsCodeBasedProject && !Params.ProgramTargets.Contains(Target)) ? ProjectDir : EngineDir;
 
 					Platform PlatformInstance = Platform.Platforms[StagePlatform];
 					UnrealTargetPlatform[] SubPlatformsToStage = PlatformInstance.GetStagePlatforms();
