Support multiple path filters in FindModifiedFiles, update modified files tasks to also consider .h

#rnx

diff --git a/Engine/Source/Programs/AutomationTool/BuildGraph/Tasks/FindModifiedFiles.cs b/Engine/Source/Programs/AutomationTool/BuildGraph/Tasks/FindModifiedFiles.cs
--- a/Engine/Source/Programs/AutomationTool/BuildGraph/Tasks/FindModifiedFiles.cs
+++ b/Engine/Source/Programs/AutomationTool/BuildGraph/Tasks/FindModifiedFiles.cs
@@ -22,9 +22,9 @@ namespace AutomationTool.Tasks
 	public class FindModifiedFilesTaskParameters
 	{
 		/// <summary>
-		/// Filter Path (default is ...)
+		///  List of file specifications separated by semicolon (default is ...)
 		/// </summary>
-		[TaskParameter(Optional = true)]
+		[TaskParameter(Optional=true, ValidationType = TaskParameterValidationType.FileSpec)]
 		public string Path = "...";
 
 		/// <summary>
@@ -82,55 +82,61 @@ namespace AutomationTool.Tasks
 		{
 			using IPerforceConnection Connection = await PerforceConnection.CreateAsync(CommandUtils.P4Settings, Log.Logger);
 
-			StringBuilder Filter = new StringBuilder($"//{Connection.Settings.ClientName}/{Parameters.Path}");
-			if (Parameters.Change > 0)
-			{
-				Filter.Append($"@={Parameters.Change}");
-			}
-			else if (Parameters.MinChange > 0)
+			SortedSet<string> AllLocalFiles = new SortedSet<string>();
+			foreach (string Path in Parameters.Path.Split(";"))
 			{
-				if (Parameters.MaxChange > 0)
+				StringBuilder Filter = new StringBuilder($"//{Connection.Settings.ClientName}/{Path}");
+				if (Parameters.Change > 0)
 				{
-					Filter.Append($"@{Parameters.MinChange},{Parameters.MaxChange}");
+					Filter.Append($"@={Parameters.Change}");
+				}
+				else if (Parameters.MinChange > 0)
+				{
+					if (Parameters.MaxChange > 0)
+					{
+						Filter.Append($"@{Parameters.MinChange},{Parameters.MaxChange}");
+					}
+					else
+					{
+						Filter.Append($"@>={Parameters.MinChange}");
+					}
 				}
 				else
 				{
-					Filter.Append($"@>={Parameters.MinChange}");
+					throw new AutomationException("Change or MinChange must be specified to FindModifiedFiles task");
 				}
-			}
-			else
-			{
-				throw new AutomationException("Change or MinChange must be specified to FindModifiedFiles task");
-			}
 
-			StreamRecord StreamRecord = await Connection.GetStreamAsync(CommandUtils.P4Env.Branch, true);
-			PerforceViewMap ViewMap = PerforceViewMap.Parse(StreamRecord.View);
+				StreamRecord StreamRecord = await Connection.GetStreamAsync(CommandUtils.P4Env.Branch, true);
+				PerforceViewMap ViewMap = PerforceViewMap.Parse(StreamRecord.View);
 
-			List<string> LocalFiles = new List<string>();
+				HashSet<string> LocalFiles = new HashSet<string>();
 
-			List<FilesRecord> Files = await Connection.FilesAsync(FilesOptions.None, Filter.ToString());
-			foreach (FilesRecord File in Files)
-			{
-				string LocalFile;
-				if (ViewMap.TryMapFile(File.DepotFile, StringComparison.OrdinalIgnoreCase, out LocalFile))
+				List<FilesRecord> Files = await Connection.FilesAsync(FilesOptions.None, Filter.ToString());
+				foreach (FilesRecord File in Files)
 				{
-					LocalFiles.Add(LocalFile);
+					string LocalFile;
+					if (ViewMap.TryMapFile(File.DepotFile, StringComparison.OrdinalIgnoreCase, out LocalFile))
+					{
+						LocalFiles.Add(LocalFile);
+					}
+					else
+					{
+						Logger.LogInformation("Unable to map {DepotFile} to workspace; skipping.", File.DepotFile);
+					}
 				}
-				else
+
+				Logger.LogInformation("Found {NumFiles} modified files matching {Filter}", LocalFiles.Count, Filter.ToString());
+				foreach (string LocalFile in LocalFiles)
 				{
-					Logger.LogInformation("Unable to map {DepotFile} to workspace; skipping.", File.DepotFile);
+					Logger.LogInformation("  {LocalFile}", LocalFile);
 				}
+				AllLocalFiles.UnionWith(LocalFiles);
 			}
 
-			Logger.LogInformation("Found {NumFiles} modified files matching {Filter}", LocalFiles.Count, Filter.ToString());
-			foreach (string LocalFile in LocalFiles)
-			{
-				Logger.LogInformation("  {LocalFile}", LocalFile);
-			}
-
+			Logger.LogInformation("Found {NumFiles} total modified files", AllLocalFiles.Count);
 			if (Parameters.Output != null)
 			{
-				await FileReference.WriteAllLinesAsync(Parameters.Output, LocalFiles);
+				await FileReference.WriteAllLinesAsync(Parameters.Output, AllLocalFiles);
 				Logger.LogInformation("Written {OutputFile}", Parameters.Output);
 			}
 		}
