- Updated SyncPerforceServers for improved functionality and updating p4 state

diff --git a/Engine/Source/Programs/AutomationTool/Scripts/SyncPerforceServers.cs b/Engine/Source/Programs/AutomationTool/Scripts/SyncPerforceServers.cs
--- a/Engine/Source/Programs/AutomationTool/Scripts/SyncPerforceServers.cs
+++ b/Engine/Source/Programs/AutomationTool/Scripts/SyncPerforceServers.cs
@@ -8,6 +8,8 @@ using System.IO;
 using Microsoft.Extensions.Logging;
 using EpicGames.Core;
 using UnrealBuildBase;
+using Amazon.S3.Model;
+
 
 /**
  * Setup of the two servers:
@@ -247,6 +249,14 @@ namespace BuildScripts.Automation
 
 			if (!bSkipReconcile)
 			{
+				// tell the incoming workspace that what is on disc is up to date - this is a lie if other people are also mirroring servers, because someone else may have updated this stream
+				// without this clientspec knowing - however, because the sync from the source server (done above) is ground truth, we don't need to actually pull anything down. reconciling
+				// will fail, however, if this client is out of date
+				Logger.LogInformation($"Syncrhonizing server state...");
+				P4IncomingClient.Sync($"-k {IncomingStreamName}/...", AllowSpew: !bHideSpew);
+
+				Logger.LogInformation($"Reconciling changed files...");
+
 				// reconcile files on the incoming stream
 				int CL = P4IncomingClient.CreateChange(Description: $"Reconciling changes on target incoming stream from {P4Env.ServerAndPort} / {SourceClient.Stream} to {Server} / {IncomingStreamName}");
 				P4IncomingClient.Reconcile(CL, $"-m {IncomingStreamName}/...", AllowSpew: !bHideSpew);
@@ -261,10 +271,13 @@ namespace BuildScripts.Automation
 
 			if (!bSkipMergeToStaging)
 			{
+				Logger.LogInformation($"Syncing Staging client {StagingStreamName}/...");
+				P4StagingClient.Sync($"{IncomingStreamName}/...", AllowSpew: !bHideSpew);
+
 				// first we merge down from Main into Staging.
-				int CL = P4StagingClient.CreateChange(Description: $"Merging {MainStreamName} to {StagingStreamName}");
-				P4StagingClient.P4($"integrate -c {CL} {MainStreamName}/... {StagingStreamName}/...");
-				if (P4StagingClient.P4($"resolve -am -c {CL}").ExitCode != 0)
+				int CL = P4StagingClient.CreateChange(Description: $"Merging {MainStreamName} to {StagingStreamName}", AllowSpew: !bHideSpew);
+				P4StagingClient.P4($"integrate -c {CL} {MainStreamName}/... {StagingStreamName}/...", AllowSpew: !bHideSpew);
+				if (P4StagingClient.P4($"resolve -am -c {CL}", AllowSpew: !bHideSpew).ExitCode != 0)
 				{
 					Logger.LogError($"Failed to resolve the merge from {MainStreamName} to {StagingStreamName}. This is unexpected, and indicates that the staging directory was dirtied outside of this process. Resolve everything in changelist {CL}, then press enter to continue. \nIf you kill this script, you can run this script again with \"-skipsync -skipreconcile\".");
 					Console.ReadLine();
@@ -281,9 +294,9 @@ namespace BuildScripts.Automation
 					Console.ReadLine();
 				}
 
-				CL = P4StagingClient.CreateChange(Description: $"Merging {IncomingStreamName} to {StagingStreamName}");
-				P4StagingClient.P4($"integrate -c {CL} {IncomingStreamName}/... {StagingStreamName}/...");
-				if (P4StagingClient.P4($"resolve -am -c {CL}").ExitCode != 0)
+				CL = P4StagingClient.CreateChange(Description: $"Merging {IncomingStreamName} to {StagingStreamName}", AllowSpew: !bHideSpew);
+				P4StagingClient.P4($"integrate -c {CL} {IncomingStreamName}/... {StagingStreamName}/...", AllowSpew: !bHideSpew);
+				if (P4StagingClient.P4($"resolve -am -c {CL}", AllowSpew: !bHideSpew).ExitCode != 0)
 				{
 					Logger.LogError($"Manual resolves are needed when merging {IncomingStreamName} to {StagingStreamName}. This is expected. Resolve changelist {CL}, then press enter to continue. \nIf you kill this script, you can run this script again with \"-skipsync -skipreconcile\" to get to this step faster.");
 					Console.ReadLine();
@@ -306,8 +319,8 @@ namespace BuildScripts.Automation
 
 			if (!bSkipCopyToMain)
 			{
-				int CL = P4MainClient.CreateChange(Description: $"Copying {StagingStreamName} to parent {MainStreamName}");
-				P4MainClient.P4($"copy -c {CL} -S {StagingStreamName}");
+				int CL = P4MainClient.CreateChange(Description: $"Copying {StagingStreamName} to parent {MainStreamName}", AllowSpew: !bHideSpew);
+				P4MainClient.P4($"copy -c {CL} -S {StagingStreamName}", AllowSpew: !bHideSpew);
 
 				try
 				{
