Adding p4 as a paramHelp for tests that might need p4v access

diff --git a/Engine/Source/Programs/AutomationTool/Gauntlet/Unreal/RunUnreal.cs b/Engine/Source/Programs/AutomationTool/Gauntlet/Unreal/RunUnreal.cs
--- a/Engine/Source/Programs/AutomationTool/Gauntlet/Unreal/RunUnreal.cs
+++ b/Engine/Source/Programs/AutomationTool/Gauntlet/Unreal/RunUnreal.cs
@@ -33,6 +33,7 @@ namespace Gauntlet
 	[ParamHelp("Configuration", "Configuration to perform tests on", Choices = new string[] { "Debug", "DebugGame", "Development", "Test", "Shipping"})]
 	[ParamHelp("Platform", "Platforms to perform tests on and their params")]
 	[ParamHelp("Dev", "Run in Dev mode", ParamType = typeof(bool))]
+	[ParamHelp("p4", "Enable p4v support", ParamType = typeof(bool))]
 	[ParamHelp("MaxDuration", "Maximum duration for test in sections", ParamType = typeof(int), DefaultValue = 3600)]
 	[ParamHelp("NoTimeout", "No maximum timeout", ParamType = typeof(bool))]
 	[ParamHelp("TestIterations", "Number of iterations to repeat this test", ParamType = typeof(int), DefaultValue = 1)]
