Fix .csproj condition

diff --git a/Engine/Source/Programs/AutomationTool/AutomationTool.csproj b/Engine/Source/Programs/AutomationTool/AutomationTool.csproj
--- a/Engine/Source/Programs/AutomationTool/AutomationTool.csproj
+++ b/Engine/Source/Programs/AutomationTool/AutomationTool.csproj
@@ -108,7 +108,27 @@
     <ProjectReference Include="..\Shared\EpicGames.Build\EpicGames.Build.csproj" PrivateAssets="all" />
     <ProjectReference Include="..\UnrealBuildTool\UnrealBuildTool.csproj" PrivateAssets="all" />
   </ItemGroup>
-  
+
+  <Choose>
+    <When Condition="'$(IsWindows)'=='true'">
+      <ItemGroup Condition="Exists('..\..\..\Restricted\NotForLicensees\Binaries\Win64\Box\Box.dll')">
+        <Reference Include="..\..\..\Restricted\NotForLicensees\Binaries\Win64\Box\Box.dll" />
+        <Content Include="..\..\..\Restricted\NotForLicensees\Binaries\Win64\Box\BoxNative.dll">
+          <CopyToOutputDirectory>PreserveNewest</CopyToOutputDirectory>
+        </Content>
+        <Content Include="..\..\..\Restricted\NotForLicensees\Binaries\Win64\Box\BoxDetours.dll">
+          <CopyToOutputDirectory>PreserveNewest</CopyToOutputDirectory>
+        </Content>
+        <Content Include="..\..\..\Restricted\NotForLicensees\Binaries\Win64\Box\zstd.dll">
+          <CopyToOutputDirectory>PreserveNewest</CopyToOutputDirectory>
+        </Content>
+      </ItemGroup>
+      <PropertyGroup Condition="Exists('..\..\..\Restricted\NotForLicensees\Binaries\Win64\Box\Box.dll')">
+        <DefineConstants>$(DefineConstants);__BOXEXECUTOR_AVAILABLE__</DefineConstants>
+      </PropertyGroup>
+    </When>
+  </Choose>
+
   <ItemGroup>
     <None Include="..\..\..\Build\BatchFiles\BuildUAT.bat"><Link>BatchFiles\BuildUAT.bat</Link></None>
     <None Include="..\..\..\Build\BatchFiles\BuildUAT.sh"><Link>BatchFiles\BuildUAT.sh</Link></None>
diff --git a/Engine/Source/Programs/UnrealBuildTool/UnrealBuildTool.csproj b/Engine/Source/Programs/UnrealBuildTool/UnrealBuildTool.csproj
--- a/Engine/Source/Programs/UnrealBuildTool/UnrealBuildTool.csproj
+++ b/Engine/Source/Programs/UnrealBuildTool/UnrealBuildTool.csproj
@@ -71,6 +71,26 @@
     </When>
   </Choose>
 
+  <Choose>
+    <When Condition="'$([System.Runtime.InteropServices.RuntimeInformation]::IsOSPlatform($([System.Runtime.InteropServices.OSPlatform]::Windows)))' == 'true'">
+      <ItemGroup Condition="Exists('..\..\..\Restricted\NotForLicensees\Binaries\Win64\Box\Box.dll')">
+        <Reference Include="..\..\..\Restricted\NotForLicensees\Binaries\Win64\Box\Box.dll" />
+        <Content Include="..\..\..\Restricted\NotForLicensees\Binaries\Win64\Box\BoxNative.dll">
+          <CopyToOutputDirectory>PreserveNewest</CopyToOutputDirectory>
+        </Content>
+        <Content Include="..\..\..\Restricted\NotForLicensees\Binaries\Win64\Box\BoxDetours.dll">
+          <CopyToOutputDirectory>PreserveNewest</CopyToOutputDirectory>
+        </Content>
+        <Content Include="..\..\..\Restricted\NotForLicensees\Binaries\Win64\Box\zstd.dll">
+          <CopyToOutputDirectory>PreserveNewest</CopyToOutputDirectory>
+        </Content>
+      </ItemGroup>
+      <PropertyGroup Condition="Exists('..\..\..\Restricted\NotForLicensees\Binaries\Win64\Box\Box.dll')">
+        <DefineConstants>$(DefineConstants);__BOXEXECUTOR_AVAILABLE__</DefineConstants>
+      </PropertyGroup>
+    </When>
+  </Choose>
+
   <ItemGroup>
     <Compile Include="..\Shared\MetaData.cs">
       <Link>Properties\MetaData.cs</Link>
