Added zenstore and iterate arguments to the ProjectParams.
Updated CookProfile.

diff --git a/Engine/Source/Programs/AutomationTool/AutomationUtils/ProjectParams.cs b/Engine/Source/Programs/AutomationTool/AutomationUtils/ProjectParams.cs
--- a/Engine/Source/Programs/AutomationTool/AutomationUtils/ProjectParams.cs
+++ b/Engine/Source/Programs/AutomationTool/AutomationUtils/ProjectParams.cs
@@ -28,6 +28,8 @@ namespace AutomationTool
 	[ParamHelp("deploy", "Deploy the project for the target platform", ParamType = typeof(bool))]
 	[ParamHelp("stage", "Put this build in a stage director", ParamType = typeof(bool))]
 	[ParamHelp("run", "Run the game after it is built (including server, if -server)", ParamType = typeof(bool))]
+	[ParamHelp("zenstore", "Save cooked output data to the Zen storage server", ParamType = typeof(bool))]
+	[ParamHelp("iterate", "Uses the iterative cooking/deploy", ParamType = typeof(bool))]
 	public interface IProjectParamsHelpers
 	{
 	}
