[ddc] allow custom args passed to "engine" DDC pak command; use this to enable SM6 by default when cooking DDC pak for engine content (required since some templates enable SM6 by default)

diff --git a/Engine/Build/InstalledEngineBuild.xml b/Engine/Build/InstalledEngineBuild.xml
--- a/Engine/Build/InstalledEngineBuild.xml
+++ b/Engine/Build/InstalledEngineBuild.xml
@@ -790,7 +790,11 @@
 			<Copy From="$(RootDir)" To="$(DDCDir)" Files="#FilteredCopyList"/>
 			<Copy From="$(CsToolsDir)" To="$(DDCDir)" Files="#Build Tools CS Binaries"/>
 
-			<Command Name="BuildDerivedDataCache" Arguments="-TempDir=&quot;$(DDCDir)&quot; -FeaturePacks=&quot;$(ProjectsToBuildDDCWin64)&quot; -TargetPlatforms=$(DDCPlatformsWin64) -HostPlatform=Win64 -SavedDir=&quot;$(SavedDir)&quot;"/>
+			<!-- Force enable D3D SM6 format; this is not enabled by default in engine settings but is used by various template content -->
+			<!-- As such we must ensure we compile engine shaders for SM6 as well -->
+			<Property Name="EngineExtraArgs" Value="&quot;-ini:Engine:[/Script/WindowsTargetPlatform.WindowsTargetSettings]:D3D12TargetedShaderFormats=PCD3D_SM6&quot;"/>
+
+			<Command Name="BuildDerivedDataCache" Arguments="-TempDir=&quot;$(DDCDir)&quot; -FeaturePacks=&quot;$(ProjectsToBuildDDCWin64)&quot; -TargetPlatforms=$(DDCPlatformsWin64) -HostPlatform=Win64 -SavedDir=&quot;$(SavedDir)&quot; -EngineContentExtraArgs=$(EngineExtraArgs)"/>
 
 			<!-- Store off the generated pack csv file for later debugging -->
 			<!-- D:\build\++UE5\Sync\LocalBuilds\InstalledDDC\Engine\DerivedDataCache\ -->
diff --git a/Engine/Source/Programs/AutomationTool/Scripts/BuildDerivedDataCache.Automation.cs b/Engine/Source/Programs/AutomationTool/Scripts/BuildDerivedDataCache.Automation.cs
--- a/Engine/Source/Programs/AutomationTool/Scripts/BuildDerivedDataCache.Automation.cs
+++ b/Engine/Source/Programs/AutomationTool/Scripts/BuildDerivedDataCache.Automation.cs
@@ -27,7 +27,7 @@ public class BuildDerivedDataCache : BuildCommand
 		string BackendName = ParseParamValue("BackendName", "CreateInstalledEnginePak");
 		string RelativePakPath = ParseParamValue("RelativePakPath", "Engine/DerivedDataCache/Compressed.ddp");
 		bool bSkipEngine = ParseParam("SkipEngine");
-
+		string EngineContentExtraArgs = ParseParamValue("EngineContentExtraArgs", string.Empty);
 
 		// Get paths to everything within the temporary directory
 		string EditorExe = CommandUtils.GetEditorCommandletExe(TempDir, HostPlatform);
@@ -78,7 +78,7 @@ public class BuildDerivedDataCache : BuildCommand
 						FilteredPlatforms.Add(TargetPlatform);
 					}
 				}
-				if(FilteredPlatforms.Count == 0)
+				if (FilteredPlatforms.Count == 0)
 				{
 					Logger.LogInformation("Did not find any project specific platforms for FeaturePack {GameName} out of supplied TargetPlatforms {ProjectSpecificPlatforms}, skipping it!", GameName, ProjectSpecificPlatforms);
 					continue;
@@ -104,9 +104,26 @@ public class BuildDerivedDataCache : BuildCommand
 		// recomputed after they do, but there is currently nothing that makes that happen automatically
 		CommandUtils.DeleteDirectory(CommandUtils.CombinePaths(TempDir, "Engine", "Saved", "Config"));
 
-		// Generate DDC for the editor, and merge all the other PAK files in
+		// Generate DDC for the editor, and merge all the other PAK files ini
+		List<string> EngineContentArgs = new() 
+		{ 
+			"-fill",
+			$"-DDC={BackendName}",
+			$"-MergePaks={CommandUtils.MakePathSafeToUseWithCommandLine(String.Join("+", ProjectPakFiles))}" 
+		};
+
+		if (!string.IsNullOrEmpty(EngineContentExtraArgs))
+		{
+			EngineContentArgs.Add(EngineContentExtraArgs);
+		}
+
+		if (bSkipEngine)
+		{
+			EngineContentArgs.Add("-projectonly");
+		}
+
 		Logger.LogInformation("Generating DDC data for engine content on {TargetPlatforms}", TargetPlatforms);
-		CommandUtils.DDCCommandlet(null, EditorExe, null, TargetPlatforms, String.Format("-fill -DDC={0} -MergePaks={1}{2}", BackendName, CommandUtils.MakePathSafeToUseWithCommandLine(String.Join("+", ProjectPakFiles)), bSkipEngine? " -projectonly" : ""));
+		CommandUtils.DDCCommandlet(null, EditorExe, null, TargetPlatforms, String.Join(" ", EngineContentArgs));
 
 		string SavedPakFile = CommandUtils.CombinePaths(SavedDir, RelativePakPath);
 		CommandUtils.CopyFile(OutputPakFile, SavedPakFile);
