Removed accidental debug log line
#rnx

diff --git a/Engine/Source/Programs/AutomationTool/Scripts/CopyBuildToStagingDirectory.Automation.cs b/Engine/Source/Programs/AutomationTool/Scripts/CopyBuildToStagingDirectory.Automation.cs
--- a/Engine/Source/Programs/AutomationTool/Scripts/CopyBuildToStagingDirectory.Automation.cs
+++ b/Engine/Source/Programs/AutomationTool/Scripts/CopyBuildToStagingDirectory.Automation.cs
@@ -770,7 +770,6 @@ namespace AutomationScripts
 						string[] CookedFiles = FileReference.ReadAllLines(CookedFilesManifestFile);
 						foreach (string RelativeFilename in CookedFiles)
 						{
-							CommandUtils.LogInformation("{0}", RelativeFilename);
 							SC.PackageStoreManifest.ZenCookedFiles.Add(FileReference.Combine(SC.PlatformCookDir, RelativeFilename).FullName);
 						}
 						return;
