Added the ability to get the crashes directory for a project

#RNX

diff --git a/Engine/Source/Programs/AutomationTool/AutomationUtils/CommandletUtils.cs b/Engine/Source/Programs/AutomationTool/AutomationUtils/CommandletUtils.cs
--- a/Engine/Source/Programs/AutomationTool/AutomationUtils/CommandletUtils.cs
+++ b/Engine/Source/Programs/AutomationTool/AutomationUtils/CommandletUtils.cs
@@ -308,7 +308,7 @@ namespace AutomationTool
 			// If we're running on a Windows build machine, copy any crash dumps into the log folder
 			if(HostPlatform.Current.HostEditorPlatform == UnrealTargetPlatform.Win64 && IsBuildMachine)
 			{
-				DirectoryInfo CrashesDir = new DirectoryInfo(DirectoryReference.Combine(DirectoryReference.FromFile(ProjectName) ?? Unreal.EngineDirectory, "Saved", "Crashes").FullName);
+				DirectoryInfo CrashesDir = new DirectoryInfo(GetCrashesDirectory(ProjectName).FullName);
 				if(CrashesDir.Exists)
 				{
 					foreach(DirectoryInfo CrashDir in CrashesDir.EnumerateDirectories())
@@ -538,5 +538,16 @@ namespace AutomationTool
 			}
 			return ProjectFullPath;
 		}
+
+		/// <summary>
+		/// Get the crashes directory for the given project file
+		/// </summary>
+		/// <param name="ProjectFullPath">Path to a project file</param>
+		/// <returns>DirectoryReference for the directory where crashes are stored for this project</returns>
+		public static DirectoryReference GetCrashesDirectory(FileReference ProjectFullPath)
+		{
+			DirectoryReference CrashesDir = DirectoryReference.Combine(DirectoryReference.FromFile(ProjectFullPath) ?? Unreal.EngineDirectory, "Saved", "Crashes");
+			return CrashesDir;
+		}
 	}
 }
