Windows platform extensions can use BootstrapPackagedGame

diff --git a/Engine/Source/Programs/AutomationTool/Win/WinPlatform.Automation.cs b/Engine/Source/Programs/AutomationTool/Win/WinPlatform.Automation.cs
--- a/Engine/Source/Programs/AutomationTool/Win/WinPlatform.Automation.cs
+++ b/Engine/Source/Programs/AutomationTool/Win/WinPlatform.Automation.cs
@@ -73,6 +73,12 @@ public class Win64Platform : Platform
 
 	public override bool IsSupported { get { return true; } }
 
+	public virtual UnrealTargetPlatform? BootstrapExePlatform
+	{
+		get { return null; }
+	}
+
+
 	public override void GetFilesToDeployOrStage(ProjectParams Params, DeploymentContext SC)
 	{
 		// Engine non-ufs (binaries)
@@ -182,7 +188,8 @@ public class Win64Platform : Platform
 
 	void StageBootstrapExecutable(DeploymentContext SC, string ExeName, FileReference TargetFile, StagedFileReference StagedRelativeTargetPath, string StagedArguments)
 	{
-		FileReference InputFile = FileReference.Combine(SC.LocalRoot, "Engine", "Binaries", SC.PlatformDir, String.Format("BootstrapPackagedGame-{0}-Shipping.exe", SC.PlatformDir));
+		UnrealTargetPlatform BootstrapPlatform = (BootstrapExePlatform ?? SC.StageTargetPlatform.PlatformType);
+		FileReference InputFile = FileReference.Combine(SC.LocalRoot, "Engine", "Binaries", BootstrapPlatform.ToString(), String.Format("BootstrapPackagedGame-{0}-Shipping.exe", BootstrapPlatform));
 		if(FileReference.Exists(InputFile))
 		{
 			// Create the new bootstrap program
