Fix launching on remote Linux box from Windows editor

diff --git a/Engine/Source/Programs/AutomationTool/AutomationUtils/CommandUtils.cs b/Engine/Source/Programs/AutomationTool/AutomationUtils/CommandUtils.cs
--- a/Engine/Source/Programs/AutomationTool/AutomationUtils/CommandUtils.cs
+++ b/Engine/Source/Programs/AutomationTool/AutomationUtils/CommandUtils.cs
@@ -758,11 +758,12 @@ namespace AutomationTool
 		/// <param name="Filenames">Filename.</param>
 		/// <returns>True if the file exists, false otherwise.</returns>
 		public static bool FileExists(string FileName)
-		{
-			var NormalizedFilename = ConvertSeparators(PathSeparator.Default, FileName);
+        {
+			// need to remove the quotes before checking to see if it exists
+			var NormalizedFilename = ConvertSeparators(PathSeparator.Default, FileName).Replace("\"", "");
 			return InternalUtils.SafeFileExists(NormalizedFilename);
-		}
-
+        }
+		
 		/// <summary>
 		/// Checks if a file(s) exists.
 		/// </summary>
@@ -782,7 +783,8 @@ namespace AutomationTool
         /// <returns>True if the file exists, false otherwise.</returns>
         public static bool FileExists(bool bQuiet, string FileName)
         {
-			var NormalizedFilename = ConvertSeparators(PathSeparator.Default, FileName);
+			// need to remove the quotes before checking to see if it exists
+			var NormalizedFilename = ConvertSeparators(PathSeparator.Default, FileName).Replace("\"", "");
 			return InternalUtils.SafeFileExists(NormalizedFilename, bQuiet);
         }
 
diff --git a/Engine/Source/Programs/AutomationTool/Linux/LinuxPlatform.Automation.cs b/Engine/Source/Programs/AutomationTool/Linux/LinuxPlatform.Automation.cs
--- a/Engine/Source/Programs/AutomationTool/Linux/LinuxPlatform.Automation.cs
+++ b/Engine/Source/Programs/AutomationTool/Linux/LinuxPlatform.Automation.cs
@@ -380,7 +380,7 @@ chmod +x {0}
 				// try contacting the device(s) and cache the key(s)
 				foreach (string DeviceAddress in ProjParams.DeviceNames)
 				{
-					RunAndLog(CmdEnv, "cmd.exe", String.Format("/c \"echo y | {0} -P 22 -ssh -t -l {1} -pw {2} {3} echo All Ok\"", PlinkPath, ProjParams.DeviceUsername, ProjParams.DevicePassword, DeviceAddress));
+					RunAndLog(CmdEnv, "cmd.exe", String.Format("/c \'echo y | {0} -P 22 -ssh -t -l {1} -pw {2} {3} echo All Ok\'", PlinkPath, ProjParams.DeviceUsername, ProjParams.DevicePassword, DeviceAddress));
 				}
 			}
 		}
