iOS Remote Mac builds don't include user defined Localization string resources.  Make sure they get included during the stub build step.

#rnx

diff --git a/Engine/Source/Programs/AutomationTool/IOS/IOSPlatform.Automation.cs b/Engine/Source/Programs/AutomationTool/IOS/IOSPlatform.Automation.cs
--- a/Engine/Source/Programs/AutomationTool/IOS/IOSPlatform.Automation.cs
+++ b/Engine/Source/Programs/AutomationTool/IOS/IOSPlatform.Automation.cs
@@ -673,6 +673,27 @@ public class IOSPlatform : ApplePlatform
 		return bIsBuiltAsFramework;
 	}
 
+	private void StageCustomLocalizationResources(ProjectParams Params, DeploymentContext SC)
+	{
+		string RelativeResourcesPath = CombinePaths("Build", "IOS", "Resources", "Localizations");
+		DirectoryReference LocalizationDirectory = DirectoryReference.Combine(Params.RawProjectPath.Directory, RelativeResourcesPath);
+		if (DirectoryReference.Exists(LocalizationDirectory))
+		{
+			IEnumerable<DirectoryReference> LocalizationDirsToStage = DirectoryReference.EnumerateDirectories(LocalizationDirectory, "*.lproj", SearchOption.TopDirectoryOnly);
+			LogInformation("There are {0} Localization directories.", LocalizationDirsToStage.Count());
+
+			foreach (DirectoryReference FullLocDirPath in LocalizationDirsToStage)
+			{
+				StagedDirectoryReference LocInStageDir = new StagedDirectoryReference(FullLocDirPath.GetDirectoryName());
+				SC.StageFiles(StagedFileType.SystemNonUFS, FullLocDirPath, StageFilesSearch.TopDirectoryOnly, LocInStageDir);
+			}
+		}
+		else
+		{
+			LogInformation("App has no custom Localization resources");
+		}
+	}
+
 	private void StageCustomLaunchScreenStoryboard(ProjectParams Params, DeploymentContext SC)
 	{
 		string InterfaceSBDirectory = Path.GetDirectoryName(Params.RawProjectPath.FullName) + "/Build/IOS/Resources/Interface/";
@@ -1462,6 +1483,9 @@ public class IOSPlatform : ApplePlatform
 				SC.StageFile(StagedFileType.SystemNonUFS, MuteCafFile, new StagedFileReference("mute.caf"));
 			}
 		}
+
+		// Copy any project defined iOS specific localization resources into Staging (ie: App Display Name)
+		StageCustomLocalizationResources(Params, SC);
 	}
 
 	protected void StageMovieFiles(DirectoryReference InputDir, DeploymentContext SC)
diff --git a/Engine/Source/Programs/UnrealBuildTool/Platform/IOS/UEDeployIOS.cs b/Engine/Source/Programs/UnrealBuildTool/Platform/IOS/UEDeployIOS.cs
--- a/Engine/Source/Programs/UnrealBuildTool/Platform/IOS/UEDeployIOS.cs
+++ b/Engine/Source/Programs/UnrealBuildTool/Platform/IOS/UEDeployIOS.cs
@@ -1,4 +1,4 @@
-﻿// Copyright Epic Games, Inc. All Rights Reserved.
+// Copyright Epic Games, Inc. All Rights Reserved.
 
 using System;
 using System.Collections;
@@ -1017,20 +1017,6 @@ namespace UnrealBuildTool
 			}
 		}
 
-		protected virtual void CopyLocalizationsResources(string InEngineDir, string AppDirectory, string BuildDirectory, string IntermediateDir, ILogger Logger)
-		{
-			string LocalizationsPath = BuildDirectory + "/Resources/Localizations";
-			if (Directory.Exists(LocalizationsPath))
-			{
-				Logger.LogInformation("Copy localizations from Resources {LocalizationsPath}, {AppDirectory}", LocalizationsPath, AppDirectory);
-				CopyFolder(BuildDirectory + "/Resources/Localizations", AppDirectory, true, delegate(string InFilename)
-				{
-					if (Path.GetFileName(InFilename).Equals(".DS_Store")) return false;
-					return true;
-				});
-			}
-		}
-
 		public bool PrepForUATPackageOrDeploy(UnrealTargetConfiguration Config, FileReference ProjectFile, string InProjectName, string InProjectDirectory, FileReference Executable, string InEngineDir, bool bForDistribution, string CookFlavor, bool bIsDataDeploy, bool bCreateStubIPA, TargetReceipt Receipt)
 		{
 			List<string> UPLScripts = CollectPluginDataPaths(Receipt.AdditionalProperties, Logger);
@@ -1261,8 +1247,6 @@ namespace UnrealBuildTool
 
 			if (!bCreateStubIPA)
 			{
-				CopyLocalizationsResources(InEngineDir, AppDirectory, BuildDirectory, IntermediateDirectory, Logger);
-
 				CopyCloudResources(InProjectDirectory, AppDirectory);
 
 				// copy additional engine framework assets in
