Fix Log warnings

#rnx

diff --git a/Engine/Source/Programs/AutomationTool/IOS/IOSPlatform.Automation.cs b/Engine/Source/Programs/AutomationTool/IOS/IOSPlatform.Automation.cs
--- a/Engine/Source/Programs/AutomationTool/IOS/IOSPlatform.Automation.cs
+++ b/Engine/Source/Programs/AutomationTool/IOS/IOSPlatform.Automation.cs
@@ -680,7 +680,7 @@ public class IOSPlatform : ApplePlatform
 		if (DirectoryReference.Exists(LocalizationDirectory))
 		{
 			IEnumerable<DirectoryReference> LocalizationDirsToStage = DirectoryReference.EnumerateDirectories(LocalizationDirectory, "*.lproj", SearchOption.TopDirectoryOnly);
-			LogInformation("There are {0} Localization directories.", LocalizationDirsToStage.Count());
+			Logger.LogInformation("There are {0} Localization directories.", LocalizationDirsToStage.Count());
 
 			foreach (DirectoryReference FullLocDirPath in LocalizationDirsToStage)
 			{
@@ -690,7 +690,7 @@ public class IOSPlatform : ApplePlatform
 		}
 		else
 		{
-			LogInformation("App has no custom Localization resources");
+			Logger.LogInformation("App has no custom Localization resources");
 		}
 	}
 
@@ -813,7 +813,7 @@ public class IOSPlatform : ApplePlatform
 		}
 		else
 		{
-			GetProvisioningData(Params.RawProjectPath, Params.Distribution, out MobileProvision, out SigningCertificate, out TeamUUID, out bAutomaticSigning);
+		GetProvisioningData(Params.RawProjectPath, Params.Distribution, out MobileProvision, out SigningCertificate, out TeamUUID, out bAutomaticSigning);
 		}
 
 		//@TODO: We should be able to use this code on both platforms, when the following issues are sorted:
@@ -2355,14 +2355,14 @@ public class IOSPlatform : ApplePlatform
 	{
 		if (HostPlatform.Current.HostEditorPlatform == UnrealTargetPlatform.Win64 || HostPlatform.Current.HostEditorPlatform == UnrealTargetPlatform.Mac)
 		{
-			int StartPos = ProjectFilePath.LastIndexOf("/");
-			int StringLength = ProjectFilePath.Length - 10; // 9 for .uproject, 1 for the /
-			string PackageName = ProjectFilePath.Substring(StartPos + 1, StringLength - StartPos);
-			string PackagePath = Path.Combine(Path.GetDirectoryName(ProjectFilePath), "Binaries", ClientPlatform);
-			if (string.IsNullOrEmpty(SourcePackage))
-			{
-				SourcePackage = Path.Combine(PackagePath, PackageName + ".ipa");
-			}
+    		int StartPos = ProjectFilePath.LastIndexOf("/");
+    		int StringLength = ProjectFilePath.Length - 10; // 9 for .uproject, 1 for the /
+    		string PackageName = ProjectFilePath.Substring(StartPos + 1, StringLength - StartPos);
+    		string PackagePath = Path.Combine(Path.GetDirectoryName(ProjectFilePath), "Binaries", ClientPlatform);
+    		if (string.IsNullOrEmpty(SourcePackage))
+    		{
+    			SourcePackage = Path.Combine(PackagePath, PackageName + ".ipa");
+    		}
 			
 			string[] IPAFiles;
 			if (!File.Exists(SourcePackage))
@@ -2397,32 +2397,32 @@ public class IOSPlatform : ApplePlatform
 			Logger.LogInformation("PackageName {PackageName}", PackageName);
 			Logger.LogInformation("PayloadPath {PayloadPath}", PayloadPath);
 
-			if (File.Exists(ZipFile))
-			{
-				Logger.LogInformation("Deleting previously present ZIP file created from IPA");
-				File.Delete(ZipFile);
-			}
-
-			File.Copy(SourcePackage, ZipFile);
-			UnzipPackage(ZipFile);
-
-			if (HostPlatform.Current.HostEditorPlatform == UnrealTargetPlatform.Win64)
-			{
-				IOSExports.PrepareRemoteMacForDebugging(CookedDataDirectory, new FileReference(ProjectFilePath), Log.Logger);
-			}
-			else
-			{
-				string ProjectPath = ProjectFilePath;
-				if (Directory.Exists(Path.Combine(ProjectPath, "Binaries", ClientPlatform, "Payload", PackageName, ".app")))
-				{
-					CopyDirectory_NoExceptions(CookedDataDirectory, Path.Combine(ProjectPath, "Binaries", ClientPlatform, "Payload", PackageName, ".app/cookeddata/"), true);
-				}
-				else
-				{
-					string ProjectRoot = SourcePackage;
-					ProjectRoot = ProjectRoot.Substring(0, ProjectRoot.LastIndexOf('/'));
-					CopyFile(SourcePackage, ProjectRoot + "/Binaries/" + ClientPlatform + "/Payload/" + PackageName + ".ipa", true);
-				}
+    		if (File.Exists(ZipFile))
+    		{
+    			Logger.LogInformation("Deleting previously present ZIP file created from IPA");
+    			File.Delete(ZipFile);
+    		}
+
+    		File.Copy(SourcePackage, ZipFile);
+    		UnzipPackage(ZipFile);
+
+    		if (HostPlatform.Current.HostEditorPlatform == UnrealTargetPlatform.Win64)
+    		{
+    			IOSExports.PrepareRemoteMacForDebugging(CookedDataDirectory, new FileReference(ProjectFilePath), Log.Logger);
+    		}
+    		else
+        	{
+        		string ProjectPath = ProjectFilePath;
+        		if (Directory.Exists(Path.Combine(ProjectPath, "Binaries", ClientPlatform, "Payload", PackageName, ".app")))
+        		{
+        			CopyDirectory_NoExceptions(CookedDataDirectory, Path.Combine(ProjectPath, "Binaries", ClientPlatform, "Payload", PackageName, ".app/cookeddata/"), true);
+        		}
+        		else
+        		{
+        			string ProjectRoot = SourcePackage;
+        			ProjectRoot = ProjectRoot.Substring(0, ProjectRoot.LastIndexOf('/'));
+        			CopyFile(SourcePackage, ProjectRoot + "/Binaries/" + ClientPlatform + "/Payload/" + PackageName + ".ipa", true);
+        		}
 			}
 			//cleanup
 			Logger.LogInformation("Deleting temp files ...");
