Improve the P4WriteConfig automation script.

- The script now takes any existing P4CONFIG value set in the environment into account. If the value is set then we use that as the filename rather than our default
-- This prevents the script from breaking pre-existing P4CONFIG set ups.
- If the user uses the -Path cmdline option to set a specific name for the file we will throw an error if it would conflict with an existing P4CONFIG setting. To override this the user can pass in a new cmdline flag '-ReplaceEnv' which is present will cause any pre-existing P4CONFIG value to be replaced with the one supplied on the cmdline.
- Added some minor tweaks to existing logging for readability.

diff --git a/Engine/Source/Programs/AutomationTool/AutomationUtils/P4WriteConfig.cs b/Engine/Source/Programs/AutomationTool/AutomationUtils/P4WriteConfig.cs
--- a/Engine/Source/Programs/AutomationTool/AutomationUtils/P4WriteConfig.cs
+++ b/Engine/Source/Programs/AutomationTool/AutomationUtils/P4WriteConfig.cs
@@ -14,6 +14,7 @@ namespace AutomationTool
 	[Help("Auto-detects P4 settings based on the current path and creates a p4config file with the relevant settings.")]
 	[ParamHelp("SetIgnore", "Adds a P4IGNORE to the default file (p4ignore in the root, or Engine/Extras/Perforce/p4ignore)", ParamType = typeof(bool), Flag = "-SetIgnore")]
 	[ParamHelp("Path", "Write to a path other than the current directory")]
+	[ParamHelp("ReplaceEnv", "When a filename passed in via 'Path' an error will be thrown if it does not match the existing value for P4CONFIG (if set), using this parameter will force the existing environment setting to be replaced", ParamType = typeof(bool), Flag = "-ReplaceEnv")]
 	[ParamHelp("p4port=<server:port>", "Optional hint/override of the server to use during lookup")]
 	[ParamHelp("p4user=<username>", "Optional hint/override of the username to use during lookup")]
 	public class P4WriteConfig : BuildCommand
@@ -29,8 +30,6 @@ namespace AutomationTool
 			bool SetIgnore = ParseParam("setignore");
 			bool ListOnly = ParseParam("listonly");
 
-			string OutputPath = ParseParamValue("path", "");
-
 			// apply any hints
 			if (!string.IsNullOrEmpty(Port))
 			{
@@ -66,7 +65,7 @@ namespace AutomationTool
 			{
 				// If a p4 ignore file is in the root, prefer it
 				string IgnorePath = Directory.EnumerateFiles(Unreal.RootDirectory.FullName, "*p4ignore*", SearchOption.TopDirectoryOnly).FirstOrDefault();
-				if (string.IsNullOrEmpty(IgnorePath) )
+				if (string.IsNullOrEmpty(IgnorePath))
 				{
 					IgnorePath = Path.Combine(Unreal.EngineDirectory.ToString(), "Extras", "Perforce", "p4ignore");
 				}
@@ -75,42 +74,136 @@ namespace AutomationTool
 
 			string P4Settings = P4Config.ToString();
 
-			if (!string.IsNullOrEmpty(OutputPath))
+			(string P4ConfigDirectory, string P4ConfigFilename) = GenerateP4ConfigPath();
+
+			string FullP4ConfigPath = Path.Combine(P4ConfigDirectory, P4ConfigFilename);
+
+			Logger.LogInformation("***\nWriting\n----------\n{Settings}----------\nto '{OutputPath}'\n***", P4Settings.Replace("\r\n", "\n"), FullP4ConfigPath);
+
+			if (!ListOnly)
 			{
-				if (!DirectoryExists(OutputPath) && !FileExists(OutputPath))
-				{
-					throw new AutomationException("Path {0} does not exist.", OutputPath);
-				}
+				File.WriteAllText(FullP4ConfigPath, P4Settings);
+
+				Logger.LogInformation("Wrote P4 settings to {OutputPath}", FullP4ConfigPath);
+
+				P4.P4(string.Format("set P4CONFIG={0}", P4ConfigFilename));
+				Logger.LogInformation("set P4CONFIG={Filename}", P4ConfigFilename);
 			}
 			else
 			{
-				OutputPath = Environment.CurrentDirectory;
+				Logger.LogInformation("Skipped write");
 			}
 
-			if (!File.Exists(OutputPath))
+			return ExitCode.Success;
+		}
+
+		/// <summary>
+		/// Generate the path of the p4config file based on the users input on the commandline and the existing P4CONFIG settings in
+		/// the current environment.
+		/// </summary>
+		/// <returns>The directory that we should write the p4config file to, and the filename to use for it</returns>
+		/// <exception cref="AutomationException">Exceptions will be thrown if the users input is invalid and the script cannot continue</exception>
+		private (string Directory, string Filename) GenerateP4ConfigPath()
+		{
+			bool ReplaceEnv = ParseParam("replaceenv");
+
+			(string OutputDirectory, string Filename) = ParseCmdlinePath();
+			
+			if (string.IsNullOrEmpty(OutputDirectory))
 			{
-				OutputPath = Path.Combine(OutputPath, "p4config.txt");
+				// We will write the file to the current directory if none was supplied on the cmdline
+				OutputDirectory = Environment.CurrentDirectory;
 			}
 
-			Console.WriteLine("***\nWriting\n{0}to - {1}\n***", P4Settings, OutputPath);
-
-			if (!ListOnly)
+			if (string.IsNullOrEmpty(Filename))
+			{
+				// If no filename was provided on the cmdline we will first try to use the existing value set in the users environment
+				// falling back to our default name if nothing has been set there either.
+				Filename = FindP4ConfigEnvValue();
+				if (string.IsNullOrEmpty(Filename))
+				{
+					Filename = DefaultFilename;
+				}
+			}
+			else
 			{
-				File.WriteAllText(OutputPath, P4Settings);
+				// As the caller provided a filename via the commandline we need to check if it will work with the current environment settings (if any)
+				string EnvFilename = FindP4ConfigEnvValue();
+
+				if (!string.IsNullOrEmpty(EnvFilename) && Filename != EnvFilename)
+				{
+					if (ReplaceEnv)
+					{
+						Logger.LogWarning("Changing P4CONFIG from '{CurrentName} to '{NewName}'", EnvFilename, Filename);
+					}
+					else
+					{
+						throw new AutomationException("The provided filename '{0}' will not work with the current environment setting 'P4CONFIG={1}'" +
+							"\nEither update the -path parameter or use '-ReplaceEnv' to change the environment setting", Filename, EnvFilename);
+					}
+				}
+			}
+
+			return (OutputDirectory, Filename);
+		}
+
+		/// <summary>
+		/// Parses the cmdline arg '-Path=' and returns it split into the directory and filename components.
+		/// </summary>
+		/// <returns>The directory and the filename that the user supplied on the commandline, either or both values can be empty strings</returns>
+		/// <exception cref="AutomationException">Throws an exception if the user specifies a directory on the commandline that does not exist</exception>
+		private (string Directory, string Filename) ParseCmdlinePath()
+		{
+			string CmdlinePath = ParseParamValue("path", "");
 
-				string OutputFile = Path.GetFileName(OutputPath);
+			string Directory = "";
+			string Filename = "";
 
-				Logger.LogInformation("Wrote P4 settings to {OutputPath}", OutputPath);
+			if (!string.IsNullOrEmpty(CmdlinePath))
+			{
+				if (Path.HasExtension(CmdlinePath))
+				{
+					Directory = Path.GetDirectoryName(CmdlinePath);
+					Filename = Path.GetFileName(CmdlinePath);
+				}
+				else
+				{
+					Directory = CmdlinePath;
+				}
 
-				P4.P4(string.Format("set P4CONFIG={0}", OutputFile));
-				Logger.LogInformation("set P4CONFIG={OutputFile}", OutputFile);
+				if (!string.IsNullOrEmpty(Directory) && !DirectoryExists(Directory))
+				{
+					throw new AutomationException("Directory {0} does not exist, please create it before calling the script", Directory);
+				}
 			}
-			else
+
+			return (Directory, Filename);
+		}
+
+		/// <summary>
+		/// Utility to return the currently set p4config value from the users environment
+		/// </summary>
+		/// <returns>The p4config filename currently set in the environment or an empty string if none is set</returns>
+		private string FindP4ConfigEnvValue()
+		{
+			string Filename = DefaultFilename;
+
+			IProcessResult P4SetResult = P4.P4("set -q P4CONFIG");
+			if (P4SetResult.ExitCode == 0)
 			{
-				Logger.LogInformation("Skipped write");
+				string[] Tokens = P4SetResult.Output.Trim().Split("=");
+				if (Tokens.Length == 2)
+				{
+					Filename = Tokens[1];
+				}
 			}
 
-			return ExitCode.Success;
+			return Filename;
 		}
+
+		/// <summary>
+		/// The standard name we should use for the p4config file if there is nothing existing set in the users environment
+		/// </summary>
+		private const string DefaultFilename = "p4config.txt";
 	}
 }
