[Automation] Add devkit login verification to Mcp automation tests

diff --git a/Engine/Source/Programs/AutomationTool/Gauntlet/Unreal/Base/Gauntlet.UnrealSession.cs b/Engine/Source/Programs/AutomationTool/Gauntlet/Unreal/Base/Gauntlet.UnrealSession.cs
--- a/Engine/Source/Programs/AutomationTool/Gauntlet/Unreal/Base/Gauntlet.UnrealSession.cs
+++ b/Engine/Source/Programs/AutomationTool/Gauntlet/Unreal/Base/Gauntlet.UnrealSession.cs
@@ -769,7 +769,7 @@ namespace Gauntlet
 				List<UnrealSessionRole> SortedRoles = SessionRoles.OrderBy(R => R.Constraint.IsIdentity() ? 1 : 0).ToList();
 
 				// first install all roles on these devices
-				foreach (var Role in SortedRoles)
+				foreach (UnrealSessionRole Role in SortedRoles)
 				{
 					ITargetDevice Device = null;
 
@@ -785,7 +785,7 @@ namespace Gauntlet
 						Device = new TargetDeviceNull(string.Format("Null{0}", Role.RoleType));
 					}
 
-					var OtherRoles = SortedRoles.Where(R => R != Role);
+					IEnumerable<UnrealSessionRole> OtherRoles = SortedRoles.Where(R => R != Role);
 
 					// create a config from the build source (this also applies the role options)
 					UnrealAppConfig AppConfig = BuildSource.CreateConfiguration(Role, OtherRoles);
@@ -867,7 +867,7 @@ namespace Gauntlet
 						IDeviceUsageReporter.RecordStart(Device.Name, Device.Platform, IDeviceUsageReporter.EventType.Install, IDeviceUsageReporter.EventState.Success, BuildSource.BuildName);
 						try
 						{
-							if (Globals.Params.ParseParam("VerifyLogin") && Device is IOnlineServiceLogin)
+							if ((Role.Options as UnrealTestConfiguration).VerifyLogin && Device is IOnlineServiceLogin)
 							{
 								Log.Info("\nVerifying device login...");
 								if (!(Device as IOnlineServiceLogin).VerifyLogin())
diff --git a/Engine/Source/Programs/AutomationTool/Gauntlet/Unreal/Base/Gauntlet.UnrealTestConfiguration.cs b/Engine/Source/Programs/AutomationTool/Gauntlet/Unreal/Base/Gauntlet.UnrealTestConfiguration.cs
--- a/Engine/Source/Programs/AutomationTool/Gauntlet/Unreal/Base/Gauntlet.UnrealTestConfiguration.cs
+++ b/Engine/Source/Programs/AutomationTool/Gauntlet/Unreal/Base/Gauntlet.UnrealTestConfiguration.cs
@@ -635,6 +635,11 @@ namespace Gauntlet
 		/// </summary>
 		public bool AllRolesExit { get; set; }
 
+		/// <summary>
+		/// Whether a given role should attempt login to a device platform service
+		/// </summary>
+		public bool VerifyLogin { get; set; }
+
 		/// <summary>
 		/// The collection of options which define heartbeat behavior
 		/// </summary>
diff --git a/Engine/Source/Programs/AutomationTool/Gauntlet/Unreal/Game/EpicGame.TestConfig.cs b/Engine/Source/Programs/AutomationTool/Gauntlet/Unreal/Game/EpicGame.TestConfig.cs
--- a/Engine/Source/Programs/AutomationTool/Gauntlet/Unreal/Game/EpicGame.TestConfig.cs
+++ b/Engine/Source/Programs/AutomationTool/Gauntlet/Unreal/Game/EpicGame.TestConfig.cs
@@ -235,10 +235,15 @@ namespace EpicGame
 			{
 				bool bNoAccountOverride = DataDrivenPlatformInfo.GetDataDrivenInfoForPlatform(ConfigRole.Platform.ToString())?.bNoAccountOverride ?? false;
 				// select an account
-				if (NoMCP == false && !bNoAccountOverride && PreAssignAccount == true)
+				if (!NoMCP)
 				{
-					Account UserAccount = AccountPool.Instance.ReserveAccount();
-					UserAccount.ApplyToConfig(AppConfig);
+					VerifyLogin = Globals.Params.ParseParam("VerifyLogin");
+
+					if (!bNoAccountOverride && PreAssignAccount)
+					{
+						Account UserAccount = AccountPool.Instance.ReserveAccount();
+						UserAccount.ApplyToConfig(AppConfig);
+					}
 				}
 			}
 
