[Automation] Build discovery platforms: Searches for PS4 also finds PS5 builds

#tests -test=BootTest -platform=ps4 -configuration=Test

diff --git a/Engine/Source/Programs/AutomationTool/Gauntlet/Framework/Gauntlet.Utils.cs b/Engine/Source/Programs/AutomationTool/Gauntlet/Framework/Gauntlet.Utils.cs
--- a/Engine/Source/Programs/AutomationTool/Gauntlet/Framework/Gauntlet.Utils.cs
+++ b/Engine/Source/Programs/AutomationTool/Gauntlet/Framework/Gauntlet.Utils.cs
@@ -1859,26 +1859,51 @@ namespace Gauntlet
 		}
 
 		/// <summary>
-		/// Returns files that match the specified regular expression in the provided folder, searching up to a maximum depth
+		/// Returns files, that match the specified regular expression, ignoring case, in the provided folder,
+		/// searching up to a maximum depth
 		/// </summary>
 		/// <param name="BaseDir">Directory to start searching in</param>
-		/// <param name="RegexPattern">Regular expression to match on</param>
-		/// <param name="RecursionDepth">Depth to search (-1 = unlimited)</param>
+		/// <param name="RegexPattern">Regular expression pattern to match on</param>
+		/// <param name="RecursionDepth">Optional depth to search (-1 = unlimited)</param>
 		/// <returns></returns>
 		public static IEnumerable<FileInfo> FindMatchingFiles(string BaseDir, string RegexPattern, int RecursionDepth = 0)
+		{
+			Regex RegExObj = new Regex(RegexPattern, RegexOptions.IgnoreCase);
+
+			return MatchFiles(BaseDir, RegExObj, RecursionDepth);
+		}
+
+		/// <summary>
+		/// Returns files that match a specified regular expression in the provided folder, searching up to a maximum depth
+		/// </summary>
+		/// <param name="BaseDir">Directory to start searching in</param>
+		/// <param name="RegExObj">RegEx to match on</param>
+		/// <param name="RecursionDepth">Optional depth to search (-1 = unlimited)</param>
+		/// <returns></returns>
+		public static IEnumerable<FileInfo> FindMatchingFiles(string BaseDir, Regex RegExObj, int RecursionDepth = 0)
+		{
+			return MatchFiles(BaseDir, RegExObj, RecursionDepth);
+		}
+
+		/// <summary>
+		/// Returns files that match the given RegEx from the provided folder, searching up to a maximum depth
+		/// </summary>
+		/// <param name="BaseDir">Directory to start searching in</param>
+		/// <param name="RegExObj">RegEx to match on</param>
+		/// <param name="RecursionDepth">Depth to search (-1 = unlimited)</param>
+		/// <returns></returns>
+		private static IEnumerable<FileInfo> MatchFiles(string BaseDir, Regex RegExObj, int RecursionDepth)
 		{
 			List<FileInfo> Found = new List<FileInfo>();
 
 			IEnumerable<DirectoryInfo> CandidateDirs = new DirectoryInfo[] { new DirectoryInfo(BaseDir) };
 
-			Regex Pattern = new Regex(RegexPattern, RegexOptions.IgnoreCase);
-
 			int CurrentDepth = 0;
 
 			do
 			{
 				// check for matching files in this set of directories
-				IEnumerable<FileInfo> MatchingFiles = CandidateDirs.SelectMany(D => Utils.SystemHelpers.GetFiles(D)).Where(F => Pattern.IsMatch(F.Name));
+				IEnumerable<FileInfo> MatchingFiles = CandidateDirs.SelectMany(D => Utils.SystemHelpers.GetFiles(D)).Where(F => RegExObj.IsMatch(F.Name));
 
 				Found.AddRange(MatchingFiles);
 
diff --git a/Engine/Source/Programs/AutomationTool/Gauntlet/Unreal/BuildSource/Gauntlet.UnrealBuildSource.cs b/Engine/Source/Programs/AutomationTool/Gauntlet/Unreal/BuildSource/Gauntlet.UnrealBuildSource.cs
--- a/Engine/Source/Programs/AutomationTool/Gauntlet/Unreal/BuildSource/Gauntlet.UnrealBuildSource.cs
+++ b/Engine/Source/Programs/AutomationTool/Gauntlet/Unreal/BuildSource/Gauntlet.UnrealBuildSource.cs
@@ -342,7 +342,7 @@ namespace Gauntlet
 					}
 				}
 
-				if (BuildPaths.Count() > 0)
+				if (BuildPaths.Any())
 				{
 					foreach (string Path in BuildPaths)
 					{
@@ -351,6 +351,11 @@ namespace Gauntlet
 
 						foreach (var BS in BuildSources)
 						{
+							if (!BS.BuildName.Contains(InPlatform.ToString(), StringComparison.OrdinalIgnoreCase))
+							{
+								continue;
+							}
+
 							IEnumerable<IBuild> Builds = BS.GetBuildsAtPath(ProjectName, Path);
 
 							foreach (IBuild Build in Builds)
