Localization: Updated Smartling to include some new directives.
- Updated downloads to only use published. The Smartling workflow will take care of publishing strings after each step. This is to emulate the "pending" behavior.
- Introduced the smartling.entity_escaping directive and set it to false for all files that are uploaded. This prevents  characters that resemble HTML tags from being escaped and garbling the translations.
#rnx

diff --git a/Engine/Source/Programs/AutomationTool/SmartlingLocalization/SmartlingLocalizationProvider.cs b/Engine/Source/Programs/AutomationTool/SmartlingLocalization/SmartlingLocalizationProvider.cs
--- a/Engine/Source/Programs/AutomationTool/SmartlingLocalization/SmartlingLocalizationProvider.cs
+++ b/Engine/Source/Programs/AutomationTool/SmartlingLocalization/SmartlingLocalizationProvider.cs
@@ -125,8 +125,8 @@ namespace EpicGames.SmartlingLocalization
 				string SmartlingFileUri = GetSmartlingFileUri(SmartlingFilename);
 				var DownloadUriBuilder = new UriBuilder(DownloadEndpoint);
 				// Changing the retrieval type to pseudo is a good way to test downloads
-				// We use pending as the retrieval type to also download files that are in progress. We need both published and in progress files on Smartling.
-				string DownloadRetrievalType = "pending";
+				// Published here doesn't necessarily mean that the strings are done translating. The Smartling configuration allows publishing every step of the way in the pipeline. We are just approximating "pending" by leveraging the Smartling workflow istead of the API.
+				string DownloadRetrievalType = "published";
 				string DownloadQueryString = $"fileUri={SmartlingFileUri}&retrievalType={DownloadRetrievalType}&includeOriginalStrings=false";
 				DownloadUriBuilder.Query = DownloadQueryString;
 				var DownloadResponse = await Client.GetAsync(DownloadUriBuilder.Uri);
@@ -264,6 +264,10 @@ namespace EpicGames.SmartlingLocalization
 					// This allows the warnings to be resolved. The regex designates anything within {} to be a placeholder value.
 					// This accounts for the common FText formatted string placeholders of {0} or {MyPlaceholderVariable}
 					UploadMultipartFormDataContent.Add(new StringContent("\\{([^}]+)\\}", Encoding.UTF8, "application/json"), "smartling.placeholder_format_custom");
+					// Controls whether base characters ( > < & " ) are "escaped" into entities 
+					// We set this as false so that we don't automatically escape and have Smartling think these are HTML tags or something.
+					// https://help.smartling.com/hc/en-us/articles/360007894594-Gettext-PO-POT
+					UploadMultipartFormDataContent.Add(new StringContent("false", Encoding.UTF8, "application/json"), "smartling.entity_escaping");
 					var UploadResponse = await Client.PostAsync(UploadEndpoint, UploadMultipartFormDataContent);
 					if (UploadResponse.IsSuccessStatusCode)
 					{
