Gauntlet - UE.AutomationTestConfig - Expose Log Idle timeout to command line argument

#rnx

diff --git a/Engine/Source/Programs/AutomationTool/Gauntlet/Unreal/Automation/UE.Automation.cs b/Engine/Source/Programs/AutomationTool/Gauntlet/Unreal/Automation/UE.Automation.cs
--- a/Engine/Source/Programs/AutomationTool/Gauntlet/Unreal/Automation/UE.Automation.cs
+++ b/Engine/Source/Programs/AutomationTool/Gauntlet/Unreal/Automation/UE.Automation.cs
@@ -130,6 +130,12 @@ namespace UE
 		[AutoParam]
 		public string DDC = "";
 
+		/// <summary>
+		/// Log Idle timeout in second
+		/// </summary>
+		[AutoParam]
+		public virtual int LogIdleTimeout { get; set; } = 30 * 60;
+
 		/// <summary>
 		/// Used for having the editor and any client communicate
 		/// </summary>
@@ -527,7 +533,11 @@ namespace UE
 		/// </summary>
 		public override void TickTest()
 		{
-			const float IdleTimeout = 30 * 60;
+			float IdleTimeout = 30 * 60;
+			if (GetConfiguration() is AutomationTestConfig Config && Config.LogIdleTimeout > 0)
+			{
+				IdleTimeout = Config.LogIdleTimeout;
+			}
 
 			List<string> ChannelEntries = new List<string>();
 			// We are primarily interested in what the editor is doing
