Localization: Introduced retry mechanism with exponential backoff for downloading files.
- The default timeout for HttpClient is 100s, this can result in certain files not being able to be downloaded in time. The httpResponse is never received. From experiments, it seems responses can take ~160s to respond during busier network traffic.
- Introduced retries up to 5 times with exponential backoff of multiples of 150s when requesteing the HttpResponse. First try is 150s, second is 300s etc
- Introduced authentication token refreshing. This cuts down the amount of time required to retrieve a new authentication token. Introduced to ensure that the token used after a retry will still be valid.
#test: Ran UAT locally with teh Localize command. Timeouts of 150s seems to work within 2 tries of exponential backoff while downloading the files. No additional timeout exceptions have been encountered so far.

diff --git a/Engine/Source/Programs/AutomationTool/SmartlingLocalization/SmartlingLocalizationProvider.cs b/Engine/Source/Programs/AutomationTool/SmartlingLocalization/SmartlingLocalizationProvider.cs
--- a/Engine/Source/Programs/AutomationTool/SmartlingLocalization/SmartlingLocalizationProvider.cs
+++ b/Engine/Source/Programs/AutomationTool/SmartlingLocalization/SmartlingLocalizationProvider.cs
@@ -12,6 +12,7 @@ using System.Net.Http.Headers;
 using AutomationTool;
 using UnrealBuildTool;
 using EpicGames.Localization;
+using System.Threading;
 
 
 #pragma warning disable SYSLIB0014
@@ -43,6 +44,7 @@ namespace EpicGames.SmartlingLocalization
 		public int RefreshExpiresIn { get; set; }
 		public string RefreshToken { get; set; }
 		public string TokenType { get; set; }
+		public DateTime LastSuccessfulUpdateTime { get; set; }
 	}
 
 	public class SmartlingRequestError
@@ -117,87 +119,114 @@ namespace EpicGames.SmartlingLocalization
 			var SmartlingFilename = GetSmartlingFilename(ProjectImportInfo.PortableObjectName, Platform);
 
 			Console.WriteLine($"Exporting: '{SmartlingFilename}' as '{ExportFile.FullName}' ({EpicLocale})");
-			// @TODOLocalization: Refactor this long try block into multiple try-catch blocks. 
-			try
+
+			string SmartlingLocale = ConvertEpicLocaleToSmartlingLocale(EpicLocale);
+			string DownloadEndpoint = $"https://api.smartling.com/files-api/v2/projects/{Config.ProjectId}/locales/{SmartlingLocale}/file";
+			string SmartlingFileUri = GetSmartlingFileUri(SmartlingFilename);
+			var DownloadUriBuilder = new UriBuilder(DownloadEndpoint);
+			// Changing the retrieval type to pseudo is a good way to test downloads
+			// Published here doesn't necessarily mean that the strings are done translating. The Smartling configuration allows publishing every step of the way in the pipeline. We are just approximating "pending" by leveraging the Smartling workflow istead of the API.
+			string DownloadRetrievalType = "published";
+			string DownloadQueryString = $"fileUri={SmartlingFileUri}&retrievalType={DownloadRetrievalType}&includeOriginalStrings=false";
+			DownloadUriBuilder.Query = DownloadQueryString;
+
+			HttpResponseMessage DownloadResponse = null;
+			int MaxTries = 5;
+			// 1 to account for the initial try 
+			int CurrentTries = 1;
+			// The base for which all the exponential back off will be derived from. By default HttpClient has a default timeout of 100s 
+			int InitialTimeOut = 150;
+			while (true)
 			{
-				string SmartlingLocale = ConvertEpicLocaleToSmartlingLocale(EpicLocale);
-				string DownloadEndpoint = $"https://api.smartling.com/files-api/v2/projects/{Config.ProjectId}/locales/{SmartlingLocale}/file";
-				string SmartlingFileUri = GetSmartlingFileUri(SmartlingFilename);
-				var DownloadUriBuilder = new UriBuilder(DownloadEndpoint);
-				// Changing the retrieval type to pseudo is a good way to test downloads
-				// Published here doesn't necessarily mean that the strings are done translating. The Smartling configuration allows publishing every step of the way in the pipeline. We are just approximating "pending" by leveraging the Smartling workflow istead of the API.
-				string DownloadRetrievalType = "published";
-				string DownloadQueryString = $"fileUri={SmartlingFileUri}&retrievalType={DownloadRetrievalType}&includeOriginalStrings=false";
-				DownloadUriBuilder.Query = DownloadQueryString;
-				var DownloadResponse = await Client.GetAsync(DownloadUriBuilder.Uri);
-				if (DownloadResponse.IsSuccessStatusCode)
+				int CurrentTimeOut= CurrentTries * InitialTimeOut;
+				Console.WriteLine($"Current time out for download response {CurrentTimeOut}s.");
+				TimeSpan Timeout = TimeSpan.FromSeconds(CurrentTimeOut);
+				using (var CancellationToken = new CancellationTokenSource(Timeout))
 				{
-					if (!CultureDirectory.Exists)
+					try
 					{
-						CultureDirectory.Create();
+						DownloadResponse = await Client.GetAsync(DownloadUriBuilder.Uri, CancellationToken.Token);
+						break;
 					}
-					bool ExportFileWasReadOnly = true;
-					if (ExportFile.Exists)
+					catch (Exception Ex)
 					{
-						// We're going to clobber the existing PO file, so make sure it's writable (it may be read-only if in Perforce)
-						ExportFileWasReadOnly = ExportFile.IsReadOnly;
-						ExportFile.IsReadOnly = false;
-					}
-					// Download the file 
-					using (var DownloadStream = await DownloadResponse.Content.ReadAsStreamAsync())
-					{
-						using (var DownloadFileStream = ExportFile.OpenWrite())
+						++CurrentTries;
+						if (CurrentTries > MaxTries)
 						{
-							await DownloadStream.CopyToAsync(DownloadFileStream);
+							Console.WriteLine($"[FAILED] Exporting: '{ExportFile.FullName}' ({EpicLocale}) teimed out. - {Ex}");
+							return;
 						}
+						Console.WriteLine($"Failed to get download response. Retrying {CurrentTries}/{MaxTries} times.");
+						// We need to retreive the authentication token again because after each timeout, we may exceed the validity of the authentication token 
+						await GetAuthenticationToken();
 					}
+				}
+			}
 
-					Console.WriteLine($"[SUCCESS] Exporting: '{SmartlingFileUri}' as '{ExportFile.FullName}' ({EpicLocale})");
-					// Reset the write status of the file
-					if (ExportFileWasReadOnly)
+			if (DownloadResponse.IsSuccessStatusCode)
+			{
+				if (!CultureDirectory.Exists)
+				{
+					CultureDirectory.Create();
+				}
+				bool ExportFileWasReadOnly = true;
+				if (ExportFile.Exists)
+				{
+					// We're going to clobber the existing PO file, so make sure it's writable (it may be read-only if in Perforce)
+					ExportFileWasReadOnly = ExportFile.IsReadOnly;
+					ExportFile.IsReadOnly = false;
+				}
+				// Download the file 
+				using (var DownloadStream = await DownloadResponse.Content.ReadAsStreamAsync())
+				{
+					using (var DownloadFileStream = ExportFile.OpenWrite())
 					{
-						ExportFile.IsReadOnly = true;
+						await DownloadStream.CopyToAsync(DownloadFileStream);
 					}
+				}
 
-					// Update the back-up copy so we can diff against what we got from Smartling, and what the gather commandlet produced
-					// This can be used to verify if the Smartling download is corrupt or if our export file is corrupted.
-					// Set the flag to true for debug. 
-					bool bCreateBackupCopy = false;
-					if (bCreateBackupCopy)
-					{
-						string ExportFileCopyPath = Path.Combine(ExportFile.DirectoryName, $"{Path.GetFileNameWithoutExtension(ExportFile.Name)}_FromSmartling{ExportFile.Extension}");
-						Console.WriteLine($"Updating Smartling copy '{ExportFileCopyPath}'");
-						var ExportFileCopy = new FileInfo(ExportFileCopyPath);
+				Console.WriteLine($"[SUCCESS] Exporting: '{SmartlingFileUri}' as '{ExportFile.FullName}' ({EpicLocale})");
+				// Reset the write status of the file
+				if (ExportFileWasReadOnly)
+				{
+					ExportFile.IsReadOnly = true;
+				}
 
-						var ExportFileCopyWasReadOnly = false;
-						if (ExportFileCopy.Exists)
-						{
-							// We're going to clobber the existing PO file, so make sure it's writable (it may be read-only if in Perforce)
-							ExportFileCopyWasReadOnly = ExportFileCopy.IsReadOnly;
-							ExportFileCopy.IsReadOnly = false;
-						}
+				// Update the back-up copy so we can diff against what we got from Smartling, and what the gather commandlet produced
+				// This can be used to verify if the Smartling download is corrupt or if our export file is corrupted.
+				// Set the flag to true for debug. 
+				bool bCreateBackupCopy = false;
+				if (bCreateBackupCopy)
+				{
+					string ExportFileCopyPath = Path.Combine(ExportFile.DirectoryName, $"{Path.GetFileNameWithoutExtension(ExportFile.Name)}_FromSmartling{ExportFile.Extension}");
+					Console.WriteLine($"Updating Smartling copy '{ExportFileCopyPath}'");
+					var ExportFileCopy = new FileInfo(ExportFileCopyPath);
 
-						ExportFile.CopyTo(ExportFileCopy.FullName, true);
-						// reset the read/write status 
-						if (ExportFileCopyWasReadOnly)
-						{
-							ExportFileCopy.IsReadOnly = true;
-						}
-						// We don't upload the ocpy to p4 as it's no longer strictly necessary. Just use the file as a lcoal debug. 
+					var ExportFileCopyWasReadOnly = false;
+					if (ExportFileCopy.Exists)
+					{
+						// We're going to clobber the existing PO file, so make sure it's writable (it may be read-only if in Perforce)
+						ExportFileCopyWasReadOnly = ExportFileCopy.IsReadOnly;
+						ExportFileCopy.IsReadOnly = false;
 					}
-				}
-				else
-				{
-					// The file may not currently exist in Smartling and will need to be uploaded first via the Upload step later on. 
-					Console.WriteLine($"[FAILED] Exporting: '{ExportFile.FullName}' ({EpicLocale}. The file may need to be uploaded first.)");
-					await PrintRequestErrors(DownloadResponse);
+
+					ExportFile.CopyTo(ExportFileCopy.FullName, true);
+					// reset the read/write status 
+					if (ExportFileCopyWasReadOnly)
+					{
+						ExportFileCopy.IsReadOnly = true;
+					}
+					// We don't upload the copy to p4 as it's no longer strictly necessary. Just use the file as a lcoal debug. 
 				}
 			}
-			catch (Exception Ex)
+			else
 			{
-				Console.WriteLine($"[FAILED] Exporting: '{ExportFile.FullName}' ({EpicLocale}) - {Ex}");
+				// The file may not currently exist in Smartling and will need to be uploaded first via the Upload step later on. 
+				Console.WriteLine($"[FAILED] Exporting: '{ExportFile.FullName}' ({EpicLocale}. The file may need to be uploaded first.)");
+				await PrintRequestErrors(DownloadResponse);
 			}
 		}
+		
 
 		private async Task PrintRequestErrors(HttpResponseMessage Response)
 		{
@@ -206,7 +235,7 @@ namespace EpicGames.SmartlingLocalization
 			var RequestErrors = ResponseEnvelope.Response.Errors;
 			foreach (SmartlingRequestError RequestError in RequestErrors)
 			{
-				Console.WriteLine($"Error:\n {RequestError.ToString()}");
+				Console.WriteLine($"Smartling Error:\n {RequestError.ToString()}");
 			}
 		}
 
@@ -345,7 +374,36 @@ namespace EpicGames.SmartlingLocalization
 
 		private async Task GetAuthenticationToken()
 		{
-			// @TODOLocalization: Implement refreshing instead of retrieving a new token each time 
+			// If the token is still valid, we just return. Otherwise we either refresh or authenticate again.
+
+			// If the last successful update + time to expire is still less than our current time accounting for some slack, we don't need to do anything to the token.
+			// We account for slack as it makes little sense to try authenticating with only 1s of validity left on a busy network
+			int Slack = 10; 
+			if (AuthenticationToken != null && AuthenticationToken.LastSuccessfulUpdateTime.AddSeconds(AuthenticationToken.ExpiresIn - Slack) < DateTime.UtcNow)
+			{
+				Console.WriteLine("Authentication token still valid. Using current authentication token.");
+				return;
+			}
+
+			if (AuthenticationToken != null && AuthenticationToken.LastSuccessfulUpdateTime.AddSeconds(AuthenticationToken.RefreshExpiresIn) > DateTime.UtcNow)
+			{
+				try
+				{
+					await RefreshAuthenticationToken();
+				}
+				catch (Exception)
+				{
+					await RequestAuthenticationToken();
+				}
+			}
+			else
+			{
+				await RequestAuthenticationToken();
+			}
+		}
+
+		private async Task RequestAuthenticationToken()
+		{
 			string AuthenticateEndpoint = "https://api.smartling.com/auth-api/v2/authenticate";
 			// serialize from Dictionary as MultipartFormDataContent not supported as a content-type 
 			var AuthenticateRequestContent = new Dictionary<string, string>()
@@ -363,6 +421,7 @@ namespace EpicGames.SmartlingLocalization
 					string AuthenticateResponseString = await AuthenticateResponse.Content.ReadAsStringAsync();
 					var AuthenticateResponseEnvelope = JsonSerializer.Deserialize<SmartlingResponseEnvelope<SmartlingDataEnvelope<SmartlingAuthenticationToken>>>(AuthenticateResponseString, JsonOptions);
 					AuthenticationToken = AuthenticateResponseEnvelope.Response.Data;
+					AuthenticationToken.LastSuccessfulUpdateTime = DateTime.UtcNow;
 					Console.WriteLine("Successfully retrieved authentication token!");
 				}
 				else
@@ -376,7 +435,40 @@ namespace EpicGames.SmartlingLocalization
 				Console.WriteLine($"Failed to retrieve authentication token. {Ex}");
 			}
 		}
-		
+
+		private async Task RefreshAuthenticationToken()
+		{
+			string RefreshEndpoint = "https://api.smartling.com/auth-api/v2/authenticate/refresh"; 
+			// serialize from Dictionary as MultipartFormDataContent not supported as a content-type 
+			var RefreshRequestContent = new Dictionary<string, string>()
+			{
+				{ "refreshToken", AuthenticationToken.RefreshToken}
+			};
+			var RefreshRequestContentJson = JsonSerializer.Serialize(RefreshRequestContent, JsonOptions);
+			var RefreshRequestStringContent = new StringContent(RefreshRequestContentJson, Encoding.UTF8, "application/json");
+			try
+			{
+				var RefreshResponse = Client.PostAsync(RefreshEndpoint, RefreshRequestStringContent).Result;
+				if (RefreshResponse.IsSuccessStatusCode)
+				{
+					string RefreshResponseString = await RefreshResponse.Content.ReadAsStringAsync();
+					var RefreshResponseEnvelope = JsonSerializer.Deserialize<SmartlingResponseEnvelope<SmartlingDataEnvelope<SmartlingAuthenticationToken>>>(RefreshResponseString, JsonOptions);
+					AuthenticationToken = RefreshResponseEnvelope.Response.Data;
+					AuthenticationToken.LastSuccessfulUpdateTime = DateTime.UtcNow;
+					Console.WriteLine("Successfully refreshed authentication token!");
+				}
+				else
+				{
+					Console.WriteLine("Failed to refresh authentication token.");
+					await PrintRequestErrors(RefreshResponse);
+				}
+			}
+			catch (Exception Ex)
+			{
+				Console.WriteLine($"Failed to refresh authentication token. {Ex}");
+			}
+		}
+
 		protected SmartlingConfig Config;
 		protected HttpClient Client;
 		private JsonSerializerOptions JsonOptions;
