Adds bStageLoose to pakfilerules which stages loose files and gathers them up into a new ondemand pak.

diff --git a/Engine/Source/Programs/AutomationTool/AutomationUtils/BundleUtils.cs b/Engine/Source/Programs/AutomationTool/AutomationUtils/BundleUtils.cs
--- a/Engine/Source/Programs/AutomationTool/AutomationUtils/BundleUtils.cs
+++ b/Engine/Source/Programs/AutomationTool/AutomationUtils/BundleUtils.cs
@@ -19,6 +19,7 @@ namespace AutomationUtils.Automation
 			public List<string> Tags { get; set; }
 			public List<string> Dependencies { get; set; }
 			public List<string> FileRegex { get; set; }
+			public List<string> Files { get; set; }
 			public bool bFoundParent { get; set; }
 			public bool bContainsShaderLibrary { get; set; }
 			public int Order { get; set; }
@@ -93,6 +94,17 @@ namespace AutomationUtils.Automation
 						Bundle.FileRegex = new List<string>();
 					}
 				}
+				{
+					List<string> Files;
+					if (BundleConfig.GetArray(SectionName, "Files", out Files))
+					{
+						Bundle.Files = Files;
+					}
+					else
+					{
+						Bundle.Files = new List<string>();
+					}
+				}
 				{
 					bool bContainsShaderLibrary;
 					if (BundleConfig.GetBool(SectionName, "ContainsShaderLibrary", out bContainsShaderLibrary))
@@ -128,7 +140,7 @@ namespace AutomationUtils.Automation
 		public static TPlatformBundleSettings MatchBundleSettings<TPlatformBundleSettings>(
 			string FileName, IReadOnlyDictionary<string, TPlatformBundleSettings> InstallBundles) where TPlatformBundleSettings : BundleSettings
 		{
-			// Try to find a matching chunk regex
+			// Try to find a matching chunk regex or exact filename
 			foreach (var Bundle in InstallBundles.Values)
 			{
 				foreach (string RegexString in Bundle.FileRegex)
@@ -138,6 +150,13 @@ namespace AutomationUtils.Automation
 						return Bundle;
 					}
 				}
+				foreach(string FileString in Bundle.Files)
+				{
+					if (FileString.Equals(FileName, StringComparison.OrdinalIgnoreCase))
+					{
+						return Bundle;
+					}
+				}
 			}
 
 			return null;
diff --git a/Engine/Source/Programs/AutomationTool/AutomationUtils/DeploymentContext.cs b/Engine/Source/Programs/AutomationTool/AutomationUtils/DeploymentContext.cs
--- a/Engine/Source/Programs/AutomationTool/AutomationUtils/DeploymentContext.cs
+++ b/Engine/Source/Programs/AutomationTool/AutomationUtils/DeploymentContext.cs
@@ -41,6 +41,8 @@ public class FilesToStage
 	/// </summary>
 	public Dictionary<StagedFileReference, FileReference> UFSFiles = new Dictionary<StagedFileReference, FileReference>();
 
+	public Dictionary<StagedFileReference, FileReference> LooseFiles = new Dictionary<StagedFileReference, FileReference>();
+
 	/// <summary>
 	/// After staging, this is a map from staged file to source file. These file are binaries, etc and can't go into a pak file.
 	/// </summary>
diff --git a/Engine/Source/Programs/AutomationTool/Scripts/CopyBuildToStagingDirectory.Automation.cs b/Engine/Source/Programs/AutomationTool/Scripts/CopyBuildToStagingDirectory.Automation.cs
--- a/Engine/Source/Programs/AutomationTool/Scripts/CopyBuildToStagingDirectory.Automation.cs
+++ b/Engine/Source/Programs/AutomationTool/Scripts/CopyBuildToStagingDirectory.Automation.cs
@@ -281,7 +281,6 @@ namespace AutomationScripts
 			string ContainerName,
 			FileReference PakOutputLocation,
 			DirectoryReference OptionalOutputLocation,
-			DirectoryReference OnDemandOutputLocation,
 			bool bCompressed,
 			EncryptionAndSigning.CryptoSettings CryptoSettings,
 			string EncryptionKeyGuid,
@@ -290,17 +289,7 @@ namespace AutomationScripts
 			bool bIsDLC)
 		{
 			StringBuilder CmdLine = new StringBuilder();
-
-			if (OnDemandOutputLocation != null)
-			{
-				CmdLine.AppendFormat("-Output={0}", MakePathSafeToUseWithCommandLine(Path.ChangeExtension(PakOutputLocation.FullName, ".tocmanifest")));
-				CmdLine.AppendFormat(" -OnDemandOutput={0}", MakePathSafeToUseWithCommandLine(OnDemandOutputLocation.FullName));
-			}
-			else
-			{
-				CmdLine.AppendFormat("-Output={0}", MakePathSafeToUseWithCommandLine(Path.ChangeExtension(PakOutputLocation.FullName, ".utoc")));
-			}
-
+			CmdLine.AppendFormat("-Output={0}", MakePathSafeToUseWithCommandLine(Path.ChangeExtension(PakOutputLocation.FullName, ".utoc")));
 			if (OptionalOutputLocation != null)
 			{
 				CmdLine.AppendFormat(" -OptionalOutput={0}", MakePathSafeToUseWithCommandLine(OptionalOutputLocation.FullName));
@@ -2148,17 +2137,7 @@ namespace AutomationScripts
 			// List of pak files to use instead of manifest
 			public List<string> OverridePaks;
 
-			// Whether this rule is defined for content-on-demand
-			public bool bOnDemand;
-
-			public static bool IsMatch(PakFileRules PakRules, KeyValuePair<string, string> StagingFile)
-			{
-				bool bMatched = !PakRules.bDisabled &&
-					((!PakRules.bUseExactFilePathMatch && PakRules.Filter != null && PakRules.Filter.Matches(StagingFile.Key)) ||
-					(PakRules.bUseExactFilePathMatch && PakRules.ExactFileMatches != null && PakRules.ExactFileMatches.Contains(StagingFile.Value)));
-				
-				return bMatched;
-			}
+			public bool bStageLoose;
 		};
 
 		/// <summary>
@@ -2289,9 +2268,9 @@ namespace AutomationScripts
 				PakRules.Name = SectionName;
 				PakRulesConfig.TryGetValue(SectionName, "bUseExactFilePathMatch", out PakRules.bUseExactFilePathMatch);
 				PakRulesConfig.TryGetValue(SectionName, "bExcludeFromPaks", out PakRules.bExcludeFromPaks);
+				PakRulesConfig.TryGetValue(SectionName, "bStageLoose", out PakRules.bStageLoose);
 				PakRulesConfig.TryGetValue(SectionName, "bOverrideChunkManifest", out PakRules.bOverrideChunkManifest);
 				PakRulesConfig.TryGetValue(SectionName, "bDisabled", out PakRules.bDisabled);
-				PakRulesConfig.TryGetValue(SectionName, "bOnDemand", out PakRules.bOnDemand);
 				string PakString;
 				PakRulesConfig.TryGetValue(SectionName, "OverridePaks", out PakString);
 
@@ -2304,22 +2283,14 @@ namespace AutomationScripts
 					PakRules.OverridePaks = null;
 				}
 
-				if (PakRules.bExcludeFromPaks)
+				if (PakRules.bExcludeFromPaks && PakRules.OverridePaks != null)
 				{
-					if (PakRules.OverridePaks != null)
-					{
-						Logger.LogWarning("Error in PakFileRules {Arg0}, set to exclude but also sets override!", PakRules.Name);
-						continue;
-					}
-					if (PakRules.bOnDemand)
-					{
-						Logger.LogWarning("Error in PakFileRules {Arg0}, set to exclude but also sets to ondemand!", PakRules.Name);
-						continue;
-					}
+					LogWarning("Error in PakFileRules {0}, set to exclude but also sets override!", PakRules.Name);
+					continue;
 				}
-				else if (PakRules.OverridePaks == null && !PakRules.bOnDemand)
+				else if (!PakRules.bExcludeFromPaks && PakRules.OverridePaks == null)
 				{
-					Logger.LogWarning("Error in PakFileRules {Arg0}, set to include but did not specify paks!", PakRules.Name);
+					LogWarning("Error in PakFileRules {0}, set to include but did not specify paks!", PakRules.Name);
 					continue;
 				}
 
@@ -2354,58 +2325,31 @@ namespace AutomationScripts
 		/// </summary>
 		/// <param name="Params"></param>
 		/// <param name="SC"></param>
-		private static void ApplyPakFileRules(List<PakFileRules> RulesList, KeyValuePair<string, string> StagingFile, HashSet<ChunkDefinition> ModifyPakList, ConcurrentDictionary<string, ChunkDefinition> ChunkNameToDefinition, out bool bExcludeFromPaks)
+		private static void ApplyPakFileRules(List<PakFileRules> RulesList, KeyValuePair<string, string> StagingFile, HashSet<ChunkDefinition> ModifyPakList, ConcurrentDictionary<string, ChunkDefinition> ChunkNameToDefinition, out bool bExcludeFromPaks, out bool bStageLoose)
 		{
 			bExcludeFromPaks = false;
+			bStageLoose = false;
 
 			if (RulesList == null)
 			{
 				return;
 			}
 
-			// Process content-on-demand rules first
-			if (ModifyPakList != null)
+			// Search in order, return on first match
+			foreach (var PakRules in RulesList)
 			{
-				foreach (PakFileRules Rule in RulesList.Where(Rule => Rule.bOnDemand))
+				bool bMatched = !PakRules.bDisabled &&
+					((!PakRules.bUseExactFilePathMatch && PakRules.Filter != null && PakRules.Filter.Matches(StagingFile.Key)) ||
+					(PakRules.bUseExactFilePathMatch && PakRules.ExactFileMatches != null && PakRules.ExactFileMatches.Contains(StagingFile.Value)));
+
+				if (bMatched)
 				{
-					if (PakFileRules.IsMatch(Rule, StagingFile))
+					if (PakRules.bStageLoose)
 					{
-						ModifyPakList.Clear();
-
-						string ChunkName = Rule.Name;
-						if (Rule.OverridePaks != null && Rule.OverridePaks.Count > 0)
-						{
-							string OverrideChunkName = Rule.OverridePaks[0];
-							if (ChunkNameToDefinition.ContainsKey(OverrideChunkName))
-							{
-								ChunkName = OverrideChunkName;
-							}
-							else
-							{
-								Logger.LogWarning("Undefined PAK override '{OverrideChunkName}' in PAK rule '{Arg1}'", OverrideChunkName, Rule.Name);
-							}
-						}
-
-						ChunkDefinition Chunk = null;
-						if (ChunkNameToDefinition.TryGetValue(ChunkName, out Chunk))
-						{
-							ModifyPakList.Add(Chunk);
-						}
-						else
-						{
-							Logger.LogWarning("Undefined chunk name '{ChunkName}' in PAK rule '{Arg1}'", ChunkName, Rule.Name);
-						}
-						
+						bStageLoose = true;
 						return;
 					}
-				}
-			}
 
-			// Search in order, return on first match
-			foreach (var PakRules in RulesList)
-			{
-				if (PakFileRules.IsMatch(PakRules, StagingFile))
-				{
 					bool bOverrideChunkAssignment = false;
 					if (ModifyPakList != null && ModifyPakList.Count > 0)
 					{
@@ -2458,6 +2402,7 @@ namespace AutomationScripts
 			}
 		}
 
+
 		/// <summary>
 		/// Creates a pak file using staging context (single manifest)
 		/// </summary>
@@ -2479,7 +2424,8 @@ namespace AutomationScripts
 			Parallel.ForEach(UnrealPakResponseFile, StagingFile =>
 			{
 				bool bExcludeFromPaks = false;
-				ApplyPakFileRules(PakRulesList, StagingFile, null, null, out bExcludeFromPaks);
+				bool bStageLoose = false;
+				ApplyPakFileRules(PakRulesList, StagingFile, null, null, out bExcludeFromPaks, out bStageLoose);
 
 				if (bExcludeFromPaks)
 				{
@@ -2636,11 +2582,6 @@ namespace AutomationScripts
 			/// Whether to rehydrate the assets when creating the pak file or not
 			/// </summary>
 			public bool bRehydrateAssets;
-			
-			/// <summary>
-			// Whether this pak file is for content-on-demand content delivery or not
-			/// </summary>
-			public bool bOnDemand;
 
 			/// <summary>
 			/// GUID of the encryption key for this pak file
@@ -2653,14 +2594,13 @@ namespace AutomationScripts
 			/// <param name="PakName">Path to the base output file for this pak file</param>
 			/// <param name="UnrealPakResponseFile">Map of files within the pak file to their source file on disk</param>
 			/// <param name="bCompressed">Whether to enable compression</param>
-			public CreatePakParams(string PakName, Dictionary<string, string> UnrealPakResponseFile, bool bCompressed, bool RehydrateAssets, string EncryptionKeyGuid, bool bOnDemand = false)
+			public CreatePakParams(string PakName, Dictionary<string, string> UnrealPakResponseFile, bool bCompressed, bool RehydrateAssets, string EncryptionKeyGuid)
 			{
 				this.PakName = PakName;
 				this.UnrealPakResponseFile = UnrealPakResponseFile;
 				this.bCompressed = bCompressed;
 				this.bRehydrateAssets = RehydrateAssets;
 				this.EncryptionKeyGuid = EncryptionKeyGuid;
-				this.bOnDemand = bOnDemand;
 			}
 		}
 
@@ -3188,13 +3128,11 @@ namespace AutomationScripts
 				OutputFilename = OutputFilename + PostFix;
 
 				StagedFileReference OutputRelativeLocation;
-				StagedDirectoryReference OnDemandOutputRelativeLocation;
 				if (Params.HasDLCName)
 				{
 					if (Params.DLCOverrideStagedSubDir != null)
 					{
 						OutputRelativeLocation = StagedFileReference.Combine(SC.RelativeProjectRootForStage, Params.DLCOverrideStagedSubDir, "Content", "Paks", Params.DLCFile.GetFileNameWithoutExtension() + OutputFilename + ".pak");
-						OnDemandOutputRelativeLocation = StagedDirectoryReference.Combine(SC.RelativeProjectRootForStage, Params.DLCOverrideStagedSubDir, "ContentOnDemand", "Paks", Params.DLCFile.GetFileNameWithoutExtension() + OutputFilename);
 					}
 					else
 					{
@@ -3209,29 +3147,23 @@ namespace AutomationScripts
 							PluginSubdirectory = Params.DLCFile.GetFileNameWithoutAnyExtensions();
 						}
 						OutputRelativeLocation = StagedFileReference.Combine(SC.RelativeProjectRootForStage, PluginSubdirectory, "Content", "Paks", SC.FinalCookPlatform, Params.DLCFile.GetFileNameWithoutExtension() + OutputFilename + ".pak");
-						OnDemandOutputRelativeLocation = StagedDirectoryReference.Combine(SC.RelativeProjectRootForStage, PluginSubdirectory, "ContentOnDemand", "Paks", SC.FinalCookPlatform, Params.DLCFile.GetFileNameWithoutExtension() + OutputFilename);
 					}
 				}
 				else
 				{
 					OutputRelativeLocation = StagedFileReference.Combine(SC.RelativeProjectRootForStage, "Content", "Paks", OutputFilename + OutputFilenameExtension);
-					OnDemandOutputRelativeLocation = StagedDirectoryReference.Combine(SC.RelativeProjectRootForStage, "ContentOnDemand", "Paks", OutputFilename);
 				}
 				if (SC.StageTargetPlatform.DeployLowerCaseFilenames(StagedFileType.UFS))
 				{
 					OutputRelativeLocation = OutputRelativeLocation.ToLowerInvariant();
 				}
 				OutputRelativeLocation = SC.StageTargetPlatform.Remap(OutputRelativeLocation);
-				
-				//TODO: Do we need to remap the ondemand output location?
-				//OnDemandOutputRelativeLocation = SC.StageTargetPlatform.Remap(OnDemandOutputRelativeLocation);
 
 				FileReference OutputLocation = FileReference.Combine(SC.RuntimeRootDir, OutputRelativeLocation.Name);
-				DirectoryReference OnDemandOutputLocation = DirectoryReference.Combine(SC.RuntimeRootDir, OnDemandOutputRelativeLocation.Name);
 
 				bool bCopiedExistingPak = false;
 
-				if (!PakParams.bOnDemand && SC.StageTargetPlatform != SC.CookSourcePlatform && !Params.IgnorePaksFromDifferentCookSource)
+				if (SC.StageTargetPlatform != SC.CookSourcePlatform && !Params.IgnorePaksFromDifferentCookSource)
 				{
 					// Check to see if we have an existing pak file we can use
 
@@ -3368,7 +3300,6 @@ namespace AutomationScripts
 								OutputLocation.GetFileNameWithoutAnyExtensions(),
 								OutputLocation,
 								SC.OptionalFileStageDirectory,
-								PakParams.bOnDemand ? OnDemandOutputLocation : null,
 								bCompressContainers,
 								CryptoSettings,
 								PakParams.EncryptionKeyGuid,
@@ -3899,7 +3830,6 @@ namespace AutomationScripts
 				ResponseFile = new ConcurrentDictionary<string, string>();
 				Manifest = null;
 				bCompressed = false;
-				bOnDemand = false;
 			}
 
 			// Name of pak file without extension, ie pakchunk0
@@ -3915,9 +3845,6 @@ namespace AutomationScripts
 			public bool bCompressed;
 			public string EncryptionKeyGuid;
 			public string RequestedEncryptionKeyGuid;
-
-			// Whether the chunk is used for content-on-demand content delivery or not
-			public bool bOnDemand;
 		}
 
 		/// <summary>
@@ -4000,6 +3927,7 @@ namespace AutomationScripts
 
 				const string OptionalBulkDataFileExtension = ".uptnl";
 				ConcurrentDictionary<string, ChunkDefinition> OptionalChunks = new ConcurrentDictionary<string, ChunkDefinition>(StringComparer.InvariantCultureIgnoreCase);
+				ConcurrentDictionary<string, ChunkDefinition> OnDemandChunks = new ConcurrentDictionary<string, ChunkDefinition>(StringComparer.InvariantCultureIgnoreCase);
 				ChunkDefinition DefaultChunk = ChunkDefinitions[DefaultChunkIndex];
 
 				Dictionary<string, List<ChunkDefinition>> FileNameToChunks = new Dictionary<string, List<ChunkDefinition>>(StringComparer.InvariantCultureIgnoreCase);
@@ -4017,41 +3945,19 @@ namespace AutomationScripts
 					}
 				}
 
-				// Add ondemand chunk definitions
-				if (PakRulesList != null)
-				{
-					foreach (PakFileRules Rule in PakRulesList.Where(R => R.bOnDemand))
-					{
-						string ChunkName = Rule.Name;
-						if (Rule.OverridePaks != null && Rule.OverridePaks.Count > 0)
-						{
-							ChunkName = Rule.OverridePaks[0];
-						}
-
-						if (ChunkDefinitions.Any(D => D.ChunkName == ChunkName))
-						{
-							continue;
-						}
-
-						Logger.LogInformation("Creating content-on-demand chunk definition '{ChunkName}' for PAK rule '{Arg1}'", ChunkName, Rule.Name);
-
-						//TODO: Add encryption/compression settings
-						ChunkDefinition Chunk = new ChunkDefinition(ChunkName);
-						Chunk.bOnDemand = true;
-						ChunkDefinitions.Add(Chunk);
-					}
-				}
-
 				ConcurrentDictionary<string, ChunkDefinition> ChunkNameToDefinition = new ConcurrentDictionary<string, ChunkDefinition>(StringComparer.InvariantCultureIgnoreCase);
 				foreach (ChunkDefinition Chunk in ChunkDefinitions)
 				{
 					ChunkNameToDefinition.TryAdd(Chunk.ChunkName, Chunk);
 				}
 
+				ConcurrentDictionary<StagedFileReference, FileReference> LooseFiles = new ConcurrentDictionary<StagedFileReference, FileReference>();
+
 				Parallel.ForEach(StagingManifestResponseFile, StagingFile =>
 				{
 					bool bAddedToChunk = false;
 					bool bExcludeFromPaks = false;
+					bool bStageLoose = false;
 					HashSet<ChunkDefinition> PakList = new HashSet<ChunkDefinition>();
 
 					string OriginalFilename = StagingFile.Key;
@@ -4085,7 +3991,18 @@ namespace AutomationScripts
 					}
 
 					// Now run through the pak rules which may override things
-					ApplyPakFileRules(PakRulesList, StagingFile, PakList, ChunkNameToDefinition, out bExcludeFromPaks);
+					ApplyPakFileRules(PakRulesList, StagingFile, PakList, ChunkNameToDefinition, out bExcludeFromPaks, out bStageLoose);
+
+					if (bStageLoose)
+					{
+						string StagingFileValue = StagingFile.Value;
+						if (StagingFileValue.StartsWith(SC.PakFileInternalRoot))
+						{
+							StagingFileValue = StagingFileValue.Substring(SC.PakFileInternalRoot.Length);
+						}
+						LooseFiles.TryAdd(new StagedFileReference(StagingFileValue), new FileReference(StagingFile.Key));
+						//return;
+					}
 
 					if (bExcludeFromPaks)
 					{
@@ -4097,23 +4014,32 @@ namespace AutomationScripts
 					{
 						ChunkDefinition TargetChunk = Chunk;
 
-						// Don't separate out any bulk data for ondemand chunks
-						if (Chunk.bOnDemand == false)
+						string OrigExt = Path.GetExtension(OriginalFilename);
+						if (bStageLoose)
 						{
-							string OrigExt = Path.GetExtension(OriginalFilename);
-							if (OrigExt.Equals(OptionalBulkDataFileExtension))
+							// make a streaming pak
+							string OnDemandChunkName = Chunk.ChunkName + "ondemand";
+							if (!OnDemandChunks.ContainsKey(OnDemandChunkName))
 							{
-								// any optional files encountered we want to put in a separate pak file
-								string OptionalChunkName = Chunk.ChunkName + "optional";
-								if (!OptionalChunks.ContainsKey(OptionalChunkName))
-								{
-									ChunkDefinition OptionalChunk = new ChunkDefinition(OptionalChunkName);
-									OptionalChunk.RequestedEncryptionKeyGuid = Chunk.RequestedEncryptionKeyGuid;
-									OptionalChunk.EncryptionKeyGuid = Chunk.EncryptionKeyGuid;
-									OptionalChunks.TryAdd(OptionalChunkName, OptionalChunk);
-								}
-								TargetChunk = OptionalChunks[OptionalChunkName];
+								ChunkDefinition OnDemandChunk = new ChunkDefinition(OnDemandChunkName);
+								OnDemandChunk.RequestedEncryptionKeyGuid = Chunk.RequestedEncryptionKeyGuid;
+								OnDemandChunk.EncryptionKeyGuid = Chunk.EncryptionKeyGuid;
+								OnDemandChunks.TryAdd(OnDemandChunkName, OnDemandChunk);
+							}
+							TargetChunk = OnDemandChunks[OnDemandChunkName];
+						}
+						else if (OrigExt.Equals(OptionalBulkDataFileExtension))
+						{
+							// any optional files encountered we want to put in a separate pak file
+							string OptionalChunkName = Chunk.ChunkName + "optional";
+							if (!OptionalChunks.ContainsKey(OptionalChunkName))
+							{
+								ChunkDefinition OptionalChunk = new ChunkDefinition(OptionalChunkName);
+								OptionalChunk.RequestedEncryptionKeyGuid = Chunk.RequestedEncryptionKeyGuid;
+								OptionalChunk.EncryptionKeyGuid = Chunk.EncryptionKeyGuid;
+								OptionalChunks.TryAdd(OptionalChunkName, OptionalChunk);
 							}
+							TargetChunk = OptionalChunks[OptionalChunkName];
 						}
 
 						TargetChunk.ResponseFile.TryAdd(StagingFile.Key, StagingFile.Value);
@@ -4128,8 +4054,15 @@ namespace AutomationScripts
 
 					if (!bAddedToChunk)
 					{
-						//LogInformation("No chunk assigned found for {0}. Using default chunk.", StagingFile.Key);
-						DefaultChunk.ResponseFile.TryAdd(StagingFile.Key, StagingFile.Value);
+						if (bStageLoose)
+						{
+							Logger.LogInformation("No chunk assignment found for OnDemand file {Arg0}. Skipping.", StagingFile.Key);
+						}
+						else
+						{
+							//LogInformation("No chunk assigned found for {0}. Using default chunk.", StagingFile.Key);
+							DefaultChunk.ResponseFile.TryAdd(StagingFile.Key, StagingFile.Value);
+						}
 					}
 				});
 
@@ -4137,6 +4070,16 @@ namespace AutomationScripts
 				{
 					ChunkDefinitions.Add(OptionalChunkIt.Value);
 				}
+
+				foreach (var OnDemandChunkIt in OnDemandChunks)
+				{
+					ChunkDefinitions.Add(OnDemandChunkIt.Value);
+				}
+
+				foreach (var LooseFile in LooseFiles)
+				{
+					SC.FilesToStage.LooseFiles.Add(LooseFile.Key, LooseFile.Value);
+				}
 			}
 
 			// chunk downloader pak file is a minimal pak file which contains no content.  It can be provided with as a minimal download so that the game can download all the content from another source.
@@ -4214,7 +4157,7 @@ namespace AutomationScripts
 					string EncryptionKeyToUse = Params.SkipEncryption ? "" : Chunk.EncryptionKeyGuid;
 					PakInputs.Add(new CreatePakParams(Chunk.ChunkName,
 						Chunk.ResponseFile.ToDictionary(entry => entry.Key, entry => entry.Value),
-						Params.Compressed || Chunk.bCompressed, Params.RehydrateAssets, EncryptionKeyToUse, Chunk.bOnDemand));
+						Params.Compressed || Chunk.bCompressed, Params.RehydrateAssets, EncryptionKeyToUse));
 				}
 			}
 
@@ -4483,6 +4426,7 @@ namespace AutomationScripts
 
 			if (!SC.Stage || Params.SkipStage)
 			{
+				CopyManifestFilesToStageDir(SC, SC.FilesToStage.LooseFiles, "LooseFiles");
 				return;
 			}
 
