Fix merge warning from deprecated log function

diff --git a/Engine/Source/Programs/AutomationTool/Scripts/CopyBuildToStagingDirectory.Automation.cs b/Engine/Source/Programs/AutomationTool/Scripts/CopyBuildToStagingDirectory.Automation.cs
--- a/Engine/Source/Programs/AutomationTool/Scripts/CopyBuildToStagingDirectory.Automation.cs
+++ b/Engine/Source/Programs/AutomationTool/Scripts/CopyBuildToStagingDirectory.Automation.cs
@@ -2285,12 +2285,12 @@ namespace AutomationScripts
 
 				if (PakRules.bExcludeFromPaks && PakRules.OverridePaks != null)
 				{
-					LogWarning("Error in PakFileRules {0}, set to exclude but also sets override!", PakRules.Name);
+					Logger.LogWarning("Error in PakFileRules {Arg0}, set to exclude but also sets override!", PakRules.Name);
 					continue;
 				}
 				else if (!PakRules.bExcludeFromPaks && PakRules.OverridePaks == null)
 				{
-					LogWarning("Error in PakFileRules {0}, set to include but did not specify paks!", PakRules.Name);
+					Logger.LogWarning("Error in PakFileRules {Arg0}, set to include but did not specify paks!", PakRules.Name);
 					continue;
 				}
 
