Fix Platform Config Directory copy for staged programs.

diff --git a/Engine/Source/Programs/AutomationTool/Scripts/CopyBuildToStagingDirectory.Automation.cs b/Engine/Source/Programs/AutomationTool/Scripts/CopyBuildToStagingDirectory.Automation.cs
--- a/Engine/Source/Programs/AutomationTool/Scripts/CopyBuildToStagingDirectory.Automation.cs
+++ b/Engine/Source/Programs/AutomationTool/Scripts/CopyBuildToStagingDirectory.Automation.cs
@@ -1659,6 +1659,7 @@ namespace AutomationScripts
 				if (ShouldStage ?? true)
 				{
 					SC.StageFile(StagedFileType.UFS, ConfigFile);
+					Logger.LogInformation("Including config file {ConfigFile}", ConfigFile);
 				}
 				else
 				{
@@ -1676,16 +1677,28 @@ namespace AutomationScripts
 		{
 			DirectoryReference PlatformEngineConfigDir = DirectoryReference.Combine(SC.LocalRoot, "Engine", "Platforms", PlatformName, "Config");
 			DirectoryReference PlatformProjectConfigDir = DirectoryReference.Combine(SC.ProjectRoot, "Platforms", PlatformName, "Config");
+			List<DirectoryReference> PlatformProgramExtendedConfigDirs = Unreal.GetExtensionDirs(DirectoryReference.Combine(SC.LocalRoot, "Engine"), "Programs")
+				.Where( InPath => InPath.ContainsName(PlatformName,0))
+				.ToList();
 
-			if (DirectoryReference.Exists(PlatformEngineConfigDir))
+			void StageConfigFileIfExist(DirectoryReference InDir)
 			{
-				StageConfigFiles(SC, PlatformEngineConfigDir, PlatformName);
-
+				if (DirectoryReference.Exists(InDir))
+				{
+					StageConfigFiles(SC, InDir, PlatformName);
+				}
 			}
-			if (DirectoryReference.Exists(PlatformProjectConfigDir))
+
+			StageConfigFileIfExist(PlatformEngineConfigDir);
+			StageConfigFileIfExist(PlatformProjectConfigDir);
+
+			// Find All Sub Program directories nested under Engine/<ExtendedDirs>/Platforms/PlatformName/Programs/ProgramName
+			PlatformProgramExtendedConfigDirs.ForEach((InConfigDir) =>
 			{
-				StageConfigFiles(SC, PlatformProjectConfigDir, PlatformName);
-			}
+				DirectoryReference ProgramConfigDir = DirectoryReference.Combine(InConfigDir, SC.ShortProjectName, "Config");
+				Logger.LogDebug("Platform Extended Directory checking: {ProgramConfigDir}", ProgramConfigDir);
+				StageConfigFileIfExist(ProgramConfigDir);
+			});
 		}
 
 		/// <summary>
