- Re-codesign content-only projects when archiving mac builds, because it modifies the .plist which breaks the signature

diff --git a/Engine/Source/Programs/AutomationTool/Mac/MacPlatform.Automation.cs b/Engine/Source/Programs/AutomationTool/Mac/MacPlatform.Automation.cs
--- a/Engine/Source/Programs/AutomationTool/Mac/MacPlatform.Automation.cs
+++ b/Engine/Source/Programs/AutomationTool/Mac/MacPlatform.Automation.cs
@@ -486,6 +486,11 @@ public class MacPlatform : ApplePlatform
 				InfoPlistContents = InfoPlistContents.Replace("<string>UnrealGame</string>", "<string>" + SC.ShortProjectName + "</string>");
 				DeleteFile(InfoPlistPath);
 				WriteAllText(InfoPlistPath, InfoPlistContents);
+
+				// we now need to re-sign the .app because we modified the .plist
+				// we codesign with ad-hoc, and if the Developer ID Application cert exists, attempt to use it, ignore any errors
+				Utils.RunLocalProcessAndReturnStdOut("/usr/bin/codesign", $"-f -s - \"{BundlePath}\"", null);
+				Utils.RunLocalProcessAndReturnStdOut("/usr/bin/codesign", $"-f -s \"Developer ID Application\" \"{BundlePath}\"", null);
 			}
 
 			if (!SC.bIsCombiningMultiplePlatforms)
