Catch xml report parsing error to warning

diff --git a/Engine/Source/Programs/AutomationTool/LowLevelTests/Utility/LowLevelTests.ReportParser.cs b/Engine/Source/Programs/AutomationTool/LowLevelTests/Utility/LowLevelTests.ReportParser.cs
--- a/Engine/Source/Programs/AutomationTool/LowLevelTests/Utility/LowLevelTests.ReportParser.cs
+++ b/Engine/Source/Programs/AutomationTool/LowLevelTests/Utility/LowLevelTests.ReportParser.cs
@@ -21,7 +21,7 @@ namespace LowLevelTests
 			}
 			catch (Exception ParseEx)
 			{
-				Log.Error("Encountered error while parsing report {0}", ParseEx.ToString());
+				Log.Warning("Encountered error while parsing report {0}", ParseEx.ToString());
 				IsValid = false;
 			}
 		}
