Game client no longer needs -zenstoreproject or -zenstorehost arguments to correctly run with a COTF server which uses Zen.
Game client now accepts only -filehostip for connection to a CTOF server, -cookonthefly is no longer necessary and is ignored.
Game client will now receive ProjectName (ProjectId), Platform (OplogId), Zen server host name and port from COTF server if running with zenstore.
Fixed issues with generating the ProjectId from a path where drive letters where lower case.

diff --git a/Engine/Source/Developer/CookOnTheFlyNetServer/Internal/CookOnTheFlyNetServer.h b/Engine/Source/Developer/CookOnTheFlyNetServer/Internal/CookOnTheFlyNetServer.h
--- a/Engine/Source/Developer/CookOnTheFlyNetServer/Internal/CookOnTheFlyNetServer.h
+++ b/Engine/Source/Developer/CookOnTheFlyNetServer/Internal/CookOnTheFlyNetServer.h
@@ -33,9 +33,6 @@ struct FCookOnTheFlyNetworkServerOptions
 
 	/* Active target platform(s) */
 	TArray<ITargetPlatform*> TargetPlatforms;
-
-	/* Zen project name */
-	FString ZenProjectName;
 };
 
 class ICookOnTheFlyClientConnection
@@ -47,6 +44,7 @@ public:
 	virtual const ITargetPlatform* GetTargetPlatform() const = 0;
 	virtual bool GetIsSingleThreaded() const = 0;
 	virtual bool SendMessage(const FCookOnTheFlyMessage& Message) = 0;
+	virtual void SetZenInfo(const FString& InProjectId, const FString& InOplogId, const FString& InHostName, uint16 InHostPort) = 0;
 };
 
 class ICookOnTheFlyNetworkServer
diff --git a/Engine/Source/Developer/CookOnTheFlyNetServer/Private/CookOnTheFlyNetServer.cpp b/Engine/Source/Developer/CookOnTheFlyNetServer/Private/CookOnTheFlyNetServer.cpp
--- a/Engine/Source/Developer/CookOnTheFlyNetServer/Private/CookOnTheFlyNetServer.cpp
+++ b/Engine/Source/Developer/CookOnTheFlyNetServer/Private/CookOnTheFlyNetServer.cpp
@@ -18,9 +18,9 @@ public:
 		switch (Options.Protocol)
 		{
 		case ECookOnTheFlyNetworkServerProtocol::Tcp:
-			return MakeShared<FCookOnTheFlyServerTCP>(Options.Port, Options.TargetPlatforms, Options.ZenProjectName);
+			return MakeShared<FCookOnTheFlyServerTCP>(Options.Port, Options.TargetPlatforms);
 		case ECookOnTheFlyNetworkServerProtocol::Platform:
-			return MakeShared<FCookOnTheFlyServerPlatformProtocol>(Options.TargetPlatforms, Options.ZenProjectName);
+			return MakeShared<FCookOnTheFlyServerPlatformProtocol>(Options.TargetPlatforms);
 		default:
 			UE_LOG(LogCookOnTheFlyNetworkServer, Fatal, TEXT("Unsupported protocol: %d"), Options.Protocol);
 			return nullptr;
diff --git a/Engine/Source/Developer/CookOnTheFlyNetServer/Private/CookOnTheFlyNetServerBase.cpp b/Engine/Source/Developer/CookOnTheFlyNetServer/Private/CookOnTheFlyNetServerBase.cpp
--- a/Engine/Source/Developer/CookOnTheFlyNetServer/Private/CookOnTheFlyNetServerBase.cpp
+++ b/Engine/Source/Developer/CookOnTheFlyNetServer/Private/CookOnTheFlyNetServerBase.cpp
@@ -91,7 +91,9 @@ bool FCookOnTheFlyClientConnectionBase::Initialize()
 	FBufferArchive HandshakeResponsePayload;
 	HandshakeResponsePayload << bIsOk;
 	HandshakeResponsePayload << PlatformNameString;
-	HandshakeResponsePayload << Owner.ZenProjectName;
+	HandshakeResponsePayload << ZenProjectId;
+	HandshakeResponsePayload << ZenHostName;
+	HandshakeResponsePayload << ZenHostPort;
 
 	if (!SendPayload(HandshakeResponsePayload))
 	{
@@ -144,9 +146,8 @@ bool FCookOnTheFlyClientConnectionBase::SendMessage(const UE::Cook::FCookOnTheFl
 	return SendPayload(Payload);
 }
 
-FCookOnTheFlyNetworkServerBase::FCookOnTheFlyNetworkServerBase(const TArray<ITargetPlatform*>& InActiveTargetPlatforms, const FString& InZenProjectName)
+FCookOnTheFlyNetworkServerBase::FCookOnTheFlyNetworkServerBase(const TArray<ITargetPlatform*>& InActiveTargetPlatforms)
 	: ActiveTargetPlatforms(InActiveTargetPlatforms)
-	, ZenProjectName(InZenProjectName)
 {
 
 }
diff --git a/Engine/Source/Developer/CookOnTheFlyNetServer/Private/CookOnTheFlyNetServerBase.h b/Engine/Source/Developer/CookOnTheFlyNetServer/Private/CookOnTheFlyNetServerBase.h
--- a/Engine/Source/Developer/CookOnTheFlyNetServer/Private/CookOnTheFlyNetServerBase.h
+++ b/Engine/Source/Developer/CookOnTheFlyNetServer/Private/CookOnTheFlyNetServerBase.h
@@ -30,6 +30,14 @@ public:
 		return bIsSingleThreaded;
 	}
 
+	void SetZenInfo(const FString& InProjectId, const FString& InOplogId, const FString& InHostName, uint16 InHostPort)
+	{
+		ZenProjectId = InProjectId;
+		check(InOplogId == PlatformName);
+		ZenHostName = InHostName;
+		ZenHostPort = InHostPort;
+	}
+
 	bool Initialize();
 
 	bool SendMessage(const UE::Cook::FCookOnTheFlyMessage& Message) override;
@@ -80,6 +88,10 @@ private:
 	TAtomic<bool> bRunning{ true };
 	FRunnableThread* WorkerThread = nullptr;
 	bool bClientConnectedBroadcasted = false;
+
+	FString ZenProjectId;
+	FString ZenHostName;
+	uint16 ZenHostPort = 0;
 };
 
 class FCookOnTheFlyNetworkServerBase
@@ -88,7 +100,7 @@ class FCookOnTheFlyNetworkServerBase
 	friend class FCookOnTheFlyClientConnectionBase;
 
 public:
-	FCookOnTheFlyNetworkServerBase(const TArray<ITargetPlatform*>& InActiveTargetPlatforms, const FString& InZenProjectName);
+	FCookOnTheFlyNetworkServerBase(const TArray<ITargetPlatform*>& InActiveTargetPlatforms);
 
 	bool ProcessRequest(FCookOnTheFlyClientConnectionBase& Connection, const UE::Cook::FCookOnTheFlyRequest& Request);
 
@@ -110,5 +122,4 @@ private:
 	FClientConnectionEvent ClientDisconnectedEvent;
 	TMap<UE::Cook::ECookOnTheFlyMessage, FHandleRequestDelegate> Handlers;
 	TArray<ITargetPlatform*> ActiveTargetPlatforms;
-	FString ZenProjectName;
 };
diff --git a/Engine/Source/Developer/CookOnTheFlyNetServer/Private/PlatformProtocolServer.cpp b/Engine/Source/Developer/CookOnTheFlyNetServer/Private/PlatformProtocolServer.cpp
--- a/Engine/Source/Developer/CookOnTheFlyNetServer/Private/PlatformProtocolServer.cpp
+++ b/Engine/Source/Developer/CookOnTheFlyNetServer/Private/PlatformProtocolServer.cpp
@@ -104,8 +104,8 @@ private:
 };
 
 
-FCookOnTheFlyServerPlatformProtocol::FCookOnTheFlyServerPlatformProtocol(const TArray<ITargetPlatform*>& InTargetPlatforms, const FString& InZenProjectName)
-	: FCookOnTheFlyNetworkServerBase(InTargetPlatforms, InZenProjectName)
+FCookOnTheFlyServerPlatformProtocol::FCookOnTheFlyServerPlatformProtocol(const TArray<ITargetPlatform*>& InTargetPlatforms)
+	: FCookOnTheFlyNetworkServerBase(InTargetPlatforms)
 	, TargetPlatforms(InTargetPlatforms)
 {
 	Running = false;
diff --git a/Engine/Source/Developer/CookOnTheFlyNetServer/Private/PlatformProtocolServer.h b/Engine/Source/Developer/CookOnTheFlyNetServer/Private/PlatformProtocolServer.h
--- a/Engine/Source/Developer/CookOnTheFlyNetServer/Private/PlatformProtocolServer.h
+++ b/Engine/Source/Developer/CookOnTheFlyNetServer/Private/PlatformProtocolServer.h
@@ -30,7 +30,7 @@ public:
 	 *
 	 * @param InFileServerOptions Network file server options
 	 */
-	FCookOnTheFlyServerPlatformProtocol(const TArray<ITargetPlatform*>& TargetPlatforms, const FString& InZenProjectName);
+	FCookOnTheFlyServerPlatformProtocol(const TArray<ITargetPlatform*>& TargetPlatforms);
 
 	/**
 	 * Destructor.
diff --git a/Engine/Source/Developer/CookOnTheFlyNetServer/Private/TCPServer.cpp b/Engine/Source/Developer/CookOnTheFlyNetServer/Private/TCPServer.cpp
--- a/Engine/Source/Developer/CookOnTheFlyNetServer/Private/TCPServer.cpp
+++ b/Engine/Source/Developer/CookOnTheFlyNetServer/Private/TCPServer.cpp
@@ -58,8 +58,8 @@ private:
 /* FCookOnTheFlyServerTCP constructors
  *****************************************************************************/
 
-FCookOnTheFlyServerTCP::FCookOnTheFlyServerTCP(int32 InPort, const TArray<ITargetPlatform*>& InTargetPlatforms, const FString& InZenProjectName)
-	: FCookOnTheFlyNetworkServerBase(InTargetPlatforms, InZenProjectName)
+FCookOnTheFlyServerTCP::FCookOnTheFlyServerTCP(int32 InPort, const TArray<ITargetPlatform*>& InTargetPlatforms)
+	: FCookOnTheFlyNetworkServerBase(InTargetPlatforms)
 	, Port(InPort)
 {
 	if (Port < 0)
diff --git a/Engine/Source/Developer/CookOnTheFlyNetServer/Private/TCPServer.h b/Engine/Source/Developer/CookOnTheFlyNetServer/Private/TCPServer.h
--- a/Engine/Source/Developer/CookOnTheFlyNetServer/Private/TCPServer.h
+++ b/Engine/Source/Developer/CookOnTheFlyNetServer/Private/TCPServer.h
@@ -24,7 +24,7 @@ public:
 	 *
 	 * @param InFileServerOptions Network file server options
 	 */
-	FCookOnTheFlyServerTCP(int32 InPort, const TArray<ITargetPlatform*>& InTargetPlatforms, const FString& InZenProjectName);
+	FCookOnTheFlyServerTCP(int32 InPort, const TArray<ITargetPlatform*>& InTargetPlatforms);
 
 	/**
 	 * Destructor.
diff --git a/Engine/Source/Developer/IoStoreUtilities/Private/ZenStoreWriter.cpp b/Engine/Source/Developer/IoStoreUtilities/Private/ZenStoreWriter.cpp
--- a/Engine/Source/Developer/IoStoreUtilities/Private/ZenStoreWriter.cpp
+++ b/Engine/Source/Developer/IoStoreUtilities/Private/ZenStoreWriter.cpp
@@ -638,6 +638,23 @@ void FZenStoreWriter::EndCook(const FCookInfo& Info)
 	UE_LOG(LogZenStoreWriter, Display, TEXT("Output:\t%d Public runtime script objects"), PackageStoreOptimizer->GetTotalScriptObjectCount());
 }
 
+FZenStoreWriter::ZenHostInfo FZenStoreWriter::GetHostInfo() const
+{
+	FZenStoreWriter::ZenHostInfo Info;
+	Info.ProjectId = ProjectId;
+	Info.OplogId = OplogId;
+	if (IsLocalConnection)
+	{
+		Info.HostName = "localhost";
+	}
+	else
+	{
+		Info.HostName = HttpClient->GetHostName();
+	}
+	Info.HostPort = HttpClient->GetPort();
+	return Info;
+}
+
 void FZenStoreWriter::BeginPackage(const FBeginPackageInfo& Info)
 {
 	FPendingPackageState& State = AddPendingPackage(Info.PackageName);
diff --git a/Engine/Source/Developer/IoStoreUtilities/Public/PackageStoreWriter.h b/Engine/Source/Developer/IoStoreUtilities/Public/PackageStoreWriter.h
--- a/Engine/Source/Developer/IoStoreUtilities/Public/PackageStoreWriter.h
+++ b/Engine/Source/Developer/IoStoreUtilities/Public/PackageStoreWriter.h
@@ -5,6 +5,8 @@
 #include "IO/PackageStore.h"
 #include "Serialization/PackageWriter.h"
 
+class FZenStoreWriter;
+
 class IPackageStoreWriter : public ICookedPackageWriter
 {
 public:
@@ -14,6 +16,12 @@ public:
 		return this;
 	}
 
+	/** Downcast function for IPackageStoreWriter that is of class FZenStoreWriter. */
+	virtual FZenStoreWriter* AsZenStoreWriter()
+	{
+		return nullptr;
+	}
+
 	struct FOplogCookInfo
 	{
 		struct FAttachment
diff --git a/Engine/Source/Developer/IoStoreUtilities/Public/ZenStoreWriter.h b/Engine/Source/Developer/IoStoreUtilities/Public/ZenStoreWriter.h
--- a/Engine/Source/Developer/IoStoreUtilities/Public/ZenStoreWriter.h
+++ b/Engine/Source/Developer/IoStoreUtilities/Public/ZenStoreWriter.h
@@ -60,6 +60,22 @@ public:
 
 	IOSTOREUTILITIES_API ~FZenStoreWriter();
 
+	/** Identify as a implmenter of this class from the IPackageStoreWriter api. */
+	IOSTOREUTILITIES_API virtual FZenStoreWriter* AsZenStoreWriter() override
+	{
+		return this;
+	};
+
+	struct ZenHostInfo
+	{
+		FString ProjectId;
+		FString OplogId;
+		FString HostName;
+		uint16 HostPort;
+	};
+
+	IOSTOREUTILITIES_API ZenHostInfo GetHostInfo() const;
+
 	IOSTOREUTILITIES_API virtual void BeginPackage(const FBeginPackageInfo& Info) override;
 	IOSTOREUTILITIES_API virtual void CommitPackage(FCommitPackageInfo&& Info) override;
 
diff --git a/Engine/Source/Editor/UnrealEd/Private/CookOnTheFlyServer.cpp b/Engine/Source/Editor/UnrealEd/Private/CookOnTheFlyServer.cpp
--- a/Engine/Source/Editor/UnrealEd/Private/CookOnTheFlyServer.cpp
+++ b/Engine/Source/Editor/UnrealEd/Private/CookOnTheFlyServer.cpp
@@ -598,10 +598,6 @@ bool UCookOnTheFlyServer::StartCookOnTheFly(FCookOnTheFlyStartupOptions InCookOn
 	{
 		NetworkServerOptions.TargetPlatforms = TPM.GetTargetPlatforms();
 	}
-	if (IsUsingZenStore())
-	{
-		NetworkServerOptions.ZenProjectName = FApp::GetZenStoreProjectId();
-	}
 
 	ICookOnTheFlyNetworkServerModule& CookOnTheFlyNetworkServerModule = FModuleManager::LoadModuleChecked<ICookOnTheFlyNetworkServerModule>(TEXT("CookOnTheFlyNetServer"));
 	CookOnTheFlyNetworkServer = CookOnTheFlyNetworkServerModule.CreateServer(NetworkServerOptions);
diff --git a/Engine/Source/Editor/UnrealEd/Private/Cooker/IoStoreCookOnTheFlyRequestManager.cpp b/Engine/Source/Editor/UnrealEd/Private/Cooker/IoStoreCookOnTheFlyRequestManager.cpp
--- a/Engine/Source/Editor/UnrealEd/Private/Cooker/IoStoreCookOnTheFlyRequestManager.cpp
+++ b/Engine/Source/Editor/UnrealEd/Private/Cooker/IoStoreCookOnTheFlyRequestManager.cpp
@@ -25,6 +25,7 @@
 #include "SocketSubsystem.h"
 #include "Templates/Function.h"
 #include "UObject/SavePackage.h"
+#include "ZenStoreWriter.h"
 
 class FPackageRegistry
 {
@@ -392,7 +393,12 @@ private:
 		{
 			IPackageStoreWriter* PackageWriter = CookOnTheFlyServer.GetPackageWriter(TargetPlatform).AsPackageStoreWriter();
 			check(PackageWriter); // This class should not be used except when COTFS is using an IPackageStoreWriter
-
+			FZenStoreWriter* ZenStoreWriter = PackageWriter->AsZenStoreWriter();
+			if (ZenStoreWriter != nullptr)
+			{
+				FZenStoreWriter::ZenHostInfo HostInfo = ZenStoreWriter->GetHostInfo();
+				Connection.SetZenInfo(HostInfo.ProjectId, HostInfo.OplogId, HostInfo.HostName, HostInfo.HostPort);
+			}
 			FName PlatformName = Connection.GetPlatformName();
 			FScopeLock _(&ContextsCriticalSection);
 			if (!PlatformContexts.Contains(PlatformName))
diff --git a/Engine/Source/Programs/AutomationTool/AutomationUtils/ProjectParams.cs b/Engine/Source/Programs/AutomationTool/AutomationUtils/ProjectParams.cs
--- a/Engine/Source/Programs/AutomationTool/AutomationUtils/ProjectParams.cs
+++ b/Engine/Source/Programs/AutomationTool/AutomationUtils/ProjectParams.cs
@@ -345,7 +345,6 @@ namespace AutomationTool
 			this.Pak = InParams.Pak;
 			this.IgnorePaksFromDifferentCookSource = InParams.IgnorePaksFromDifferentCookSource;
 			this.IoStore = InParams.IoStore;
-			this.Cook4IoStore = InParams.Cook4IoStore;
 			this.ZenStore = InParams.ZenStore;
 			this.NoZenAutoLaunch = InParams.NoZenAutoLaunch;
 			this.GenerateOptimizationData = InParams.GenerateOptimizationData;
@@ -550,7 +549,6 @@ namespace AutomationTool
 			bool? Pak = null,
 			bool? IgnorePaksFromDifferentCookSource = null,
 			bool? IoStore = null,
-			bool? Cook4IoStore = null,
 			bool? ZenStore = null,
 			string NoZenAutoLaunch = null,
 			bool? SkipIoStore = null,
@@ -754,12 +752,6 @@ namespace AutomationTool
 			this.IgnorePaksFromDifferentCookSource = GetParamValueIfNotSpecified(Command, IgnorePaksFromDifferentCookSource, this.IgnorePaksFromDifferentCookSource, "IgnorePaksFromDifferentCookSource");
 			this.IoStore = GetParamValueIfNotSpecified(Command, IoStore, this.IoStore, "iostore");
 			this.SkipIoStore = GetParamValueIfNotSpecified(Command, SkipIoStore, this.SkipIoStore, "skipiostore");
-			this.Cook4IoStore = GetParamValueIfNotSpecified(Command, Cook4IoStore, this.Cook4IoStore, "cook4iostore");
-			if (this.Cook4IoStore)
-			{
-				this.IoStore = true;
-				this.AdditionalCookerOptions += " -IoStore";
-			}
 			this.ZenStore = GetParamValueIfNotSpecified(Command, ZenStore, this.ZenStore, "zenstore");
 			if (this.ZenStore && this.Cook && !this.SkipCook)
 			{
@@ -1412,12 +1404,6 @@ namespace AutomationTool
 		/// </summary>
 		[Help("iostore", "generate I/O store container file(s)")]
 		public bool IoStore { private set; get; }
-
-		/// <summary>
-		/// Shared: True if the cooker should write directly to container file(s)
-		/// </summary>
-		[Help("cook4iostore", "generate I/O store container file(s)")]
-		public bool Cook4IoStore { private set; get; }
 		
 		/// <summary>
 		/// Shared: True if the cooker should store the cooked output to the Zen storage server
@@ -2097,7 +2083,7 @@ namespace AutomationTool
 
 
 		/// <summary>
-		/// Run: The client runs with cooked data provided by UnrealFileServer, command line: -fileserver
+		/// Run: The client runs with cooked data provided by UnrealFileServer, command line: -fileserver = CookByTheBook with ZenServer (zenstore)
 		/// </summary>
 		[Help("fileserver", "run the client with cooked data provided by UnrealFileServer")]
 		public bool FileServer { private set; get; }
@@ -3181,7 +3167,6 @@ namespace AutomationTool
 				Logger.LogDebug("IgnorePaksFromDifferentCookSource={IgnorePaksFromDifferentCookSource}", IgnorePaksFromDifferentCookSource);
 				Logger.LogDebug("IoStore={IoStore}", IoStore);
 				Logger.LogDebug("SkipIoStore={SkipIoStore}", SkipIoStore);
-				Logger.LogDebug("Cook4IoStore={Cook4IoStore}", Cook4IoStore);
 				Logger.LogDebug("ZenStore={ZenStore}", ZenStore);
 				Logger.LogDebug("NoZenAutoLaunch={NoZenAutoLaunch}", NoZenAutoLaunch);
 				Logger.LogDebug("SkipEncryption={SkipEncryption}", SkipEncryption);
diff --git a/Engine/Source/Programs/AutomationTool/Scripts/RunProjectCommand.Automation.cs b/Engine/Source/Programs/AutomationTool/Scripts/RunProjectCommand.Automation.cs
--- a/Engine/Source/Programs/AutomationTool/Scripts/RunProjectCommand.Automation.cs
+++ b/Engine/Source/Programs/AutomationTool/Scripts/RunProjectCommand.Automation.cs
@@ -581,14 +581,13 @@ namespace AutomationScripts
 					TempCmdLine += Params.MapToRun + " ";
 				}
 
+				String FileHostCommandline = GetFileHostCommandline(Params, SC);
+				if (!string.IsNullOrEmpty(FileHostCommandline))
+				{
+					TempCmdLine += FileHostCommandline + " ";
+				}
 				if (Params.CookOnTheFly || Params.FileServer)
 				{
-					String FileHostCommandline = GetFileHostCommandline(Params, SC);
-					if (!string.IsNullOrEmpty(FileHostCommandline))
-					{
-						TempCmdLine += FileHostCommandline + " ";
-					}
-
 					if (Params.CookOnTheFlyStreaming)
 					{
 						TempCmdLine += "-streaming ";
@@ -818,12 +817,12 @@ namespace AutomationScripts
 				}
 			}
 
-			if (Params.ZenStore)
+			if (Params.CookOnTheFly)
+			{
+				FileHostParams += "-filehostip=";
+			}
+			else if (Params.ZenStore)
 			{
-				if (Params.CookOnTheFly)
-				{
-					FileHostParams += "-cookonthefly ";
-				}
 				FileHostParams += "-zenstoreproject=" + ProjectUtils.GetProjectPathId(SC.RawProjectPath) + " ";
 				FileHostParams += "-zenstorehost=";
 			}
diff --git a/Engine/Source/Runtime/CookOnTheFly/Internal/CookOnTheFly.h b/Engine/Source/Runtime/CookOnTheFly/Internal/CookOnTheFly.h
--- a/Engine/Source/Runtime/CookOnTheFly/Internal/CookOnTheFly.h
+++ b/Engine/Source/Runtime/CookOnTheFly/Internal/CookOnTheFly.h
@@ -307,6 +307,10 @@ public:
 
 	virtual const FString& GetPlatformName() const = 0;
 
+	virtual const FString& GetZenHostName() const = 0;
+
+	virtual const uint16 GetZenHostPort() const = 0;
+
 	/**
 	 * Returns whether connected to the cook-on-the-fly server.
 	 */
diff --git a/Engine/Source/Runtime/CookOnTheFly/Private/CookOnTheFlyServerConnection.cpp b/Engine/Source/Runtime/CookOnTheFly/Private/CookOnTheFlyServerConnection.cpp
--- a/Engine/Source/Runtime/CookOnTheFly/Private/CookOnTheFlyServerConnection.cpp
+++ b/Engine/Source/Runtime/CookOnTheFly/Private/CookOnTheFlyServerConnection.cpp
@@ -61,6 +61,14 @@ public:
 		}
 		ResponsePayload << PlatformName;
 		ResponsePayload << ZenProjectName;
+		ResponsePayload << ZenHostName;
+		ResponsePayload << ZenHostPort;
+
+		// If our server is running its zenserver locally, we get "localhost" as a reply and we should use the same host for the zenserver
+		if (ZenHostName == "localhost")
+		{
+			ZenHostName = Transport->GetHostName();
+		}
 
 		FCoreDelegates::OnEnginePreExit.AddRaw(this, &FCookOnTheFlyServerConnection::OnEnginePreExit);
 
@@ -93,6 +101,16 @@ public:
 		return ZenProjectName;
 	}
 
+	virtual const FString& GetZenHostName() const override
+	{
+		return ZenHostName;
+	}
+
+	virtual const uint16 GetZenHostPort() const override
+	{
+		return ZenHostPort;
+	}
+
 	virtual const FString& GetPlatformName() const override
 	{
 		return PlatformName;
@@ -307,6 +325,8 @@ private:
 	const bool bIsSingleThreaded;
 	FString PlatformName;
 	FString ZenProjectName;
+	FString ZenHostName;
+	uint16 ZenHostPort;
 
 	FCriticalSection RequestsCriticalSection;
 	TMap<uint32, TUniquePtr<FPendingRequest>> PendingRequests;
diff --git a/Engine/Source/Runtime/CookOnTheFly/Private/CookOnTheFlyServerConnection.h b/Engine/Source/Runtime/CookOnTheFly/Private/CookOnTheFlyServerConnection.h
--- a/Engine/Source/Runtime/CookOnTheFly/Private/CookOnTheFlyServerConnection.h
+++ b/Engine/Source/Runtime/CookOnTheFly/Private/CookOnTheFlyServerConnection.h
@@ -17,6 +17,7 @@ public:
 	virtual bool SendPayload(const TArray<uint8>& Payload) = 0;
 	virtual bool HasPendingPayload() = 0;
 	virtual bool ReceivePayload(FArrayReader& Payload) = 0;
+	virtual FString GetHostName() const = 0;
 
 };
 
diff --git a/Engine/Source/Runtime/CookOnTheFly/Private/PlatformTransport.cpp b/Engine/Source/Runtime/CookOnTheFly/Private/PlatformTransport.cpp
--- a/Engine/Source/Runtime/CookOnTheFly/Private/PlatformTransport.cpp
+++ b/Engine/Source/Runtime/CookOnTheFly/Private/PlatformTransport.cpp
@@ -115,6 +115,11 @@ bool FPlatformTransport::HasPendingPayload()
 	return false;
 }
 
+FString FPlatformTransport::GetHostName() const
+{
+	return "platform://";
+}
+
 void FPlatformTransport::Disconnect()
 {
 	IPlatformHostCommunication& HostCommunication = FPlatformMisc::GetPlatformHostCommunication();
diff --git a/Engine/Source/Runtime/CookOnTheFly/Private/PlatformTransport.h b/Engine/Source/Runtime/CookOnTheFly/Private/PlatformTransport.h
--- a/Engine/Source/Runtime/CookOnTheFly/Private/PlatformTransport.h
+++ b/Engine/Source/Runtime/CookOnTheFly/Private/PlatformTransport.h
@@ -22,6 +22,7 @@ public:
 	virtual bool SendPayload(const TArray<uint8>& Payload) override;
 	virtual bool ReceivePayload(FArrayReader& Payload) override;
 	virtual bool HasPendingPayload() override;
+	virtual FString GetHostName() const override;
 	virtual void Disconnect() override;
 	//~ End ICookOnTheFlyServerTransport interface
 
diff --git a/Engine/Source/Runtime/CookOnTheFly/Private/TCPTransport.cpp b/Engine/Source/Runtime/CookOnTheFly/Private/TCPTransport.cpp
--- a/Engine/Source/Runtime/CookOnTheFly/Private/TCPTransport.cpp
+++ b/Engine/Source/Runtime/CookOnTheFly/Private/TCPTransport.cpp
@@ -56,6 +56,7 @@ bool FTCPTransport::Initialize(const TCHAR* InHostIp)
 			FileSocket = NULL;
 			UE_LOG(LogCookOnTheFly, Error, TEXT("Failed to connect to COTF server at %s."), *Addr->ToString(true));
 		}
+		HostName = Addr->ToString(false);
 	}
 
 #if USE_MCSOCKET_FOR_NFS
@@ -95,6 +96,11 @@ bool FTCPTransport::HasPendingPayload()
 #endif
 }
 
+FString FTCPTransport::GetHostName() const
+{
+	return HostName;
+}
+
 void FTCPTransport::Disconnect()
 {
 #if USE_MCSOCKET_FOR_NFS
diff --git a/Engine/Source/Runtime/CookOnTheFly/Private/TCPTransport.h b/Engine/Source/Runtime/CookOnTheFly/Private/TCPTransport.h
--- a/Engine/Source/Runtime/CookOnTheFly/Private/TCPTransport.h
+++ b/Engine/Source/Runtime/CookOnTheFly/Private/TCPTransport.h
@@ -19,10 +19,11 @@ public:
 	virtual bool ReceivePayload(FArrayReader& Payload) override;
 	virtual bool HasPendingPayload() override;
 	virtual void Disconnect() override;
+	virtual FString GetHostName() const override;
 
 private: 
 
 	class FSocket*		FileSocket;
 	class FMultichannelTcpSocket* MCSocket;
-
+	FString HostName;
 };
diff --git a/Engine/Source/Runtime/Core/Private/Misc/App.cpp b/Engine/Source/Runtime/Core/Private/Misc/App.cpp
--- a/Engine/Source/Runtime/Core/Private/Misc/App.cpp
+++ b/Engine/Source/Runtime/Core/Private/Misc/App.cpp
@@ -396,8 +396,8 @@ FString FApp::GetZenStoreProjectId()
 	if (FPaths::IsProjectFilePathSet())
 	{
 		FString ProjectFilePath = FPaths::GetProjectFilePath();
-		FPaths::NormalizeFilename(ProjectFilePath);
 		FString AbsProjectFilePath = IFileManager::Get().ConvertToAbsolutePathForExternalAppForRead(*ProjectFilePath);
+		AbsProjectFilePath = FPaths::FindCorrectCase(AbsProjectFilePath);
 		FTCHARToUTF8 AbsProjectFilePathUTF8(*AbsProjectFilePath);
 
 		FString HashString = FMD5::HashBytes((unsigned char*)AbsProjectFilePathUTF8.Get(), AbsProjectFilePathUTF8.Length()).Left(8);
@@ -408,4 +408,4 @@ FString FApp::GetZenStoreProjectId()
 	UE_LOG(LogInit, Fatal, TEXT("-ZenStoreProject command line argument is required to run from Zen"));
 #endif
 	return FString();
-}
\ No newline at end of file
+}
diff --git a/Engine/Source/Runtime/Core/Private/Misc/CoreGlobals.cpp b/Engine/Source/Runtime/Core/Private/Misc/CoreGlobals.cpp
--- a/Engine/Source/Runtime/Core/Private/Misc/CoreGlobals.cpp
+++ b/Engine/Source/Runtime/Core/Private/Misc/CoreGlobals.cpp
@@ -750,11 +750,7 @@ bool IsRunningCookOnTheFly()
 		FCookOnTheFlyCommandline(const TCHAR* CmdLine)
 		{
 			FString Host;
-			bParsed = FParse::Param(CmdLine, TEXT("CookOnTheFly")) && FParse::Value(CmdLine, TEXT("-ZenStoreHost="), Host);
-			if (!bParsed)	
-			{
-				bParsed = FParse::Value(CmdLine, TEXT("-FileHostIP="), Host);
-			}
+			bParsed = FParse::Value(CmdLine, TEXT("-FileHostIP="), Host);
 		}
 	} CookOnTheFlyCommandline(FCommandLine::Get());
 	
diff --git a/Engine/Source/Runtime/Core/Private/Misc/Paths.cpp b/Engine/Source/Runtime/Core/Private/Misc/Paths.cpp
--- a/Engine/Source/Runtime/Core/Private/Misc/Paths.cpp
+++ b/Engine/Source/Runtime/Core/Private/Misc/Paths.cpp
@@ -18,9 +18,16 @@
 #include "Misc/LazySingleton.h"
 #include "Misc/Parse.h"
 #include "Misc/ScopeLock.h"
+#include "Misc/PathViews.h"
 #include "String/ParseTokens.h"
 #include "UObject/NameTypes.h"
 
+#if PLATFORM_WINDOWS
+#include "Windows/AllowWindowsPlatformTypes.h"
+#include <fileapi.h>
+#include "Windows/HideWindowsPlatformTypes.h"
+#endif
+
 DEFINE_LOG_CATEGORY_STATIC(LogPaths, Log, All);
 
 #if !defined(SUPPORTS_LOGS_IN_USERDIR)
@@ -767,6 +774,58 @@ void FPaths::SetProjectFilePath( const FString& NewGameProjectFilePath )
 	FPaths::NormalizeFilename(StaticData.GameProjectFilePath);
 }
 
+FString FPaths::FindCorrectCase(const FString& Path)
+{
+#if PLATFORM_WINDOWS
+	// GetFilenameOnDisk on Windows will resolve directory junctions and resolving those here has negative consequences
+	// for workflows that use a junction at their root (eg: p4 gets confused about paths and operations fail).
+	// There is a way to get a case-accurate path on Windows without resolving directory junctions, but it is slow.
+	// We can use it here for this one-off situation without causing all uses of GetFilenameOnDisk to be slower.
+	FString ProjectFilePath = FPaths::GetProjectFilePath();
+	TStringBuilder<MAX_PATH> Builder;
+	FPathViews::IterateComponents(
+		ProjectFilePath,
+		[&Builder](FStringView CurrentPathComponent)
+		{
+			if (Builder.Len() != 0)
+			{
+				Builder.AppendChar(TEXT('/'));
+			}
+
+			// Any volume name should be upper case
+			const bool bIsVolumeSegment = CurrentPathComponent.EndsWith(TEXT(':'));
+			if (bIsVolumeSegment)
+			{
+				Builder.Append(FString(CurrentPathComponent).ToUpper());
+				return;
+			}
+
+			int32 LenBeforeCurrentComponent = Builder.Len();
+			Builder.Append(CurrentPathComponent);
+
+			// Skip over all segments that are either empty or contain relative transforms, they should remain as-is
+			const bool bIsIgnoredSegment = CurrentPathComponent.IsEmpty() || CurrentPathComponent.Equals(TEXTVIEW(".")) || CurrentPathComponent.Equals(TEXTVIEW(".."));
+			if (bIsIgnoredSegment)
+			{
+				return;
+			}
+
+			WIN32_FIND_DATAW Data;
+			HANDLE Handle = FindFirstFileW(StringCast<WIDECHAR>(*Builder, Builder.Len() + 1).Get(), &Data);
+			if (Handle != INVALID_HANDLE_VALUE)
+			{
+				Builder.RemoveSuffix(Builder.Len() - LenBeforeCurrentComponent);
+				Builder.Append(Data.cFileName);
+				FindClose(Handle);
+			}
+		}
+	);
+	return Builder.ToString();
+#else
+	return IFileManager::Get().GetFilenameOnDisk(*Path);
+#endif
+}
+
 FString FPaths::GetExtension( const FString& InPath, bool bIncludeDot )
 {
 	const FString Filename = GetCleanFilename(InPath);
diff --git a/Engine/Source/Runtime/Core/Public/Misc/Paths.h b/Engine/Source/Runtime/Core/Public/Misc/Paths.h
--- a/Engine/Source/Runtime/Core/Public/Misc/Paths.h
+++ b/Engine/Source/Runtime/Core/Public/Misc/Paths.h
@@ -494,6 +494,9 @@ public:
 	/** Convert all / and \ to TEXT("/") */
 	static void NormalizeFilename(FString& InPath);
 
+	/** Converts the path casing to match the casing found on disk without resolving directory junctions */
+	static FString FindCorrectCase(const FString& Path);
+
 	/**
 	 * Checks if two paths are the same.
 	 *
diff --git a/Engine/Source/Runtime/Launch/Private/LaunchEngineLoop.cpp b/Engine/Source/Runtime/Launch/Private/LaunchEngineLoop.cpp
--- a/Engine/Source/Runtime/Launch/Private/LaunchEngineLoop.cpp
+++ b/Engine/Source/Runtime/Launch/Private/LaunchEngineLoop.cpp
@@ -958,47 +958,9 @@ void LaunchFixProjectPathCase()
 {
 	if (FPaths::IsProjectFilePathSet())
 	{
-#if PLATFORM_WINDOWS
-		// GetFilenameOnDisk on Windows will resolve directory junctions and resolving those here has negative consequences
-		// for workflows that use a junction at their root (eg: p4 gets confused about paths and operations fail).
-		// There is a way to get a case-accurate path on Windows without resolving directory junctions, but it is slow.
-		// We can use it here for this one-off situation without causing all uses of GetFilenameOnDisk to be slower.
 		FString ProjectFilePath = FPaths::GetProjectFilePath();
-		TStringBuilder<MAX_PATH> Builder;
-		FPathViews::IterateComponents(
-			ProjectFilePath,
-			[&Builder](FStringView CurrentPathComponent)
-			{
-				if (Builder.Len() != 0)
-				{
-					Builder.AppendChar(TEXT('/'));
-				}
-
-				int32 LenBeforeCurrentComponent = Builder.Len();
-
-				Builder.Append(CurrentPathComponent);
-
-				// Skip over all segments that are either empty or contain relative transforms or start with the volume separator, they should remain as-is
-				const bool bIsIgnoredSegment = CurrentPathComponent.IsEmpty() || CurrentPathComponent.Equals(TEXTVIEW(".")) || CurrentPathComponent.Equals(TEXTVIEW("..")) || CurrentPathComponent.EndsWith(TEXT(':'));
-				if (bIsIgnoredSegment)
-				{
-					return;
-				}
-
-				WIN32_FIND_DATAW Data;
-				HANDLE Handle = FindFirstFileW(StringCast<WIDECHAR>(*Builder, Builder.Len() + 1).Get(), &Data);
-				if (Handle != INVALID_HANDLE_VALUE)
-				{
-					Builder.RemoveSuffix(Builder.Len() - LenBeforeCurrentComponent);
-					Builder.Append(Data.cFileName);
-					FindClose(Handle);
-				}
-			}
-		);
-		FPaths::SetProjectFilePath(Builder.ToString());
-#else
-		FPaths::SetProjectFilePath(IFileManager::Get().GetFilenameOnDisk(*FPaths::GetProjectFilePath()));
-#endif
+		FString ProjectFilePathCorrectCase = FPaths::FindCorrectCase(ProjectFilePath);
+		FPaths::SetProjectFilePath(ProjectFilePathCorrectCase);
 	}
 }
 
diff --git a/Engine/Source/Runtime/StorageServerClient/Private/StorageServerPlatformFile.cpp b/Engine/Source/Runtime/StorageServerClient/Private/StorageServerPlatformFile.cpp
--- a/Engine/Source/Runtime/StorageServerClient/Private/StorageServerPlatformFile.cpp
+++ b/Engine/Source/Runtime/StorageServerClient/Private/StorageServerPlatformFile.cpp
@@ -249,6 +249,16 @@ FStorageServerPlatformFile::~FStorageServerPlatformFile()
 
 bool FStorageServerPlatformFile::ShouldBeUsed(IPlatformFile* Inner, const TCHAR* CmdLine) const
 {
+#if WITH_COTF
+	UE::Cook::ICookOnTheFlyModule& CookOnTheFlyModule = FModuleManager::LoadModuleChecked<UE::Cook::ICookOnTheFlyModule>(TEXT("CookOnTheFly"));
+	TSharedPtr<UE::Cook::ICookOnTheFlyServerConnection> DefaultConnection = CookOnTheFlyModule.GetDefaultServerConnection();
+	if (DefaultConnection.IsValid() && !DefaultConnection->GetZenProjectName().IsEmpty())
+	{
+		HostAddrs.Add(DefaultConnection->GetZenHostName());
+		HostPort = DefaultConnection->GetZenHostPort();
+		return true;
+	}
+#endif
 	FString Host;
 	if (FParse::Value(FCommandLine::Get(), TEXT("-ZenStoreHost="), Host))
 	{
@@ -257,18 +267,6 @@ bool FStorageServerPlatformFile::ShouldBeUsed(IPlatformFile* Inner, const TCHAR*
 			HostAddrs.Add(Host);
 		}
 	}
-	else 
-	{
-#if WITH_COTF
-		UE::Cook::ICookOnTheFlyModule& CookOnTheFlyModule = FModuleManager::LoadModuleChecked<UE::Cook::ICookOnTheFlyModule>(TEXT("CookOnTheFly"));
-		TSharedPtr<UE::Cook::ICookOnTheFlyServerConnection> DefaultConnection = CookOnTheFlyModule.GetDefaultServerConnection();
-		if (DefaultConnection.IsValid() && !DefaultConnection->GetZenProjectName().IsEmpty())
-		{
-			HostAddrs.Add(DefaultConnection->GetHost());
-		}
-#endif
-	}
-
 	return HostAddrs.Num() > 0;
 }
 
@@ -300,7 +298,7 @@ void FStorageServerPlatformFile::InitializeAfterProjectFilePath()
 	Connection.Reset(new FStorageServerConnection());
 	const TCHAR* ProjectOverride = ServerProject.IsEmpty() ? nullptr : *ServerProject;
 	const TCHAR* PlatformOverride = ServerPlatform.IsEmpty() ? nullptr : *ServerPlatform;
-	if (Connection->Initialize(HostAddrs, 1337, ProjectOverride, PlatformOverride))
+	if (Connection->Initialize(HostAddrs, HostPort, ProjectOverride, PlatformOverride))
 	{
 		if (SendGetFileListMessage())
 		{
diff --git a/Engine/Source/Runtime/StorageServerClient/Private/StorageServerPlatformFile.h b/Engine/Source/Runtime/StorageServerClient/Private/StorageServerPlatformFile.h
--- a/Engine/Source/Runtime/StorageServerClient/Private/StorageServerPlatformFile.h
+++ b/Engine/Source/Runtime/StorageServerClient/Private/StorageServerPlatformFile.h
@@ -123,6 +123,7 @@ private:
 	FString ServerProject;
 	FString ServerPlatform;
 	mutable TArray<FString> HostAddrs;
+	mutable uint16 HostPort = 1337;
 };
 
 #endif
