- When archiving a Mac build, replace CEF3 framework with a proper version of the framework that contains symlinks, that the Mac App Store submission process will accept. It is non-trivial to attempt to add the framework to a Mac build using legacy Xcode project mode, so we fix up on Archive.

diff --git a/Engine/Source/Programs/AutomationTool/Mac/MacPlatform.Automation.cs b/Engine/Source/Programs/AutomationTool/Mac/MacPlatform.Automation.cs
--- a/Engine/Source/Programs/AutomationTool/Mac/MacPlatform.Automation.cs
+++ b/Engine/Source/Programs/AutomationTool/Mac/MacPlatform.Automation.cs
@@ -10,6 +10,7 @@ using UnrealBuildTool;
 using UnrealBuildBase;
 using EpicGames.Core;
 using Microsoft.Extensions.Logging;
+using System.Security.Cryptography;
 
 public class MacPlatform : ApplePlatform
 {
@@ -400,6 +401,36 @@ public class MacPlatform : ApplePlatform
 		PrintRunTime();
 	}
 
+	private void FixupFrameworks(string TargetPath)
+	{
+		DirectoryReference TargetCEFDir = DirectoryReference.Combine(new DirectoryReference(TargetPath), "Engine/Binaries/ThirdParty/CEF3/Mac");
+		DirectoryReference X86Framework = DirectoryReference.Combine(TargetCEFDir, "Chromium Embedded Framework x86.framework");
+		DirectoryReference X86Versions = DirectoryReference.Combine(X86Framework, "Versions");
+		DirectoryReference Arm64Framework = DirectoryReference.Combine(TargetCEFDir, "Chromium Embedded Framework arm64.framework");
+		DirectoryReference Arm64Versions = DirectoryReference.Combine(Arm64Framework, "Versions");
+
+		DirectoryReference EngineCEFDir = DirectoryReference.Combine(Unreal.EngineDirectory, "Binaries/ThirdParty/CEF3/Mac");
+		FileReference X86Zip = FileReference.Combine(EngineCEFDir, "Chromium Embedded Framework x86.framework.zip");
+		FileReference Arm64Zip = FileReference.Combine(EngineCEFDir, "Chromium Embedded Framework arm64.framework.zip");
+
+		// if the archive has a framework without Versions directory, it won't be allowed for App Store submission, so replace it with the zipped version
+		// that has the proper symlinks 
+		if (DirectoryReference.Exists(X86Framework) && !DirectoryReference.Exists(X86Versions))
+		{
+			Logger.LogInformation($"Replacing {X86Framework} with {X86Zip}...");
+
+			DirectoryReference.Delete(X86Framework, true);
+			Utils.RunLocalProcessAndLogOutput("/usr/bin/unzip", $"-q -o \"{X86Zip}\" -d \"{TargetCEFDir}\" -x \"__MACOSX/*\" \"*.DS_Store\"", Logger);
+		}
+		if (DirectoryReference.Exists(Arm64Framework) && !DirectoryReference.Exists(Arm64Versions))
+		{
+			Logger.LogInformation($"Replacing {Arm64Framework} with {Arm64Zip}...");
+
+			DirectoryReference.Delete(Arm64Framework, true);
+			Utils.RunLocalProcessAndLogOutput("/usr/bin/unzip", $"-q -o \"{Arm64Zip}\" -d \"{TargetCEFDir}\" -x \"__MACOSX/*\" \"*.DS_Store\"", Logger);
+		}
+	}
+
 	public override void ProcessArchivedProject(ProjectParams Params, DeploymentContext SC)
 	{
 		if (Params.CreateAppBundle)
@@ -455,6 +486,8 @@ public class MacPlatform : ApplePlatform
 						RenameDirectory(DirPath, TargetDirPath, true);
 					}
 				}
+
+				FixupFrameworks(TargetPath);
 			}
 
 			// Update executable name, icon and entry in Info.plist
@@ -493,6 +526,11 @@ public class MacPlatform : ApplePlatform
 				Utils.RunLocalProcessAndReturnStdOut("/usr/bin/codesign", $"-f -s \"Developer ID Application\" \"{BundlePath}\"", null);
 			}
 
+			// we now need to re-sign the .app because we modified the .plist
+			// we codesign with ad-hoc, and if the Developer ID Application cert exists, attempt to use it, ignore any errors
+			Utils.RunLocalProcessAndReturnStdOut("/usr/bin/codesign", $"-f -s - \"{BundlePath}\"", null);
+			Utils.RunLocalProcessAndReturnStdOut("/usr/bin/codesign", $"-f -s \"Developer ID Application\" \"{BundlePath}\"", null);
+
 			if (!SC.bIsCombiningMultiplePlatforms)
 			{
 				// creating these directories when the content isn't moved into the application causes it 
