PR #10253: AutomationTool: Respect specified error level for commandlets (Contributed by [REDACTED_PERSON])

diff --git a/Engine/Source/Programs/AutomationTool/AutomationUtils/CommandletUtils.cs b/Engine/Source/Programs/AutomationTool/AutomationUtils/CommandletUtils.cs
--- a/Engine/Source/Programs/AutomationTool/AutomationUtils/CommandletUtils.cs
+++ b/Engine/Source/Programs/AutomationTool/AutomationUtils/CommandletUtils.cs
@@ -264,7 +264,7 @@ namespace AutomationTool
 			DateTime StartTime = DateTime.UtcNow;
 
 			StartRunCommandlet(ProjectName, UnrealExe, Commandlet, Parameters, ERunOptions.Default, out LocalLogFile, out RunResult);
-			FinishRunCommandlet(ProjectName, Commandlet, StartTime, RunResult, LocalLogFile, out DestLogFile);
+			FinishRunCommandlet(ProjectName, Commandlet, StartTime, RunResult, LocalLogFile, out DestLogFile, ErrorLevel);
 		}
 
 		public static void StartRunCommandlet(FileReference ProjectName, string UnrealExe, string Commandlet, string Parameters, ERunOptions RunOptions, out string LocalLogFile, out IProcessResult RunResult, ProcessResult.SpewFilterCallbackType SpewFilterCallback=null)
@@ -304,7 +304,7 @@ namespace AutomationTool
 		}
 
 		public static void FinishRunCommandlet(FileReference ProjectName, string Commandlet, DateTime StartTime, IProcessResult RunResult, string LocalLogFile, out string DestLogFile, int ErrorLevel = 1)
-		{ 
+		{
 			// If we're running on a Windows build machine, copy any crash dumps into the log folder
 			if(HostPlatform.Current.HostEditorPlatform == UnrealTargetPlatform.Win64 && IsBuildMachine)
 			{
@@ -346,11 +346,11 @@ namespace AutomationTool
 					Logger.LogInformation("Pausing before checking for crash logs...");
 					Thread.Sleep(10 * 1000);
 				}
-				
+
 				// Create a list of directories containing crash logs, and add the system log folder
 				List<string> CrashDirs = new List<string>();
 				CrashDirs.Add("/Library/Logs/DiagnosticReports");
-					
+
 				// Add the user's log directory too
 				string HomeDir = Environment.GetEnvironmentVariable("HOME");
 				if(!String.IsNullOrEmpty(HomeDir))
@@ -375,7 +375,7 @@ namespace AutomationTool
 						// Not all account types can access /Library/Logs/DiagnosticReports
 					}
 				}
-				
+
 				// Dump them all to the log
 				foreach(FileInfo CrashFileInfo in CrashFileInfos)
 				{
@@ -423,7 +423,7 @@ namespace AutomationTool
 //				CommandUtils.LogWarning("Failed to find directory {0} will not save stats", ProjectStatsDirectory);
 //			}
 
-			// Whether it was copied correctly or not, delete the local log as it was only a temporary file. 
+			// Whether it was copied correctly or not, delete the local log as it was only a temporary file.
 			CommandUtils.DeleteFile_NoExceptions(LocalLogFile);
 
 			// Throw an exception if the execution failed. Draw attention to signal exit codes on Posix systems, rather than just printing the exit code
@@ -444,7 +444,7 @@ namespace AutomationTool
 				throw new CommandletException(DestLogFile, RunResult.ExitCode, "Editor terminated with exit code {0}{1} while running {2}{3}; see log {4}", RunResult.ExitCode, ExitCodeDesc, Commandlet, (ProjectName == null)? "" : String.Format(" for {0}", ProjectName), DestLogFile) { OutputFormat = AutomationExceptionOutputFormat.Minimal };
 			}
 		}
-		
+
 		/// <summary>
 		/// Returns the default path of the editor executable to use for running commandlets.
 		/// </summary>
