Fixed issue with Smartling localization autehntication timeouts.
- Slack for authentication has been increased to 30s to account for heavy network traffic.
- Fixed logical error in determining authentication expiration time.
- Demoted Smartling Http errors from errors to warnings to prevent over-aggressive Horde notifications.
#rnx
#test: Ran a localization import job locally to verify that all authentication tokens are valid with a slack of 470s to simulate heavy network traffic. Ran a Horde job to ensure Setup Build ran successfully.

diff --git a/Engine/Source/Programs/AutomationTool/SmartlingLocalization/SmartlingLocalizationProvider.cs b/Engine/Source/Programs/AutomationTool/SmartlingLocalization/SmartlingLocalizationProvider.cs
--- a/Engine/Source/Programs/AutomationTool/SmartlingLocalization/SmartlingLocalizationProvider.cs
+++ b/Engine/Source/Programs/AutomationTool/SmartlingLocalization/SmartlingLocalizationProvider.cs
@@ -235,7 +235,7 @@ namespace EpicGames.SmartlingLocalization
 			var RequestErrors = ResponseEnvelope.Response.Errors;
 			foreach (SmartlingRequestError RequestError in RequestErrors)
 			{
-				Console.WriteLine($"Smartling Error:\n {RequestError.ToString()}");
+				Console.WriteLine($"Smartling Warning:\n {RequestError.ToString()}");
 			}
 		}
 
@@ -378,8 +378,8 @@ namespace EpicGames.SmartlingLocalization
 
 			// If the last successful update + time to expire is still less than our current time accounting for some slack, we don't need to do anything to the token.
 			// We account for slack as it makes little sense to try authenticating with only 1s of validity left on a busy network
-			int Slack = 10; 
-			if (AuthenticationToken != null && AuthenticationToken.LastSuccessfulUpdateTime.AddSeconds(AuthenticationToken.ExpiresIn - Slack) < DateTime.UtcNow)
+			int Slack = 30; 
+			if (AuthenticationToken != null && AuthenticationToken.LastSuccessfulUpdateTime.AddSeconds(AuthenticationToken.ExpiresIn - Slack) > DateTime.UtcNow)
 			{
 				Console.WriteLine("Authentication token still valid. Using current authentication token.");
 				return;
