BuildGraph: Ignore empty values for file arguments in BuildGraph tasks.

diff --git a/Engine/Source/Programs/AutomationTool/BuildGraph/BgNodeExecutor.cs b/Engine/Source/Programs/AutomationTool/BuildGraph/BgNodeExecutor.cs
--- a/Engine/Source/Programs/AutomationTool/BuildGraph/BgNodeExecutor.cs
+++ b/Engine/Source/Programs/AutomationTool/BuildGraph/BgNodeExecutor.cs
@@ -227,7 +227,14 @@ namespace AutomationTool
 					{
 						// Parse it and assign it to the parameters object
 						object? FieldValue = ParseValue(Value, Parameter.ValueType);
-						Parameter.FieldInfo.SetValue(ParametersObject, FieldValue);
+						if (FieldValue != null)
+						{
+							Parameter.FieldInfo.SetValue(ParametersObject, FieldValue);
+						}
+						else if (!Parameter.Optional)
+						{
+							Logger.LogScriptError(TaskInfo.Location, "Empty value for parameter '{AttrName}' is not allowed.", Name);
+						}
 					}
 					else
 					{
@@ -243,8 +250,11 @@ namespace AutomationTool
 						List<string> ValueStrings = BgTaskImpl.SplitDelimitedList(Value);
 						foreach (string ValueString in ValueStrings)
 						{
-							object ElementValue = ParseValue(ValueString, Parameter.ValueType)!;
-							Parameter.CollectionType.InvokeMember("Add", BindingFlags.InvokeMethod | BindingFlags.Instance | BindingFlags.Public, null, CollectionValue, new object[] { ElementValue });
+							object? ElementValue = ParseValue(ValueString, Parameter.ValueType);
+							if (ElementValue != null)
+							{
+								Parameter.CollectionType.InvokeMember("Add", BindingFlags.InvokeMethod | BindingFlags.Instance | BindingFlags.Public, null, CollectionValue, new object[] { ElementValue });
+							}
 						}
 					}
 				}
@@ -314,11 +324,25 @@ namespace AutomationTool
 			}
 			else if (ValueType == typeof(FileReference))
 			{
-				return BgTaskImpl.ResolveFile(ValueText);
+				if (String.IsNullOrEmpty(ValueText))
+				{
+					return null;
+				}
+				else
+				{
+					return BgTaskImpl.ResolveFile(ValueText);
+				}
 			}
 			else if (ValueType == typeof(DirectoryReference))
 			{
-				return BgTaskImpl.ResolveDirectory(ValueText);
+				if (String.IsNullOrEmpty(ValueText))
+				{
+					return null;
+				}
+				else
+				{
+					return BgTaskImpl.ResolveDirectory(ValueText);
+				}
 			}
 
 			TypeConverter Converter = TypeDescriptor.GetConverter(ValueType);
