BuildGraph: Add a command for verifying that a project only references NuGet packages with allowed licenses.

diff --git a/Engine/Source/Programs/AutomationTool/BuildGraph/Tasks/NuGetLicenseCheckTask.cs b/Engine/Source/Programs/AutomationTool/BuildGraph/Tasks/NuGetLicenseCheckTask.cs
new file mode 100644
--- /dev/null
+++ b/Engine/Source/Programs/AutomationTool/BuildGraph/Tasks/NuGetLicenseCheckTask.cs
@@ -0,0 +1,308 @@
+// Copyright Epic Games, Inc. All Rights Reserved.
+
+using AutomationTool;
+using EpicGames.Core;
+using Microsoft.Extensions.Logging;
+using System;
+using System.CodeDom;
+using System.Collections.Generic;
+using System.Data;
+using System.IO;
+using System.Linq;
+using System.Net.Http;
+using System.Net.Http.Headers;
+using System.Text;
+using System.Text.RegularExpressions;
+using System.Threading.Tasks;
+using System.Web;
+using System.Xml;
+using UnrealBuildBase;
+
+namespace AutomationTool.Tasks
+{
+	/// <summary>
+	/// Parameters for a NuGetLicenseCheck task
+	/// </summary>
+	public class NuGetLicenseCheckTaskParameters
+	{
+		/// <summary>
+		/// Base directory for running the command
+		/// </summary>
+		[TaskParameter]
+		public string BaseDir;
+
+		/// <summary>
+		/// Directory containing allowed licenses
+		/// </summary>
+		[TaskParameter(Optional = true)]
+		public DirectoryReference LicenseDir;
+	}
+
+	/// <summary>
+	/// Spawns Docker and waits for it to complete.
+	/// </summary>
+	[TaskElement("NuGet-LicenseCheck", typeof(NuGetLicenseCheckTaskParameters))]
+	public class NuGetLicenseCheckTask : SpawnTaskBase
+	{
+		/// <summary>
+		/// Parameters for this task
+		/// </summary>
+		NuGetLicenseCheckTaskParameters Parameters;
+
+		/// <summary>
+		/// Construct a Docker task
+		/// </summary>
+		/// <param name="InParameters">Parameters for the task</param>
+		public NuGetLicenseCheckTask(NuGetLicenseCheckTaskParameters InParameters)
+		{
+			Parameters = InParameters;
+		}
+
+		class PackageInfo
+		{
+			public string Name;
+			public string Version;
+			public LicenseInfo License;
+			public string LicenseSource;
+		}
+
+		class LicenseInfo
+		{
+			public IoHash Hash;
+			public string Text;
+			public string NormalizedText;
+			public string Extension;
+			public bool Approved;
+		}
+
+		LicenseInfo FindOrAddLicense(Dictionary<IoHash, LicenseInfo> Licenses, string Text, string Extension)
+		{
+			string NormalizedText = Text;
+			NormalizedText = Regex.Replace(NormalizedText, @"^\s+", "", RegexOptions.Multiline);
+			NormalizedText = Regex.Replace(NormalizedText, @"\s+$", "", RegexOptions.Multiline);
+			NormalizedText = Regex.Replace(NormalizedText, "^(?:MIT License|The MIT License \\(MIT\\))\n", "", RegexOptions.Multiline);
+			NormalizedText = Regex.Replace(NormalizedText, "^Copyright \\(c\\)[^\n]*\\s*(?:All rights reserved\\.?\\s*)?", "", RegexOptions.Multiline);
+			NormalizedText = Regex.Replace(NormalizedText, @"\s+", " ");
+			NormalizedText = NormalizedText.Trim();
+
+			byte[] Data = Encoding.UTF8.GetBytes(NormalizedText);
+			IoHash Hash = IoHash.Compute(Data);
+
+			LicenseInfo LicenseInfo;
+			if (!Licenses.TryGetValue(Hash, out LicenseInfo))
+			{
+				LicenseInfo = new LicenseInfo();
+				LicenseInfo.Hash = Hash;
+				LicenseInfo.Text = Text;
+				LicenseInfo.NormalizedText = NormalizedText;
+				LicenseInfo.Extension = Extension;
+				Licenses.Add(Hash, LicenseInfo);
+			}
+			return LicenseInfo;
+		}
+
+		/// <summary>
+		/// Execute the task.
+		/// </summary>
+		/// <param name="Job">Information about the current job</param>
+		/// <param name="BuildProducts">Set of build products produced by this node.</param>
+		/// <param name="TagNameToFileSet">Mapping from tag names to the set of files they include</param>
+		public override async Task ExecuteAsync(JobContext Job, HashSet<FileReference> BuildProducts, Dictionary<string, HashSet<FileReference>> TagNameToFileSet)
+		{
+			Dictionary<string, PackageInfo> Packages = new Dictionary<string, PackageInfo>();
+
+			string OutputFilter(string Line)
+			{
+				Match Match = Regex.Match(Line, @"^\s*>\s*([^ ]+)\s+(?:[^ ]+\s+)?([^ ]+)\s*$");
+				if (Match.Success)
+				{
+					PackageInfo Info = new PackageInfo();
+					Info.Name = Match.Groups[1].Value;
+					Info.Version = Match.Groups[2].Value;
+					Packages.Add($"{Info.Name}@{Info.Version}", Info);
+				}
+				return Line;
+			}
+
+			const string UnknownPrefix = "Unknown-";
+
+			IProcessResult Result = await ExecuteAsync(Unreal.DotnetPath.FullName, "list package --include-transitive", WorkingDir: Parameters.BaseDir, LogOutput: false, SpewFilterCallback: OutputFilter);
+			if (Result.ExitCode != 0)
+			{
+				throw new AutomationException("Docker terminated with an exit code indicating an error ({0})", Result.ExitCode);
+			}
+
+			DirectoryReference PackageRootDir = DirectoryReference.Combine(DirectoryReference.GetSpecialFolder(Environment.SpecialFolder.UserProfile), ".nuget", "packages");
+			if (!DirectoryReference.Exists(PackageRootDir))
+			{
+				throw new AutomationException("Missing NuGet package cache at {0}", PackageRootDir);
+			}
+
+			Dictionary<IoHash, LicenseInfo> Licenses = new Dictionary<IoHash, LicenseInfo>();
+			if (Parameters.LicenseDir != null)
+			{
+				Logger.LogInformation("Reading allowed licenses from {LicenseDir}", Parameters.LicenseDir);
+				foreach (FileReference File in DirectoryReference.EnumerateFiles(Parameters.LicenseDir))
+				{
+					if (!File.GetFileName().StartsWith(UnknownPrefix, StringComparison.OrdinalIgnoreCase))
+					{
+						string Text = await FileReference.ReadAllTextAsync(File);
+						LicenseInfo License = FindOrAddLicense(Licenses, Text, File.GetFileNameWithoutExtension());
+						License.Approved = true;
+					}
+				}
+			}
+
+			Dictionary<string, LicenseInfo> LicenseUrlToInfo = new Dictionary<string, LicenseInfo>(StringComparer.OrdinalIgnoreCase);
+
+			Logger.LogInformation("Referenced Packages:");
+			Logger.LogInformation("");
+			foreach (PackageInfo Info in Packages.Values.OrderBy(x => x.Name).ThenBy(x => x.Version))
+			{
+				DirectoryReference PackageDir = DirectoryReference.Combine(PackageRootDir, Info.Name, Info.Version);
+				if (!DirectoryReference.Exists(PackageDir))
+				{
+					Logger.LogInformation("  {Name,-60} NuGet package not found", Info.Name);
+					continue;
+				}
+
+				FileReference NuSpecFile = FileReference.Combine(PackageDir, $"{Info.Name}.nuspec");
+				if (!FileReference.Exists(NuSpecFile))
+				{
+					Logger.LogWarning("Missing package descriptor: {NuSpecFile}", NuSpecFile);
+					continue;
+				}
+
+				using (Stream Stream = FileReference.Open(NuSpecFile, FileMode.Open, FileAccess.Read, FileShare.Read))
+				{
+					XmlTextReader XmlReader = new XmlTextReader(Stream);
+					XmlReader.Namespaces = false;
+
+					XmlDocument XmlDocument = new XmlDocument();
+					XmlDocument.Load(XmlReader);
+
+					if (Info.License == null)
+					{
+						XmlNode LicenseNode = XmlDocument.SelectSingleNode("/package/metadata/license");
+						if (LicenseNode?.Attributes["type"]?.InnerText?.Equals("file", StringComparison.Ordinal) ?? false)
+						{
+							FileReference LicenseFile = FileReference.Combine(PackageDir, LicenseNode.InnerText);
+							if (FileReference.Exists(LicenseFile))
+							{
+								string Text = await FileReference.ReadAllTextAsync(LicenseFile);
+								Info.License = FindOrAddLicense(Licenses, Text, LicenseFile.GetExtension());
+								Info.LicenseSource = LicenseFile.FullName;
+							}
+						}
+					}
+
+					if (Info.License == null)
+					{
+						XmlNode LicenseUrlNode = XmlDocument.SelectSingleNode("/package/metadata/licenseUrl");
+
+						string LicenseUrl = LicenseUrlNode?.InnerText;
+						if (LicenseUrl != null)
+						{
+							LicenseUrl = Regex.Replace(LicenseUrl, @"^https://github.com/(.*)/blob/(.*)$", @"https://raw.githubusercontent.com/$1/$2");
+							Info.LicenseSource = LicenseUrl;
+
+							if (!LicenseUrlToInfo.TryGetValue(LicenseUrl, out Info.License))
+							{
+								using (HttpClient Client = new HttpClient())
+								{
+									using HttpResponseMessage Response = await Client.GetAsync(LicenseUrl);
+									if (!Response.IsSuccessStatusCode)
+									{
+										Logger.LogError("Unable to fetch license from {LicenseUrl}", LicenseUrl);
+									}
+									else
+									{
+										string Text = await Response.Content.ReadAsStringAsync();
+										string Type = (Response.Content.Headers.ContentType?.MediaType == "text/html") ? ".html" : ".txt";
+										Info.License = FindOrAddLicense(Licenses, Text, Type);
+										LicenseUrlToInfo.Add(LicenseUrl, Info.License);
+									}
+								}
+							}
+						}
+					}
+				}
+
+				if (Info.License == null)
+				{
+					Logger.LogError("  {Name,-60} {Version,-10} No license metadata found", Info.Name, Info.Version);
+				}
+				else if (!Info.License.Approved)
+				{
+					Logger.LogWarning("  {Name,-60} {Version,-10} {Hash}", Info.Name, Info.Version, Info.License.Hash);
+				}
+				else
+				{
+					Logger.LogInformation("  {Name,-60} {Version,-10} {Hash}", Info.Name, Info.Version, Info.License.Hash);
+				}
+			}
+
+			Dictionary<LicenseInfo, List<PackageInfo>> MissingLicenses = new Dictionary<LicenseInfo, List<PackageInfo>>();
+			foreach (PackageInfo PackageInfo in Packages.Values)
+			{
+				if (PackageInfo.License != null && !PackageInfo.License.Approved)
+				{
+					List<PackageInfo> LicensePackages;
+					if (!MissingLicenses.TryGetValue(PackageInfo.License, out LicensePackages))
+					{
+						LicensePackages = new List<PackageInfo>();
+						MissingLicenses.Add(PackageInfo.License, LicensePackages);
+					}
+					LicensePackages.Add(PackageInfo);
+				}
+			}
+
+			if (MissingLicenses.Count > 0)
+			{
+				DirectoryReference LicenseDir = Parameters.LicenseDir ?? DirectoryReference.Combine(Unreal.RootDirectory, "Engine", "Saved", "Licenses");
+				DirectoryReference.CreateDirectory(LicenseDir);
+
+				Logger.LogInformation("");
+				Logger.LogInformation("Missing licenses:");
+				foreach ((LicenseInfo MissingLicense, List<PackageInfo> MissingLicensePackages) in MissingLicenses.OrderBy(x => x.Key.Hash))
+				{
+					FileReference OutputFile = FileReference.Combine(LicenseDir, $"{UnknownPrefix}{MissingLicense.Hash}{MissingLicense.Extension}");
+					await FileReference.WriteAllTextAsync(OutputFile, MissingLicense.Text);
+
+					Logger.LogInformation("");
+					Logger.LogInformation("  {LicenseFile}", OutputFile);
+					foreach (PackageInfo LicensePackage in MissingLicensePackages)
+					{
+						Logger.LogInformation("  -> {Name} {Version} ({Source})", LicensePackage.Name, LicensePackage.Version, LicensePackage.LicenseSource);
+					}
+				}
+			}
+		}
+
+		/// <summary>
+		/// Output this task out to an XML writer.
+		/// </summary>
+		public override void Write(XmlWriter Writer)
+		{
+			Write(Writer, Parameters);
+		}
+
+		/// <summary>
+		/// Find all the tags which are used as inputs to this task
+		/// </summary>
+		/// <returns>The tag names which are read by this task</returns>
+		public override IEnumerable<string> FindConsumedTagNames()
+		{
+			yield break;
+		}
+
+		/// <summary>
+		/// Find all the tags which are modified by this task
+		/// </summary>
+		/// <returns>The tag names which are modified by this task</returns>
+		public override IEnumerable<string> FindProducedTagNames()
+		{
+			yield break;
+		}
+	}
+}
diff --git a/Engine/Source/Programs/AutomationTool/BuildGraph/Tasks/SpawnTask.cs b/Engine/Source/Programs/AutomationTool/BuildGraph/Tasks/SpawnTask.cs
--- a/Engine/Source/Programs/AutomationTool/BuildGraph/Tasks/SpawnTask.cs
+++ b/Engine/Source/Programs/AutomationTool/BuildGraph/Tasks/SpawnTask.cs
@@ -81,7 +81,7 @@ namespace AutomationTool.Tasks
 			CommandUtils.ERunOptions Options = CommandUtils.ERunOptions.Default;
 			if (!LogOutput)
 			{
-				Options &= ~CommandUtils.ERunOptions.AllowSpew;
+				Options |= CommandUtils.ERunOptions.SpewIsVerbose;
 			}
 
 			IProcessResult Result = CommandUtils.Run(Exe, Arguments, Env: EnvVars, WorkingDir: WorkingDir, Options: Options, Input: Input, SpewFilterCallback: SpewFilterCallback);
