Horde: Query the configured NuGet package folders before finding dependencies.

diff --git a/Engine/Source/Programs/AutomationTool/BuildGraph/Tasks/NuGetLicenseCheckTask.cs b/Engine/Source/Programs/AutomationTool/BuildGraph/Tasks/NuGetLicenseCheckTask.cs
--- a/Engine/Source/Programs/AutomationTool/BuildGraph/Tasks/NuGetLicenseCheckTask.cs
+++ b/Engine/Source/Programs/AutomationTool/BuildGraph/Tasks/NuGetLicenseCheckTask.cs
@@ -115,9 +115,34 @@ namespace AutomationTool.Tasks
 		/// <param name="TagNameToFileSet">Mapping from tag names to the set of files they include</param>
 		public override async Task ExecuteAsync(JobContext Job, HashSet<FileReference> BuildProducts, Dictionary<string, HashSet<FileReference>> TagNameToFileSet)
 		{
-			Dictionary<string, PackageInfo> Packages = new Dictionary<string, PackageInfo>();
+			IProcessResult NuGetOutput = await ExecuteAsync(Unreal.DotnetPath.FullName, $"nuget locals global-packages --list", LogOutput: false);
+			if (NuGetOutput.ExitCode != 0)
+			{
+				throw new AutomationException("DotNet terminated with an exit code indicating an error ({0})", NuGetOutput.ExitCode);
+			}
+
+			List<DirectoryReference> NuGetPackageDirs = new List<DirectoryReference>();
+			foreach (string Line in NuGetOutput.Output.Split('\n'))
+			{
+				int ColonIdx = Line.IndexOf(':');
+				if (ColonIdx != -1)
+				{
+					DirectoryReference NuGetPackageDir = new DirectoryReference(Line.Substring(ColonIdx + 1).Trim());
+					Logger.LogInformation("Using NuGet package directory: {Path}", NuGetPackageDir);
+					NuGetPackageDirs.Add(NuGetPackageDir);
+				}
+			}
+
+			const string UnknownPrefix = "Unknown-";
 
-			string OutputFilter(string Line)
+			IProcessResult PackageListOutput = await ExecuteAsync(Unreal.DotnetPath.FullName, "list package --include-transitive", WorkingDir: Parameters.BaseDir, LogOutput: false);
+			if (PackageListOutput.ExitCode != 0)
+			{
+				throw new AutomationException("DotNet terminated with an exit code indicating an error ({0})", PackageListOutput.ExitCode);
+			}
+
+			Dictionary<string, PackageInfo> Packages = new Dictionary<string, PackageInfo>();
+			foreach (string Line in PackageListOutput.Output.Split('\n'))
 			{
 				Match Match = Regex.Match(Line, @"^\s*>\s*([^ ]+)\s+(?:[^ ]+\s+)?([^ ]+)\s*$");
 				if (Match.Success)
@@ -127,15 +152,6 @@ namespace AutomationTool.Tasks
 					Info.Version = Match.Groups[2].Value;
 					Packages.Add($"{Info.Name}@{Info.Version}", Info);
 				}
-				return Line;
-			}
-
-			const string UnknownPrefix = "Unknown-";
-
-			IProcessResult Result = await ExecuteAsync(Unreal.DotnetPath.FullName, "list package --include-transitive", WorkingDir: Parameters.BaseDir, LogOutput: false, SpewFilterCallback: OutputFilter);
-			if (Result.ExitCode != 0)
-			{
-				throw new AutomationException("Docker terminated with an exit code indicating an error ({0})", Result.ExitCode);
 			}
 
 			DirectoryReference PackageRootDir = DirectoryReference.Combine(DirectoryReference.GetSpecialFolder(Environment.SpecialFolder.UserProfile), ".nuget", "packages");
@@ -177,8 +193,8 @@ namespace AutomationTool.Tasks
 					continue;
 				}
 
-				DirectoryReference PackageDir = DirectoryReference.Combine(PackageRootDir, Info.Name, Info.Version);
-				if (!DirectoryReference.Exists(PackageDir))
+				DirectoryReference PackageDir = NuGetPackageDirs.Select(x => DirectoryReference.Combine(x, Info.Name, Info.Version)).FirstOrDefault(x => DirectoryReference.Exists(x));
+				if (PackageDir == null)
 				{
 					Logger.LogInformation("  {Name,-60} {Version,-10} NuGet package not found", Info.Name, Info.Version);
 					continue;
diff --git a/Engine/Source/Programs/Horde/BuildHorde.xml b/Engine/Source/Programs/Horde/BuildHorde.xml
--- a/Engine/Source/Programs/Horde/BuildHorde.xml
+++ b/Engine/Source/Programs/Horde/BuildHorde.xml
@@ -17,6 +17,7 @@
 		<Property Name="StagingDir" Value="$(RootDir)/Engine/Saved/Horde.Server"/>
 		
 		<Node Name="Check HordeServer Licenses">
+			<DotNet Arguments="build &quot;$(RootDir)/Engine/Source/Programs/Horde/Horde.Server/Horde.Server.csproj&quot; $(VersionArguments)"/>
 			<Property Name="IgnorePackages">
 				[REDACTED_EMAIL]
 				[REDACTED_EMAIL]
