BuildGraph: Use lowercase directories when looking for NuGet packages, to fix case sensitivity on Linux.

diff --git a/Engine/Source/Programs/AutomationTool/BuildGraph/Tasks/NuGetLicenseCheckTask.cs b/Engine/Source/Programs/AutomationTool/BuildGraph/Tasks/NuGetLicenseCheckTask.cs
--- a/Engine/Source/Programs/AutomationTool/BuildGraph/Tasks/NuGetLicenseCheckTask.cs
+++ b/Engine/Source/Programs/AutomationTool/BuildGraph/Tasks/NuGetLicenseCheckTask.cs
@@ -193,7 +193,7 @@ namespace AutomationTool.Tasks
 					continue;
 				}
 
-				DirectoryReference PackageDir = NuGetPackageDirs.Select(x => DirectoryReference.Combine(x, Info.Name, Info.Version)).FirstOrDefault(x => DirectoryReference.Exists(x));
+				DirectoryReference PackageDir = NuGetPackageDirs.Select(x => DirectoryReference.Combine(x, Info.Name.ToLowerInvariant(), Info.Version.ToLowerInvariant())).FirstOrDefault(x => DirectoryReference.Exists(x));
 				if (PackageDir == null)
 				{
 					Logger.LogInformation("  {Name,-60} {Version,-10} NuGet package not found", Info.Name, Info.Version);
