BuildGraph: Fix another case sensitivity issue for NuGetLicenseCheck.

diff --git a/Engine/Source/Programs/AutomationTool/BuildGraph/Tasks/NuGetLicenseCheckTask.cs b/Engine/Source/Programs/AutomationTool/BuildGraph/Tasks/NuGetLicenseCheckTask.cs
--- a/Engine/Source/Programs/AutomationTool/BuildGraph/Tasks/NuGetLicenseCheckTask.cs
+++ b/Engine/Source/Programs/AutomationTool/BuildGraph/Tasks/NuGetLicenseCheckTask.cs
@@ -200,7 +200,7 @@ namespace AutomationTool.Tasks
 					continue;
 				}
 
-				FileReference NuSpecFile = FileReference.Combine(PackageDir, $"{Info.Name}.nuspec");
+				FileReference NuSpecFile = FileReference.Combine(PackageDir, $"{Info.Name.ToLowerInvariant()}.nuspec");
 				if (!FileReference.Exists(NuSpecFile))
 				{
 					Logger.LogWarning("Missing package descriptor: {NuSpecFile}", NuSpecFile);
