BuildGraph: Fix \r characters being matched when parsing NuGet package output on Windows.

diff --git a/Engine/Source/Programs/AutomationTool/BuildGraph/Tasks/NuGetLicenseCheckTask.cs b/Engine/Source/Programs/AutomationTool/BuildGraph/Tasks/NuGetLicenseCheckTask.cs
--- a/Engine/Source/Programs/AutomationTool/BuildGraph/Tasks/NuGetLicenseCheckTask.cs
+++ b/Engine/Source/Programs/AutomationTool/BuildGraph/Tasks/NuGetLicenseCheckTask.cs
@@ -144,7 +144,7 @@ namespace AutomationTool.Tasks
 			Dictionary<string, PackageInfo> Packages = new Dictionary<string, PackageInfo>();
 			foreach (string Line in PackageListOutput.Output.Split('\n'))
 			{
-				Match Match = Regex.Match(Line, @"^\s*>\s*([^ ]+)\s+(?:[^ ]+\s+)?([^ ]+)\s*$");
+				Match Match = Regex.Match(Line, @"^\s*>\s*([^\s]+)\s+(?:[^\s]+\s+)?([^\s]+)\s*$");
 				if (Match.Success)
 				{
 					PackageInfo Info = new PackageInfo();
