Log UE critical failure as Gauntlet_FatalEvent
+ limit the number of reported not run tests
+ clean up some part of summary report

#rnx

diff --git a/Engine/Source/Programs/AutomationTool/Gauntlet/Framework/Utils/Gauntlet.HordeReport.cs b/Engine/Source/Programs/AutomationTool/Gauntlet/Framework/Utils/Gauntlet.HordeReport.cs
--- a/Engine/Source/Programs/AutomationTool/Gauntlet/Framework/Utils/Gauntlet.HordeReport.cs
+++ b/Engine/Source/Programs/AutomationTool/Gauntlet/Framework/Utils/Gauntlet.HordeReport.cs
@@ -981,7 +981,7 @@ namespace Gauntlet
 								Context = string.Format("{0} @line:{1}", InTestEntry.Filename, InTestEntry.LineNumber.ToString());
 							}
 							// Tag critical failure as crash
-							if (InTestEntry.Event.Type == EventType.Error && Tag == "entry" && InTestEntry.Event.Message.Contains("critical failure"))
+							if (InTestEntry.Event.Type == EventType.Error && Tag == "entry" && InTestEntry.Event.IsCriticalFailure)
 							{
 								Tag = "crash";
 							}
diff --git a/Engine/Source/Programs/AutomationTool/Gauntlet/Unreal/Automation/UE.Automation.cs b/Engine/Source/Programs/AutomationTool/Gauntlet/Unreal/Automation/UE.Automation.cs
--- a/Engine/Source/Programs/AutomationTool/Gauntlet/Unreal/Automation/UE.Automation.cs
+++ b/Engine/Source/Programs/AutomationTool/Gauntlet/Unreal/Automation/UE.Automation.cs
@@ -11,6 +11,7 @@ using System.Threading.Tasks;
 using UnrealBuildTool;
 using EpicGames.Core;
 using Log = Gauntlet.Log;
+using Microsoft.Extensions.Logging;
 
 namespace UE
 {
@@ -672,17 +673,16 @@ namespace UE
 						}
 						else
 						{
-							ErrorMessage = "Engine encountered a critical failure. \n";
 							if (FatalError != null)
 							{
-								ErrorMessage += FatalError.FormatForLog();
+								ErrorMessage = FatalError.FormatForLog();
 							}
 							else
 							{
-								ErrorMessage += "No callstack found in the log.";
+								ErrorMessage = "No callstack found in the log.";
 							}
 						}
-						Test.AddError(ErrorMessage);
+						Test.AddError(ErrorMessage, !HasTimeout);
 						if (!CanRetry() || JsonTestPassResults.NotRun == 0)
 						{
 							// Setting the test as fail because no retry will be done anymore.
@@ -875,7 +875,7 @@ namespace UE
 				IEnumerable<UnrealAutomatedTestResult> FailedTests = AllTests.Where(T => T.IsComplete && T.HasFailed);
 				IEnumerable<UnrealAutomatedTestResult> TestsWithWarnings = AllTests.Where(T => T.HasSucceeded && T.HasWarnings);
 
-				Func<IEnumerable<string>, IEnumerable<string>> CapErrorOrWarningList = (E) =>
+				Func<IEnumerable<UnrealAutomationEvent>, IEnumerable<UnrealAutomationEvent>> CapErrorOrWarningList = (E) =>
 				{
 					if (E.Count() > kMaxErrorsOrWarningsToDisplay)
 					{
@@ -901,15 +901,16 @@ namespace UE
 							{
 								string Message = !Result.IsComplete ? " * Test '{Name}' did not complete." : " * Test '{Name}' failed.";
 								Log.Error(KnownLogEvents.Gauntlet_UnrealEngineTestEvent, Message, Result.FullTestPath);
-								var Errors = CapErrorOrWarningList(Result.ErrorEvents.Select(E => E.Message).Distinct());
+								var Errors = CapErrorOrWarningList(Result.ErrorEvents.Distinct());
 								foreach (var Error in Errors)
 								{
-									Log.Error(KnownLogEvents.Gauntlet_UnrealEngineTestEvent, "    " + Error);
+									EventId ErrorEventType = Error.IsCriticalFailure ? KnownLogEvents.Gauntlet_FatalEvent : KnownLogEvents.Gauntlet_UnrealEngineTestEvent;
+									Log.Error(ErrorEventType, "    " + Error.FormatToString());
 								}
-								var Warnings = CapErrorOrWarningList(Result.WarningEvents.Select(E => E.Message).Distinct());
+								var Warnings = CapErrorOrWarningList(Result.WarningEvents.Distinct());
 								foreach (var Warning in Warnings)
 								{
-									Log.Warning(KnownLogEvents.Gauntlet_UnrealEngineTestEvent, "    " + Warning);
+									Log.Warning(KnownLogEvents.Gauntlet_UnrealEngineTestEvent, "    " + Warning.FormatToString());
 								}
 								Log.Info("");
 							}
@@ -917,8 +918,13 @@ namespace UE
 						var NotRun = IncompleteTests.Where(T => T.State == TestStateType.NotRun);
 						if (NotRun.Any())
 						{
+							int TotalNotRun = NotRun.Count();
 							Log.Info(" ### The following Engine test(s) were not run:");
-							Log.Info(string.Join("\n", NotRun.Select(T => " * "+T.FullTestPath)));
+							Log.Info(string.Join("\n", NotRun.Take(10).Select(T => " * "+T.FullTestPath)));
+							if (TotalNotRun > 10)
+							{
+								Log.Info(" (and {Count} more)", TotalNotRun - 10);
+							}
 							Log.Info("");
 						}
 					}
@@ -929,18 +935,18 @@ namespace UE
 					{
 						Log.Error(KnownLogEvents.Gauntlet_TestEvent, " * No tests were executed.");
 
-						IEnumerable<string> Errors = Events.Where(E => E.IsError).Select(E => E.Message).Distinct();
-						IEnumerable<string> Warnings = Events.Where(E => E.IsWarning).Select(E => E.Message).Distinct();
+						IEnumerable<UnrealAutomationEvent> Errors = Events.Where(E => E.IsError).Distinct();
+						IEnumerable<UnrealAutomationEvent> Warnings = Events.Where(E => E.IsWarning).Distinct();
 						if (Errors.Any() || Warnings.Any())
 						{
 							Log.Info("   See log above for details.");
 							foreach (var Error in CapErrorOrWarningList(Errors))
 							{
-								Log.Error(KnownLogEvents.Gauntlet_UnrealEngineTestEvent, "    " + Error);
+								Log.Error(KnownLogEvents.Gauntlet_UnrealEngineTestEvent, "    " + Error.FormatToString());
 							}
 							foreach (var Warning in CapErrorOrWarningList(Warnings))
 							{
-								Log.Warning(KnownLogEvents.Gauntlet_UnrealEngineTestEvent, "    " + Warning);
+								Log.Warning(KnownLogEvents.Gauntlet_UnrealEngineTestEvent, "    " + Warning.FormatToString());
 							}
 							Log.Info("");
 						}
@@ -959,12 +965,12 @@ namespace UE
 							foreach (UnrealAutomatedTestResult Result in FailedTests)
 							{
 								Log.Error(KnownLogEvents.Gauntlet_UnrealEngineTestEvent, " * Test '{Name}' failed.", Result.FullTestPath);
-								// only show the last N items
-								IEnumerable<string> Events = Result.ErrorEvents.Select(E => E.Message).Distinct();
+								IEnumerable<UnrealAutomationEvent> Events = Result.ErrorEvents.Distinct();
 								Events = CapErrorOrWarningList(Events);
 								foreach (var Event in Events)
 								{
-									Log.Error(KnownLogEvents.Gauntlet_UnrealEngineTestEvent, "    " + Event);
+									EventId ErrorEventType = Event.IsCriticalFailure ? KnownLogEvents.Gauntlet_FatalEvent : KnownLogEvents.Gauntlet_UnrealEngineTestEvent;
+									Log.Error(ErrorEventType, "    " + Event.FormatToString());
 								}
 								Log.Info("");
 							}
@@ -977,12 +983,12 @@ namespace UE
 							foreach (UnrealAutomatedTestResult Result in TestsWithWarnings)
 							{
 								Log.Warning(KnownLogEvents.Gauntlet_UnrealEngineTestEvent, " * Test '{Name}' completed with warnings.", Result.FullTestPath);
-								// only show the last N items
-								IEnumerable<string> WarningEvents = Result.WarningEvents.Select(E => E.Message).Distinct();
+								// only show the first N items
+								IEnumerable<UnrealAutomationEvent> WarningEvents = Result.WarningEvents.Distinct();
 								WarningEvents = CapErrorOrWarningList(WarningEvents);
 								foreach (var Event in WarningEvents)
 								{
-									Log.Warning(KnownLogEvents.Gauntlet_UnrealEngineTestEvent, "    " + Event);
+									Log.Warning(KnownLogEvents.Gauntlet_UnrealEngineTestEvent, "    " + Event.FormatToString());
 								}
 								Log.Info("");
 							}
@@ -995,12 +1001,12 @@ namespace UE
 							foreach (UnrealAutomatedTestResult Result in IncompleteTests)
 							{
 								Log.Error(KnownLogEvents.Gauntlet_UnrealEngineTestEvent, " * Test '{Name}' did not run or complete.", Result.FullTestPath);
-								// only show the last N items
-								IEnumerable<string> ErrorAndWarningEvents = Result.WarningAndErrorEvents.Select(E => E.IsError? E.Message : E.Message).Distinct();
+								// only show the first N items
+								IEnumerable<UnrealAutomationEvent> ErrorAndWarningEvents = Result.WarningAndErrorEvents.Distinct();
 								ErrorAndWarningEvents = CapErrorOrWarningList(ErrorAndWarningEvents);
 								foreach (var Event in ErrorAndWarningEvents)
 								{
-									Log.Error(KnownLogEvents.Gauntlet_UnrealEngineTestEvent, "    " + Event);
+									Log.Error(KnownLogEvents.Gauntlet_UnrealEngineTestEvent, "    " + Event.FormatToString());
 								}
 							}
 							Log.Info("");
diff --git a/Engine/Source/Programs/AutomationTool/Gauntlet/Unreal/Base/Gauntlet.UnrealTestNode.cs b/Engine/Source/Programs/AutomationTool/Gauntlet/Unreal/Base/Gauntlet.UnrealTestNode.cs
--- a/Engine/Source/Programs/AutomationTool/Gauntlet/Unreal/Base/Gauntlet.UnrealTestNode.cs
+++ b/Engine/Source/Programs/AutomationTool/Gauntlet/Unreal/Base/Gauntlet.UnrealTestNode.cs
@@ -2248,15 +2248,8 @@ namespace Gauntlet
 			// Good/Bad news upfront
 			string Prefix = TestFailed ? "Error: " : "";
 			string WarningStatement = (HasWarnings && !TestFailed)  ? " With Warnings" : "";
-			string ResultString = string.Format(" ### {0}{1} {2}{3} ***", Prefix, this.Name, GetTestResult(), WarningStatement);
-			if(TestFailed)
-			{
-				Log.Error(KnownLogEvents.Gauntlet_TestEvent, ResultString);
-			}
-			else
-			{
-				Log.Info(ResultString);
-			}
+			string ResultString = string.Format(" ### {0}{1} {2}{3} ***\n", Prefix, this.Name, GetTestResult(), WarningStatement);
+			Log.Info(ResultString);
 
 			IEnumerable<UnrealRoleResult> SortedRoles = RoleResults.OrderBy(R => R.ProcessResult == UnrealProcessResult.ExitOk);
 
@@ -2266,7 +2259,7 @@ namespace Gauntlet
 			{
 				string RoleName = RoleResult.Artifacts.SessionRole.RoleType.ToString();
 
-				Log.Info(string.Format("\n {0} Role: {1} ({2}, ExitCode {3})\n", RoleName, RoleResult.Summary, RoleResult.ProcessResult, RoleResult.ExitCode));
+				Log.Info(string.Format(" {0} Role: {1} ({2}, ExitCode {3})", RoleName, RoleResult.Summary, RoleResult.ProcessResult, RoleResult.ExitCode));
 
 				FatalErrors += RoleResult.LogSummary.FatalError != null ? 1 : 0;
 				Ensures += RoleResult.LogSummary.Ensures.Count();
@@ -2274,14 +2267,15 @@ namespace Gauntlet
 				Warnings += RoleResult.LogSummary.Warnings.Count();
 			}
 
+			Log.Info("");
 			Log.Info(string.Join("\n", new string[] {
-				string.Format("Context: {0}", Context.ToString()),
-				FatalErrors > 0 ? string.Format("FatalErrors: {0}", FatalErrors) : null,
-				Ensures > 0 ? string.Format("Ensures: {0}", Ensures) : null,
-				Errors > 0 ? string.Format("Log Errors: {0}", Errors) : null,
-				Warnings > 0 ? string.Format("Log Warnings: {0}", Warnings) : null,
-				string.Format("Result: {0}", GetTestResult())
-			}.Where(L=>!string.IsNullOrEmpty(L)).Select(L=>" * "+L)));
+				string.Format(" * Context: {0}", Context.ToString()),
+				FatalErrors > 0 ? string.Format(" * FatalErrors: {0}", FatalErrors) : null,
+				Ensures > 0 ? string.Format(" * Ensures: {0}", Ensures) : null,
+				Errors > 0 ? string.Format(" * Log Errors: {0}", Errors) : null,
+				Warnings > 0 ? string.Format(" * Log Warnings: {0}", Warnings) : null,
+				string.Format(" * Result: {0}", GetTestResult())
+			}.Where(L=>!string.IsNullOrEmpty(L))));
 
 			if (TestFailed && GetRolesThatFailed().Where(R => R.ProcessResult == UnrealProcessResult.Unknown).Any())
 			{
diff --git a/Engine/Source/Programs/AutomationTool/Gauntlet/Unreal/Utils/Gauntlet.UnrealTestResults.cs b/Engine/Source/Programs/AutomationTool/Gauntlet/Unreal/Utils/Gauntlet.UnrealTestResults.cs
--- a/Engine/Source/Programs/AutomationTool/Gauntlet/Unreal/Utils/Gauntlet.UnrealTestResults.cs
+++ b/Engine/Source/Programs/AutomationTool/Gauntlet/Unreal/Utils/Gauntlet.UnrealTestResults.cs
@@ -66,12 +66,18 @@ namespace Gauntlet
 		[JsonPropertyName("artifact")]
 		public string Artifact { get; set; }
 
+		const string CriticalFailureString = "critical failure";
+
 		public UnrealAutomationEvent()
 		{ }
-		public UnrealAutomationEvent(EventType InType, string InMessage)
+		public UnrealAutomationEvent(EventType InType, string InMessage, bool bIsCriticalFailure = false)
 		{
 			Type = InType;
 			Message = InMessage;
+			if(bIsCriticalFailure)
+			{
+				SetAsCriticalFailure();
+			}
 		}
 
 		[JsonIgnore]
@@ -93,6 +99,47 @@ namespace Gauntlet
 			}
 
 		}
+
+		[JsonIgnore]
+		public bool IsCriticalFailure
+		{
+			get
+			{
+				return IsError && Message.Contains(CriticalFailureString);
+			}
+		}
+		public void SetAsCriticalFailure()
+		{
+			if (!IsCriticalFailure)
+			{
+				Type = EventType.Error;
+				Message = $"Engine encountered a {CriticalFailureString}. \n" + Message;
+			}
+		}
+
+		public string FormatToString()
+		{
+			return Message;
+		}
+
+		protected bool Equals(UnrealAutomationEvent other)
+		{
+			return string.Equals(Type, other.Type) && string.Equals(Message, other.Message);
+		}
+
+		public override bool Equals(object obj)
+		{
+			if (ReferenceEquals(null, obj)) return false;
+			if (ReferenceEquals(this, obj)) return true;
+			if (obj.GetType() != this.GetType()) return false;
+			return Equals((UnrealAutomationEvent)obj);
+		}
+
+		public override int GetHashCode()
+		{
+			return Type.GetHashCode() ^ Message.GetHashCode();
+		}
+
 	}
 	public class UnrealAutomationEntry
 	{
@@ -136,15 +183,15 @@ namespace Gauntlet
 			Entries = new List<UnrealAutomationEntry>();
 		}
 
-		public void AddEvent(EventType EventType, string Message)
+		public void AddEvent(EventType EventType, string Message, bool bIsCriticalFailure = false)
 		{
-			var Event = new UnrealAutomationEvent(EventType, Message);
+			var Event = new UnrealAutomationEvent(EventType, Message, bIsCriticalFailure);
 			var Entry = new UnrealAutomationEntry();
 			Entry.Event = Event;
 			Entry.Timestamp = System.DateTime.Now.ToString("yyyy.MM.dd-HH.mm.ss");
 			Entries.Add(Entry);
 
-			switch(EventType)
+			switch(Event.Type)
 			{
 				case EventType.Error:
 					Errors++;
@@ -154,9 +201,9 @@ namespace Gauntlet
 					break;
 			}
 		}
-		public void AddError(string Message)
+		public void AddError(string Message, bool bIsCriticalFailure = false)
 		{
-			AddEvent(EventType.Error, Message);
+			AddEvent(EventType.Error, Message, bIsCriticalFailure);
 		}
 		public void AddWarning(string Message)
 		{
