Fixed deprecated log warnings

diff --git a/Engine/Source/Programs/AutomationTool/Scripts/CopyBuildToStagingDirectory.Automation.cs b/Engine/Source/Programs/AutomationTool/Scripts/CopyBuildToStagingDirectory.Automation.cs
--- a/Engine/Source/Programs/AutomationTool/Scripts/CopyBuildToStagingDirectory.Automation.cs
+++ b/Engine/Source/Programs/AutomationTool/Scripts/CopyBuildToStagingDirectory.Automation.cs
@@ -2328,12 +2328,12 @@ namespace AutomationScripts
 				{
 					if (PakRules.OverridePaks != null)
 					{
-						LogWarning("Error in PakFileRules {Arg0}, set to exclude but also sets override!", PakRules.Name);
+						Logger.LogWarning("Error in PakFileRules {Arg0}, set to exclude but also sets override!", PakRules.Name);
 						continue;
 					}
 					if (PakRules.bOnDemand)
 					{
-						LogWarning("Error in PakFileRules {Arg0}, set to exclude but also sets to ondemand!", PakRules.Name);
+						Logger.LogWarning("Error in PakFileRules {Arg0}, set to exclude but also sets to ondemand!", PakRules.Name);
 						continue;
 					}
 				}
@@ -2403,7 +2403,7 @@ namespace AutomationScripts
 							}
 							else
 							{
-								LogWarning("Undefined PAK override '{0}' in PAK rule '{1}'", OverrideChunkName, Rule.Name);
+								Logger.LogWarning("Undefined PAK override '{Arg0}' in PAK rule '{Arg1}'", OverrideChunkName, Rule.Name);
 							}
 						}
 
@@ -2414,7 +2414,7 @@ namespace AutomationScripts
 						}
 						else
 						{
-							LogWarning("Undefined chunk name '{0}' in PAK rule '{1}'", ChunkName, Rule.Name);
+							Logger.LogWarning("Undefined chunk name '{Arg0}' in PAK rule '{Arg1}'", ChunkName, Rule.Name);
 						}
 						
 						return;
@@ -4062,13 +4062,13 @@ namespace AutomationScripts
 							continue;
 						}
 
-						Log.TraceInformation("Creating content-on-demand chunk definition '{0}' for PAK rule '{1}'", ChunkName, Rule.Name);
+						Logger.LogInformation("Creating content-on-demand chunk definition '{ChunkName}' for PAK rule '{Arg1}'", ChunkName, Rule.Name);
 
 						bool bIsEncryptionKeyValid = true;
 						if (string.IsNullOrEmpty(Rule.EncryptionKeyGuid))
 						{
 							bIsEncryptionKeyValid = false;
-							Log.TraceWarning("No encryption key specified for content-on-demand PAK rule '{0}'", Rule.Name);
+							Logger.LogWarning("No encryption key specified for content-on-demand PAK rule '{Arg0}'", Rule.Name);
 						}
 						else
 						{
@@ -4076,13 +4076,13 @@ namespace AutomationScripts
 							if (!Guid.TryParse(Rule.EncryptionKeyGuid, out KeyGuid))
 							{
 								bIsEncryptionKeyValid = false;
-								Log.TraceWarning("Invalid encryption key GUID specified for content-on-demand PAK rule '{0}'", Rule.Name);
+								Logger.LogWarning("Invalid encryption key GUID specified for content-on-demand PAK rule '{Arg0}'", Rule.Name);
 							}
 
 							if (!PakCryptoSettings.ContainsEncryptionKey(Rule.EncryptionKeyGuid))
 							{
 								bIsEncryptionKeyValid = false;
-								Log.TraceWarning("No matching encryption key found for PAK rule '{0}', GUID '{1}'", Rule.Name, Rule.EncryptionKeyGuid);
+								Logger.LogWarning("No matching encryption key found for PAK rule '{Arg0}', GUID '{Arg1}'", Rule.Name, Rule.EncryptionKeyGuid);
 							}
 						}
 
