- Fixed logging warning

diff --git a/Engine/Source/Programs/AutomationTool/Scripts/GenerateDSYM.Automation.cs b/Engine/Source/Programs/AutomationTool/Scripts/GenerateDSYM.Automation.cs
--- a/Engine/Source/Programs/AutomationTool/Scripts/GenerateDSYM.Automation.cs
+++ b/Engine/Source/Programs/AutomationTool/Scripts/GenerateDSYM.Automation.cs
@@ -96,7 +96,7 @@ public class GenerateDSYM : BuildCommand
 					dSYM = Path.Combine(Path.GetDirectoryName(Binary.Substring(0, Binary.LastIndexOf(".app"))), Path.GetFileName(dSYM));
 				}
 
-				Log.TraceInformation("  Generating for {0} -> {1}", Binary, dSYM);
+				Logger.LogInformation("  Generating for {Binary} -> {dSYM}", Binary, dSYM);
 
 				string Command;
 				if (bSaveFlat)
