Add Overflow 3 and 4 for Android OBB files
OverflowOBBFileLimit project setting to set number allowed (up to 4, default 2)

#android

diff --git a/Engine/Source/Programs/AutomationTool/Android/AndroidPlatform.Automation.cs b/Engine/Source/Programs/AutomationTool/Android/AndroidPlatform.Automation.cs
--- a/Engine/Source/Programs/AutomationTool/Android/AndroidPlatform.Automation.cs
+++ b/Engine/Source/Programs/AutomationTool/Android/AndroidPlatform.Automation.cs
@@ -960,12 +960,18 @@ public class AndroidPlatform : Platform
 		return bAllowPatchOBBFile;
 	}
 
-	private bool AllowOverflowOBBFiles(DeploymentContext SC)
+	private int AllowOverflowOBBFiles(DeploymentContext SC)
 	{
+		int FileLimit = 0;
 		ConfigHierarchy Ini = ConfigCache.ReadHierarchy(ConfigHierarchyType.Engine, DirectoryReference.FromFile(SC.RawProjectPath), SC.StageTargetPlatform.PlatformType);
 		bool bAllowOverflowOBBFiles = false;
 		Ini.GetBool("/Script/AndroidRuntimeSettings.AndroidRuntimeSettings", "bAllowOverflowOBBFiles", out bAllowOverflowOBBFiles);
-		return bAllowOverflowOBBFiles;
+		if (bAllowOverflowOBBFiles)
+		{
+			FileLimit = 2;
+			Ini.GetInt32("/Script/AndroidRuntimeSettings.AndroidRuntimeSettings", "OverflowOBBFileLimit", out FileLimit);
+		}
+		return FileLimit;
 	}
 
 	private bool CreateOBBFile(DeploymentContext SC, string OutputFilename, List<FileReference> FilesForObb)
@@ -1079,6 +1085,8 @@ public class AndroidPlatform : Platform
 		string LocalPatchName = SC.StageDirectory.FullName + ".patch.obb";
 		string LocalOverflow1Name = SC.StageDirectory.FullName + ".overflow1.obb";
 		string LocalOverflow2Name = SC.StageDirectory.FullName + ".overflow2.obb";
+		string LocalOverflow3Name = SC.StageDirectory.FullName + ".overflow3.obb";
+		string LocalOverflow4Name = SC.StageDirectory.FullName + ".overflow4.obb";
 
 		FileFilter ObbFileFilter = new FileFilter(FileFilterType.Include);
 		ConfigHierarchy EngineIni = ConfigCache.ReadHierarchy(ConfigHierarchyType.Engine, DirectoryReference.FromFile(Params.RawProjectPath), UnrealTargetPlatform.Android);
@@ -1155,14 +1163,28 @@ public class AndroidPlatform : Platform
 				File.Delete(LocalOverflow2Name);
 			}
 
+			// Always delete the target overflow3 OBB file if it exists
+			if (File.Exists(LocalOverflow3Name))
+			{
+				File.Delete(LocalOverflow3Name);
+			}
+
+			// Always delete the target overflow4 OBB file if it exists
+			if (File.Exists(LocalOverflow4Name))
+			{
+				File.Delete(LocalOverflow4Name);
+			}
+
 			List<FileReference> FilesToObb = FilesForObb;
 			List<FileReference> FilesToPatch = new List<FileReference>();
 			List<FileReference> FilesToOverflow1 = new List<FileReference>();
 			List<FileReference> FilesToOverflow2 = new List<FileReference>();
+			List<FileReference> FilesToOverflow3 = new List<FileReference>();
+			List<FileReference> FilesToOverflow4 = new List<FileReference>();
 
 			if (AllowPatchOBBFile(SC))
 			{
-				bool bAllowOverflowOBBs = AllowOverflowOBBFiles(SC);
+				int AllowOverflowOBBLimit = AllowOverflowOBBFiles(SC);
 
 				FilesToObb = new List<FileReference>();
 
@@ -1173,6 +1195,8 @@ public class AndroidPlatform : Platform
 				Int64 PatchObbSize = MinimumObbSize;
 				Int64 Overflow1ObbSize = MinimumObbSize;
 				Int64 Overflow2ObbSize = MinimumObbSize;
+				Int64 Overflow3ObbSize = MinimumObbSize;
+				Int64 Overflow4ObbSize = MinimumObbSize;
 				foreach (FileReference FileRef in FilesForObb)
 				{
 					FileInfo LocalFileInfo = new FileInfo(FileRef.FullName);
@@ -1193,24 +1217,25 @@ public class AndroidPlatform : Platform
 						FilesToPatch.Add(FileRef);
 						PatchObbSize += FileRequirements;
 					}
-					else if (bAllowOverflowOBBs)
+					else if (AllowOverflowOBBLimit >= 1 && Overflow1ObbSize + FileRequirements < OBBSizeAllowed)
 					{
-						if (Overflow1ObbSize + FileRequirements < OBBSizeAllowed)
-						{
-							FilesToOverflow1.Add(FileRef);
-							Overflow1ObbSize += FileRequirements;
-						}
-						else if (Overflow2ObbSize + FileRequirements < OBBSizeAllowed)
-						{
-							FilesToOverflow2.Add(FileRef);
-							Overflow2ObbSize += FileRequirements;
-						}
-						else
-						{
-							// no room in either file
-							Logger.LogInformation("{Text}", "Failed to build OBB: " + LocalObbName);
-							throw new AutomationException(ExitCode.Error_AndroidOBBError, "Stage Failed. Could not build OBB {0}. The file may be too big to fit in an OBB ({1} limit)", LocalObbName, LimitString);
-						}
+						FilesToOverflow1.Add(FileRef);
+						Overflow1ObbSize += FileRequirements;
+					}
+					else if (AllowOverflowOBBLimit >= 2 && Overflow2ObbSize + FileRequirements < OBBSizeAllowed)
+					{
+						FilesToOverflow2.Add(FileRef);
+						Overflow2ObbSize += FileRequirements;
+					}
+					else if (AllowOverflowOBBLimit >= 3 && Overflow3ObbSize + FileRequirements < OBBSizeAllowed)
+					{
+						FilesToOverflow3.Add(FileRef);
+						Overflow3ObbSize += FileRequirements;
+					}
+					else if (AllowOverflowOBBLimit >= 4 && Overflow4ObbSize + FileRequirements < OBBSizeAllowed)
+					{
+						FilesToOverflow4.Add(FileRef);
+						Overflow4ObbSize += FileRequirements;
 					}
 					else
 					{
@@ -1244,7 +1269,7 @@ public class AndroidPlatform : Platform
 				if (!CreateOBBFile(SC, LocalOverflow1Name, FilesToOverflow1))
 				{
 					Logger.LogInformation("{Text}", "Failed to build OBB: " + LocalOverflow1Name);
-					throw new AutomationException(ExitCode.Error_AndroidOBBError, "Stage Failed. Could not build OBB {0}. The file may be too big to fit in an OBB ({1} limit)", LocalPatchName, LimitString);
+					throw new AutomationException(ExitCode.Error_AndroidOBBError, "Stage Failed. Could not build OBB {0}. The file may be too big to fit in an OBB ({1} limit)", LocalOverflow1Name, LimitString);
 				}
 			}
 
@@ -1254,7 +1279,27 @@ public class AndroidPlatform : Platform
 				if (!CreateOBBFile(SC, LocalOverflow2Name, FilesToOverflow2))
 				{
 					Logger.LogInformation("{Text}", "Failed to build OBB: " + LocalOverflow2Name);
-					throw new AutomationException(ExitCode.Error_AndroidOBBError, "Stage Failed. Could not build OBB {0}. The file may be too big to fit in an OBB ({1} limit)", LocalPatchName, LimitString);
+					throw new AutomationException(ExitCode.Error_AndroidOBBError, "Stage Failed. Could not build OBB {0}. The file may be too big to fit in an OBB ({1} limit)", LocalOverflow2Name, LimitString);
+				}
+			}
+
+			// Now create the overflow3 OBB as a ZIP archive if required.
+			if (FilesToOverflow3.Count() > 0)
+			{
+				if (!CreateOBBFile(SC, LocalOverflow3Name, FilesToOverflow3))
+				{
+					Logger.LogInformation("{Text}", "Failed to build OBB: " + LocalOverflow3Name);
+					throw new AutomationException(ExitCode.Error_AndroidOBBError, "Stage Failed. Could not build OBB {0}. The file may be too big to fit in an OBB ({1} limit)", LocalOverflow3Name, LimitString);
+				}
+			}
+
+			// Now create the overflow4 OBB as a ZIP archive if required.
+			if (FilesToOverflow4.Count() > 0)
+			{
+				if (!CreateOBBFile(SC, LocalOverflow4Name, FilesToOverflow4))
+				{
+					Logger.LogInformation("{Text}", "Failed to build OBB: " + LocalOverflow4Name);
+					throw new AutomationException(ExitCode.Error_AndroidOBBError, "Stage Failed. Could not build OBB {0}. The file may be too big to fit in an OBB ({1} limit)", LocalOverflow4Name, LimitString);
 				}
 			}
 		}
@@ -1338,6 +1383,10 @@ public class AndroidPlatform : Platform
 			string Overflow1Name = "";
 			string DeviceOverflow2Name = "";
 			string Overflow2Name = "";
+			string DeviceOverflow3Name = "";
+			string Overflow3Name = "";
+			string DeviceOverflow4Name = "";
+			string Overflow4Name = "";
 			if (!bPackageDataInsideApk)
 			{
 				DeviceObbName = GetDeviceObbName(ApkName, SC);
@@ -1376,6 +1425,26 @@ public class AndroidPlatform : Platform
 					// apply store version to OBB to make it unique for PlayStore upload
 					UpdateObbStoreVersion(Overflow2Name);
 				}
+
+				if (File.Exists(LocalOverflow3Name))
+				{
+					DeviceOverflow3Name = GetDeviceOverflowName(ApkName, SC, 3);
+					Overflow3Name = GetFinalOverflowName(ApkName, SC, 3);
+					CopyFile(LocalOverflow3Name, Overflow3Name);
+
+					// apply store version to OBB to make it unique for PlayStore upload
+					UpdateObbStoreVersion(Overflow3Name);
+				}
+
+				if (File.Exists(LocalOverflow4Name))
+				{
+					DeviceOverflow4Name = GetDeviceOverflowName(ApkName, SC, 4);
+					Overflow4Name = GetFinalOverflowName(ApkName, SC, 4);
+					CopyFile(LocalOverflow4Name, Overflow4Name);
+
+					// apply store version to OBB to make it unique for PlayStore upload
+					UpdateObbStoreVersion(Overflow4Name);
+				}
 			}
 
 			// check for optional universal apk
@@ -1415,7 +1484,8 @@ public class AndroidPlatform : Platform
 				string PackageName = GetPackageInfo(ApkName, SC, false);
 				string BatchName = GetFinalBatchName(ApkName, SC, false, EBatchType.Install, Target);
 				string[] BatchLines = GenerateInstallBatchFile(bPackageDataInsideApk, PackageName, ApkName, Params, ObbName, DeviceObbName, false, PatchName, DevicePatchName, false, 
-					Overflow1Name, DeviceOverflow1Name, false, Overflow2Name, DeviceOverflow2Name, false, bIsPC, Params.Distribution, TargetSDKVersion > 22, bDisablePerfHarden, bUseAFS, bUseAFSProject, AFSToken, Target);
+					Overflow1Name, DeviceOverflow1Name, false, Overflow2Name, DeviceOverflow2Name, false, Overflow3Name, DeviceOverflow3Name, false, Overflow4Name, DeviceOverflow4Name, false, 
+					bIsPC, Params.Distribution, TargetSDKVersion > 22, bDisablePerfHarden, bUseAFS, bUseAFSProject, AFSToken, Target);
 				if (bHaveAPK)
 				{
 					// make a batch file that can be used to install the .apk and .obb files
@@ -1435,7 +1505,8 @@ public class AndroidPlatform : Platform
 					UniversalBatchName = GetFinalBatchName(UniversalApkName, SC, false, EBatchType.Install, Target);
 					// make a batch file that can be used to install the .apk
 					string[] UniversalBatchLines = GenerateInstallBatchFile(bPackageDataInsideApk, PackageName, UniversalApkName, Params, ObbName, DeviceObbName, false, PatchName, DevicePatchName, false,
-						Overflow1Name, DeviceOverflow1Name, false, Overflow2Name, DeviceOverflow2Name, false, bIsPC, Params.Distribution, TargetSDKVersion > 22, bDisablePerfHarden, bUseAFS, bUseAFSProject, AFSToken, Target);
+						Overflow1Name, DeviceOverflow1Name, false, Overflow2Name, DeviceOverflow2Name, false, Overflow3Name, DeviceOverflow3Name, false, Overflow4Name, DeviceOverflow4Name, false, 
+						bIsPC, Params.Distribution, TargetSDKVersion > 22, bDisablePerfHarden, bUseAFS, bUseAFSProject, AFSToken, Target);
 					File.WriteAllText(UniversalBatchName, string.Join(LineEnding, UniversalBatchLines) + LineEnding);
 				}
 
@@ -1517,6 +1588,7 @@ public class AndroidPlatform : Platform
 
 	private string[] GenerateInstallBatchFile(bool bPackageDataInsideApk, string PackageName, string ApkName, ProjectParams Params, string ObbName, string DeviceObbName, bool bNoObbInstall,
 		string PatchName, string DevicePatchName, bool bNoPatchInstall, string Overflow1Name, string DeviceOverflow1Name, bool bNoOverflow1Install, string Overflow2Name, string DeviceOverflow2Name, bool bNoOverflow2Install,
+		string Overflow3Name, string DeviceOverflow3Name, bool bNoOverflow3Install, string Overflow4Name, string DeviceOverflow4Name, bool bNoOverflow4Install,
 		bool bIsPC, bool bIsDistribution, bool bRequireRuntimeStoragePermission, bool bDisablePerfHarden, bool bUseAFS, bool bUseAFSProject, string AFSToken, UnrealTargetPlatform Target)
     {
         string[] BatchLines = null;
@@ -1534,6 +1606,8 @@ public class AndroidPlatform : Platform
 		bool bHavePatch = (PatchName != "");
 		bool bHaveOverflow1 = (Overflow1Name != "");
 		bool bHaveOverflow2 = (Overflow2Name != "");
+		bool bHaveOverflow3 = (Overflow3Name != "");
+		bool bHaveOverflow4 = (Overflow4Name != "");
 
 		string AFSExecutable = GetAFSExecutable(Target);
 		string AFSCommonArg = "-p " + PackageName;
@@ -1553,6 +1627,8 @@ public class AndroidPlatform : Platform
 			string PatchInstallCommand = bNoPatchInstall ? "\t$ADB $DEVICE shell 'rm -r $EXTERNAL_STORAGE/" + DevicePatchName + "'" : "\t$ADB $DEVICE push " + Path.GetFileName(PatchName) + (bIsDistribution ? " " + TargetAndroidTemp : " $STORAGE/") + DevicePatchName;
 			string Overflow1InstallCommand = bNoOverflow1Install ? "\t$ADB $DEVICE shell 'rm -r $EXTERNAL_STORAGE/" + DeviceOverflow1Name + "'" : "\t$ADB $DEVICE push " + Path.GetFileName(Overflow1Name) + (bIsDistribution ? " " + TargetAndroidTemp : " $STORAGE/") + DeviceOverflow1Name;
 			string Overflow2InstallCommand = bNoOverflow2Install ? "\t$ADB $DEVICE shell 'rm -r $EXTERNAL_STORAGE/" + DeviceOverflow2Name + "'" : "\t$ADB $DEVICE push " + Path.GetFileName(Overflow2Name) + (bIsDistribution ? " " + TargetAndroidTemp : " $STORAGE/") + DeviceOverflow2Name;
+			string Overflow3InstallCommand = bNoOverflow3Install ? "\t$ADB $DEVICE shell 'rm -r $EXTERNAL_STORAGE/" + DeviceOverflow3Name + "'" : "\t$ADB $DEVICE push " + Path.GetFileName(Overflow3Name) + (bIsDistribution ? " " + TargetAndroidTemp : " $STORAGE/") + DeviceOverflow3Name;
+			string Overflow4InstallCommand = bNoOverflow4Install ? "\t$ADB $DEVICE shell 'rm -r $EXTERNAL_STORAGE/" + DeviceOverflow4Name + "'" : "\t$ADB $DEVICE push " + Path.GetFileName(Overflow4Name) + (bIsDistribution ? " " + TargetAndroidTemp : " $STORAGE/") + DeviceOverflow4Name;
 
 			if (bUseAFS)
 			{
@@ -1571,6 +1647,8 @@ public class AndroidPlatform : Platform
 				PatchInstallCommand = bNoPatchInstall ? AFSCommand + " deletefile '^patchobb'" : AFSCommand + " push " + Path.GetFileName(PatchName) + " '^patchobb'";
 				Overflow1InstallCommand = bNoOverflow1Install ? AFSCommand + " deletefile '^overflow1obb'" : AFSCommand + " push " + Path.GetFileName(Overflow1Name) + " '^overflow1obb'";
 				Overflow2InstallCommand = bNoOverflow2Install ? AFSCommand + " deletefile '^overflow2obb'" : AFSCommand + " push " + Path.GetFileName(Overflow2Name) + " '^overflow2obb'";
+				Overflow3InstallCommand = bNoOverflow3Install ? AFSCommand + " deletefile '^overflow3obb'" : AFSCommand + " push " + Path.GetFileName(Overflow3Name) + " '^overflow3obb'";
+				Overflow4InstallCommand = bNoOverflow4Install ? AFSCommand + " deletefile '^overflow4obb'" : AFSCommand + " push " + Path.GetFileName(Overflow4Name) + " '^overflow4obb'";
 			}
 
 			Logger.LogInformation("Writing shell script for install with {Arg0}", bPackageDataInsideApk ? "data in APK" : "separate obb");
@@ -1610,6 +1688,8 @@ public class AndroidPlatform : Platform
 						!bHavePatch ? "" : (bPackageDataInsideApk ? "" : PatchInstallCommand),
 						!bHaveOverflow1 ? "" : (bPackageDataInsideApk ? "" : Overflow1InstallCommand),
 						!bHaveOverflow2 ? "" : (bPackageDataInsideApk ? "" : Overflow2InstallCommand),
+						!bHaveOverflow3 ? "" : (bPackageDataInsideApk ? "" : Overflow3InstallCommand),
+						!bHaveOverflow4 ? "" : (bPackageDataInsideApk ? "" : Overflow4InstallCommand),
 						bDontMoveOBB ? "" : "\t\t$ADB $DEVICE shell mkdir $STORAGE/Android/" + TargetAndroidLocation + PackageName, // don't check for error since installing may create the obb directory
 						bDontMoveOBB ? "" : "\t\t$ADB $DEVICE shell mv " + TargetAndroidTemp + TargetAndroidLocation + PackageName + " $STORAGE/Android/" + TargetAndroidLocation,
 						bDontMoveOBB ? "" : "\t\t$ADB $DEVICE shell rm -r " + TargetAndroidTemp + TargetAndroidLocation,
@@ -1638,6 +1718,8 @@ public class AndroidPlatform : Platform
 			string PatchInstallCommand = bNoPatchInstall ? "%ADB% %DEVICE% shell rm -r %STORAGE%/" + DevicePatchName : "%ADB% %DEVICE% push " + Path.GetFileName(PatchName) + (bIsDistribution ? " " + TargetAndroidTemp : " %STORAGE%/") + DevicePatchName;
 			string Overflow1InstallCommand = bNoOverflow1Install ? "%ADB% %DEVICE% shell rm -r %STORAGE%/" + DeviceOverflow1Name : "%ADB% %DEVICE% push " + Path.GetFileName(Overflow1Name) + (bIsDistribution ? " " + TargetAndroidTemp : " %STORAGE%/") + DeviceOverflow1Name;
 			string Overflow2InstallCommand = bNoOverflow2Install ? "%ADB% %DEVICE% shell rm -r %STORAGE%/" + DeviceOverflow2Name : "%ADB% %DEVICE% push " + Path.GetFileName(Overflow2Name) + (bIsDistribution ? " " + TargetAndroidTemp : " %STORAGE%/") + DeviceOverflow2Name;
+			string Overflow3InstallCommand = bNoOverflow3Install ? "%ADB% %DEVICE% shell rm -r %STORAGE%/" + DeviceOverflow3Name : "%ADB% %DEVICE% push " + Path.GetFileName(Overflow3Name) + (bIsDistribution ? " " + TargetAndroidTemp : " %STORAGE%/") + DeviceOverflow3Name;
+			string Overflow4InstallCommand = bNoOverflow4Install ? "%ADB% %DEVICE% shell rm -r %STORAGE%/" + DeviceOverflow4Name : "%ADB% %DEVICE% push " + Path.GetFileName(Overflow4Name) + (bIsDistribution ? " " + TargetAndroidTemp : " %STORAGE%/") + DeviceOverflow4Name;
 
 			if (bUseAFS)
 			{
@@ -1656,6 +1738,8 @@ public class AndroidPlatform : Platform
 				PatchInstallCommand = bNoPatchInstall ? AFSCommand + " deletefile \"^patchobb\"" : AFSCommand + " push " + Path.GetFileName(PatchName) + " \"^patchobb\"";
 				Overflow1InstallCommand = bNoOverflow1Install ? AFSCommand + " deletefile \"^overflow1obb\"" : AFSCommand + " push " + Path.GetFileName(Overflow1Name) + " \"^overflow1obb\"";
 				Overflow2InstallCommand = bNoOverflow2Install ? AFSCommand + " deletefile \"^overflow2obb\"" : AFSCommand + " push " + Path.GetFileName(Overflow2Name) + " \"^overflow2obb\"";
+				Overflow3InstallCommand = bNoOverflow3Install ? AFSCommand + " deletefile \"^overflow3obb\"" : AFSCommand + " push " + Path.GetFileName(Overflow3Name) + " \"^overflow3obb\"";
+				Overflow4InstallCommand = bNoOverflow4Install ? AFSCommand + " deletefile \"^overflow4obb\"" : AFSCommand + " push " + Path.GetFileName(Overflow4Name) + " \"^overflow4obb\"";
 			}
 
 			Logger.LogInformation("Writing bat for install with {Arg0}", bPackageDataInsideApk ? "data in APK" : "separate OBB");
@@ -1692,6 +1776,10 @@ public class AndroidPlatform : Platform
 						!bHaveOverflow1 ? "" : (bPackageDataInsideApk ? "" : "if \"%ERRORLEVEL%\" NEQ \"0\" goto Error"),
 						!bHaveOverflow2 ? "" : (bPackageDataInsideApk ? "" : Overflow2InstallCommand),
 						!bHaveOverflow2 ? "" : (bPackageDataInsideApk ? "" : "if \"%ERRORLEVEL%\" NEQ \"0\" goto Error"),
+						!bHaveOverflow3 ? "" : (bPackageDataInsideApk ? "" : Overflow3InstallCommand),
+						!bHaveOverflow3 ? "" : (bPackageDataInsideApk ? "" : "if \"%ERRORLEVEL%\" NEQ \"0\" goto Error"),
+						!bHaveOverflow4 ? "" : (bPackageDataInsideApk ? "" : Overflow4InstallCommand),
+						!bHaveOverflow4 ? "" : (bPackageDataInsideApk ? "" : "if \"%ERRORLEVEL%\" NEQ \"0\" goto Error"),
 						bDontMoveOBB ? "" : "%ADB% %DEVICE% shell mkdir %STORAGE%/Android/" + TargetAndroidLocation + PackageName, // don't check for error since installing may create the obb directory
 						bDontMoveOBB ? "" : "%ADB% %DEVICE% shell mv " + TargetAndroidTemp + TargetAndroidLocation + PackageName + " %STORAGE%/Android/" + TargetAndroidLocation,
 						bDontMoveOBB ? "" : "if \"%ERRORLEVEL%\" NEQ \"0\" goto Error",
@@ -1851,6 +1939,8 @@ public class AndroidPlatform : Platform
 			string PatchName = GetFinalPatchName(ApkName, SC);
 			string Overflow1Name = GetFinalOverflowName(ApkName, SC, 1);
 			string Overflow2Name = GetFinalOverflowName(ApkName, SC, 2);
+			string Overflow3Name = GetFinalOverflowName(ApkName, SC, 3);
+			string Overflow4Name = GetFinalOverflowName(ApkName, SC, 4);
 			bool bBuildWithHiddenSymbolVisibility = BuildWithHiddenSymbolVisibility(SC);
 			bool bSaveSymbols = GetSaveSymbols(SC);
 
@@ -1969,6 +2059,14 @@ public class AndroidPlatform : Platform
 					{
 						SC.ArchiveFiles(Path.GetDirectoryName(Overflow2Name), Path.GetFileName(Overflow2Name));
 					}
+					if (FileExists(Overflow3Name))
+					{
+						SC.ArchiveFiles(Path.GetDirectoryName(Overflow3Name), Path.GetFileName(Overflow3Name));
+					}
+					if (FileExists(Overflow4Name))
+					{
+						SC.ArchiveFiles(Path.GetDirectoryName(Overflow4Name), Path.GetFileName(Overflow4Name));
+					}
 				}
 			}
 
@@ -2393,6 +2491,8 @@ public class AndroidPlatform : Platform
 			string DevicePatchName = StorageLocation + "/" + GetDevicePatchName(ApkName, SC);
 			string DeviceOverflow1Name = StorageLocation + "/" + GetDeviceOverflowName(ApkName, SC, 1);
 			string DeviceOverflow2Name = StorageLocation + "/" + GetDeviceOverflowName(ApkName, SC, 2);
+			string DeviceOverflow3Name = StorageLocation + "/" + GetDeviceOverflowName(ApkName, SC, 3);
+			string DeviceOverflow4Name = StorageLocation + "/" + GetDeviceOverflowName(ApkName, SC, 4);
 			string RemoteDir = client.Query("^project", true);
 			if (RemoteDir == null)
 			{
@@ -2871,6 +2971,8 @@ public class AndroidPlatform : Platform
 				client.FileDelete(DevicePatchName);
 				client.FileDelete(DeviceOverflow1Name);
 				client.FileDelete(DeviceOverflow2Name);
+				client.FileDelete(DeviceOverflow3Name);
+				client.FileDelete(DeviceOverflow4Name);
 
 				// We now have a minimal set of file & dir entries we need
 				// to deploy. Files we deploy will get individually copied
@@ -2909,6 +3011,20 @@ public class AndroidPlatform : Platform
 				{
 					client.FileWrite(Overflow2Path, DeviceOverflow2Name);
 				}
+
+				// deploy the overflow3 if there is one
+				string Overflow3Path = Path.Combine(SC.StageDirectory.FullName, GetFinalOverflowName(ApkName, SC, 3));
+				if (File.Exists(Overflow3Path))
+				{
+					client.FileWrite(Overflow3Path, DeviceOverflow3Name);
+				}
+
+				// deploy the overflow4 if there is one
+				string Overflow4Path = Path.Combine(SC.StageDirectory.FullName, GetFinalOverflowName(ApkName, SC, 4));
+				if (File.Exists(Overflow4Path))
+				{
+					client.FileWrite(Overflow4Path, DeviceOverflow4Name);
+				}
 			}
 			else
 			{
@@ -2963,6 +3079,8 @@ public class AndroidPlatform : Platform
 			string DevicePatchName = StorageLocation + "/" + GetDevicePatchName(ApkName, SC);
 			string DeviceOverflow1Name = StorageLocation + "/" + GetDeviceOverflowName(ApkName, SC, 1);
 			string DeviceOverflow2Name = StorageLocation + "/" + GetDeviceOverflowName(ApkName, SC, 2);
+			string DeviceOverflow3Name = StorageLocation + "/" + GetDeviceOverflowName(ApkName, SC, 3);
+			string DeviceOverflow4Name = StorageLocation + "/" + GetDeviceOverflowName(ApkName, SC, 4);
 			string RemoteDir = StorageLocation + "/UnrealGame/" + Params.ShortProjectName;
 
 			if (bDisablePerfHarden)
@@ -3359,6 +3477,8 @@ public class AndroidPlatform : Platform
 				RunAdbCommand(Params, DeviceName, "shell rm " + DevicePatchName);
 				RunAdbCommand(Params, DeviceName, "shell rm " + DeviceOverflow1Name);
 				RunAdbCommand(Params, DeviceName, "shell rm " + DeviceOverflow2Name);
+				RunAdbCommand(Params, DeviceName, "shell rm " + DeviceOverflow3Name);
+				RunAdbCommand(Params, DeviceName, "shell rm " + DeviceOverflow4Name);
 			}
 			else if (SC.Archive)
             {
@@ -3405,6 +3525,28 @@ public class AndroidPlatform : Platform
 					string Commandline = string.Format("{0} \"{1}\" \"{2}\"", BaseCommandline, Overflow2Path, DeviceOverflow2Name);
 					RunAdbCommand(Params, DeviceName, Commandline);
 				}
+
+				// deploy the overflow3 if there is one
+				string Overflow3Path = Path.Combine(SC.StageDirectory.FullName, GetFinalOverflowName(ApkName, SC, 3));
+				if (File.Exists(Overflow3Path))
+				{
+					// cache some strings
+					string BaseCommandline = "push";
+
+					string Commandline = string.Format("{0} \"{1}\" \"{2}\"", BaseCommandline, Overflow3Path, DeviceOverflow3Name);
+					RunAdbCommand(Params, DeviceName, Commandline);
+				}
+
+				// deploy the overflow4 if there is one
+				string Overflow4Path = Path.Combine(SC.StageDirectory.FullName, GetFinalOverflowName(ApkName, SC, 4));
+				if (File.Exists(Overflow4Path))
+				{
+					// cache some strings
+					string BaseCommandline = "push";
+
+					string Commandline = string.Format("{0} \"{1}\" \"{2}\"", BaseCommandline, Overflow4Path, DeviceOverflow4Name);
+					RunAdbCommand(Params, DeviceName, Commandline);
+				}
 			}
 			else
             {
diff --git a/Engine/Source/Runtime/Core/Private/Android/AndroidPlatformFile.cpp b/Engine/Source/Runtime/Core/Private/Android/AndroidPlatformFile.cpp
--- a/Engine/Source/Runtime/Core/Private/Android/AndroidPlatformFile.cpp
+++ b/Engine/Source/Runtime/Core/Private/Android/AndroidPlatformFile.cpp
@@ -98,6 +98,10 @@ FString GOBBPatchFilePath;
 FString GOBBOverflow1FilePath;
 // Obb Overflow2 filepath
 FString GOBBOverflow2FilePath;
+// Obb Overflow3 filepath
+FString GOBBOverflow3FilePath;
+// Obb Overflow2 filepath
+FString GOBBOverflow4FilePath;
 // Internal File Direcory Path (for application) - setup during load
 FString GInternalFilePath;
 // External File Direcory Path (for application) - setup during load
@@ -1220,6 +1224,8 @@ public:
 			{
 				FString Overflow1OBBName = FString::Printf(TEXT("overflow1.%d.%s.obb"), GAndroidPackageVersion, *GPackageName);
 				FString Overflow2OBBName = FString::Printf(TEXT("overflow2.%d.%s.obb"), GAndroidPackageVersion, *GPackageName);
+				FString Overflow3OBBName = FString::Printf(TEXT("overflow3.%d.%s.obb"), GAndroidPackageVersion, *GPackageName);
+				FString Overflow4OBBName = FString::Printf(TEXT("overflow4.%d.%s.obb"), GAndroidPackageVersion, *GPackageName);
 
 				if (!GOBBOverflow1FilePath.IsEmpty() && FileExists(*GOBBOverflow1FilePath, true))
 				{
@@ -1256,6 +1262,43 @@ public:
 					MountOBB(*GOBBOverflow2FilePath);
 					FPlatformMisc::LowLevelOutputDebugStringf(TEXT("Mounted overflow2 OBB: %s"), *GOBBOverflow2FilePath);
 				}
+
+				if (!GOBBOverflow3FilePath.IsEmpty() && FileExists(*GOBBOverflow3FilePath, true))
+				{
+					MountOBB(*GOBBOverflow3FilePath);
+					FPlatformMisc::LowLevelOutputDebugStringf(TEXT("Mounted overflow3 OBB: %s"), *GOBBOverflow3FilePath);
+				}
+				else if (FileExists(*(OBBDir1 / Overflow3OBBName), true))
+				{
+					GOBBOverflow3FilePath = OBBDir1 / Overflow3OBBName;
+					MountOBB(*GOBBOverflow3FilePath);
+					FPlatformMisc::LowLevelOutputDebugStringf(TEXT("Mounted overflow3 OBB: %s"), *GOBBOverflow3FilePath);
+				}
+				else if (FileExists(*(OBBDir2 / Overflow3OBBName), true))
+				{
+					GOBBOverflow3FilePath = OBBDir2 / Overflow3OBBName;
+					MountOBB(*GOBBOverflow3FilePath);
+					FPlatformMisc::LowLevelOutputDebugStringf(TEXT("Mounted overflow2 OBB: %s"), *GOBBOverflow3FilePath);
+				}
+
+				if (!GOBBOverflow4FilePath.IsEmpty() && FileExists(*GOBBOverflow4FilePath, true))
+				{
+					MountOBB(*GOBBOverflow4FilePath);
+					FPlatformMisc::LowLevelOutputDebugStringf(TEXT("Mounted overflow4 OBB: %s"), *GOBBOverflow4FilePath);
+				}
+				else if (FileExists(*(OBBDir1 / Overflow4OBBName), true))
+				{
+					GOBBOverflow4FilePath = OBBDir1 / Overflow4OBBName;
+					MountOBB(*GOBBOverflow4FilePath);
+					FPlatformMisc::LowLevelOutputDebugStringf(TEXT("Mounted overflow4 OBB: %s"), *GOBBOverflow4FilePath);
+				}
+				else if (FileExists(*(OBBDir2 / Overflow4OBBName), true))
+				{
+					GOBBOverflow4FilePath = OBBDir2 / Overflow4OBBName;
+					MountOBB(*GOBBOverflow4FilePath);
+					FPlatformMisc::LowLevelOutputDebugStringf(TEXT("Mounted overflow4 OBB: %s"), *GOBBOverflow4FilePath);
+				}
+
 			}
 		}
 
