Make the AutomationTool emit an error code when it terminates on an exception. This will result in failing the operation instead of reporting a false success. An example is when aborting for multiple instances in RunSingleInstanceThread.

diff --git a/Engine/Source/Programs/AutomationTool/Program.cs b/Engine/Source/Programs/AutomationTool/Program.cs
--- a/Engine/Source/Programs/AutomationTool/Program.cs
+++ b/Engine/Source/Programs/AutomationTool/Program.cs
@@ -414,7 +414,7 @@ namespace AutomationToolDriver
 			}
 
 			// Enter the main program section
-			ExitCode ReturnCode = ExitCode.Success;
+			ExitCode ReturnCode = ExitCode.Error_Unknown;
 			try
 			{
 				// Set the working directory to the Unreal root directory
