BuildGraph: Add support for explicitly labeling artifacts from a build using the <Artifact> element.

- Each artifact is named, and by default any files added to a tag with the name of the artifact will be treated as belonging to it.
- Artifacts can have include an optional set of keys that can be queried against on Horde.
- The host system (eg. Horde) is deemed responsible for archiving artifacts for later retrieval.

diff --git a/Engine/Source/Programs/AutomationTool/BuildGraph/BgSchema.cs b/Engine/Source/Programs/AutomationTool/BuildGraph/BgSchema.cs
--- a/Engine/Source/Programs/AutomationTool/BuildGraph/BgSchema.cs
+++ b/Engine/Source/Programs/AutomationTool/BuildGraph/BgSchema.cs
@@ -139,6 +139,7 @@ namespace AutomationTool
 		Node,
 		NodeBody,
 		Aggregate,
+		Artifact,
 		Report,
 		Badge,
 		Label,
@@ -322,6 +323,7 @@ namespace AutomationTool
 			newSchema.Items.Add(CreateNodeType());
 			newSchema.Items.Add(CreateNodeBodyType(taskNameToType));
 			newSchema.Items.Add(CreateAggregateType());
+			newSchema.Items.Add(CreateArtifactType());
 			newSchema.Items.Add(CreateReportType());
 			newSchema.Items.Add(CreateBadgeType());
 			newSchema.Items.Add(CreateLabelType());
@@ -478,6 +480,7 @@ namespace AutomationTool
 			graphChoice.Items.Add(CreateSchemaElement("Extend", ScriptSchemaStandardType.Extend));
 			graphChoice.Items.Add(CreateSchemaElement("Agent", ScriptSchemaStandardType.Agent));
 			graphChoice.Items.Add(CreateSchemaElement("Aggregate", ScriptSchemaStandardType.Aggregate));
+			graphChoice.Items.Add(CreateSchemaElement("Artifact", ScriptSchemaStandardType.Artifact));
 			graphChoice.Items.Add(CreateSchemaElement("Report", ScriptSchemaStandardType.Report));
 			graphChoice.Items.Add(CreateSchemaElement("Badge", ScriptSchemaStandardType.Badge));
 			graphChoice.Items.Add(CreateSchemaElement("Label", ScriptSchemaStandardType.Label));
@@ -534,6 +537,7 @@ namespace AutomationTool
 			agentChoice.Items.Add(CreateSchemaElement("Trace", ScriptSchemaStandardType.Trace));
 			agentChoice.Items.Add(CreateSchemaElement("Label", ScriptSchemaStandardType.Label));
 			agentChoice.Items.Add(CreateSchemaElement("Annotate", ScriptSchemaStandardType.Annotate));
+			agentChoice.Items.Add(CreateSchemaElement("Artifact", ScriptSchemaStandardType.Artifact));
 			agentChoice.Items.Add(CreateSchemaElement("Warning", ScriptSchemaStandardType.Warning));
 			agentChoice.Items.Add(CreateSchemaElement("Error", ScriptSchemaStandardType.Error));
 			agentChoice.Items.Add(CreateSchemaElement("Expand", ScriptSchemaStandardType.Expand));
@@ -629,6 +633,21 @@ namespace AutomationTool
 			return aggregateType;
 		}
 
+		/// <summary>
+		/// Creates the schema type representing the artifact type
+		/// </summary>
+		/// <returns>Type definition for an artifact</returns>
+		static XmlSchemaType CreateArtifactType()
+		{
+			XmlSchemaComplexType artifactType = new XmlSchemaComplexType();
+			artifactType.Name = GetTypeName(ScriptSchemaStandardType.Artifact);
+			artifactType.Attributes.Add(CreateSchemaAttribute("Name", ScriptSchemaStandardType.Name, XmlSchemaUse.Required));
+			artifactType.Attributes.Add(CreateSchemaAttribute("Tag", ScriptSchemaStandardType.Tag, XmlSchemaUse.Optional));
+			artifactType.Attributes.Add(CreateSchemaAttribute("Keys", ScriptSchemaStandardType.BalancedString, XmlSchemaUse.Optional));
+			artifactType.Attributes.Add(CreateSchemaAttribute("If", ScriptSchemaStandardType.BalancedString, XmlSchemaUse.Optional));
+			return artifactType;
+		}
+
 		/// <summary>
 		/// Creates the schema type representing the report type
 		/// </summary>
diff --git a/Engine/Source/Programs/AutomationTool/BuildGraph/BgScriptReader.cs b/Engine/Source/Programs/AutomationTool/BuildGraph/BgScriptReader.cs
--- a/Engine/Source/Programs/AutomationTool/BuildGraph/BgScriptReader.cs
+++ b/Engine/Source/Programs/AutomationTool/BuildGraph/BgScriptReader.cs
@@ -472,6 +472,9 @@ namespace AutomationTool
 					case "Aggregate":
 						await ReadAggregateAsync(childElement);
 						break;
+					case "Artifact":
+						await ReadArtifactAsync(childElement);
+						break;
 					case "Report":
 						await ReadReportAsync(childElement);
 						break;
@@ -1003,6 +1006,9 @@ namespace AutomationTool
 					case "Aggregate":
 						await ReadAggregateAsync(childElement);
 						break;
+					case "Artifact":
+						await ReadArtifactAsync(childElement);
+						break;
 					case "Trace":
 						await ReadDiagnosticAsync(childElement, LogLevel.Information);
 						break;
@@ -1099,6 +1105,35 @@ namespace AutomationTool
 			}
 		}
 
+		/// <summary>
+		/// Reads the definition for an artifact
+		/// </summary>
+		/// <param name="element">Xml element to read the definition from</param>
+		async Task ReadArtifactAsync(BgScriptElement element)
+		{
+			string? name;
+			if (await EvaluateConditionAsync(element) && TryReadObjectName(element, out name))
+			{
+				if (_graph.NameToArtifact.ContainsKey(name))
+				{
+					LogError(element, "'{0}' is already defined; cannot add a second time", name);
+				}
+				else
+				{
+					string tag = ReadAttribute(element, "Tag");
+					if (String.IsNullOrEmpty(tag))
+					{
+						tag = $"#{name}";
+					}
+
+					string[] keys = ReadListAttribute(element, "Keys");
+
+					BgArtifactDef newArtifact = new BgArtifactDef(name, tag, keys);
+					_graph.NameToArtifact.Add(name, newArtifact);
+				}
+			}
+		}
+
 		/// <summary>
 		/// Reads the definition for a report
 		/// </summary>
diff --git a/Engine/Source/Programs/Shared/EpicGames.BuildGraph/BgArtifactDef.cs b/Engine/Source/Programs/Shared/EpicGames.BuildGraph/BgArtifactDef.cs
new file mode 100644
--- /dev/null
+++ b/Engine/Source/Programs/Shared/EpicGames.BuildGraph/BgArtifactDef.cs
@@ -0,0 +1,49 @@
+// Copyright Epic Games, Inc. All Rights Reserved.
+
+using System.Collections.Generic;
+
+namespace EpicGames.BuildGraph
+{
+	/// <summary>
+	/// Defines an artifact produced by the build
+	/// </summary>
+	public class BgArtifactDef
+	{
+		/// <summary>
+		/// Name of this artifact
+		/// </summary>
+		public string Name { get; }
+
+		/// <summary>
+		/// Tag to use for the artifact. Uses the artifact name by default.
+		/// </summary>
+		public string Tag { get; }
+
+		/// <summary>
+		/// Keys that can be used to find the artifact
+		/// </summary>
+		public IReadOnlyList<string> Keys { get; }
+
+		/// <summary>
+		/// Constructor
+		/// </summary>
+		/// <param name="name">Name of the artifact</param>
+		/// <param name="tag">Name of the tag producing this artifact</param>
+		/// <param name="keys">Keys that can be used to find the artifact</param>
+		public BgArtifactDef(string name, string tag, IReadOnlyList<string> keys)
+		{
+			Name = name;
+			Tag = tag;
+			Keys = keys;
+		}
+
+		/// <summary>
+		/// Get the name of this badge
+		/// </summary>
+		/// <returns>The name of this badge</returns>
+		public override string ToString()
+		{
+			return Name;
+		}
+	}
+}
diff --git a/Engine/Source/Programs/Shared/EpicGames.BuildGraph/BgGraphDef.cs b/Engine/Source/Programs/Shared/EpicGames.BuildGraph/BgGraphDef.cs
--- a/Engine/Source/Programs/Shared/EpicGames.BuildGraph/BgGraphDef.cs
+++ b/Engine/Source/Programs/Shared/EpicGames.BuildGraph/BgGraphDef.cs
@@ -80,6 +80,11 @@ namespace EpicGames.BuildGraph
 		/// </summary>
 		public Dictionary<string, BgAggregateDef> NameToAggregate { get; private set; } = new Dictionary<string, BgAggregateDef>(StringComparer.OrdinalIgnoreCase);
 
+		/// <summary>
+		/// Mapping of artifact names to their definitions. Artifacts will be produced from matching tag names.
+		/// </summary>
+		public Dictionary<string, BgArtifactDef> NameToArtifact { get; private set; } = new Dictionary<string, BgArtifactDef>(StringComparer.OrdinalIgnoreCase);
+
 		/// <summary>
 		/// List of badges that can be displayed for this build
 		/// </summary>
@@ -422,6 +427,18 @@ namespace EpicGames.BuildGraph
 			using (JsonWriter jsonWriter = new JsonWriter(file.FullName))
 			{
 				jsonWriter.WriteObjectStart();
+
+				jsonWriter.WriteArrayStart("Artifacts");
+				foreach (BgArtifactDef artifact in NameToArtifact.Values)
+				{
+					jsonWriter.WriteObjectStart();
+					jsonWriter.WriteValue("Name", artifact.Name);
+					jsonWriter.WriteValue("Tag", artifact.Tag);
+					jsonWriter.WriteStringArrayField("Keys", artifact.Keys);
+					jsonWriter.WriteObjectEnd();
+				}
+				jsonWriter.WriteArrayEnd();
+
 				jsonWriter.WriteArrayStart("Groups");
 				foreach (BgAgentDef agent in Agents)
 				{
