BuildGraph: Fix Tag property not being respected in WriteTextFile task.

diff --git a/Engine/Source/Programs/AutomationTool/BuildGraph/Tasks/WriteTextFileTask.cs b/Engine/Source/Programs/AutomationTool/BuildGraph/Tasks/WriteTextFileTask.cs
--- a/Engine/Source/Programs/AutomationTool/BuildGraph/Tasks/WriteTextFileTask.cs
+++ b/Engine/Source/Programs/AutomationTool/BuildGraph/Tasks/WriteTextFileTask.cs
@@ -112,6 +112,15 @@ namespace AutomationTool.Tasks
 				Logger.LogInformation("{Text}", string.Format("Writing text to file '{0}': {1}", Parameters.File, FileText));
 				await FileReference.WriteAllTextAsync(Parameters.File, FileText);
 			}
+
+			// Apply the optional tag to the build products
+			foreach (string TagName in FindTagNamesFromList(Parameters.Tag))
+			{
+				FindOrAddTagSet(TagNameToFileSet, TagName).Add(Parameters.File);
+			}
+
+			// Add them to the set of build products
+			BuildProducts.Add(Parameters.File);
 		}
 
 		/// <summary>
