BuildGraph: Store artifact keys in the tag manifest rather than in the exported graph definition, since they may include non-constant values from build arguments.

diff --git a/Engine/Source/Programs/AutomationTool/BuildGraph/BgScriptReader.cs b/Engine/Source/Programs/AutomationTool/BuildGraph/BgScriptReader.cs
--- a/Engine/Source/Programs/AutomationTool/BuildGraph/BgScriptReader.cs
+++ b/Engine/Source/Programs/AutomationTool/BuildGraph/BgScriptReader.cs
@@ -1114,23 +1114,16 @@ namespace AutomationTool
 			string? name;
 			if (await EvaluateConditionAsync(element) && TryReadObjectName(element, out name))
 			{
-				if (_graph.NameToArtifact.ContainsKey(name))
+				string tag = ReadAttribute(element, "Tag");
+				if (String.IsNullOrEmpty(tag))
 				{
-					LogError(element, "'{0}' is already defined; cannot add a second time", name);
+					tag = $"#{name}";
 				}
-				else
-				{
-					string tag = ReadAttribute(element, "Tag");
-					if (String.IsNullOrEmpty(tag))
-					{
-						tag = $"#{name}";
-					}
 
-					string[] keys = ReadListAttribute(element, "Keys");
+				string[] keys = ReadListAttribute(element, "Keys");
 
-					BgArtifactDef newArtifact = new BgArtifactDef(name, tag, keys);
-					_graph.NameToArtifact.Add(name, newArtifact);
-				}
+				BgArtifactDef newArtifact = new BgArtifactDef(name, tag, keys);
+				_graph.Artifacts.Add(newArtifact);
 			}
 		}
 
diff --git a/Engine/Source/Programs/AutomationTool/BuildGraph/BuildGraph.cs b/Engine/Source/Programs/AutomationTool/BuildGraph/BuildGraph.cs
--- a/Engine/Source/Programs/AutomationTool/BuildGraph/BuildGraph.cs
+++ b/Engine/Source/Programs/AutomationTool/BuildGraph/BuildGraph.cs
@@ -975,7 +975,7 @@ namespace AutomationTool
 			HashSet<TempStorageBlock> InputStorageBlocks = new HashSet<TempStorageBlock>();
 			foreach(BgNodeOutput Input in Node.Inputs)
 			{
-				TempStorageFileList FileList = Storage.ReadFileList(Input.ProducingNode.Name, Input.TagName);
+				TempStorageTagManifest FileList = Storage.ReadFileList(Input.ProducingNode.Name, Input.TagName);
 				TagNameToFileSet[Input.TagName] = FileList.ToFileSet(RootDir);
 				InputStorageBlocks.UnionWith(FileList.Blocks);
 			}
@@ -1140,6 +1140,9 @@ namespace AutomationTool
 					Storage.Archive(Node.Name, Pair.Key, Pair.Value.ToArray(), Pair.Value.Any(x => ReferencedOutputFiles.Contains(x)));
 				}
 
+				// Find all the output tags that are published artifacts
+				Dictionary<string, BgArtifactDef> outputNameToArtifact = Graph.Artifacts.ToDictionary(x => x.Tag, x => x);
+
 				// Publish all the output tags
 				foreach (BgNodeOutput Output in Node.Outputs)
 				{
@@ -1154,8 +1157,14 @@ namespace AutomationTool
 							StorageBlocks.Add(StorageBlock);
 						}
 					}
+					
+					IEnumerable<string> Keys = Enumerable.Empty<string>();
+					if (outputNameToArtifact.TryGetValue(Output.TagName, out BgArtifactDef artifact))
+					{
+						Keys = artifact.Keys;
+					}
 
-					Storage.WriteFileList(Node.Name, Output.TagName, Files, StorageBlocks.ToArray());
+					Storage.WriteFileList(Node.Name, Output.TagName, Files, StorageBlocks.ToArray(), Keys);
 				}
 			}
 
diff --git a/Engine/Source/Programs/AutomationTool/BuildGraph/TempStorage.cs b/Engine/Source/Programs/AutomationTool/BuildGraph/TempStorage.cs
--- a/Engine/Source/Programs/AutomationTool/BuildGraph/TempStorage.cs
+++ b/Engine/Source/Programs/AutomationTool/BuildGraph/TempStorage.cs
@@ -411,7 +411,8 @@ namespace AutomationTool
 	/// <summary>
 	/// Stores the contents of a tagged file set
 	/// </summary>
-	public class TempStorageFileList
+	[XmlRoot(ElementName = "TempStorageFileList")]
+	public class TempStorageTagManifest
 	{
 		/// <summary>
 		/// List of files that are in this tag set, relative to the root directory
@@ -434,15 +435,22 @@ namespace AutomationTool
 		[XmlArrayItem("Block")]
 		public TempStorageBlock[] Blocks;
 
+		/// <summary>
+		/// List of keys for published artifacts
+		/// </summary>
+		[XmlArray]
+		[XmlArrayItem("ArtifactKey")]
+		public string[] ArtifactKeys;
+
 		/// <summary>
 		/// Construct a static Xml serializer to avoid throwing an exception searching for the reflection info at runtime
 		/// </summary>
-		static XmlSerializer Serializer = XmlSerializer.FromTypes(new Type[]{ typeof(TempStorageFileList) })[0];
+		static XmlSerializer Serializer = XmlSerializer.FromTypes(new Type[]{ typeof(TempStorageTagManifest) })[0];
 
 		/// <summary>
 		/// Construct an empty file list for deserialization
 		/// </summary>
-		private TempStorageFileList()
+		private TempStorageTagManifest()
 		{
 		}
 
@@ -452,7 +460,8 @@ namespace AutomationTool
 		/// <param name="InFiles">List of full file paths</param>
 		/// <param name="RootDir">Root folder for all the files. All files must be relative to this RootDir.</param>
 		/// <param name="InBlocks">Referenced storage blocks required for these files</param>
-		public TempStorageFileList(IEnumerable<FileReference> InFiles, DirectoryReference RootDir, IEnumerable<TempStorageBlock> InBlocks)
+		/// <param name="InArtifactKeys">Keys for published artifacts</param>
+		public TempStorageTagManifest(IEnumerable<FileReference> InFiles, DirectoryReference RootDir, IEnumerable<TempStorageBlock> InBlocks, IEnumerable<string> InArtifactKeys)
 		{
 			List<string> NewLocalFiles = new List<string>();
 			List<string> NewExternalFiles = new List<string>();
@@ -471,17 +480,18 @@ namespace AutomationTool
 			ExternalFiles = NewExternalFiles.ToArray();
 
 			Blocks = InBlocks.ToArray();
+			ArtifactKeys = InArtifactKeys.ToArray();
 		}
 
 		/// <summary>
 		/// Load this list of files from disk
 		/// </summary>
 		/// <param name="File">File to load</param>
-		static public TempStorageFileList Load(FileReference File)
+		static public TempStorageTagManifest Load(FileReference File)
 		{
 			using(StreamReader Reader = new StreamReader(File.FullName))
 			{
-				return (TempStorageFileList)Serializer.Deserialize(Reader);
+				return (TempStorageTagManifest)Serializer.Deserialize(Reader);
 			}
 		}
 
@@ -679,7 +689,7 @@ namespace AutomationTool
 				}
 
 				// Read the manifest and add the referenced blocks to be checked
-				TempStorageFileList LocalFileList = TempStorageFileList.Load(LocalFileListLocation);
+				TempStorageTagManifest LocalFileList = TempStorageTagManifest.Load(LocalFileListLocation);
 				Blocks.UnionWith(LocalFileList.Blocks);
 			}
 
@@ -729,16 +739,16 @@ namespace AutomationTool
 		/// <param name="NodeName">Name of the node which produced the tag set</param>
 		/// <param name="TagName">Name of the tag, with a '#' prefix</param>
 		/// <returns>The set of files</returns>
-		public TempStorageFileList ReadFileList(string NodeName, string TagName)
+		public TempStorageTagManifest ReadFileList(string NodeName, string TagName)
 		{
-			TempStorageFileList FileList = null;
+			TempStorageTagManifest FileList = null;
 
 			// Try to read the tag set from the local directory
 			FileReference LocalFileListLocation = GetTaggedFileListLocation(LocalDir, NodeName, TagName);
 			if(FileReference.Exists(LocalFileListLocation))
 			{
 				Logger.LogInformation("Reading local file list from {Arg0}", LocalFileListLocation.FullName);
-				FileList = TempStorageFileList.Load(LocalFileListLocation);
+				FileList = TempStorageTagManifest.Load(LocalFileListLocation);
 			}
 			else
 			{
@@ -766,7 +776,7 @@ namespace AutomationTool
 					{
 						// Read the shared manifest
 						Logger.LogInformation("Copying shared tag set from {Arg0} to {Arg1}", SharedFileListLocation.FullName, LocalFileListLocation.FullName);
-						FileList = TempStorageFileList.Load(SharedFileListLocation);
+						FileList = TempStorageTagManifest.Load(SharedFileListLocation);
 					}, Attempts, TimeSpan.FromSeconds(5));
 				}
 				catch
@@ -788,11 +798,12 @@ namespace AutomationTool
 		/// <param name="TagName">Name of the tag, with a '#' prefix</param>
 		/// <param name="Files">List of files in this set</param>
 		/// <param name="Blocks">List of referenced storage blocks</param>
+		/// <param name="ArtifactKeys">Keys for published artifacts</param>
 		/// <returns>The set of files</returns>
-		public void WriteFileList(string NodeName, string TagName, IEnumerable<FileReference> Files, IEnumerable<TempStorageBlock> Blocks)
+		public void WriteFileList(string NodeName, string TagName, IEnumerable<FileReference> Files, IEnumerable<TempStorageBlock> Blocks, IEnumerable<string> ArtifactKeys)
 		{
 			// Create the file list
-			TempStorageFileList FileList = new TempStorageFileList(Files, RootDir, Blocks);
+			TempStorageTagManifest FileList = new TempStorageTagManifest(Files, RootDir, Blocks, ArtifactKeys);
 
 			// Save the set of files to the local and shared locations
 			FileReference LocalFileListLocation = GetTaggedFileListLocation(LocalDir, NodeName, TagName);
diff --git a/Engine/Source/Programs/Shared/EpicGames.BuildGraph/BgGraphDef.cs b/Engine/Source/Programs/Shared/EpicGames.BuildGraph/BgGraphDef.cs
--- a/Engine/Source/Programs/Shared/EpicGames.BuildGraph/BgGraphDef.cs
+++ b/Engine/Source/Programs/Shared/EpicGames.BuildGraph/BgGraphDef.cs
@@ -83,7 +83,7 @@ namespace EpicGames.BuildGraph
 		/// <summary>
 		/// Mapping of artifact names to their definitions. Artifacts will be produced from matching tag names.
 		/// </summary>
-		public Dictionary<string, BgArtifactDef> NameToArtifact { get; private set; } = new Dictionary<string, BgArtifactDef>(StringComparer.OrdinalIgnoreCase);
+		public List<BgArtifactDef> Artifacts { get; } = new List<BgArtifactDef>();
 
 		/// <summary>
 		/// List of badges that can be displayed for this build
@@ -429,12 +429,11 @@ namespace EpicGames.BuildGraph
 				jsonWriter.WriteObjectStart();
 
 				jsonWriter.WriteArrayStart("Artifacts");
-				foreach (BgArtifactDef artifact in NameToArtifact.Values)
+				foreach (BgArtifactDef artifact in Artifacts)
 				{
 					jsonWriter.WriteObjectStart();
 					jsonWriter.WriteValue("Name", artifact.Name);
 					jsonWriter.WriteValue("Tag", artifact.Tag);
-					jsonWriter.WriteStringArrayField("Keys", artifact.Keys);
 					jsonWriter.WriteObjectEnd();
 				}
 				jsonWriter.WriteArrayEnd();
diff --git a/Engine/Source/Programs/Shared/EpicGames.BuildGraph/BgNodeDef.cs b/Engine/Source/Programs/Shared/EpicGames.BuildGraph/BgNodeDef.cs
--- a/Engine/Source/Programs/Shared/EpicGames.BuildGraph/BgNodeDef.cs
+++ b/Engine/Source/Programs/Shared/EpicGames.BuildGraph/BgNodeDef.cs
@@ -30,12 +30,12 @@ namespace EpicGames.BuildGraph
 		/// <summary>
 		/// Constructor
 		/// </summary>
-		/// <param name="inProducingNode">Node which produces the given output</param>
-		/// <param name="inTagName">Name of the tag</param>
-		public BgNodeOutput(BgNodeDef inProducingNode, string inTagName)
+		/// <param name="producingNode">Node which produces the given output</param>
+		/// <param name="tagName">Name of the tag</param>
+		public BgNodeOutput(BgNodeDef producingNode, string tagName)
 		{
-			ProducingNode = inProducingNode;
-			TagName = inTagName;
+			ProducingNode = producingNode;
+			TagName = tagName;
 		}
 
 		/// <summary>
