Low Level Tests: timeout detection from within test app, inactivity detection via Gauntlet
Default timeouts for tests are:
- optional, only enabled when set via --timeout=X (in minutes)
- standard practice to help triage tests and bound their running time; slower tests can be tagged appropriately and run with the nightly tests
- useful to detect segmentation faults and other issues by enforcing a maximum period of inactivity

diff --git a/Engine/Source/Developer/LowLevelTestsRunner/Private/TestListeners/EventListener.cpp b/Engine/Source/Developer/LowLevelTestsRunner/Private/TestListeners/EventListener.cpp
--- a/Engine/Source/Developer/LowLevelTestsRunner/Private/TestListeners/EventListener.cpp
+++ b/Engine/Source/Developer/LowLevelTestsRunner/Private/TestListeners/EventListener.cpp
@@ -10,6 +10,10 @@
 #include <catch2/reporters/catch_reporter_registrars.hpp>
 #include <catch2/reporters/catch_reporter_streaming_base.hpp>
 
+#include <thread>
+#include <chrono>
+#include <string>
+
 namespace UE::LowLevelTests
 {
 
@@ -19,6 +23,32 @@ class FEventListener final : public Catch::EventListenerBase
 {
 private:
 	FDateTime TestStarted;
+	std::thread* TimeoutThread = nullptr;
+	std::string CurrentTest;
+
+	void logMessageOnTimeout()
+	{
+		std::string CurrentTestLocal;
+		if (ITestRunner* TestRunner = ITestRunner::Get())
+		{
+			while (true)
+			{
+				if (CurrentTestLocal != CurrentTest)
+				{
+					CurrentTestLocal = CurrentTest;
+					std::this_thread::sleep_for(std::chrono::minutes(TestRunner->GetTimeoutMinutes()));
+					UE_CLOG(CurrentTestLocal == CurrentTest, LogLowLevelTests, Error,
+						TEXT("Timeout detected: Test case \"%hs\" has been running for more than %d minute(s)"), CurrentTest.c_str(), TestRunner->GetTimeoutMinutes());
+					
+				}
+				else // test still running after timeout reached, prevent CPU spinning
+				{
+					std::this_thread::sleep_for(std::chrono::seconds(1));
+				}
+			}
+		}
+	}
+
 public:
 	using EventListenerBase::EventListenerBase;
 
@@ -26,6 +56,13 @@ public:
 	{
 		if (ITestRunner* TestRunner = ITestRunner::Get())
 		{
+			CurrentTest = TestInfo.name;
+
+			if (TimeoutThread == nullptr && TestRunner->GetTimeoutMinutes() > 0)
+			{
+				TimeoutThread = new std::thread(&FEventListener::logMessageOnTimeout, this);
+			}
+
 			TestStarted = FDateTime::Now();
 			UE_CLOG(TestRunner->IsDebugMode(), LogLowLevelTests, Display,
 				TEXT("Started test case \"%hs\" at %hs(%" SIZE_T_FMT ") with tags %hs"),
diff --git a/Engine/Source/Developer/LowLevelTestsRunner/Private/TestRunner.cpp b/Engine/Source/Developer/LowLevelTestsRunner/Private/TestRunner.cpp
--- a/Engine/Source/Developer/LowLevelTestsRunner/Private/TestRunner.cpp
+++ b/Engine/Source/Developer/LowLevelTestsRunner/Private/TestRunner.cpp
@@ -64,6 +64,8 @@ public:
 	bool HasLogOutput() const final { return bLogOutput || bDebugMode; }
 	bool IsDebugMode() const final { return bDebugMode; }
 
+	int32 GetTimeoutMinutes() const final { return TimeoutMinutes; }
+
 private:
 	TArray<const ANSICHAR*> CatchArgs;
 	FStringBuilderBase ExtraArgs;
@@ -75,6 +77,7 @@ private:
 	bool bWaitForInputToTerminate = false;
 	bool bAttachToDebugger = false;
 	int32 SleepOnInitSeconds = 0;
+	int32 TimeoutMinutes = 0;
 };
 
 FTestRunner::FTestRunner()
@@ -113,6 +116,10 @@ void FTestRunner::ParseCommandLine(TConstArrayView<const ANSICHAR*> Args)
 		{
 			LexFromString(SleepOnInitSeconds, WriteToString<16>(Arg.RightChop(8)).ToView());
 		}
+		else if (Arg.StartsWith(ANSITEXTVIEW("--timeout=")))
+		{
+			LexFromString(TimeoutMinutes, WriteToString<16>(Arg.RightChop(10)).ToView());
+		}
 		else if (Arg == ANSITEXTVIEW("--global-setup"))
 		{
 			bGlobalSetup = true;
diff --git a/Engine/Source/Developer/LowLevelTestsRunner/Private/TestRunnerPrivate.h b/Engine/Source/Developer/LowLevelTestsRunner/Private/TestRunnerPrivate.h
--- a/Engine/Source/Developer/LowLevelTestsRunner/Private/TestRunnerPrivate.h
+++ b/Engine/Source/Developer/LowLevelTestsRunner/Private/TestRunnerPrivate.h
@@ -17,6 +17,8 @@ public:
 	virtual bool HasLogOutput() const = 0;
 	/** Returns true if the test runner has extra debug functionality enabled. */
 	virtual bool IsDebugMode() const = 0;
+	/** Returns a value indicating per test allowed timeout in minutes. */
+	virtual int GetTimeoutMinutes() const = 0;
 
 protected:
 	ITestRunner();
diff --git a/Engine/Source/Programs/AutomationTool/LowLevelTests/RunLowLevelTests.cs b/Engine/Source/Programs/AutomationTool/LowLevelTests/RunLowLevelTests.cs
--- a/Engine/Source/Programs/AutomationTool/LowLevelTests/RunLowLevelTests.cs
+++ b/Engine/Source/Programs/AutomationTool/LowLevelTests/RunLowLevelTests.cs
@@ -173,6 +173,8 @@ namespace LowLevelTests
 		[AutoParam("console")]
 		public string ReportType;
 
+		public int Timeout;
+
 		public Type BuildSourceType { get; protected set; }
 
 		[AutoParam(UnrealTargetConfiguration.Development)]
@@ -210,6 +212,8 @@ namespace LowLevelTests
 			LogDir = Path.GetFullPath(LogDir);
 			TempDir = Path.GetFullPath(TempDir);
 
+			Timeout = Params.ParseValue("timeout=", 0);
+
 			Build = Params.ParseValue("build=", null);
 			TestApp = Globals.Params.ParseValue("testapp=", "");
 
@@ -244,16 +248,18 @@ namespace LowLevelTests
 		private int Sleep { get; set; }
 		private bool AttachToDebugger { get; set; }
 		private string ReportType { get; set; }
+		private int PerTestTimeout { get; set; }
 
 		public UnrealDeviceReservation UnrealDeviceReservation { get; private set; }
 
-		public LowLevelTestsSession(LowLevelTestsBuildSource InBuildSource, string InTags, int InSleep, bool InAttachToDebugger, string InReportType)
+		public LowLevelTestsSession(LowLevelTestsBuildSource InBuildSource, string InTags, int InSleep, bool InAttachToDebugger, string InReportType, int InPerTestTimeout = 0)
 		{
 			BuildSource = InBuildSource;
 			Tags = InTags;
 			Sleep = InSleep;
 			AttachToDebugger = InAttachToDebugger;
 			ReportType = InReportType;
+			PerTestTimeout = InPerTestTimeout;
 			UnrealDeviceReservation = new UnrealDeviceReservation();
 		}
 
@@ -277,7 +283,7 @@ namespace LowLevelTests
 
 			// TargetDevice<Platform> classes have a hard dependency on UnrealAppConfig instead of IAppConfig.
 			// More refactoring needed to support non-packaged applications that can be run natively from a path on the device.
-			UnrealAppConfig AppConfig = BuildSource.GetUnrealAppConfig(Tags, Sleep, AttachToDebugger, ReportType);
+			UnrealAppConfig AppConfig = BuildSource.GetUnrealAppConfig(Tags, Sleep, AttachToDebugger, ReportType, PerTestTimeout);
 
 			IEnumerable<ITargetDevice> DevicesToInstallOn = UnrealDeviceReservation.ReservedDevices.ToArray();
 			ITargetDevice Device = DevicesToInstallOn.Where(D => D.IsConnected && D.Platform == BuildSource.Platform).First();
@@ -396,12 +402,15 @@ namespace LowLevelTests
 
 		public UnrealDeviceTargetConstraint Constraint;
 
-		public LowLevelTestContext(LowLevelTestsBuildSource InBuildInfo, LowLevelTestRoleContext InRoleContext, LowLevelTestExecutorOptions InOptions)
+		public int PerTestTimeout { get; private set; }
+
+		public LowLevelTestContext(LowLevelTestsBuildSource InBuildInfo, LowLevelTestRoleContext InRoleContext, LowLevelTestExecutorOptions InOptions, int InPerTestTimeout = 0)
 		{
 			BuildInfo = InBuildInfo;
 			Options = InOptions;
 			TestParams = new Params(new string[0]);
 			RoleContext = InRoleContext;
+			PerTestTimeout = InPerTestTimeout;
 		}
 
 		public object Clone()
@@ -611,7 +620,7 @@ namespace LowLevelTests
 				.FirstOrDefault();
 		}
 
-		public UnrealAppConfig GetUnrealAppConfig(string InTags, int InSleep, bool InAttachToDebugger, string InReportType)
+		public UnrealAppConfig GetUnrealAppConfig(string InTags, int InSleep, bool InAttachToDebugger, string InReportType, int InPerTestTimeout = 0)
 		{
 			if (CachedConfig == null)
 			{
@@ -637,6 +646,10 @@ namespace LowLevelTests
 				{
 					CachedConfig.CommandLineParams.AddRawCommandline(String.Format("--sleep={0}", InSleep));
 				}
+				if (InPerTestTimeout > 0)
+				{
+					CachedConfig.CommandLineParams.AddRawCommandline(String.Format("--timeout={0}", InPerTestTimeout));
+				}
 				CachedConfig.CommandLineParams.AddRawCommandline("--debug");
 				CachedConfig.CommandLineParams.AddRawCommandline("--log");
 				if (InAttachToDebugger)
diff --git a/Engine/Source/Programs/AutomationTool/LowLevelTests/Tests/Gauntlet.LowLevelTestNode.cs b/Engine/Source/Programs/AutomationTool/LowLevelTests/Tests/Gauntlet.LowLevelTestNode.cs
--- a/Engine/Source/Programs/AutomationTool/LowLevelTests/Tests/Gauntlet.LowLevelTestNode.cs
+++ b/Engine/Source/Programs/AutomationTool/LowLevelTests/Tests/Gauntlet.LowLevelTestNode.cs
@@ -22,6 +22,9 @@ namespace LowLevelTests
 
 		public LowLevelTestsSession LowLevelTestsApp { get; private set; }
 
+		private int LastStdoutSeekPos = 0;
+		private string[] CurrentProcessedLines;
+
 		public LowLevelTests(LowLevelTestContext InContext)
 		{
 			Context = InContext;
@@ -37,11 +40,14 @@ namespace LowLevelTests
 
 		public override float MaxDuration { protected set; get; }
 
+		private DateTime InactivityStart = DateTime.MinValue;
+		private TimeSpan InactivityPeriod = TimeSpan.Zero;
+
 		public override bool IsReadyToStart()
 		{
 			if (LowLevelTestsApp == null)
 			{
-				LowLevelTestsApp = new LowLevelTestsSession(Context.BuildInfo, Context.Options.Tags, Context.Options.Sleep, Context.Options.AttachToDebugger, Context.Options.ReportType);
+				LowLevelTestsApp = new LowLevelTestsSession(Context.BuildInfo, Context.Options.Tags, Context.Options.Sleep, Context.Options.AttachToDebugger, Context.Options.ReportType, Context.Options.Timeout);
 			}
 
 			return LowLevelTestsApp.TryReserveDevices();
@@ -95,13 +101,48 @@ namespace LowLevelTests
 
 		public override void TickTest()
 		{
-			if (TestInstance != null && TestInstance.HasExited)
+			if (TestInstance != null)
 			{
-				if (TestInstance.WasKilled)
+				if (TestInstance.HasExited)
 				{
-					LowLevelTestResult = TestResult.Failed;
+					if (TestInstance.WasKilled)
+					{
+						LowLevelTestResult = TestResult.Failed;
+					}
+					MarkTestComplete();
+				}
+				else
+				{
+					ParseLowLevelTestsLog();
+
+					if (CheckForTimeout())
+					{
+						Log.Error("Timeout detected from application logged events, stopping.");
+						MarkTestComplete();
+						LowLevelTestResult = TestResult.TimedOut;
+					}
+					else if (CurrentProcessedLines != null && CurrentProcessedLines.Length > 0)
+					{
+						InactivityStart = DateTime.MinValue;
+					}
+					else if ((CurrentProcessedLines == null || CurrentProcessedLines.Length == 0) && InactivityStart == DateTime.MinValue)
+					{
+						InactivityStart = DateTime.Now;
+					}
+					else if (InactivityStart != DateTime.MinValue)
+					{
+						InactivityPeriod = DateTime.Now - InactivityStart;
+					}
+
+					if (Context.Options.Timeout != 0 && InactivityPeriod.TotalMinutes > Context.Options.Timeout + 0.5)
+					{
+						Log.Error($"Test application didn't log any test events after timeout period of {Context.Options.Timeout} minutes, stopping.");
+						MarkTestComplete();
+						LowLevelTestResult = TestResult.TimedOut;
+					}
+
+					CurrentProcessedLines = null;
 				}
-				MarkTestComplete();
 			}
 		}
 
@@ -139,7 +180,12 @@ namespace LowLevelTests
 				{
 					ClientOutputWriter.Write(StdOut);
 				}
-				File.Copy(ClientOutputLog, Path.Combine(LogDir, ClientLogFile));
+
+				string DestClientLogFile = Path.Combine(LogDir, ClientLogFile);
+				if (DestClientLogFile != ClientOutputLog)
+				{
+					File.Copy(ClientOutputLog, DestClientLogFile, true);
+				}
 			}
 
 			ILowLevelTestsReporting LowLevelTestsReporting = Gauntlet.Utils.InterfaceHelpers.FindImplementations<ILowLevelTestsReporting>(true)
@@ -162,7 +208,7 @@ namespace LowLevelTests
 			string ExitReason = "";
 			if (TestInstance.WasKilled)
 			{
-				if (InReason == StopReason.MaxDuration)
+				if (InReason == StopReason.MaxDuration || LowLevelTestResult == TestResult.TimedOut)
 				{
 					LowLevelTestResult = TestResult.TimedOut;
 					ExitReason = "Timed Out";
@@ -170,7 +216,7 @@ namespace LowLevelTests
 				else
 				{
 					LowLevelTestResult = TestResult.Failed;
-					ExitReason = "Process was killed by Gauntlet.";
+					ExitReason = $"Process was killed by Gauntlet with reason {InReason.ToString()}.";
 				}
 			}
 			else if (TestInstance.ExitCode != 0)
@@ -227,5 +273,35 @@ namespace LowLevelTests
 				LowLevelTestsApp = null;
 			}
 		}
+
+		private void ParseLowLevelTestsLog()
+		{
+			// Parse new lines from Stdout, if any
+			if (LastStdoutSeekPos < TestInstance.StdOut.Length)
+			{
+				CurrentProcessedLines = TestInstance.StdOut
+					.Substring(LastStdoutSeekPos)
+					.Split("\n")
+					.Where(Line => Line.Contains("LogLowLevelTests"))
+					.ToArray();
+				LastStdoutSeekPos = TestInstance.StdOut.Length - 1;
+			}
+		}
+
+		private bool CheckForTimeout()
+		{
+			if (CurrentProcessedLines == null || CurrentProcessedLines.Length == 0)
+			{
+				return false;
+			}
+			foreach (string Line in CurrentProcessedLines)
+			{
+				if (Line.Contains("Timeout detected"))
+				{
+					return true;
+				}
+			}
+			return false;
+		}
 	}
 }
