BuildGraph: Always register a "docker compose down" command to run after a "docker compose up".

diff --git a/Engine/Source/Programs/AutomationTool/BuildGraph/Task.cs b/Engine/Source/Programs/AutomationTool/BuildGraph/Task.cs
--- a/Engine/Source/Programs/AutomationTool/BuildGraph/Task.cs
+++ b/Engine/Source/Programs/AutomationTool/BuildGraph/Task.cs
@@ -14,6 +14,7 @@ using OpenTracing;
 using UnrealBuildBase;
 using UnrealBuildTool;
 using Microsoft.Extensions.Logging;
+using System.Runtime.InteropServices;
 
 using static AutomationTool.CommandUtils;
 
@@ -432,6 +433,20 @@ namespace AutomationTool
 		{
 			return Text.Split(';').Select(x => x.Trim()).Where(x => x.Length > 0).ToList();
 		}
+
+		/// <summary>
+		/// Add cleanup commands to run after the job completes
+		/// </summary>
+		/// <param name="NewLines">Lines to add to the cleanup script</param>
+		public static async Task AddCleanupCommandsAsync(IEnumerable<string> NewLines)
+		{
+			string CleanupScriptEnvVar = Environment.GetEnvironmentVariable("UE_HORDE_CLEANUP");
+			if (!String.IsNullOrEmpty(CleanupScriptEnvVar))
+			{
+				FileReference CleanupScript = new FileReference(CleanupScriptEnvVar);
+				await FileReference.AppendAllLinesAsync(CleanupScript, NewLines);
+			}
+		}
 	}
 
 	/// <summary>
diff --git a/Engine/Source/Programs/AutomationTool/BuildGraph/Tasks/DockerComposeUpTask.cs b/Engine/Source/Programs/AutomationTool/BuildGraph/Tasks/DockerComposeUpTask.cs
--- a/Engine/Source/Programs/AutomationTool/BuildGraph/Tasks/DockerComposeUpTask.cs
+++ b/Engine/Source/Programs/AutomationTool/BuildGraph/Tasks/DockerComposeUpTask.cs
@@ -60,6 +60,14 @@ namespace AutomationTool.Tasks
 		/// <param name="TagNameToFileSet">Mapping from tag names to the set of files they include</param>
 		public override async Task ExecuteAsync(JobContext Job, HashSet<FileReference> BuildProducts, Dictionary<string, HashSet<FileReference>> TagNameToFileSet)
 		{
+			FileReference LogFile = FileReference.Combine(Unreal.RootDirectory, "Engine/Programs/AutomationTool/Saved/Logs/docker-compose-logs.txt");
+			string[] Lines =
+			{
+				$"docker-compose -f {Parameters.File.QuoteArgument()} logs --no-color > {LogFile.FullName.QuoteArgument()}",
+				$"docker-compose -f {Parameters.File.QuoteArgument()} down"
+			};
+			await AddCleanupCommandsAsync(Lines);
+
 			StringBuilder Arguments = new StringBuilder("--ansi never ");
 			if (!String.IsNullOrEmpty(Parameters.File))
 			{
diff --git a/Engine/Source/Programs/UnrealCloudDDC/BuildUnrealCloudDDC.xml b/Engine/Source/Programs/UnrealCloudDDC/BuildUnrealCloudDDC.xml
--- a/Engine/Source/Programs/UnrealCloudDDC/BuildUnrealCloudDDC.xml
+++ b/Engine/Source/Programs/UnrealCloudDDC/BuildUnrealCloudDDC.xml
@@ -11,10 +11,6 @@
 	<Agent Name="Jupiter" Type="Linux">
 		<Property Name="ComposeFile" Value="$(RootDir)/Engine/Source/Programs/UnrealCloudDDC/Composes/docker-compose-tests.yml"/>
 		<Node Name="Test Jupiter">
-			<!-- Make sure to run docker compose down after the build no matter what, also copy the logs to AutomationTool logs so we get the logs in Horde -->
-			<WriteTextFile File="$(RootDir)/Cleanup.sh" Text="docker-compose -f $(ComposeFile) logs --no-color > $(RootDir)/Engine/Programs/AutomationTool/Saved/Logs/docker-compose-logs.txt"/>
-			<WriteTextFile File="$(RootDir)/Cleanup.sh" Append="true" Text="docker-compose -f $(ComposeFile) down"/>
-
 			<Docker-Compose-Up File="$(ComposeFile)"/>
 
 			<!-- Sleep a while to let scylla finish starting up before we execute the tests -->
