BuildGraph: Add a task to create a report for the Horde dashboard. Experimental; may be removed in future.

diff --git a/Engine/Source/Programs/AutomationTool/BuildGraph/Tasks/HordeCreateReportTask.cs b/Engine/Source/Programs/AutomationTool/BuildGraph/Tasks/HordeCreateReportTask.cs
new file mode 100644
--- /dev/null
+++ b/Engine/Source/Programs/AutomationTool/BuildGraph/Tasks/HordeCreateReportTask.cs
@@ -0,0 +1,112 @@
+using EpicGames.Core;
+using Microsoft.Extensions.Logging;
+using System;
+using System.Collections.Generic;
+using System.IO;
+using System.Linq;
+using System.Text;
+using System.Text.Json;
+using System.Threading.Tasks;
+using System.Xml;
+
+namespace AutomationTool.Tasks
+{
+	/// <summary>
+	/// Parameters for a <see cref="HordeCreateReportTask"/>.
+	/// </summary>
+	public class HordeCreateReportTaskParameters
+	{
+		/// <summary>
+		/// Name for the report
+		/// </summary>
+		[TaskParameter]
+		public string Name;
+
+		/// <summary>
+		/// Where to display the report
+		/// </summary>
+		[TaskParameter]
+		public string Scope;
+
+		/// <summary>
+		/// Where to show the report
+		/// </summary>
+		[TaskParameter]
+		public string Placement;
+
+		/// <summary>
+		/// Text to be displayed
+		/// </summary>
+		[TaskParameter]
+		public string Text;
+	}
+
+	/// <summary>
+	/// Creates a Horde report file, which will be displayed on the dashboard with any job running this task.
+	/// </summary>
+	[TaskElement("Horde-CreateReport", typeof(HordeCreateReportTaskParameters))]
+	public class HordeCreateReportTask : BgTaskImpl
+	{
+		/// <summary>
+		/// Parameters for this task.
+		/// </summary>
+		HordeCreateReportTaskParameters Parameters;
+
+		/// <summary>
+		/// Constructor.
+		/// </summary>
+		/// <param name="InParameters">Parameters for this task.</param>
+		public HordeCreateReportTask(HordeCreateReportTaskParameters InParameters)
+		{
+			Parameters = InParameters;
+		}
+
+		/// <summary>
+		/// Execute the task.
+		/// </summary>
+		/// <param name="Job">Information about the current job.</param>
+		/// <param name="BuildProducts">Set of build products produced by this node.</param>
+		/// <param name="TagNameToFileSet">Mapping from tag names to the set of files they include.</param>
+		public override async Task ExecuteAsync(JobContext Job, HashSet<FileReference> BuildProducts, Dictionary<string, HashSet<FileReference>> TagNameToFileSet)
+		{
+			FileReference ReportTextFile = FileReference.Combine(new DirectoryReference(CommandUtils.CmdEnv.LogFolder), $"{Parameters.Name}.md");
+			await FileReference.WriteAllTextAsync(ReportTextFile, Parameters.Text);
+
+			FileReference ReportJsonFile = FileReference.Combine(new DirectoryReference(CommandUtils.CmdEnv.LogFolder), $"{Parameters.Name}.json.md");
+			using (FileStream ReportJsonStream = FileReference.Open(ReportJsonFile, FileMode.Create, FileAccess.Write, FileShare.Read))
+			{
+				using (Utf8JsonWriter Writer = new Utf8JsonWriter(ReportJsonStream))
+				{
+					Writer.WriteStartObject();
+					Writer.WriteString("scope", Parameters.Scope);
+					Writer.WriteString("name", Parameters.Name);
+					Writer.WriteString("placement", Parameters.Placement);
+					Writer.WriteString("file", ReportTextFile.GetFileName());
+					Writer.WriteEndObject();
+				}
+			}
+
+			Logger.LogInformation("Written report to {TextFile} and {JsonFile}: \"{Text}\"", ReportTextFile, ReportJsonFile, Parameters.Text);
+		}
+
+		/// <summary>
+		/// Output this task out to an XML writer.
+		/// </summary>
+		public override void Write(XmlWriter Writer)
+		{
+			Write(Writer, Parameters);
+		}
+
+		/// <summary>
+		/// Find all the tags which are used as inputs to this task
+		/// </summary>
+		/// <returns>The tag names which are read by this task</returns>
+		public override IEnumerable<string> FindConsumedTagNames() => Enumerable.Empty<string>();
+
+		/// <summary>
+		/// Find all the tags which are modified by this task
+		/// </summary>
+		/// <returns>The tag names which are modified by this task</returns>
+		public override IEnumerable<string> FindProducedTagNames() => Enumerable.Empty<string>();
+	}
+}
