BuildGraph: Add missing copyright notice.

diff --git a/Engine/Source/Programs/AutomationTool/BuildGraph/Tasks/HordeCreateReportTask.cs b/Engine/Source/Programs/AutomationTool/BuildGraph/Tasks/HordeCreateReportTask.cs
--- a/Engine/Source/Programs/AutomationTool/BuildGraph/Tasks/HordeCreateReportTask.cs
+++ b/Engine/Source/Programs/AutomationTool/BuildGraph/Tasks/HordeCreateReportTask.cs
@@ -1,3 +1,5 @@
+// Copyright Epic Games, Inc. All Rights Reserved.
+
 using EpicGames.Core;
 using Microsoft.Extensions.Logging;
 using System;
