Separating out SM6 shaders form the main game data.

diff --git a/Engine/Source/Programs/AutomationTool/Scripts/CopyBuildToStagingDirectory.Automation.cs b/Engine/Source/Programs/AutomationTool/Scripts/CopyBuildToStagingDirectory.Automation.cs
--- a/Engine/Source/Programs/AutomationTool/Scripts/CopyBuildToStagingDirectory.Automation.cs
+++ b/Engine/Source/Programs/AutomationTool/Scripts/CopyBuildToStagingDirectory.Automation.cs
@@ -4046,10 +4046,10 @@ namespace AutomationScripts
 					}
 				}
 
-				// Add ondemand chunk definitions
+				// Add chunk definitions
 				if (PakRulesList != null)
 				{
-					foreach (PakFileRules Rule in PakRulesList.Where(R => R.bOnDemand))
+					foreach (PakFileRules Rule in PakRulesList)
 					{
 						string ChunkName = Rule.Name;
 						if (Rule.OverridePaks != null && Rule.OverridePaks.Count > 0)
@@ -4062,32 +4062,38 @@ namespace AutomationScripts
 							continue;
 						}
 
-						Logger.LogInformation("Creating content-on-demand chunk definition '{ChunkName}' for PAK rule '{Arg1}'", ChunkName, Rule.Name);
-
 						bool bIsEncryptionKeyValid = true;
-						if (string.IsNullOrEmpty(Rule.EncryptionKeyGuid))
-						{
-							bIsEncryptionKeyValid = false;
-							Logger.LogWarning("No encryption key specified for content-on-demand PAK rule '{Arg0}'", Rule.Name);
-						}
-						else
 						{
-							Guid KeyGuid;
-							if (!Guid.TryParse(Rule.EncryptionKeyGuid, out KeyGuid))
+							Logger.LogInformation("Creating content-on-demand chunk definition '{Arg0}' for PAK rule '{Arg1}'", ChunkName, Rule.Name);
+
+							if (string.IsNullOrEmpty(Rule.EncryptionKeyGuid))
 							{
 								bIsEncryptionKeyValid = false;
-								Logger.LogWarning("Invalid encryption key GUID specified for content-on-demand PAK rule '{Arg0}'", Rule.Name);
+								Logger.LogWarning("No encryption key specified for content-on-demand PAK rule '{Arg0}'", Rule.Name);
 							}
-
-							if (!PakCryptoSettings.ContainsEncryptionKey(Rule.EncryptionKeyGuid))
+							else
 							{
-								bIsEncryptionKeyValid = false;
-								Logger.LogWarning("No matching encryption key found for PAK rule '{Arg0}', GUID '{Arg1}'", Rule.Name, Rule.EncryptionKeyGuid);
+								Guid KeyGuid;
+								if (!Guid.TryParse(Rule.EncryptionKeyGuid, out KeyGuid))
+								{
+									bIsEncryptionKeyValid = false;
+									Logger.LogWarning("Invalid encryption key GUID specified for content-on-demand PAK rule '{Arg0}'", Rule.Name);
+								}
+
+								if (!PakCryptoSettings.ContainsEncryptionKey(Rule.EncryptionKeyGuid))
+								{
+									bIsEncryptionKeyValid = false;
+									Logger.LogWarning("No matching encryption key found for PAK rule '{Arg0}', GUID '{Arg1}'", Rule.Name, Rule.EncryptionKeyGuid);
+								}
 							}
 						}
+						else
+						{
+							Logger.LogInformation("Creating chunk definition '{Arg0}' for PAK rule '{Arg1}'", ChunkName, Rule.Name);
+						}
 
 						ChunkDefinition Chunk = new ChunkDefinition(ChunkName);
-						Chunk.bOnDemand = true;
+						Chunk.bOnDemand = Rule.bOnDemand;
 						Chunk.bCompressed = true;
 						Chunk.EncryptionKeyGuid = bIsEncryptionKeyValid ? Rule.EncryptionKeyGuid : null;
 						ChunkDefinitions.Add(Chunk);
