Fixing missing if block created during previous merge.

diff --git a/Engine/Source/Programs/AutomationTool/Scripts/CopyBuildToStagingDirectory.Automation.cs b/Engine/Source/Programs/AutomationTool/Scripts/CopyBuildToStagingDirectory.Automation.cs
--- a/Engine/Source/Programs/AutomationTool/Scripts/CopyBuildToStagingDirectory.Automation.cs
+++ b/Engine/Source/Programs/AutomationTool/Scripts/CopyBuildToStagingDirectory.Automation.cs
@@ -4063,6 +4063,7 @@ namespace AutomationScripts
 						}
 
 						bool bIsEncryptionKeyValid = true;
+						if (Rule.bOnDemand)
 						{
 							Logger.LogInformation("Creating content-on-demand chunk definition '{Arg0}' for PAK rule '{Arg1}'", ChunkName, Rule.Name);
 
