Automatically detect a if a staged target is using zenstore without pak and configure zenserver host and project id.

diff --git a/Engine/Source/Developer/IoStoreUtilities/Private/ZenStoreWriter.cpp b/Engine/Source/Developer/IoStoreUtilities/Private/ZenStoreWriter.cpp
--- a/Engine/Source/Developer/IoStoreUtilities/Private/ZenStoreWriter.cpp
+++ b/Engine/Source/Developer/IoStoreUtilities/Private/ZenStoreWriter.cpp
@@ -381,10 +381,30 @@ void FZenStoreWriter::Initialize(const FCookInfo& Info)
 			IFileManager::Get().DeleteDirectory(*OutputPath, bRequireExists, bTree);
 		}
 
-		FString OplogLifetimeMarkerPath = OutputPath / (ProjectId + TEXT(".") + OplogId + TEXT(".projectstore"));
+		FString OplogLifetimeMarkerPath = OutputPath / TEXT(".projectstore");
 		TUniquePtr<FArchive> OplogMarker(IFileManager::Get().CreateFileWriter(*OplogLifetimeMarkerPath));
 
 		bool bOplogEstablished = HttpClient->TryCreateOplog(ProjectId, OplogId, OplogLifetimeMarkerPath, Info.bFullBuild);
+
+		if (bOplogEstablished && OplogMarker)
+		{
+			FCbWriter ManifestWriter;
+			ManifestWriter.BeginObject();
+			ManifestWriter.BeginObject("zenserver");
+			bool IsRunningLocally = false;
+#if UE_WITH_ZEN
+			IsRunningLocally = HttpClient->GetZenServiceInstance().IsServiceRunningLocally();
+#endif
+			ManifestWriter << "islocalhost" << IsRunningLocally;
+			ManifestWriter << "hostname" << HttpClient->GetHostName();
+			ManifestWriter << "hostport" << HttpClient->GetPort();
+			ManifestWriter << "projectid" << ProjectId;
+			ManifestWriter << "oplogid" << OplogId;
+			ManifestWriter.EndObject();
+			ManifestWriter.EndObject();
+			SaveCompactBinary(*OplogMarker, ManifestWriter.Save());
+		}
+
 		OplogMarker.Reset();
 
 		if (!bOplogEstablished && CleanBuild)
diff --git a/Engine/Source/Programs/AutomationTool/Scripts/CopyBuildToStagingDirectory.Automation.cs b/Engine/Source/Programs/AutomationTool/Scripts/CopyBuildToStagingDirectory.Automation.cs
--- a/Engine/Source/Programs/AutomationTool/Scripts/CopyBuildToStagingDirectory.Automation.cs
+++ b/Engine/Source/Programs/AutomationTool/Scripts/CopyBuildToStagingDirectory.Automation.cs
@@ -867,27 +867,17 @@
 			}
 		}
 
-		public static void CreateStagingManifest(ProjectParams Params, DeploymentContext SC)
+		public static bool SetUpStagingSourceDirectories(ProjectParams Params, DeploymentContext SC)
 		{
-			if (!Params.Stage)
+			if (SC.CookPlatform == null)
 			{
-				return;
+				return false;
 			}
-			var ThisPlatform = SC.StageTargetPlatform;
 
-			Logger.LogInformation("Creating Staging Manifest...");
-
-			ConfigHierarchy PlatformGameConfig = ConfigCache.ReadHierarchy(ConfigHierarchyType.Game, DirectoryReference.FromFile(Params.RawProjectPath), SC.StageTargetPlatform.IniPlatformType, SC.CustomConfig);
-
-			if (Params.HasIterateSharedCookedBuild)
+			if (SC.PlatformCookDir != null && SC.MetadataDir != null)
 			{
-				// can't do shared cooked builds with DLC that's madness!!
-				//check( Params.HasDLCName == false );
-
-				// stage all the previously staged files
-				SC.StageFiles(StagedFileType.NonUFS, DirectoryReference.Combine(SC.ProjectRoot, "Saved", "SharedIterativeBuild", SC.CookPlatform, "Staged"), StageFilesSearch.AllDirectories, StagedDirectoryReference.Root); // remap to the root directory
+				return true;
 			}
-			bool bCreatePluginManifest = false;
 			if (Params.HasDLCName)
 			{
 				// Making a plugin
@@ -909,7 +899,6 @@
 				{
 					DLCCookedSubDir = DLCRoot.MakeRelativeTo(SC.LocalRoot);
 				}
-
 				// Put all of the cooked dir into the staged dir
 				if (String.IsNullOrEmpty(Params.CookOutputDir))
 				{
@@ -923,10 +912,88 @@
 						SC.PlatformCookDir = DirectoryReference.Combine(SC.PlatformCookDir, SC.CookPlatform);
 					}
 				}
+				SC.MetadataDir = DirectoryReference.Combine(SC.PlatformCookDir, DLCCookedSubDir, "Metadata");
+			}
+			else
+			{
+				if (!Params.CookOnTheFly && !Params.SkipCookOnTheFly) // only stage the UFS files if we are not using cook on the fly
+				{
+					// Get the final output directory for cooked data
+					DirectoryReference CookOutputDir;
+					if (!String.IsNullOrEmpty(Params.CookOutputDir))
+					{
+						CookOutputDir = new DirectoryReference(Params.CookOutputDir);
+						if (!CookOutputDir.GetDirectoryName().Equals(SC.CookPlatform, StringComparison.InvariantCultureIgnoreCase))
+						{
+							CookOutputDir = DirectoryReference.Combine(CookOutputDir, SC.CookPlatform);
+						}
+					}
+					else if (Params.CookInEditor)
+					{
+						CookOutputDir = DirectoryReference.Combine(SC.ProjectRoot, "Saved", "EditorCooked", SC.CookPlatform);
+					}
+					else
+					{
+						CookOutputDir = DirectoryReference.Combine(SC.ProjectRoot, "Saved", "Cooked", SC.CookPlatform);
+					}
 
+					SC.PlatformCookDir = CookOutputDir;
+					SC.MetadataDir = DirectoryReference.Combine(CookOutputDir, SC.ShortProjectName, "Metadata");
+				}
+				else
+				{
+					// UE-58423
+					DirectoryReference CookOutputDir;
+					if (!String.IsNullOrEmpty(Params.CookOutputDir))
+					{
+						CookOutputDir = DirectoryReference.Combine(new DirectoryReference(Params.CookOutputDir), SC.CookPlatform);
+					}
+					else if (Params.CookInEditor)
+					{
+						CookOutputDir = DirectoryReference.Combine(SC.ProjectRoot, "Saved", "EditorCooked", SC.CookPlatform);
+					}
+					else
+					{
+						CookOutputDir = DirectoryReference.Combine(SC.ProjectRoot, "Saved", "Cooked", SC.CookPlatform);
+					}
+					SC.PlatformCookDir = CookOutputDir;
+					SC.MetadataDir = DirectoryReference.Combine(CookOutputDir, SC.ShortProjectName, "Metadata");
+				}
+			}
+			return true;
+		}
+
+		public static void CreateStagingManifest(ProjectParams Params, DeploymentContext SC)
+		{
+			if (!Params.Stage)
+			{
+				return;
+			}
+			var ThisPlatform = SC.StageTargetPlatform;
+
+			Logger.LogInformation("Creating Staging Manifest...");
+
+			if (!SetUpStagingSourceDirectories(Params, SC))
+			{
+				return;
+			}
+
+			ConfigHierarchy PlatformGameConfig = ConfigCache.ReadHierarchy(ConfigHierarchyType.Game, DirectoryReference.FromFile(Params.RawProjectPath), SC.StageTargetPlatform.IniPlatformType, SC.CustomConfig);
+
+			if (Params.HasIterateSharedCookedBuild)
+			{
+				// can't do shared cooked builds with DLC that's madness!!
+				//check( Params.HasDLCName == false );
+
+				// stage all the previously staged files
+				SC.StageFiles(StagedFileType.NonUFS, DirectoryReference.Combine(SC.ProjectRoot, "Saved", "SharedIterativeBuild", SC.CookPlatform, "Staged"), StageFilesSearch.AllDirectories, StagedDirectoryReference.Root); // remap to the root directory
+			}
+			bool bCreatePluginManifest = false;
+			if (Params.HasDLCName)
+			{
+				// Making a plugin
 				DirectoryReference PlatformEngineDir = DirectoryReference.Combine(SC.PlatformCookDir, "Engine");
 				DirectoryReference ProjectMetadataDir = DirectoryReference.Combine(SC.PlatformCookDir, SC.ShortProjectName, "Metadata");
-				SC.MetadataDir = DirectoryReference.Combine(SC.PlatformCookDir, DLCCookedSubDir, "Metadata");
 
 				// The .uplugin file is staged differently for different DLC
 				// The .uplugin file doesn't actually exist for mobile DLC
@@ -951,7 +1018,7 @@
 				}
 
 				// Put the binaries into the staged dir
-				DirectoryReference BinariesDir = DirectoryReference.Combine(DLCRoot, "Binaries");
+				DirectoryReference BinariesDir = DirectoryReference.Combine(Params.DLCFile.Directory, "Binaries");
 				if (DirectoryReference.Exists(BinariesDir))
 				{
 					SC.StageFiles(StagedFileType.NonUFS, BinariesDir, "libUE4-*.so", StageFilesSearch.AllDirectories);
@@ -961,7 +1028,7 @@
 				}
 
 				// Put the config files into the staged dir
-				DirectoryReference ConfigDir = DirectoryReference.Combine(DLCRoot, "Config");
+				DirectoryReference ConfigDir = DirectoryReference.Combine(Params.DLCFile.Directory, "Config");
 				if (DirectoryReference.Exists(ConfigDir))
 				{
 					SC.StageFiles(StagedFileType.UFS, ConfigDir, "*.ini", StageFilesSearch.AllDirectories);
@@ -969,7 +1036,7 @@
 
 				if (Params.DLCActLikePatch)
 				{
-					DirectoryReference DLCContent = DirectoryReference.Combine(DLCRoot, "Content");
+					DirectoryReference DLCContent = DirectoryReference.Combine(Params.DLCFile.Directory, "Content");
 					string RelativeDLCContentPath = DLCContent.MakeRelativeTo(SC.LocalRoot);
 					string RelativeRootContentPath = DirectoryReference.Combine(SC.ProjectRoot, "Content").MakeRelativeTo(SC.LocalRoot);
 
@@ -1022,7 +1089,7 @@
 					}
 				}
 
-				FileReference PluginSettingsFile = FileReference.Combine(DLCRoot, "Config", "PluginSettings.ini");
+				FileReference PluginSettingsFile = FileReference.Combine(Params.DLCFile.Directory, "Config", "PluginSettings.ini");
 				if (FileReference.Exists(PluginSettingsFile))
 				{
 					ConfigFile File = new ConfigFile(PluginSettingsFile);
@@ -1033,7 +1100,7 @@
 						{
 							if (Line.Key == "AdditionalNonUSFDirectories")
 							{
-								SC.StageFiles(StagedFileType.NonUFS, DirectoryReference.Combine(DLCRoot, Line.Value), "*.*", StageFilesSearch.AllDirectories);
+								SC.StageFiles(StagedFileType.NonUFS, DirectoryReference.Combine(Params.DLCFile.Directory, Line.Value), "*.*", StageFilesSearch.AllDirectories);
 							}
 						}
 					}
@@ -1047,7 +1114,7 @@
 							{
 								if (Line.Key == "AdditionalNonUSFDirectories")
 								{
-									SC.StageFiles(StagedFileType.NonUFS, DirectoryReference.Combine(DLCRoot, Line.Value), "*.*", StageFilesSearch.AllDirectories);
+									SC.StageFiles(StagedFileType.NonUFS, DirectoryReference.Combine(Params.DLCFile.Directory, Line.Value), "*.*", StageFilesSearch.AllDirectories);
 								}
 							}
 						}
@@ -1058,7 +1125,6 @@
 			}
 			else
 			{
-
 				if (!Params.IterateSharedBuildUsePrecompiledExe)
 				{
 					ThisPlatform.GetFilesToDeployOrStage(Params, SC);
@@ -1096,28 +1162,6 @@
 
 				if (!Params.CookOnTheFly && !Params.SkipCookOnTheFly) // only stage the UFS files if we are not using cook on the fly
 				{
-					// Get the final output directory for cooked data
-					DirectoryReference CookOutputDir;
-					if (!String.IsNullOrEmpty(Params.CookOutputDir))
-					{
-						CookOutputDir = new DirectoryReference(Params.CookOutputDir);
-						if (!CookOutputDir.GetDirectoryName().Equals(SC.CookPlatform, StringComparison.InvariantCultureIgnoreCase))
-						{
-							CookOutputDir = DirectoryReference.Combine(CookOutputDir, SC.CookPlatform);
-						}
-					}
-					else if (Params.CookInEditor)
-					{
-						CookOutputDir = DirectoryReference.Combine(SC.ProjectRoot, "Saved", "EditorCooked", SC.CookPlatform);
-					}
-					else
-					{
-						CookOutputDir = DirectoryReference.Combine(SC.ProjectRoot, "Saved", "Cooked", SC.CookPlatform);
-					}
-
-					SC.PlatformCookDir = CookOutputDir;
-					SC.MetadataDir = DirectoryReference.Combine(CookOutputDir, SC.ShortProjectName, "Metadata");
-
 					// Initialize internationalization preset.
 					string InternationalizationPreset = GetInternationalizationPreset(Params, PlatformGameConfig);
 
@@ -1337,9 +1381,9 @@
 					}
 
 					// Stage all the cooked data. Currently not filtering this by restricted folders, since we shouldn't mask invalid references by filtering them out.
-					if (DirectoryReference.Exists(CookOutputDir))
+					if (DirectoryReference.Exists(SC.PlatformCookDir))
 					{
-						List<FileReference> CookedFiles = DirectoryReference.EnumerateFiles(CookOutputDir, "*", SearchOption.AllDirectories).ToList();
+						List<FileReference> CookedFiles = DirectoryReference.EnumerateFiles(SC.PlatformCookDir, "*", SearchOption.AllDirectories).ToList();
 
 						// When cooking to Zen get the list of cooked package files from the manifest
 						LoadPackageStoreManifest(Params, SC);
@@ -1369,7 +1413,7 @@
 							// metallib files cannot *currently* be staged as UFS as the Metal API needs to mmap them from files on disk in order to function efficiently
 							if (!CookedFile.HasExtension(".json") && !CookedFile.HasExtension(".metallib") && !CookedFile.HasExtension(".utoc") && !CookedFile.HasExtension(".ucas"))
 							{
-								SC.StageFile(StagedFileType.UFS, CookedFile, new StagedFileReference(CookedFile.MakeRelativeTo(CookOutputDir)));
+								SC.StageFile(StagedFileType.UFS, CookedFile, new StagedFileReference(CookedFile.MakeRelativeTo(SC.PlatformCookDir)));
 							}
 						}
 					}
@@ -1390,7 +1434,7 @@
 						{
 							if (PlatformExports.IsPlatformAvailable(CrashReportPlatform))
 							{
-								Architecture = UnrealArchitectureConfig.ForPlatform(CrashReportPlatform).ActiveArchitectures(Params.RawProjectPath, TargetName:null);
+								Architecture = UnrealArchitectureConfig.ForPlatform(CrashReportPlatform).ActiveArchitectures(Params.RawProjectPath, TargetName: null);
 							}
 						}
 
@@ -1511,22 +1555,6 @@
 							}
 						}
 					}
-
-					// UE-58423
-					DirectoryReference CookOutputDir;
-					if (!String.IsNullOrEmpty(Params.CookOutputDir))
-					{
-						CookOutputDir = DirectoryReference.Combine(new DirectoryReference(Params.CookOutputDir), SC.CookPlatform);
-					}
-					else if (Params.CookInEditor)
-					{
-						CookOutputDir = DirectoryReference.Combine(SC.ProjectRoot, "Saved", "EditorCooked", SC.CookPlatform);
-					}
-					else
-					{
-						CookOutputDir = DirectoryReference.Combine(SC.ProjectRoot, "Saved", "Cooked", SC.CookPlatform);
-					}
-					SC.MetadataDir = DirectoryReference.Combine(CookOutputDir, SC.ShortProjectName, "Metadata");
 				}
 
 				StageCookerSupportFilesFromReceipt(SC);
@@ -1623,7 +1651,7 @@
 			{
 				SC.FilesToStage.NonUFSFiles = SC.FilesToStage.NonUFSFiles.ToDictionary(x => SC.StageTargetPlatform.Remap(x.Key), x => x.Value);
 			}
-				
+
 			if (SC.StageTargetPlatform.RemapFileType(StagedFileType.DebugNonUFS))
 			{
 				SC.FilesToStage.NonUFSDebugFiles = SC.FilesToStage.NonUFSDebugFiles.ToDictionary(x => SC.StageTargetPlatform.Remap(x.Key), x => x.Value);
diff --git a/Engine/Source/Programs/AutomationTool/Scripts/RunProjectCommand.Automation.cs b/Engine/Source/Programs/AutomationTool/Scripts/RunProjectCommand.Automation.cs
--- a/Engine/Source/Programs/AutomationTool/Scripts/RunProjectCommand.Automation.cs
+++ b/Engine/Source/Programs/AutomationTool/Scripts/RunProjectCommand.Automation.cs
@@ -15,6 +15,7 @@ using UnrealBuildTool;
 using EpicGames.Core;
 using UnrealBuildBase;
 using Microsoft.Extensions.Logging;
+using EpicGames.Serialization;
 
 namespace AutomationScripts
 {
@@ -567,6 +568,13 @@ namespace AutomationScripts
 			// Get client app name and command line.
 			string TempCmdLine = SC.ProjectArgForCommandLines + " ";
 			var PlatformName = TargetPlatform.ToString();
+
+			String FileHostCommandline = GetFileHostCommandline(Params, SC);
+			if (!string.IsNullOrEmpty(FileHostCommandline))
+			{
+				TempCmdLine += FileHostCommandline + " ";
+			}
+
 			if (Params.Cook || Params.CookOnTheFly)
 			{
 				List<FileReference> Exes = SC.StageTargetPlatform.GetExecutableNames(SC);
@@ -581,11 +589,6 @@ namespace AutomationScripts
 					TempCmdLine += Params.MapToRun + " ";
 				}
 
-				String FileHostCommandline = GetFileHostCommandline(Params, SC);
-				if (!string.IsNullOrEmpty(FileHostCommandline))
-				{
-					TempCmdLine += FileHostCommandline + " ";
-				}
 				if (Params.CookOnTheFly || Params.FileServer)
 				{
 					if (Params.CookOnTheFlyStreaming)
@@ -787,16 +790,71 @@ namespace AutomationScripts
 			return HostAddresses.ToList();
 		}
 
+		private static FileReference FindZenProjectStoreMarker(ProjectParams Params, DeploymentContext SC)
+		{
+			DirectoryReference ProjectStoreDir = null;
+			if (Params.Stage)
+			{
+				if (!SetUpStagingSourceDirectories(Params, SC))
+				{
+					return null;
+				}
+				ProjectStoreDir = SC.PlatformCookDir;
+			}
+			else if (Params.Deploy)
+			{
+				ProjectStoreDir = SC.StageDirectory;
+			}
+			if (ProjectStoreDir == null)
+			{
+				return null;
+			}
+			// Check for stage with zenstore without PAK?
+			FileReference PackageStoreManifestFile = FileReference.Combine(ProjectStoreDir, ".projectstore");
+			System.IO.FileInfo PackageStoreManifestFileInfo = PackageStoreManifestFile.ToFileInfo();
+			if (PackageStoreManifestFileInfo.Exists)
+			{
+				return PackageStoreManifestFile;
+			}
+			return null;
+		}
+
 		private static string GetFileHostCommandline(ProjectParams Params, DeploymentContext SC)
 		{
 			string FileHostParams = "";
-			if (!Params.CookOnTheFly && !Params.FileServer)
+			FileReference ZenStoreMarkerFile = FindZenProjectStoreMarker(Params, SC);
+			if (!Params.CookOnTheFly && !Params.FileServer && ZenStoreMarkerFile == null)
 			{
 				return FileHostParams;
 			}
 
-			List<string> HostAddresses = GetFileHostAddresses(SC);
-			SC.StageTargetPlatform.ModifyFileHostAddresses(HostAddresses);
+			string ProjectId = ProjectUtils.GetProjectPathId(SC.RawProjectPath);
+			ushort ZenHostPort = 1337;
+			List<string> HostAddresses = null;
+			if (ZenStoreMarkerFile != null)
+			{
+				byte[] ProjectStoreData = File.ReadAllBytes(ZenStoreMarkerFile.FullName);
+				CbObject ProjectStoreObject = new CbField(ProjectStoreData).AsObject();
+				CbObject ZenServerObject = ProjectStoreObject["zenserver"].AsObject();
+				if (ZenServerObject != CbObject.Empty)
+				{
+					ZenHostPort = ZenServerObject["hostport"].AsUInt16(ZenHostPort);
+					if (!ZenServerObject["islocalhost"].AsBool())
+					{
+						string HostName = ZenServerObject["hostname"].AsString();
+						if (!string.IsNullOrEmpty(HostName))
+						{
+							HostAddresses = new List<string> { HostName };
+						}
+					}
+					ProjectId = ZenServerObject["projectid"].AsString(ProjectId);
+				}
+			}
+			if (HostAddresses == null)
+			{
+				HostAddresses = GetFileHostAddresses(SC);
+				SC.StageTargetPlatform.ModifyFileHostAddresses(HostAddresses);
+			}
 
 			if (!IsNullOrEmpty(Params.Port))
 			{
@@ -821,9 +879,10 @@ namespace AutomationScripts
 			{
 				FileHostParams += "-filehostip=";
 			}
-			else if (Params.ZenStore)
+			else if (Params.ZenStore || ZenStoreMarkerFile != null)
 			{
-				FileHostParams += "-zenstoreproject=" + ProjectUtils.GetProjectPathId(SC.RawProjectPath) + " ";
+				FileHostParams += "-zenstoreproject=" + ProjectId + " ";
+				FileHostParams += "-zenstoreport=" + ZenHostPort + " ";
 				FileHostParams += "-zenstorehost=";
 			}
 			else
diff --git a/Engine/Source/Runtime/StorageServerClient/Private/StorageServerPlatformFile.cpp b/Engine/Source/Runtime/StorageServerClient/Private/StorageServerPlatformFile.cpp
--- a/Engine/Source/Runtime/StorageServerClient/Private/StorageServerPlatformFile.cpp
+++ b/Engine/Source/Runtime/StorageServerClient/Private/StorageServerPlatformFile.cpp
@@ -3,6 +3,7 @@
 #include "StorageServerPlatformFile.h"
 #include "StorageServerIoDispatcherBackend.h"
 #include "HAL/IPlatformFileModule.h"
+#include "HAL/FileManagerGeneric.h"
 #include "Misc/CommandLine.h"
 #include "Misc/Paths.h"
 #include "Misc/ScopeRWLock.h"
@@ -15,6 +16,7 @@
 #include "Misc/CoreDelegates.h"
 #include "Modules/ModuleManager.h"
 #include "CookOnTheFly.h"
+#include "Serialization/CompactBinarySerialization.h"
 
 DEFINE_LOG_CATEGORY_STATIC(LogStorageServerPlatformFile, Log, All);
 
@@ -247,6 +249,29 @@ FStorageServerPlatformFile::~FStorageServerPlatformFile()
 {
 }
 
+TUniquePtr<FArchive> FStorageServerPlatformFile::TryFindProjectStoreMarkerFile(IPlatformFile* Inner) const
+{
+	if (Inner == nullptr)
+	{
+		return nullptr;
+	}
+	FString RootDir = FPaths::ConvertRelativePathToFull(FPaths::RootDir());
+	FString CookedDir = FPaths::ConvertRelativePathToFull(FPaths::Combine(*FPaths::ProjectDir(), TEXT("Saved"), TEXT("Cooked"), FPlatformProperties::PlatformName()));
+	TArray<FString> PotentialProjectStorePaths;
+	PotentialProjectStorePaths.Add(RootDir);
+	PotentialProjectStorePaths.Add(CookedDir);
+
+	for (const FString& ProjectStorePath : PotentialProjectStorePaths)
+	{
+		FString ProjectMarkerPath = ProjectStorePath / TEXT(".projectstore");
+		if (IFileHandle* ProjectStoreMarkerHandle = Inner->OpenRead(*ProjectMarkerPath); ProjectStoreMarkerHandle != nullptr)
+		{
+			return TUniquePtr<FArchive>(new FArchiveFileReaderGeneric(ProjectStoreMarkerHandle, *ProjectMarkerPath, ProjectStoreMarkerHandle->Size()));
+		}
+	}
+	return nullptr;
+}
+
 bool FStorageServerPlatformFile::ShouldBeUsed(IPlatformFile* Inner, const TCHAR* CmdLine) const
 {
 #if WITH_COTF
@@ -259,6 +284,22 @@ bool FStorageServerPlatformFile::ShouldBeUsed(IPlatformFile* Inner, const TCHAR*
 		return true;
 	}
 #endif
+	TUniquePtr<FArchive> ProjectStoreMarkerReader = TryFindProjectStoreMarkerFile(Inner);
+	if (ProjectStoreMarkerReader != nullptr)
+	{
+		FCbObject ProjectStoreObject = LoadCompactBinary(*ProjectStoreMarkerReader).AsObject();
+
+		if (FCbFieldView ZenServerField = ProjectStoreObject["zenserver"])
+		{
+			if (FUtf8StringView HostName = ZenServerField["hostname"].AsString(); !HostName.IsEmpty())
+			{
+				HostAddrs.Add(FString(HostName));
+			}
+			HostPort = ZenServerField["hostport"].AsUInt16(HostPort);
+			UE_LOG(LogStorageServerPlatformFile, Display, TEXT("Found .projectstore"));
+		}
+	}
+
 	FString Host;
 	if (FParse::Value(FCommandLine::Get(), TEXT("-ZenStoreHost="), Host))
 	{
@@ -267,6 +308,7 @@ bool FStorageServerPlatformFile::ShouldBeUsed(IPlatformFile* Inner, const TCHAR*
 			HostAddrs.Add(Host);
 		}
 	}
+	FParse::Value(CmdLine, TEXT("-ZenStorePort="), HostPort);
 	return HostAddrs.Num() > 0;
 }
 
@@ -276,6 +318,19 @@ bool FStorageServerPlatformFile::Initialize(IPlatformFile* Inner, const TCHAR* C
 	if (HostAddrs.Num() > 0)
 	{
 		// Don't initialize the connection yet because we want to incorporate project file path information into the initialization.
+
+		TUniquePtr<FArchive> ProjectStoreMarkerReader = TryFindProjectStoreMarkerFile(Inner);
+		if (ProjectStoreMarkerReader != nullptr)
+		{
+			FCbObject ProjectStoreObject = LoadCompactBinary(*ProjectStoreMarkerReader).AsObject();
+
+			if (FCbFieldView ZenServerField = ProjectStoreObject["zenserver"])
+			{
+				ServerProject = FString(ZenServerField["projectid"].AsString());
+				ServerPlatform = FString(ZenServerField["oplogid"].AsString());
+			}
+		}
+	
 		FParse::Value(CmdLine, TEXT("-ZenStoreProject="), ServerProject);
 		FParse::Value(CmdLine, TEXT("-ZenStorePlatform="), ServerPlatform);
 		return true;
diff --git a/Engine/Source/Runtime/StorageServerClient/Private/StorageServerPlatformFile.h b/Engine/Source/Runtime/StorageServerClient/Private/StorageServerPlatformFile.h
--- a/Engine/Source/Runtime/StorageServerClient/Private/StorageServerPlatformFile.h
+++ b/Engine/Source/Runtime/StorageServerClient/Private/StorageServerPlatformFile.h
@@ -111,7 +111,8 @@ private:
 #if WITH_COTF
 	void OnCookOnTheFlyMessage(const UE::Cook::FCookOnTheFlyMessage& Message);
 #endif
-	
+	TUniquePtr<FArchive> TryFindProjectStoreMarkerFile(IPlatformFile* Inner) const;
+
 	IPlatformFile* LowerLevel = nullptr;
 	FStringView ServerEngineDirView = FStringView(TEXT("/{engine}/"));
 	FStringView ServerProjectDirView = FStringView(TEXT("/{project}/"));
