Updated CommandTask to propagate -uselocalbuildstorage to sub tasks

#RNX

diff --git a/Engine/Source/Programs/AutomationTool/BuildGraph/Tasks/CommandTask.cs b/Engine/Source/Programs/AutomationTool/BuildGraph/Tasks/CommandTask.cs
--- a/Engine/Source/Programs/AutomationTool/BuildGraph/Tasks/CommandTask.cs
+++ b/Engine/Source/Programs/AutomationTool/BuildGraph/Tasks/CommandTask.cs
@@ -104,6 +104,14 @@ namespace AutomationTool.Tasks
 					CommandLine.Append("-nosubmit ");
 				}
 			}
+			if (Parameters.Arguments == null || !Parameters.Arguments.CaseInsensitiveContains("-uselocalbuildstorage"))
+			{
+				if (GlobalCommandLine.UseLocalBuildStorage)
+				{
+					CommandLine.Append("-uselocalbuildstorage ");
+				}
+			}
+
 			CommandLine.Append("-NoCompile ");
 			CommandLine.Append(Parameters.Name);
 			if (!String.IsNullOrEmpty(Parameters.Arguments))
