- Put machine name into the default Turnkey autosdk client spec name to avoid conflicts between same user on multiple machines, and detect if whatever user chose is in use, and suggest another name with an increasing index appended

diff --git a/Engine/Source/Programs/AutomationTool/Turnkey/Providers/PerforceCopyProvider.cs b/Engine/Source/Programs/AutomationTool/Turnkey/Providers/PerforceCopyProvider.cs
--- a/Engine/Source/Programs/AutomationTool/Turnkey/Providers/PerforceCopyProvider.cs
+++ b/Engine/Source/Programs/AutomationTool/Turnkey/Providers/PerforceCopyProvider.cs
@@ -278,8 +278,31 @@ namespace Turnkey
 						return false;
 					}
 
-					// get clientspec name
-					string ClientName = TurnkeyUtils.ReadInput("Enter clientspec name:", string.Format("{0}_sdks", Username));
+					string ClientName = "";
+					int Index = 0;
+					string BaseClientName = $"{Username}_{Hostname}_sdks";
+					while (ClientName == "")
+					{
+						// get clientspec name
+						string DefaultName = BaseClientName + (Index == 0 ? "" : $"_{Index}");
+						string TestClientName = TurnkeyUtils.ReadInput("Enter clientspec name:", DefaultName);
+
+						// make sure the clientname doesn't already exist (unlikely with hostname in it, but just to be sure)
+						if (PerforceConnection.GetClientInfo(TestClientName, true) != null)
+						{
+							TurnkeyUtils.Log("Client {0} is already in use, please choose another name", TestClientName);
+							// if the user picked a name, use that as the new base when testing, and for appending index to
+							if (TestClientName != DefaultName)
+							{
+								BaseClientName = TestClientName;
+							}
+							Index++;
+						}
+						else
+						{
+							ClientName = TestClientName;
+						}
+					}
 
 					// get local pathname
 					// if we had UE_SDKS_ROOT already set, we use it as a default
