Disable OnDemand bundle generation

diff --git a/Engine/Source/Programs/AutomationTool/Scripts/CopyBuildToStagingDirectory.Automation.cs b/Engine/Source/Programs/AutomationTool/Scripts/CopyBuildToStagingDirectory.Automation.cs
--- a/Engine/Source/Programs/AutomationTool/Scripts/CopyBuildToStagingDirectory.Automation.cs
+++ b/Engine/Source/Programs/AutomationTool/Scripts/CopyBuildToStagingDirectory.Automation.cs
@@ -4185,7 +4185,9 @@ namespace AutomationScripts
 							StagingFileValue = StagingFileValue.Substring(SC.PakFileInternalRoot.Length);
 						}
 						LooseFiles.TryAdd(new StagedFileReference(StagingFileValue), new FileReference(StagingFile.Key));
-						//return;
+
+						// Temporarily disable aggregate chunks
+						return;
 					}
 
 					if (bExcludeFromPaks)
