NuGet license check - if the same dependency exist in multiple places ignore them

diff --git a/Engine/Source/Programs/AutomationTool/BuildGraph/Tasks/NuGetLicenseCheckTask.cs b/Engine/Source/Programs/AutomationTool/BuildGraph/Tasks/NuGetLicenseCheckTask.cs
--- a/Engine/Source/Programs/AutomationTool/BuildGraph/Tasks/NuGetLicenseCheckTask.cs
+++ b/Engine/Source/Programs/AutomationTool/BuildGraph/Tasks/NuGetLicenseCheckTask.cs
@@ -150,7 +150,7 @@ namespace AutomationTool.Tasks
 					PackageInfo Info = new PackageInfo();
 					Info.Name = Match.Groups[1].Value;
 					Info.Version = Match.Groups[2].Value;
-					Packages.Add($"{Info.Name}@{Info.Version}", Info);
+					Packages.TryAdd($"{Info.Name}@{Info.Version}", Info);
 				}
 			}
 
