P4WriteConfig: If an ignore file is found in the root, set P4IGNORE to the file name rather than the path

#rnx

diff --git a/Engine/Source/Programs/AutomationTool/AutomationUtils/P4WriteConfig.cs b/Engine/Source/Programs/AutomationTool/AutomationUtils/P4WriteConfig.cs
--- a/Engine/Source/Programs/AutomationTool/AutomationUtils/P4WriteConfig.cs
+++ b/Engine/Source/Programs/AutomationTool/AutomationUtils/P4WriteConfig.cs
@@ -69,6 +69,11 @@ namespace AutomationTool
 				{
 					IgnorePath = Path.Combine(Unreal.EngineDirectory.ToString(), "Extras", "Perforce", "p4ignore");
 				}
+				else
+				{
+					// If a file was found in the root, use the filename so will apply to any same-named file in the tree
+					IgnorePath = Path.GetFileName(IgnorePath);
+				}
 				P4Config.AppendLine($"P4IGNORE={IgnorePath}");
 			}
 
