Use /sdcard/Android/obb for OBBs for all build types

#android

diff --git a/Engine/Source/Programs/AutomationTool/Android/AndroidPlatform.Automation.cs b/Engine/Source/Programs/AutomationTool/Android/AndroidPlatform.Automation.cs
--- a/Engine/Source/Programs/AutomationTool/Android/AndroidPlatform.Automation.cs
+++ b/Engine/Source/Programs/AutomationTool/Android/AndroidPlatform.Automation.cs
@@ -1601,7 +1601,7 @@ public class AndroidPlatform : Platform
 		bool bNeedGrantStoragePermission = bRequireRuntimeStoragePermission && !bIsDistribution;
 
 		// We can't always push directly to Android/obb so uploads to Download then moves it
-		bool bDontMoveOBB = bUseAFS ? true : (bPackageDataInsideApk || !bIsDistribution);
+		bool bDontMoveOBB = bUseAFS ? true : bPackageDataInsideApk;
 
 		bool bHavePatch = (PatchName != "");
 		bool bHaveOverflow1 = (Overflow1Name != "");
@@ -1623,12 +1623,12 @@ public class AndroidPlatform : Platform
 
 			// If it is a distribution build, push to $STORAGE/Android/obb folder instead of $STORAGE/obb folder.
 			// Note that $STORAGE/Android/obb will be the folder that contains the obb if you download the app from playstore.
-			string OBBInstallCommand = bNoObbInstall ? "\t$ADB $DEVICE shell 'rm -r $EXTERNAL_STORAGE/" + DeviceObbName + "'" : "\t$ADB $DEVICE push " + Path.GetFileName(ObbName) + (bIsDistribution ? " " + TargetAndroidTemp : " $STORAGE/") + DeviceObbName;
-			string PatchInstallCommand = bNoPatchInstall ? "\t$ADB $DEVICE shell 'rm -r $EXTERNAL_STORAGE/" + DevicePatchName + "'" : "\t$ADB $DEVICE push " + Path.GetFileName(PatchName) + (bIsDistribution ? " " + TargetAndroidTemp : " $STORAGE/") + DevicePatchName;
-			string Overflow1InstallCommand = bNoOverflow1Install ? "\t$ADB $DEVICE shell 'rm -r $EXTERNAL_STORAGE/" + DeviceOverflow1Name + "'" : "\t$ADB $DEVICE push " + Path.GetFileName(Overflow1Name) + (bIsDistribution ? " " + TargetAndroidTemp : " $STORAGE/") + DeviceOverflow1Name;
-			string Overflow2InstallCommand = bNoOverflow2Install ? "\t$ADB $DEVICE shell 'rm -r $EXTERNAL_STORAGE/" + DeviceOverflow2Name + "'" : "\t$ADB $DEVICE push " + Path.GetFileName(Overflow2Name) + (bIsDistribution ? " " + TargetAndroidTemp : " $STORAGE/") + DeviceOverflow2Name;
-			string Overflow3InstallCommand = bNoOverflow3Install ? "\t$ADB $DEVICE shell 'rm -r $EXTERNAL_STORAGE/" + DeviceOverflow3Name + "'" : "\t$ADB $DEVICE push " + Path.GetFileName(Overflow3Name) + (bIsDistribution ? " " + TargetAndroidTemp : " $STORAGE/") + DeviceOverflow3Name;
-			string Overflow4InstallCommand = bNoOverflow4Install ? "\t$ADB $DEVICE shell 'rm -r $EXTERNAL_STORAGE/" + DeviceOverflow4Name + "'" : "\t$ADB $DEVICE push " + Path.GetFileName(Overflow4Name) + (bIsDistribution ? " " + TargetAndroidTemp : " $STORAGE/") + DeviceOverflow4Name;
+			string OBBInstallCommand = bNoObbInstall ? "\t$ADB $DEVICE shell 'rm -r $EXTERNAL_STORAGE/" + DeviceObbName + "'" : "\t$ADB $DEVICE push " + Path.GetFileName(ObbName) + " " + TargetAndroidTemp + DeviceObbName;
+			string PatchInstallCommand = bNoPatchInstall ? "\t$ADB $DEVICE shell 'rm -r $EXTERNAL_STORAGE/" + DevicePatchName + "'" : "\t$ADB $DEVICE push " + Path.GetFileName(PatchName) + TargetAndroidTemp + DevicePatchName;
+			string Overflow1InstallCommand = bNoOverflow1Install ? "\t$ADB $DEVICE shell 'rm -r $EXTERNAL_STORAGE/" + DeviceOverflow1Name + "'" : "\t$ADB $DEVICE push " + Path.GetFileName(Overflow1Name) + " " + TargetAndroidTemp + DeviceOverflow1Name;
+			string Overflow2InstallCommand = bNoOverflow2Install ? "\t$ADB $DEVICE shell 'rm -r $EXTERNAL_STORAGE/" + DeviceOverflow2Name + "'" : "\t$ADB $DEVICE push " + Path.GetFileName(Overflow2Name) + " " + TargetAndroidTemp + DeviceOverflow2Name;
+			string Overflow3InstallCommand = bNoOverflow3Install ? "\t$ADB $DEVICE shell 'rm -r $EXTERNAL_STORAGE/" + DeviceOverflow3Name + "'" : "\t$ADB $DEVICE push " + Path.GetFileName(Overflow3Name) + " " + TargetAndroidTemp + DeviceOverflow3Name;
+			string Overflow4InstallCommand = bNoOverflow4Install ? "\t$ADB $DEVICE shell 'rm -r $EXTERNAL_STORAGE/" + DeviceOverflow4Name + "'" : "\t$ADB $DEVICE push " + Path.GetFileName(Overflow4Name) + " " + TargetAndroidTemp + DeviceOverflow4Name;
 
 			if (bUseAFS)
 			{
@@ -1714,12 +1714,12 @@ public class AndroidPlatform : Platform
 			string APKInstallCommand = "%ADB% %DEVICE% install " + Path.GetFileName(ApkName);
 			string APKReinstallCommand = "";
 
-			string OBBInstallCommand = bNoObbInstall ? "%ADB% %DEVICE% shell rm -r %STORAGE%/" + DeviceObbName : "%ADB% %DEVICE% push " + Path.GetFileName(ObbName) + (bIsDistribution ? " " + TargetAndroidTemp : " %STORAGE%/") + DeviceObbName;
-			string PatchInstallCommand = bNoPatchInstall ? "%ADB% %DEVICE% shell rm -r %STORAGE%/" + DevicePatchName : "%ADB% %DEVICE% push " + Path.GetFileName(PatchName) + (bIsDistribution ? " " + TargetAndroidTemp : " %STORAGE%/") + DevicePatchName;
-			string Overflow1InstallCommand = bNoOverflow1Install ? "%ADB% %DEVICE% shell rm -r %STORAGE%/" + DeviceOverflow1Name : "%ADB% %DEVICE% push " + Path.GetFileName(Overflow1Name) + (bIsDistribution ? " " + TargetAndroidTemp : " %STORAGE%/") + DeviceOverflow1Name;
-			string Overflow2InstallCommand = bNoOverflow2Install ? "%ADB% %DEVICE% shell rm -r %STORAGE%/" + DeviceOverflow2Name : "%ADB% %DEVICE% push " + Path.GetFileName(Overflow2Name) + (bIsDistribution ? " " + TargetAndroidTemp : " %STORAGE%/") + DeviceOverflow2Name;
-			string Overflow3InstallCommand = bNoOverflow3Install ? "%ADB% %DEVICE% shell rm -r %STORAGE%/" + DeviceOverflow3Name : "%ADB% %DEVICE% push " + Path.GetFileName(Overflow3Name) + (bIsDistribution ? " " + TargetAndroidTemp : " %STORAGE%/") + DeviceOverflow3Name;
-			string Overflow4InstallCommand = bNoOverflow4Install ? "%ADB% %DEVICE% shell rm -r %STORAGE%/" + DeviceOverflow4Name : "%ADB% %DEVICE% push " + Path.GetFileName(Overflow4Name) + (bIsDistribution ? " " + TargetAndroidTemp : " %STORAGE%/") + DeviceOverflow4Name;
+			string OBBInstallCommand = bNoObbInstall ? "%ADB% %DEVICE% shell rm -r %STORAGE%/" + DeviceObbName : "%ADB% %DEVICE% push " + Path.GetFileName(ObbName) + " " + TargetAndroidTemp + DeviceObbName;
+			string PatchInstallCommand = bNoPatchInstall ? "%ADB% %DEVICE% shell rm -r %STORAGE%/" + DevicePatchName : "%ADB% %DEVICE% push " + Path.GetFileName(PatchName) + " " + TargetAndroidTemp + DevicePatchName;
+			string Overflow1InstallCommand = bNoOverflow1Install ? "%ADB% %DEVICE% shell rm -r %STORAGE%/" + DeviceOverflow1Name : "%ADB% %DEVICE% push " + Path.GetFileName(Overflow1Name) + " " + TargetAndroidTemp + DeviceOverflow1Name;
+			string Overflow2InstallCommand = bNoOverflow2Install ? "%ADB% %DEVICE% shell rm -r %STORAGE%/" + DeviceOverflow2Name : "%ADB% %DEVICE% push " + Path.GetFileName(Overflow2Name) + " " + TargetAndroidTemp + DeviceOverflow2Name;
+			string Overflow3InstallCommand = bNoOverflow3Install ? "%ADB% %DEVICE% shell rm -r %STORAGE%/" + DeviceOverflow3Name : "%ADB% %DEVICE% push " + Path.GetFileName(Overflow3Name) + " " + TargetAndroidTemp + DeviceOverflow3Name;
+			string Overflow4InstallCommand = bNoOverflow4Install ? "%ADB% %DEVICE% shell rm -r %STORAGE%/" + DeviceOverflow4Name : "%ADB% %DEVICE% push " + Path.GetFileName(Overflow4Name) + " " + TargetAndroidTemp + DeviceOverflow4Name;
 
 			if (bUseAFS)
 			{
