Only archive AAB if bundle enabled for Android

#android

diff --git a/Engine/Source/Programs/AutomationTool/Android/AndroidPlatform.Automation.cs b/Engine/Source/Programs/AutomationTool/Android/AndroidPlatform.Automation.cs
--- a/Engine/Source/Programs/AutomationTool/Android/AndroidPlatform.Automation.cs
+++ b/Engine/Source/Programs/AutomationTool/Android/AndroidPlatform.Automation.cs
@@ -1920,6 +1920,7 @@ public class AndroidPlatform : Platform
 		UnrealArchitectures Architectures = GetDeploymentArchitectures(Params, SC);
 		bool bMakeSeparateApks = UnrealBuildTool.AndroidExports.ShouldMakeSeparateApks();
 		bool bPackageDataInsideApk = UnrealBuildTool.AndroidExports.CreateDeploymentHandler(Params.RawProjectPath, Params.ForcePackageData).GetPackageDataInsideApk();
+		bool bBundleEnabled = GetEnableBundle(SC);
 
 		bool bAFSEnablePlugin;
 		string AFSToken;
@@ -1951,53 +1952,56 @@ public class AndroidPlatform : Platform
 			bool bHaveAAB = false;
 			bool bHaveUniversal = false;
 
-			// copy optional app bundle if exists
 			string AppBundleName = Path.Combine(APKDirectory, APKNameWithoutExtension + ".aab");
-			if (FileExists(AppBundleName))
-			{
-				bHaveAAB = true;
-				SC.ArchiveFiles(APKDirectory, Path.GetFileName(AppBundleName));
-			}
-			else
+			string APKSName = Path.Combine(APKDirectory, APKNameWithoutExtension + ".apks");
+			string UniversalApkName = Path.Combine(APKDirectory, APKNameWithoutExtension + "_universal.apk");
+			if (bBundleEnabled)
 			{
-				AppBundleName = Path.Combine(APKDirectory, APKBareNameWithoutExtension + ".aab");
+				// copy optional app bundle if exists
 				if (FileExists(AppBundleName))
 				{
 					bHaveAAB = true;
 					SC.ArchiveFiles(APKDirectory, Path.GetFileName(AppBundleName));
 				}
-			}
+				else
+				{
+					AppBundleName = Path.Combine(APKDirectory, APKBareNameWithoutExtension + ".aab");
+					if (FileExists(AppBundleName))
+					{
+						bHaveAAB = true;
+						SC.ArchiveFiles(APKDirectory, Path.GetFileName(AppBundleName));
+					}
+				}
 
-			// copy optional apks (zip of split apks) if exists
-			string APKSName = Path.Combine(APKDirectory, APKNameWithoutExtension + ".apks");
-			if (FileExists(APKSName))
-			{
-				SC.ArchiveFiles(APKDirectory, Path.GetFileName(APKSName));
-			}
-			else
-			{
-				APKSName = Path.Combine(APKDirectory, APKBareNameWithoutExtension + ".apks");
+				// copy optional apks (zip of split apks) if exists
 				if (FileExists(APKSName))
 				{
 					SC.ArchiveFiles(APKDirectory, Path.GetFileName(APKSName));
 				}
-			}
+				else
+				{
+					APKSName = Path.Combine(APKDirectory, APKBareNameWithoutExtension + ".apks");
+					if (FileExists(APKSName))
+					{
+						SC.ArchiveFiles(APKDirectory, Path.GetFileName(APKSName));
+					}
+				}
 
-			// copy optional universal apk if exists
-			string UniversalApkName = Path.Combine(APKDirectory, APKNameWithoutExtension + "_universal.apk");
-			if (FileExists(UniversalApkName))
-			{
-				bHaveUniversal = true;
-				SC.ArchiveFiles(APKDirectory, Path.GetFileName(UniversalApkName));
-			}
-			else
-			{
-				UniversalApkName = Path.Combine(APKDirectory, APKBareNameWithoutExtension + "_universal.apk");
+				// copy optional universal apk if exists
 				if (FileExists(UniversalApkName))
 				{
 					bHaveUniversal = true;
 					SC.ArchiveFiles(APKDirectory, Path.GetFileName(UniversalApkName));
 				}
+				else
+				{
+					UniversalApkName = Path.Combine(APKDirectory, APKBareNameWithoutExtension + "_universal.apk");
+					if (FileExists(UniversalApkName))
+					{
+						bHaveUniversal = true;
+						SC.ArchiveFiles(APKDirectory, Path.GetFileName(UniversalApkName));
+					}
+				}
 			}
 
 			// copy optional AFSProject apk if exists
