- Fixed QuickLaunch for IOS working with Modern xcode
- No longer try to grab the output from IOS with QuickLaunch as idevicedebug is killing the app too quickly
- The modern Staged directory can now be shared between Development/Shipping/etc, for grabbing any non-Framework .dylibs (like steam)
- Removed the extra --generate-entitlements-der option now that Xcode also always adds it

diff --git a/Engine/Platforms/Hololens/Source/Programs/AutomationTool/HoloLensPlatform.Automation.cs b/Engine/Platforms/Hololens/Source/Programs/AutomationTool/HoloLensPlatform.Automation.cs
--- a/Engine/Platforms/Hololens/Source/Programs/AutomationTool/HoloLensPlatform.Automation.cs
+++ b/Engine/Platforms/Hololens/Source/Programs/AutomationTool/HoloLensPlatform.Automation.cs
@@ -1268,7 +1268,7 @@ namespace HoloLens.Automation
                 return;
             }
 
-            bool bRequiresPackage = Params.Package || SC.StageTargetPlatform.RequiresPackageToDeploy;
+            bool bRequiresPackage = Params.Package || SC.StageTargetPlatform.RequiresPackageToDeploy(Params);
 			string Name;
 			string Publisher;
 			GetPackageInfo(Params, out Name, out Publisher);
@@ -1383,7 +1383,7 @@ namespace HoloLens.Automation
 			string Publisher;
 			GetPackageInfo(Params, out Name, out Publisher);
 
-			if (Params.Package || SC.StageTargetPlatform.RequiresPackageToDeploy)
+			if (Params.Package || SC.StageTargetPlatform.RequiresPackageToDeploy(Params))
 			{
 				DefaultDevicePortalConnection conn = new DefaultDevicePortalConnection(DeviceAddress, Params.DeviceUsername, Params.DevicePassword);
 				DevicePortal portal = new DevicePortal(conn);
diff --git a/Engine/Source/Programs/AutomationTool/Apple/ApplePlatform.Automation.cs b/Engine/Source/Programs/AutomationTool/Apple/ApplePlatform.Automation.cs
--- a/Engine/Source/Programs/AutomationTool/Apple/ApplePlatform.Automation.cs
+++ b/Engine/Source/Programs/AutomationTool/Apple/ApplePlatform.Automation.cs
@@ -71,11 +71,7 @@ public abstract class ApplePlatform : Platform
 
 	public override void PostStagingFileCopy(ProjectParams Params, DeploymentContext SC)
 	{
-		ConfigHierarchy Ini = ConfigCache.ReadHierarchy(ConfigHierarchyType.Engine, Params.RawProjectPath.Directory!, SC.StageTargetPlatform.PlatformType);
-		bool bUseModernXcode;
-		Ini.TryGetValue("/Script/MacTargetPlatform.XcodeProjectSettings", "bUseModernXcode", out bUseModernXcode);
-
-		if (bUseModernXcode)
+		if (MacExports.UseModernXcode(Params.RawProjectPath))
 		{
 			// modern does all the distribution via .xcarchive
 			bool bForDistribution = false;
diff --git a/Engine/Source/Programs/AutomationTool/AutomationUtils/Platform.cs b/Engine/Source/Programs/AutomationTool/AutomationUtils/Platform.cs
--- a/Engine/Source/Programs/AutomationTool/AutomationUtils/Platform.cs
+++ b/Engine/Source/Programs/AutomationTool/AutomationUtils/Platform.cs
@@ -735,9 +735,9 @@ namespace AutomationTool
 		/// <summary>
 		///  Returns whether the platform requires a package to deploy to a device
 		/// </summary>
-		public virtual bool RequiresPackageToDeploy
+		public virtual bool RequiresPackageToDeploy(ProjectParams Params)
 		{
-			get { return false; }
+			return false;
 		}
 
 		/// <summary>
diff --git a/Engine/Source/Programs/AutomationTool/IOS/IOSPlatform.Automation.cs b/Engine/Source/Programs/AutomationTool/IOS/IOSPlatform.Automation.cs
--- a/Engine/Source/Programs/AutomationTool/IOS/IOSPlatform.Automation.cs
+++ b/Engine/Source/Programs/AutomationTool/IOS/IOSPlatform.Automation.cs
@@ -791,16 +791,14 @@ public class IOSPlatform : ApplePlatform
 
 		var TargetConfiguration = SC.StageTargetConfigurations[0];
 
-		ConfigHierarchy Ini = ConfigCache.ReadHierarchy(ConfigHierarchyType.Engine, Params.RawProjectPath.Directory!, UnrealTargetPlatform.IOS);
-		bool bUseModernXcode;
-		Ini.TryGetValue("/Script/MacTargetPlatform.XcodeProjectSettings", "bUseModernXcode", out bUseModernXcode);
-
 		string MobileProvision;
 		string SigningCertificate;
 		string TeamUUID;
 		bool bAutomaticSigning;
-		if (bUseModernXcode)
+		if (MacExports.UseModernXcode(Params.RawProjectPath))
 		{
+			ConfigHierarchy Ini = ConfigCache.ReadHierarchy(ConfigHierarchyType.Engine, Params.RawProjectPath.Directory!, UnrealTargetPlatform.IOS);
+
 			Ini.TryGetValue("/Script/MacTargetPlatform.XcodeProjectSettings", "bUseModernCodeSigning", out bAutomaticSigning);
 			Ini.TryGetValue("/Script/MacTargetPlatform.XcodeProjectSettings", "ModernSigningTeam", out TeamUUID);
 			//			Ini.TryGetValue("/Script/MacTargetPlatform.XcodeProjectSettings", "ModernSigningPrefix", out SigningPrefix);
@@ -813,7 +811,7 @@ public class IOSPlatform : ApplePlatform
 		}
 		else
 		{
-		GetProvisioningData(Params.RawProjectPath, Params.Distribution, out MobileProvision, out SigningCertificate, out TeamUUID, out bAutomaticSigning);
+			GetProvisioningData(Params.RawProjectPath, Params.Distribution, out MobileProvision, out SigningCertificate, out TeamUUID, out bAutomaticSigning);
 		}
 
 		//@TODO: We should be able to use this code on both platforms, when the following issues are sorted:
@@ -1113,14 +1111,7 @@ public class IOSPlatform : ApplePlatform
 
 	private void CodeSign(string BaseDirectory, string GameName, FileReference RawProjectPath, UnrealTargetConfiguration TargetConfig, string LocalRoot, string ProjectName, string ProjectDirectory, bool IsCode, bool Distribution, string Provision, string Certificate, string Team, bool bAutomaticSigning, string SchemeName, string SchemeConfiguration)
 	{
-		bool bUseModernXcode = false;
-		if (OperatingSystem.IsMacOS())
-		{
-			ConfigHierarchy Ini = ConfigCache.ReadHierarchy(ConfigHierarchyType.Engine, RawProjectPath.Directory, UnrealTargetPlatform.IOS);
-			Ini.TryGetValue("/Script/MacTargetPlatform.XcodeProjectSettings", "bUseModernXcode", out bUseModernXcode);
-		}
-
-		if (bUseModernXcode)
+		if (MacExports.UseModernXcode(RawProjectPath))
 		{
 			// we need an xcode project to continue
 			DirectoryReference GeneratedProject = null;
@@ -1132,7 +1123,7 @@ public class IOSPlatform : ApplePlatform
 				throw new AutomationException("iOS couldn't find the appropriate Xcode Project for " + RawProjectPath);
 			}
 
-			int ReturnCode = IOSExports.FinalizeAppWithXcode(GeneratedProject, TargetPlatformType, SchemeName ?? GameName, SchemeConfiguration ?? TargetConfig.ToString(), Provision, Certificate, Team, "", bAutomaticSigning, Distribution, bUseModernXcode, Logger);
+			int ReturnCode = IOSExports.FinalizeAppWithXcode(GeneratedProject, TargetPlatformType, SchemeName ?? GameName, SchemeConfiguration ?? TargetConfig.ToString(), Provision, Certificate, Team, "", bAutomaticSigning, Distribution, bUseModernXcode:true, Logger);
 			if (ReturnCode != 0)
 			{
 				throw new AutomationException(ExitCode.Error_FailedToCodeSign, "CodeSign Failed");
@@ -1868,20 +1859,39 @@ public class IOSPlatform : ApplePlatform
 		{
 			throw new AutomationException("iOS cannot deploy a package made for distribution.");
 		}
-		var TargetConfiguration = SC.StageTargetConfigurations[0];
-		var ProjectIPA = MakeIPAFileName(TargetConfiguration, Params, SC, true);
-		var StagedIPA = SC.StageDirectory + "\\" + Path.GetFileName(ProjectIPA);
 
-		// verify the .ipa exists
-		if (!FileExists(StagedIPA))
+		// modern mode simply deploys the staged .app which has been fully completed
+		string AppToDeploy;
+		if (MacExports.UseModernXcode(Params.RawProjectPath))
 		{
-			StagedIPA = ProjectIPA;
+			AppToDeploy = FileReference.Combine(SC.StageDirectory, SC.StageExecutables[0] + ".app").FullName;
+
+			// @todo: deal with iterative deploy - we have the uncombined .app in the Staged dir inin Binaries/IOS
+			if (Params.IterativeDeploy)
+			{
+				throw new AutomationException("Iterative deploy is currently not supported with modern, but it shojld be straightforward");
+			}
+		}
+		else
+		{
+			UnrealTargetConfiguration TargetConfiguration = SC.StageTargetConfigurations[0];
+			string ProjectIPA = MakeIPAFileName(TargetConfiguration, Params, SC, true);
+			string StagedIPA = SC.StageDirectory + "\\" + Path.GetFileName(ProjectIPA);
+
+			// verify the .ipa exists
 			if (!FileExists(StagedIPA))
 			{
-				throw new AutomationException("DEPLOY FAILED - {0} was not found", ProjectIPA);
+				StagedIPA = ProjectIPA;
+				if (!FileExists(StagedIPA))
+				{
+					throw new AutomationException("DEPLOY FAILED - {0} was not found", ProjectIPA);
+				}
 			}
+
+			AppToDeploy = StagedIPA;
 		}
 
+
 		// if iterative deploy, determine the file delta
 		string BundleIdentifier = "";
 		bool bNeedsIPA = true;
@@ -1906,13 +1916,10 @@ public class IOSPlatform : ApplePlatform
 			}
 		}
 
-		// Add a commandline for this deploy, if the config allows it.
-		string AdditionalCommandline = (Params.FileServer || Params.CookOnTheFly) ? "" : (" -additionalcommandline " + "\"" + Params.RunCommandline + "\"");
-
 		// deploy the .ipa
 		var DeviceInstaller = GetPathToLibiMobileDeviceTool("ideviceinstaller");
 
-		string LibimobileDeviceArguments =  "-u " + Params.DeviceNames[0] + " -i " + "\"" +  Path.GetFullPath(StagedIPA) + "\"";
+		string LibimobileDeviceArguments =  "-u " + Params.DeviceNames[0] + " -i " + "\"" +  Path.GetFullPath(AppToDeploy) + "\"";
 		LibimobileDeviceArguments = GetLibimobileDeviceNetworkedArgument(LibimobileDeviceArguments, Params.DeviceNames[0]);
 
 		// check for it in the stage directory
@@ -1994,10 +2001,30 @@ public class IOSPlatform : ApplePlatform
 				throw new AutomationException("Can only run on a single specified device, but {0} were specified", Params.Devices.Count);
 			}
 
+			string PlistFile;
 			string BundleIdentifier = "";
-			if (File.Exists(Params.BaseStageDirectory + "/" + PlatformName + "/Info.plist"))
+
+			// Params.StageDirectory will be something like /.../Saved/StagedBuilds
+			// ClientApp will be something like /.../Saved/StagedBuilds/IOSClient/QAGame/Binaries/IOS/QAGameClient
+			// so we want the first directory in the CLientApp after Params.StageDirectory	
+			string TargetStageDir = ClientApp.Substring(0, ClientApp.IndexOf('/', Params.BaseStageDirectory.Length + 1));
+			if (MacExports.UseModernXcode(Params.RawProjectPath))
+			{
+				// modern mode runs from the staged dir
+				PlistFile = Path.Combine(TargetStageDir, Path.GetFileNameWithoutExtension(ClientApp) + ".app", "Info.plist");
+			}
+			else
 			{
-				string Contents = File.ReadAllText(Params.BaseStageDirectory + "/" + PlatformName + "/Info.plist");
+				PlistFile = Path.Combine(TargetStageDir, "Info.plist");
+			}
+
+			Logger.LogWarning("LOOKING IN PLIST {PListfile}", PlistFile);
+			Logger.LogWarning("ClientApp is {CLient}", ClientApp);
+			Logger.LogWarning("BaseStageDir is is {stage}", Params.BaseStageDirectory);
+
+			if (File.Exists(PlistFile))
+			{
+				string Contents = File.ReadAllText(PlistFile);
 				int Pos = Contents.IndexOf("CFBundleIdentifier");
 				Pos = Contents.IndexOf("<string>", Pos) + 8;
 				int EndPos = Contents.IndexOf("</string>", Pos);
@@ -2006,10 +2033,11 @@ public class IOSPlatform : ApplePlatform
 
 			string Program = GetPathToLibiMobileDeviceTool("idevicedebug"); ;
 			string Arguments = " -u '" + Params.DeviceNames[0] + "'";
+			Arguments += " --detach";
 			Arguments = GetLibimobileDeviceNetworkedArgument(Arguments, Params.DeviceNames[0]);
 			Arguments += " run '" + BundleIdentifier + "'";
 
-			IProcessResult ClientProcess = Run(Program, Arguments, null, ClientRunFlags | ERunOptions.NoWaitForExit);
+			IProcessResult ClientProcess = Run(Program, Arguments, null, ClientRunFlags);
 			if (ClientProcess.ExitCode == -1)
 			{
 				Console.WriteLine("The application {0} has been installed on the device {1} but it cannot be launched automatically because the device does not contain the required developer software. You can launch {0} the manually by clicking its icon on the device.", BundleIdentifier, Params.DeviceNames[0]);
@@ -2019,7 +2047,8 @@ public class IOSPlatform : ApplePlatform
 				return Result;
 
 			}
-			return new IOSClientProcess(ClientProcess, Params.DeviceNames[0]);
+			return ClientProcess;
+
 		}
 		else
 		{
@@ -2029,27 +2058,6 @@ public class IOSPlatform : ApplePlatform
 		}
 	}
 
-	public override void PostRunClient(IProcessResult Result, ProjectParams Params)
-	{
-		if (UnrealBuildTool.BuildHostPlatform.Current.Platform == UnrealTargetPlatform.Mac)
-		{
-			string LaunchTracePath = Params.BaseStageDirectory + "/" + PlatformName + "/launch.trace";
-			Console.WriteLine("Deleting " + LaunchTracePath);
-			if (Directory.Exists(LaunchTracePath))
-			{
-				Directory.Delete(LaunchTracePath, true);
-			}
-
-			switch (Result.ExitCode)
-			{
-				case 253:
-					throw new AutomationException(ExitCode.Error_DeviceNotSetupForDevelopment, "Launch Failure");
-				case 255:
-					throw new AutomationException(ExitCode.Error_DeviceOSNewerThanSDK, "Launch Failure");
-			}
-		}
-	}
-
 	private static int GetChunkCount(ProjectParams Params, DeploymentContext SC)
 	{
 		var ChunkListFilename = GetChunkPakManifestListFilename(Params, SC);
@@ -2277,9 +2285,9 @@ public class IOSPlatform : ApplePlatform
 		base.PostStagingFileCopy(Params, SC);
 	}
 
-	public override bool RequiresPackageToDeploy
+	public override bool RequiresPackageToDeploy(ProjectParams Params)
 	{
-		get { return true; }
+		return !MacExports.UseModernXcode(Params.RawProjectPath);
 	}
 
 	public override HashSet<StagedFileReference> GetFilesForCRCCheck()
diff --git a/Engine/Source/Programs/AutomationTool/Scripts/PackageCommand.Automation.cs b/Engine/Source/Programs/AutomationTool/Scripts/PackageCommand.Automation.cs
--- a/Engine/Source/Programs/AutomationTool/Scripts/PackageCommand.Automation.cs
+++ b/Engine/Source/Programs/AutomationTool/Scripts/PackageCommand.Automation.cs
@@ -30,7 +30,7 @@ namespace AutomationScripts
 				bool bShouldPackage = false; 
 				foreach (var SC in DeployContextList)
 				{
-					if (Params.Package || (SC.StageTargetPlatform.RequiresPackageToDeploy && Params.Deploy))
+					if (Params.Package || (SC.StageTargetPlatform.RequiresPackageToDeploy(Params) && Params.Deploy))
 					{
 						bShouldPackage = true;
 						break;
@@ -44,7 +44,7 @@ namespace AutomationScripts
 
 					foreach (var SC in DeployContextList)
 					{
-						if (Params.Package || (SC.StageTargetPlatform.RequiresPackageToDeploy && Params.Deploy))
+						if (Params.Package || (SC.StageTargetPlatform.RequiresPackageToDeploy(Params) && Params.Deploy))
 						{
 							SC.StageTargetPlatform.Package(Params, SC, WorkingCL);
 						}
diff --git a/Engine/Source/Programs/UnrealBuildTool/Platform/Mac/MacExports.cs b/Engine/Source/Programs/UnrealBuildTool/Platform/Mac/MacExports.cs
--- a/Engine/Source/Programs/UnrealBuildTool/Platform/Mac/MacExports.cs
+++ b/Engine/Source/Programs/UnrealBuildTool/Platform/Mac/MacExports.cs
@@ -74,6 +74,23 @@ namespace UnrealBuildTool
 			}
 		}
 
+		/// <summary>
+		/// Is the current project using modern xcode?
+		/// </summary>
+		/// <param name="ProjectFile"></param>
+		/// <returns></returns>
+		public static bool UseModernXcode(FileReference? ProjectFile)
+		{
+			// Modern Xcode mode does this now
+			bool bUseModernXcode = false;
+			if (OperatingSystem.IsMacOS())
+			{
+				ConfigHierarchy Ini = ConfigCache.ReadHierarchy(ConfigHierarchyType.Engine, ProjectFile?.Directory, UnrealTargetPlatform.Mac);
+				Ini.TryGetValue("/Script/MacTargetPlatform.XcodeProjectSettings", "bUseModernXcode", out bUseModernXcode);
+			}
+			return bUseModernXcode;
+		}
+
 		/// <summary>
 		/// Strips symbols from a file
 		/// </summary>
diff --git a/Engine/Source/Programs/UnrealBuildTool/ProjectFiles/Xcode/XcodeProject.cs b/Engine/Source/Programs/UnrealBuildTool/ProjectFiles/Xcode/XcodeProject.cs
--- a/Engine/Source/Programs/UnrealBuildTool/ProjectFiles/Xcode/XcodeProject.cs
+++ b/Engine/Source/Programs/UnrealBuildTool/ProjectFiles/Xcode/XcodeProject.cs
@@ -1298,8 +1298,9 @@ namespace UnrealBuildTool.XcodeProjectXcconfig
 				{
 					CopyScript.AddRange(new string[]
 					{
-						$"echo \\\"Syncing ${{STAGED_DIR}}/{EngineOrProject}/Binaries/Mac/${{CONTENTS_FOLDER_PATH}}/ to ${{CONFIGURATION_BUILD_DIR}}/${{CONTENTS_FOLDER_PATH}}/\\\"",
-						$"rsync -av --include='*/' --include='*.dylib' --exclude='*' \\\"${{STAGED_DIR}}/{EngineOrProject}/Binaries/Mac/${{CONTENTS_FOLDER_PATH}}/MacOS/\\\" \\\"${{CONFIGURATION_BUILD_DIR}}/${{CONTENTS_FOLDER_PATH}}/MacOS\\\"",
+						// copy from whatever .app is in Binaries/Mac - this is so we can make a Shipping app from a Development staged directory
+						$"echo \\\"Syncing `echo ${{STAGED_DIR}}/{EngineOrProject}/Binaries/Mac/*.app`/${{BUNDLE_CONTENTS_FOLDER_PATH}}/ to ${{CONFIGURATION_BUILD_DIR}}/${{CONTENTS_FOLDER_PATH}}/\\\"",
+						$"rsync -av --include='*/' --include='*.dylib' --exclude='*' \\\"`echo ${{STAGED_DIR}}/{EngineOrProject}/Binaries/Mac/*.app`/${{BUNDLE_CONTENTS_FOLDER_PATH}}MacOS/\\\" \\\"${{CONFIGURATION_BUILD_DIR}}/${{CONTENTS_FOLDER_PATH}}/MacOS\\\"",
 					});
 				}
 			}
@@ -1607,8 +1608,6 @@ namespace UnrealBuildTool.XcodeProjectXcconfig
 
 			Xcconfig.AppendLine("");
 			Xcconfig.AppendLine("// Constant settings (same for all platforms and targets)");
-			// #jira UE-143619: Pre Monterey macOS requires this option for a packaged app to run on iOS15 due to new code signature format. Could be removed once Monterey is miniuS.
-			Xcconfig.AppendLine("OTHER_CODE_SIGN_FLAGS = --generate-entitlement-der");
 			Xcconfig.AppendLine("INFOPLIST_OUTPUT_FORMAT = xml");
 			Xcconfig.AppendLine("COMBINE_HIDPI_IMAGES = YES");
 			Xcconfig.AppendLine("USE_HEADERMAP = NO");
diff --git a/Engine/Source/Programs/UnrealBuildTool/ToolChain/AppleToolChain.cs b/Engine/Source/Programs/UnrealBuildTool/ToolChain/AppleToolChain.cs
--- a/Engine/Source/Programs/UnrealBuildTool/ToolChain/AppleToolChain.cs
+++ b/Engine/Source/Programs/UnrealBuildTool/ToolChain/AppleToolChain.cs
@@ -180,11 +180,7 @@ namespace UnrealBuildTool
 
 		protected static bool UseModernXcode(FileReference? ProjectFile)
 		{
-			// Modern Xcode mode does this now
-			bool _bUseModernXcode = false;
-			ConfigHierarchy Ini = ConfigCache.ReadHierarchy(ConfigHierarchyType.Engine, ProjectFile?.Directory, UnrealTargetPlatform.Mac);
-			Ini.TryGetValue("/Script/MacTargetPlatform.XcodeProjectSettings", "bUseModernXcode", out _bUseModernXcode);
-			return _bUseModernXcode;
+			return MacExports.UseModernXcode(ProjectFile);
 		}
 
 		protected bool bUseModernXcode => UseModernXcode(ProjectFile);
