Gauntlet - Use the most appropriate Role to represent UnrealTestNode context

#rnx

diff --git a/Engine/Source/Programs/AutomationTool/Gauntlet/Unreal/Base/Gauntlet.UnrealTestNode.cs b/Engine/Source/Programs/AutomationTool/Gauntlet/Unreal/Base/Gauntlet.UnrealTestNode.cs
--- a/Engine/Source/Programs/AutomationTool/Gauntlet/Unreal/Base/Gauntlet.UnrealTestNode.cs
+++ b/Engine/Source/Programs/AutomationTool/Gauntlet/Unreal/Base/Gauntlet.UnrealTestNode.cs
@@ -423,7 +423,10 @@ namespace Gauntlet
 				return Name;
 			}
 
-			return string.Format("{0} ({1})", Name, Context);
+			var Config = GetConfiguration();
+			string RoleContext = (Config is UnrealTestConfiguration) ? Context.GetRoleContext(Config.GetMainRequiredRole().Type).ToString() : Context.ToString();
+
+			return string.Format("{0} ({1})", Name, RoleContext);
 		}
 
 		/// <summary>
