- Added a couple functions to CommandUtils for archive to handle copying symlinks properly (with a mode enum to select how to handle symlinks: Follow, Retain, Ignore - with defaults based on platform to be similar to how existing code was setup)
  - FindFilesAndSymlinks
  - CopyFileOrSymlink

diff --git a/Engine/Source/Programs/AutomationTool/AutomationUtils/CommandUtils.cs b/Engine/Source/Programs/AutomationTool/AutomationUtils/CommandUtils.cs
--- a/Engine/Source/Programs/AutomationTool/AutomationUtils/CommandUtils.cs
+++ b/Engine/Source/Programs/AutomationTool/AutomationUtils/CommandUtils.cs
@@ -1897,6 +1897,104 @@ namespace AutomationTool
 			}
 		}
 
+		public enum SymlinkMode
+		{
+			Ignore,
+			Follow,
+			Retain,
+		}
+
+		// detemine how the OS wants to handle Symlinks when archiving, by default
+		public static SymlinkMode DefaultSymlinkMode
+		{
+			get
+			{
+				return OperatingSystem.IsMacOS() ? SymlinkMode.Retain :
+				(OperatingSystem.IsLinux() ? SymlinkMode.Ignore :
+				SymlinkMode.Follow);
+			}
+		}
+
+		/// <summary>
+		/// Assumes directory has already been checked to exist, and that it is not a symlink itself
+		/// The SymlinkMode here refers to directory links, as the file symlinks will be dealt with during the copy
+		/// </summary>
+		public static void FindFilesAndSymlinks(string RootPath, string Wildcard, bool bRecursive, SymlinkMode DirSymlinkMode, List<string> FoundFiles)
+		{
+			// default behaviour can just do standard findfiles
+			if (DirSymlinkMode == SymlinkMode.Follow)
+			{
+				FoundFiles.AddRange(Directory.GetFiles(RootPath, Wildcard, bRecursive ? SearchOption.AllDirectories : SearchOption.TopDirectoryOnly));
+			}
+			else if (DirSymlinkMode == SymlinkMode.Retain || DirSymlinkMode == SymlinkMode.Ignore)
+			{
+				DirectoryInfo DirInfo = new DirectoryInfo(RootPath);
+
+				// get all files in this dir
+				FoundFiles.AddRange(Directory.GetFiles(RootPath, Wildcard, SearchOption.TopDirectoryOnly));
+
+				// now walk over the directories, and recurse into true directories, and process the symlinks
+				foreach (DirectoryInfo Dir in DirInfo.EnumerateDirectories(Wildcard, SearchOption.TopDirectoryOnly))
+				{
+					if (Dir.Attributes.HasFlag(FileAttributes.ReparsePoint))
+					{
+						// if mode is Ingore, then do nothing with it
+						if (DirSymlinkMode == SymlinkMode.Retain)
+						{
+							// add the direcetory as a file since we are going to treat it as a file later
+							FoundFiles.Add(Dir.FullName);
+						}
+					}
+					else
+					{
+						FindFilesAndSymlinks(Dir.FullName, Wildcard, bRecursive, DirSymlinkMode, FoundFiles);
+					}
+				}
+			}
+		}
+
+		public static void FindFilesAndSymlinks(string RootPath, string Wildcard, bool bRecursive, List<string> FoundFiles)
+		{
+			FindFilesAndSymlinks(RootPath, Wildcard, bRecursive, DefaultSymlinkMode, FoundFiles);
+		}
+
+		public static void CopyFileOrSymlink(FileReference Source, FileReference Dest, SymlinkMode SymlinkMode)
+		{
+			if (SymlinkMode == SymlinkMode.Retain || SymlinkMode == SymlinkMode.Ignore)
+			{
+				if (File.GetAttributes(Source.FullName).HasFlag(FileAttributes.ReparsePoint))
+				{
+					if (SymlinkMode == SymlinkMode.Retain)
+					{
+						if (OperatingSystem.IsWindows())
+						{
+							throw new AutomationException("Windows ReparsePoint copying is not supported at this time");
+						}
+						else
+						{
+							Logger.LogInformation("Retaining symlink {0} as {1}", Source, Dest);
+							DirectoryReference.CreateDirectory(Dest.Directory);
+							Utils.RunLocalProcessAndReturnStdOut("/usr/bin/env", $"cp -a \"{Source}\" \"{Dest}\"", null);
+						}
+					}
+					else
+					{
+						Logger.LogInformation("Ignoring symlink {0}", Source);
+					}
+					return;
+				}
+			}
+
+			// if we didn't handle a symlink above, just copy it normally
+			InternalUtils.SafeCopyFile(Source.FullName, Dest.FullName, bSafeCreateDirectory: true);
+		}
+
+		public static void CopyFileOrSymlink(FileReference Source, FileReference Dest)
+		{
+			CopyFileOrSymlink(Source, Dest, DefaultSymlinkMode);
+		}
+
+
 		/// <summary>
 		/// Gets environment variable value.
 		/// </summary>
diff --git a/Engine/Source/Programs/AutomationTool/AutomationUtils/DeploymentContext.cs b/Engine/Source/Programs/AutomationTool/AutomationUtils/DeploymentContext.cs
--- a/Engine/Source/Programs/AutomationTool/AutomationUtils/DeploymentContext.cs
+++ b/Engine/Source/Programs/AutomationTool/AutomationUtils/DeploymentContext.cs
@@ -371,10 +371,10 @@ public class DeploymentContext //: ProjectParams
 	/// </summary>
 	public bool bIsCombiningMultiplePlatforms = false;
 
-    /// <summary>
-    /// If true if this platform is using streaming install chunk manifests
-    /// </summary>
-    public bool PlatformUsesChunkManifests = false;
+	/// <summary>
+	/// If true if this platform is using streaming install chunk manifests
+	/// </summary>
+	public bool PlatformUsesChunkManifests = false;
 
 	/// <summary>
 	/// Temporary setting to exclude non cooked packages from I/O store container file(s)
@@ -395,7 +395,7 @@ public class DeploymentContext //: ProjectParams
 		DirectoryReference BaseArchiveDirectory,
 		string CookerSupportFilesSubdirectory,
 		Platform InSourcePlatform,
-        Platform InTargetPlatform,
+		Platform InTargetPlatform,
 		List<UnrealTargetConfiguration> InTargetConfigurations,
 		IEnumerable<StageTarget> InStageTargets,
 		List<String> InStageExecutables,
@@ -407,7 +407,7 @@ public class DeploymentContext //: ProjectParams
 		bool InArchive,
 		bool InProgram,
 		bool IsClientInsteadOfNoEditor,
-        bool InForceChunkManifests,
+		bool InForceChunkManifests,
 		bool InSeparateDebugStageDirectory,
 		DirectoryReference InDLCRoot
 		)
@@ -416,29 +416,29 @@ public class DeploymentContext //: ProjectParams
 		RawProjectPath = RawProjectPathOrName;
 		DedicatedServer = InServer;
 		LocalRoot = InLocalRoot;
-        CookSourcePlatform = InSourcePlatform;
+		CookSourcePlatform = InSourcePlatform;
 		StageTargetPlatform = InTargetPlatform;
 		StageTargetConfigurations = new List<UnrealTargetConfiguration>(InTargetConfigurations);
 		StageTargets = new List<StageTarget>(InStageTargets);
 		StageExecutables = InStageExecutables;
-        IsCodeBasedProject = ProjectUtils.IsCodeBasedUProjectFile(RawProjectPath, StageTargetPlatform.PlatformType, StageTargetConfigurations);
+		IsCodeBasedProject = ProjectUtils.IsCodeBasedUProjectFile(RawProjectPath, StageTargetPlatform.PlatformType, StageTargetConfigurations);
 		ShortProjectName = ProjectUtils.GetShortProjectName(RawProjectPath);
 		Stage = InStage;
 		Archive = InArchive;
 		DLCRoot = InDLCRoot;
 
-        if (CookSourcePlatform != null && InCooked)
-        {
+		if (CookSourcePlatform != null && InCooked)
+		{
 			CookPlatform = CookSourcePlatform.GetCookPlatform(DedicatedServer, IsClientInsteadOfNoEditor);
-        }
-        else if (CookSourcePlatform != null && InProgram)
-        {
-            CookPlatform = CookSourcePlatform.GetCookPlatform(false, false);
-        }
-        else
-        {
-            CookPlatform = "";
-        }
+		}
+		else if (CookSourcePlatform != null && InProgram)
+		{
+			CookPlatform = CookSourcePlatform.GetCookPlatform(false, false);
+		}
+		else
+		{
+			CookPlatform = "";
+		}
 
 		if (StageTargetPlatform != null && InCooked)
 		{
@@ -446,11 +446,11 @@ public class DeploymentContext //: ProjectParams
 		}
 		else if (StageTargetPlatform != null && InProgram)
 		{
-            FinalCookPlatform = StageTargetPlatform.GetCookPlatform(false, false);
+			FinalCookPlatform = StageTargetPlatform.GetCookPlatform(false, false);
 		}
 		else
 		{
-            FinalCookPlatform = "";
+			FinalCookPlatform = "";
 		}
 
 		PlatformDir = StageTargetPlatform.PlatformType.ToString();
@@ -458,7 +458,7 @@ public class DeploymentContext //: ProjectParams
 		if (BaseStageDirectory != null)
 		{
 			StageDirectory = DirectoryReference.Combine(BaseStageDirectory, FinalCookPlatform);
-			DebugStageDirectory = InSeparateDebugStageDirectory? DirectoryReference.Combine(BaseStageDirectory, FinalCookPlatform + "Debug") : StageDirectory;
+			DebugStageDirectory = InSeparateDebugStageDirectory ? DirectoryReference.Combine(BaseStageDirectory, FinalCookPlatform + "Debug") : StageDirectory;
 		}
 		this.OptionalFileStageDirectory = OptionalFileStageDirectory;
 		this.OptionalFileInputDirectory = OptionalFileInputDirectory;
@@ -563,7 +563,7 @@ public class DeploymentContext //: ProjectParams
 				}
 
 				string ToDir;
-				if(!Properties.TryGetValue("To", out ToDir))
+				if (!Properties.TryGetValue("To", out ToDir))
 				{
 					throw new AutomationException("Missing 'To' property in '{0}'", RemapDirectory);
 				}
@@ -576,7 +576,7 @@ public class DeploymentContext //: ProjectParams
 		List<string> DirectoriesAllowListStrings;
 		if (GameConfig.GetArray("Staging", "AllowedDirectories", out DirectoriesAllowListStrings))
 		{
-			foreach(string AllowedDir in DirectoriesAllowListStrings)
+			foreach (string AllowedDir in DirectoriesAllowListStrings)
 			{
 				DirectoriesAllowList.Add(new StagedDirectoryReference(AllowedDir));
 			}
@@ -665,9 +665,9 @@ public class DeploymentContext //: ProjectParams
 	private static void ReadAllowDenyFileList(ConfigHierarchy Config, string SectionName, string KeyName, HashSet<StagedFileReference> FilesRef)
 	{
 		List<string> FileNames;
-		if(Config.GetArray(SectionName, KeyName, out FileNames))
+		if (Config.GetArray(SectionName, KeyName, out FileNames))
 		{
-			foreach(string FileName in FileNames)
+			foreach (string FileName in FileNames)
 			{
 				FilesRef.Add(new StagedFileReference(FileName));
 			}
@@ -713,18 +713,18 @@ public class DeploymentContext //: ProjectParams
 		if (!DirectoryReference.Exists(BaseDir))
 		{
 			return;
-		}	
+		}
 
 		// Enumerate all the files in this directory
 		Files.AddRange(DirectoryReference.EnumerateFiles(BaseDir, Pattern));
 
 		// Recurse through subdirectories if necessary
-		if(Option == StageFilesSearch.AllDirectories)
+		if (Option == StageFilesSearch.AllDirectories)
 		{
-			foreach(DirectoryReference SubDir in DirectoryReference.EnumerateDirectories(BaseDir))
+			foreach (DirectoryReference SubDir in DirectoryReference.EnumerateDirectories(BaseDir))
 			{
 				string Name = SubDir.GetDirectoryName();
-				if(!RestrictedFolderNames.Contains(Name))
+				if (!RestrictedFolderNames.Contains(Name))
 				{
 					FindFilesToStageInternal(SubDir, Pattern, Option, Files);
 				}
@@ -740,34 +740,34 @@ public class DeploymentContext //: ProjectParams
 	public StagedFileReference GetStagedFileLocation(FileReference InputFile)
 	{
 		StagedFileReference OutputFile;
-		if(InputFile.IsUnderDirectory(ProjectRoot))
+		if (InputFile.IsUnderDirectory(ProjectRoot))
 		{
 			OutputFile = StagedFileReference.Combine(RelativeProjectRootForStage, InputFile.MakeRelativeTo(ProjectRoot));
 		}
-        else if (InputFile.HasExtension(".uplugin"))
-        {
+		else if (InputFile.HasExtension(".uplugin"))
+		{
 			DirectoryReference EnterpriseRoot = DirectoryReference.Combine(EngineRoot, "..", "Enterprise"); // Enterprise plugins aren't under the project additional plugin directories, so they shouldn't be remapped
-            if (InputFile.IsUnderDirectory(EngineRoot) || InputFile.IsUnderDirectory(EnterpriseRoot))
+			if (InputFile.IsUnderDirectory(EngineRoot) || InputFile.IsUnderDirectory(EnterpriseRoot))
 			{
 				OutputFile = new StagedFileReference(InputFile.MakeRelativeTo(LocalRoot));
 			}
-            else
+			else
 			{
 				// This is a plugin that lives outside of the Engine/Plugins or Game/Plugins directory so needs to be remapped for staging/packaging
 				// We need to remap C:\SomePath\PluginName\PluginName.uplugin to RemappedPlugins\PluginName\PluginName.uplugin
 				OutputFile = new StagedFileReference(String.Format("RemappedPlugins/{0}/{1}", InputFile.GetFileNameWithoutExtension(), InputFile.GetFileName()));
 			}
-        }
-        else if (InputFile.IsUnderDirectory(LocalRoot))
-        {
+		}
+		else if (InputFile.IsUnderDirectory(LocalRoot))
+		{
 			OutputFile = new StagedFileReference(InputFile.MakeRelativeTo(LocalRoot));
-        }
-        else if (DLCRoot != null && InputFile.IsUnderDirectory(DLCRoot))
+		}
+		else if (DLCRoot != null && InputFile.IsUnderDirectory(DLCRoot))
 		{
 			OutputFile = new StagedFileReference(InputFile.MakeRelativeTo(DLCRoot));
 		}
 		else
-        {
+		{
 			throw new AutomationException("Can't deploy {0} because it doesn't start with {1} or {2}", InputFile, ProjectRoot, LocalRoot);
 		}
 		return OutputFile;
@@ -885,7 +885,7 @@ public class DeploymentContext //: ProjectParams
 	/// <param name="StagedFile">Location of the file in the staging directory</param>
 	public void StageCrashReporterFile(StagedFileType FileType, FileReference InputFile, StagedFileReference StagedFile)
 	{
-		if(FileType == StagedFileType.UFS)
+		if (FileType == StagedFileType.UFS)
 		{
 			CrashReporterUFSFiles[StagedFile] = InputFile;
 		}
@@ -916,7 +916,7 @@ public class DeploymentContext //: ProjectParams
 	public void StageCrashReporterFiles(StagedFileType FileType, DirectoryReference InputDir, StageFilesSearch Option, StagedDirectoryReference OutputDir)
 	{
 		List<FileReference> InputFiles = FindFilesToStage(InputDir, Option);
-		foreach(FileReference InputFile in InputFiles)
+		foreach (FileReference InputFile in InputFiles)
 		{
 			StagedFileReference StagedFile = StagedFileReference.Combine(OutputDir, InputFile.MakeRelativeTo(InputDir));
 			StageCrashReporterFile(FileType, InputFile, StagedFile);
@@ -926,7 +926,7 @@ public class DeploymentContext //: ProjectParams
 	public void StageBuildProductsFromReceipt(TargetReceipt Receipt, bool RequireDependenciesToExist, bool TreatNonShippingBinariesAsDebugFiles)
 	{
 		// Stage all the build products needed at runtime
-		foreach(BuildProduct BuildProduct in Receipt.BuildProducts)
+		foreach (BuildProduct BuildProduct in Receipt.BuildProducts)
 		{
 			// allow missing files if needed
 			if (RequireDependenciesToExist == false && FileReference.Exists(BuildProduct.Path) == false)
@@ -934,7 +934,7 @@ public class DeploymentContext //: ProjectParams
 				continue;
 			}
 
-			if(BuildProduct.Type == BuildProductType.Executable || BuildProduct.Type == BuildProductType.DynamicLibrary || BuildProduct.Type == BuildProductType.RequiredResource)
+			if (BuildProduct.Type == BuildProductType.Executable || BuildProduct.Type == BuildProductType.DynamicLibrary || BuildProduct.Type == BuildProductType.RequiredResource)
 			{
 				StagedFileType FileTypeToUse = StagedFileType.NonUFS;
 				if (TreatNonShippingBinariesAsDebugFiles && Receipt.Configuration != UnrealTargetConfiguration.Shipping)
@@ -944,7 +944,7 @@ public class DeploymentContext //: ProjectParams
 
 				StageFile(FileTypeToUse, BuildProduct.Path);
 			}
-			else if(BuildProduct.Type == BuildProductType.SymbolFile || BuildProduct.Type == BuildProductType.MapFile)
+			else if (BuildProduct.Type == BuildProductType.SymbolFile || BuildProduct.Type == BuildProductType.MapFile)
 			{
 				// Symbol files aren't true dependencies so we can skip if they don't exist
 				if (FileReference.Exists(BuildProduct.Path))
@@ -963,7 +963,7 @@ public class DeploymentContext //: ProjectParams
 		ExcludePatterns.Add(".../*.uasset");
 
 		// Also stage any additional runtime dependencies, like ThirdParty DLLs
-		foreach(RuntimeDependency RuntimeDependency in Receipt.RuntimeDependencies)
+		foreach (RuntimeDependency RuntimeDependency in Receipt.RuntimeDependencies)
 		{
 			// allow missing files if needed
 			if ((RequireDependenciesToExist && RuntimeDependency.Type != StagedFileType.DebugNonUFS) || FileReference.Exists(RuntimeDependency.Path))
@@ -979,7 +979,8 @@ public class DeploymentContext //: ProjectParams
 
 		if (CommandUtils.DirectoryExists(InPath))
 		{
-			var All = CommandUtils.FindFiles(Wildcard, bRecursive, InPath);
+			List<string> All = new();
+			CommandUtils.FindFilesAndSymlinks(InPath, Wildcard, bRecursive, All);
 
 			var Exclude = new HashSet<string>();
 			if (ExcludeWildcard != null)
diff --git a/Engine/Source/Programs/AutomationTool/Scripts/ArchiveCommand.Automation.cs b/Engine/Source/Programs/AutomationTool/Scripts/ArchiveCommand.Automation.cs
--- a/Engine/Source/Programs/AutomationTool/Scripts/ArchiveCommand.Automation.cs
+++ b/Engine/Source/Programs/AutomationTool/Scripts/ArchiveCommand.Automation.cs
@@ -41,7 +41,7 @@ namespace AutomationScripts
 				{
 					FileReference Src = new FileReference(Pair.Key);
 					FileReference Dest = FileReference.Combine(SC.ArchiveDirectory, Pair.Value);
-					CopyFileIncremental(Src, Dest);
+					CopyFileOrSymlink(Src, Dest);
 				}
 			}
 		}
