Modify RunLowLevelTests to have a -SkipStage flag to skip staging process for applications already-packaged.

diff --git a/Engine/Source/Programs/AutomationTool/LowLevelTests/RunLowLevelTests.cs b/Engine/Source/Programs/AutomationTool/LowLevelTests/RunLowLevelTests.cs
--- a/Engine/Source/Programs/AutomationTool/LowLevelTests/RunLowLevelTests.cs
+++ b/Engine/Source/Programs/AutomationTool/LowLevelTests/RunLowLevelTests.cs
@@ -58,7 +58,8 @@ namespace LowLevelTests
 				ContextOptions.TestApp,
 				ContextOptions.Build,
 				ContextOptions.Platform,
-				ContextOptions.Configuration);
+				ContextOptions.Configuration,
+				ContextOptions.SkipStage);
 
 			SetupDevices(TestPlatform, ContextOptions);
 
@@ -165,6 +166,9 @@ namespace LowLevelTests
 
 		public bool AttachToDebugger;
 
+		[AutoParam(false)]
+		public bool SkipStage;
+
 		public string Build;
 
 		[AutoParam("")]
@@ -220,6 +224,8 @@ namespace LowLevelTests
 			Tags = Params.ParseValue("tags=", null);
 			AttachToDebugger = Params.ParseParam("attachtodebugger");
 
+			SkipStage = Params.ParseParam("skipstage");
+
 			ReportType = Params.ParseValue("reporttype=", "console");
 
 			string PlatformArgString = Params.ParseValue("platform=", null);
@@ -472,36 +478,51 @@ namespace LowLevelTests
 	{
 		bool CanSupportPlatform(UnrealTargetPlatform InPlatform);
 
-		LowLevelTestsBuild CreateBuild(UnrealTargetPlatform InPlatform, UnrealTargetConfiguration InConfiguration, string InTestApp, string InBuildPath);
-
+		LowLevelTestsBuild CreateBuild(UnrealTargetPlatform InPlatform, UnrealTargetConfiguration InConfiguration, string InTestApp, string InBuildPath, bool bSkipStage);
 		protected static string GetExecutable(UnrealTargetPlatform InPlatform, UnrealTargetConfiguration InConfiguration, string InTestApp, string InBuildPath, string FileRegEx)
 		{
-			IEnumerable<string> Executables = DirectoryUtils.FindFiles(InBuildPath, new Regex(FileRegEx));
+			IEnumerable<string> Executables = DirectoryUtils.FindMatchingFiles(InBuildPath, FileRegEx, -1).Select(FileInfo => FileInfo.FullName);
 			string ParentDirPath;
 			string BuildExecutableName;
 			foreach (string Executable in Executables)
 			{
+				Log.VeryVerbose("Found path: {0}", Executable);
 				ParentDirPath = Directory.GetParent(Executable).FullName;
 				BuildExecutableName = Path.GetFileNameWithoutExtension(Executable);
 
-				if (ParentDirPath.ToLower().Contains(InPlatform.ToString().ToLower()))
+				if (!ParentDirPath.Contains(InPlatform.ToString(), StringComparison.OrdinalIgnoreCase))
+				{
+					Log.Error("ParentPath did not have platform ({1}) : {0}", ParentDirPath, InPlatform);
+					continue;
+				}
+
+				if (!ParentDirPath.Contains(InTestApp.ToString(), StringComparison.OrdinalIgnoreCase))
 				{
-					if (ParentDirPath.ToLower().Contains(InTestApp.ToString().ToLower()))
+					Log.Error("ParentPath did not have app ({1}) : {0}", ParentDirPath, InTestApp);
+					continue;
+				}
+
+				// Executable name must not contain any configuration or platform name
+				Log.Verbose("Config type: {0}", InConfiguration);
+				if (InConfiguration == UnrealTargetConfiguration.Development)
+				{
+					if (BuildExecutableName.Equals(InTestApp, StringComparison.OrdinalIgnoreCase) 
+						&& !BuildExecutableName.Contains(InPlatform.ToString(), StringComparison.OrdinalIgnoreCase))
 					{
-						// Executable name must not contain any configuration or platform name
-						if (InConfiguration == UnrealTargetConfiguration.Development)
-						{
-							if (BuildExecutableName.CompareTo(InTestApp) == 0 && !BuildExecutableName.Contains(InPlatform.ToString()))
-							{
-								return Path.GetRelativePath(InBuildPath, Executable);
-							}
-						}
-						else if (BuildExecutableName.Contains(InTestApp) && BuildExecutableName.Contains(InPlatform.ToString()))
-						{
-							return Path.GetRelativePath(InBuildPath, Executable);
-						}
+						Log.VeryVerbose("Output Development Executable: {0}", Path.GetRelativePath(InBuildPath, Executable));
+						return Path.GetRelativePath(InBuildPath, Executable);
 					}
 				}
+
+				if (!(BuildExecutableName.Contains(InTestApp, StringComparison.OrdinalIgnoreCase) 
+					&& BuildExecutableName.Contains(InPlatform.ToString(), StringComparison.OrdinalIgnoreCase)))
+				{
+					Log.Error("BuildExecutableName did not have expected name ({0}) or is missing platform in name ({1}): {2}", InTestApp, InPlatform.ToString(), BuildExecutableName);
+					continue;
+				}
+
+				Log.VeryVerbose("Output Executable: {0}", Path.GetRelativePath(InBuildPath, Executable));
+				return Path.GetRelativePath(InBuildPath, Executable);
 			}
 			throw new AutomationException("Cannot find low level test executable for {0} in build path {1} for {2} using regex \"{3}\"", InPlatform, InBuildPath, InTestApp, FileRegEx);
 		}
@@ -533,23 +554,26 @@ namespace LowLevelTests
 			return InPlatform.IsInGroup(UnrealPlatformGroup.Desktop);
 		}
 
-		public LowLevelTestsBuild CreateBuild(UnrealTargetPlatform InPlatform, UnrealTargetConfiguration InConfiguration, string InTestApp, string InBuildPath)
+		public LowLevelTestsBuild CreateBuild(UnrealTargetPlatform InPlatform, UnrealTargetConfiguration InConfiguration, string InTestApp, string InBuildPath, bool bSkipStage)
+		{
+			string ExecutablePath = ILowLevelTestsBuildFactory.GetExecutable(InPlatform, InConfiguration, InTestApp, InBuildPath, GetExecutableRegex(InPlatform));
+			return new LowLevelTestsBuild(InPlatform, InConfiguration, InBuildPath, ExecutablePath);
+		}
+
+		public string GetExecutableRegex(UnrealTargetPlatform InPlatform)
 		{
-			string DesktopExecutableRegEx;
 			if (InPlatform.IsInGroup(UnrealPlatformGroup.Windows))
 			{
-				DesktopExecutableRegEx = @"\w+Tests(?:-\w+)?(?:-\w+)?.exe$";
+				return @"\w+Tests(?:-\w+)?(?:-\w+)?.exe$";
 			}
 			else if (InPlatform == UnrealTargetPlatform.Linux || InPlatform == UnrealTargetPlatform.Mac)
 			{
-				DesktopExecutableRegEx = @"\w+Tests$";
+				return @"\w+Tests$";
 			}
 			else
 			{
 				throw new AutomationException("Cannot create build for non-desktop platform " + InPlatform);
 			}
-			string ExecutablePath = ILowLevelTestsBuildFactory.GetExecutable(InPlatform, InConfiguration, InTestApp, InBuildPath, DesktopExecutableRegEx);
-			return new LowLevelTestsBuild(InPlatform, InConfiguration, InBuildPath, ExecutablePath);
 		}
 	}
 
@@ -591,21 +615,23 @@ namespace LowLevelTests
 		public UnrealTargetConfiguration Configuration { get; protected set; }
 		public LowLevelTestsBuild DiscoveredBuild { get; protected set; }
 
-		public LowLevelTestsBuildSource(string InTestApp, string InBuildPath, UnrealTargetPlatform InTargetPlatform, UnrealTargetConfiguration InConfiguration)
+		public LowLevelTestsBuildSource(string InTestApp, string InBuildPath, UnrealTargetPlatform InTargetPlatform, UnrealTargetConfiguration InConfiguration, bool InSkipStage)
 		{
 			TestApp = InTestApp;
 			Platform = InTargetPlatform;
 			BuildPath = InBuildPath;
 			Configuration = InConfiguration;
-			InitBuildSource(InTestApp, InBuildPath, InTargetPlatform, InConfiguration);
+			InitBuildSource(InTestApp, InBuildPath, InTargetPlatform, InConfiguration, InSkipStage);
 		}
 
-		protected void InitBuildSource(string InTestApp, string InBuildPath, UnrealTargetPlatform InTargetPlatform, UnrealTargetConfiguration InConfiguration)
+		protected void InitBuildSource(string InTestApp, string InBuildPath, UnrealTargetPlatform InTargetPlatform, UnrealTargetConfiguration InConfiguration, bool bSkipStage)
 		{
 			LowLevelTestsBuildFactory = Gauntlet.Utils.InterfaceHelpers.FindImplementations<ILowLevelTestsBuildFactory>(true)
 				.Where(B => B.CanSupportPlatform(InTargetPlatform))
 				.First();
-			DiscoveredBuild = LowLevelTestsBuildFactory.CreateBuild(InTargetPlatform, InConfiguration, InTestApp, InBuildPath);
+			
+			DiscoveredBuild = LowLevelTestsBuildFactory.CreateBuild(InTargetPlatform, InConfiguration, InTestApp, InBuildPath, bSkipStage);
+
 			if (DiscoveredBuild == null)
 			{
 				throw new AutomationException("No builds were discovered at path {0} matching test app name {1} and target platform {2}", InBuildPath, InTestApp, InTargetPlatform);
