- Packaging and archiving for modern IOS/Mac:
  - When doing standard BuildCookRun, package will simply finalize the (likely already finalize) .app in Binaries/<platform>
  - When using Distribution mode, packge will create a .xcarchive in the Xcode Library (as seen in Organizer) to be able to distribute in various modes
  - Archiving will copy either the .app or .xcarchive made above to the archive directory (it will always copy the most recent .xcarchive, since we don't create the directory name, we can't know which archive to copy, so the most recent one wins - not an issue when using -package, because it always makes a new .xcarchive)
- A few code cleanup/refactors in related

diff --git a/Engine/Build/BatchFiles/Mac/ThinApp.sh b/Engine/Build/BatchFiles/Mac/ThinApp.sh
--- a/Engine/Build/BatchFiles/Mac/ThinApp.sh
+++ b/Engine/Build/BatchFiles/Mac/ThinApp.sh
@@ -20,9 +20,6 @@ function thin {
 				echo "Thinning $1 down to ${keeparches[*]}"
 				mv "$1" "$1.tmp"
 
-				# remove any existing signature or ditto will throw a warning - Xcode will be re-codesigning everything after this runs
-				codesign --remove-signature "$1.tmp"
-
 				archparams=(${keeparches[@]/#/"--arch "})
 				ditto "${archparams[@]}" "$1.tmp" "$1"
 				rm "$1.tmp"
diff --git a/Engine/Source/Programs/AutomationTool/Apple/ApplePlatform.Automation.cs b/Engine/Source/Programs/AutomationTool/Apple/ApplePlatform.Automation.cs
--- a/Engine/Source/Programs/AutomationTool/Apple/ApplePlatform.Automation.cs
+++ b/Engine/Source/Programs/AutomationTool/Apple/ApplePlatform.Automation.cs
@@ -8,7 +8,8 @@ using UnrealBuildTool;
 using EpicGames.Core;
 using UnrealBuildBase;
 using System.Diagnostics;
-
+using System.Collections.Generic;
+using Microsoft.Extensions.Logging;
 
 public abstract class ApplePlatform : Platform
 {
@@ -73,22 +74,146 @@ public abstract class ApplePlatform : Platform
 	{
 		if (MacExports.UseModernXcode(Params.RawProjectPath))
 		{
-			// modern does all the distribution via .xcarchive
-			bool bForDistribution = false;
-
 			foreach (TargetReceipt Target in SC.StageTargets.Select(x => x.Receipt))
 			{
-				Console.WriteLine($"GENERATING STUB PROJECT");
-				DirectoryReference GeneratedProjectFile;
-				IOSExports.GenerateRunOnlyXcodeProject(SC.RawProjectPath, Target.Platform, Target.TargetName, bForDistribution, Logger, out GeneratedProjectFile);
+				MacExports.BuildWithModernXcode(SC.RawProjectPath, Target.Platform, Target.Configuration, Target.TargetName, bArchiveForDistro: false, Logger);
+			}
+		}
+	}
+
+	private DirectoryReference GetArchivePath(StageTarget Target, DeploymentContext SC)
+	{
+		DirectoryReference UserDir = new DirectoryReference(Environment.GetFolderPath(Environment.SpecialFolder.UserProfile));
+		DirectoryReference Library = DirectoryReference.Combine(UserDir, "Library/Developer/Xcode/Archives");
+
+		// order date named folders (use creating data, not name, but same thing)
+		List<DirectoryReference> DateDirs = DirectoryReference.EnumerateDirectories(Library).ToList();
+		DateDirs.SortBy(x => Directory.GetCreationTime(x.FullName));
+		DateDirs.Reverse();
+
+
+		// go through each folder, starting at most recent, looking for an archive for the target
+		foreach (DirectoryReference DateDir in DateDirs)
+		{
+			Logger.LogInformation("Looking in Xcode archive dir {0}...", DateDir);
+
+			// find the most recent archive for this target (based on name of target, this ignores Development vs Shipping, but 
+			// since Distribution is meant only for Shipping it's ok
+			string Wildcard = $"{Target.Receipt.TargetName} *.xcarchive";
+			List<DirectoryReference> XcArchives = DirectoryReference.EnumerateDirectories(DateDir, Wildcard).ToList();
+			if (XcArchives.Count > 0)
+			{
+				XcArchives.SortBy(x => Directory.GetCreationTime(x.FullName));
+				DirectoryReference XcArchive = XcArchives.Last();
+
+				Logger.LogInformation("Found xcarchive dir {0}", XcArchive);
+				return XcArchive;
+			}
+		}
+
+		return null;
+	}
+
+	private DirectoryReference GetFinalAppPath(StageTarget Target, DeploymentContext SC)
+	{
+		DirectoryReference AppDir;
+		// get the executable from the receipt
+		FileReference Executable = Target.Receipt.BuildProducts.First(x => x.Type == BuildProductType.Executable).Path;
+
+		// if we are in a .app, then use that
+		if (Executable.FullName.Contains(".app/"))
+		{
+			AppDir = Executable.Directory;
+			// go up until we find the .app - Mac and IOS are different, so this handles both cases
+			while (AppDir.HasExtension(".app") == false)
+			{
+				AppDir = AppDir.ParentDirectory;
+			}
+		}
+		// otherwise use the .app next to the executable
+		else
+		{
+			AppDir = new DirectoryReference(Executable.FullName + ".app");
+		}
+
+		return AppDir;
+	}
+
+	public override void Package(ProjectParams Params, DeploymentContext SC, int WorkingCL)
+	{
+		if (MacExports.UseModernXcode(Params.RawProjectPath))
+		{
+			foreach (StageTarget Target in SC.StageTargets)
+			{
+				TargetReceipt Receipt = Target.Receipt;
+
+				// if we we packaging for distrbution, we will create a .xcarchive which can be used to submit to app stores, or exported for other distribution methods
+				// the archive will be created in the standard Archives location accessible via Xcode. Using -archive will copy it out into
+				// the specified location for use as needed
+				if (MacExports.BuildWithModernXcode(Params.RawProjectPath, Receipt.Platform, Receipt.Configuration, Receipt.TargetName, Params.Distribution, Logger))
+				{
+					Logger.LogInformation("=====================================================================================");
+					if (Params.Distribution)
+					{
+						Logger.LogInformation("Created .xcarchive in Xcode's Library, which can be seen in Xcode's Organizer window");
+						Logger.LogInformation("You may use this to validate and prepare for vvarious distribution methods");
+					}
+					else
+					{
+						Logger.LogInformation("Finalized {App} for running fully self-contained", GetFinalAppPath(Target, SC));
+					}
+					Logger.LogInformation("=====================================================================================");
+				}
+			}
+		}
+
+		// package up the program, potentially with an installer for Mac
+		PrintRunTime();
+	}
+
+	public override void GetFilesToArchive(ProjectParams Params, DeploymentContext SC)
+	{
+		if (!MacExports.UseModernXcode(Params.RawProjectPath))
+		{
+			return;
+		}
 
+		Logger.LogInformation("staging targets: '{tagets}', '{configs}'", string.Join(", ", Params.ClientCookedTargets),
+			string.Join(", ", SC.StageTargetConfigurations));
+		foreach (StageTarget Target in SC.StageTargets)
+		{
+			//// copy the xcarchive and any exported files
+			//string ExportMode = "Development";
+			//DirectoryReference ExportPath = GetExportPath(ExportMode, SC);
 
-				Console.WriteLine("FINALIZEING APP with {0}, target {1}", GeneratedProjectFile, Target.TargetName);
-				string ExtraOptions = $"SYMROOT=\"{SC.StageDirectory.ParentDirectory}\" EFFECTIVE_PLATFORM_NAME={SC.StageDirectory.GetDirectoryName()}";
-				// run xcodebuild on the generated project to make the .app
-				IOSExports.FinalizeAppWithModernXcode(GeneratedProjectFile!, Target.Platform, Target.TargetName, Target.Configuration.ToString(), ExtraOptions, bForDistribution, Logger);
-				Console.WriteLine("DOME FINALIZEING APP with {0}, target {1}", GeneratedProjectFile, Target.TargetName);
+			// distribution mode we want to archive the .xcarchive that was created during Package
+			DirectoryReference ArchiveSource;
+			if (Params.Distribution)
+			{
+				// find the most recent .xcarchive in the Xcode archives library
+				ArchiveSource = GetArchivePath(Target, SC);
+
+				if (ArchiveSource == null)
+				{
+					Logger.LogError("Unable to find a .xcarchive in Xcode's Library to archive to {ArchiveDir}", SC.ArchiveDirectory);
+					return;
+				}
 			}
+			else
+			{
+				ArchiveSource = GetFinalAppPath(Target, SC);
+				if (!DirectoryReference.Exists(ArchiveSource))
+				{
+					Logger.LogError("Unable to find the expected application ({App}) for achiving", ArchiveSource);
+					return;
+				}
+			}
+
+			Logger.LogInformation("=====================================================================================");
+			Logger.LogInformation("Copying {Type} package {ArchiveSource} to archive directory {ArchiveDir}", Params.Distribution ? "Distribution" : "Development", ArchiveSource, SC.ArchiveDirectory);
+			Logger.LogInformation("=====================================================================================");
+
+			SC.ArchiveFiles(ArchiveSource.FullName, NewPath: ArchiveSource.GetDirectoryName());
 		}
 	}
 }
diff --git a/Engine/Source/Programs/AutomationTool/IOS/IOSPlatform.Automation.cs b/Engine/Source/Programs/AutomationTool/IOS/IOSPlatform.Automation.cs
--- a/Engine/Source/Programs/AutomationTool/IOS/IOSPlatform.Automation.cs
+++ b/Engine/Source/Programs/AutomationTool/IOS/IOSPlatform.Automation.cs
@@ -769,6 +769,13 @@ public class IOSPlatform : ApplePlatform
 
 	public override void Package(ProjectParams Params, DeploymentContext SC, int WorkingCL)
 	{
+		// use the shared packaging with modern mode
+		if (MacExports.UseModernXcode(Params.RawProjectPath))
+		{
+			base.Package(Params, SC, WorkingCL);
+			return;
+		}
+
 		Logger.LogInformation("Package {Arg0}", Params.RawProjectPath);
 
 		bool bIsBuiltAsFramework = IsBuiltAsFramework(Params, SC);
@@ -1509,6 +1516,13 @@ public class IOSPlatform : ApplePlatform
 
 	public override void GetFilesToArchive(ProjectParams Params, DeploymentContext SC)
 	{
+		// use the shared archiving with modern mode
+		if (MacExports.UseModernXcode(Params.RawProjectPath))
+		{
+			base.GetFilesToArchive(Params, SC);
+			return;
+		}
+
 		if (SC.StageTargetConfigurations.Count != 1)
 		{
 			throw new AutomationException("iOS is currently only able to package one target configuration at a time, but StageTargetConfigurations contained {0} configurations", SC.StageTargetConfigurations.Count);
diff --git a/Engine/Source/Programs/AutomationTool/Mac/MacPlatform.Automation.cs b/Engine/Source/Programs/AutomationTool/Mac/MacPlatform.Automation.cs
--- a/Engine/Source/Programs/AutomationTool/Mac/MacPlatform.Automation.cs
+++ b/Engine/Source/Programs/AutomationTool/Mac/MacPlatform.Automation.cs
@@ -395,12 +395,6 @@ public class MacPlatform : ApplePlatform
 		}
 	}
 
-	public override void Package(ProjectParams Params, DeploymentContext SC, int WorkingCL)
-	{
-		// package up the program, potentially with an installer for Mac
-		PrintRunTime();
-	}
-
 	private void FixupFrameworks(string TargetPath)
 	{
 		DirectoryReference TargetCEFDir = DirectoryReference.Combine(new DirectoryReference(TargetPath), "Engine/Binaries/ThirdParty/CEF3/Mac");
@@ -433,6 +427,12 @@ public class MacPlatform : ApplePlatform
 
 	public override void ProcessArchivedProject(ProjectParams Params, DeploymentContext SC)
 	{
+		// nothing to do with modern
+		if (MacExports.UseModernXcode(Params.RawProjectPath))
+		{
+			return;
+		}
+
 		if (Params.CreateAppBundle)
 		{
 			string ExeName = SC.StageExecutables[0];
diff --git a/Engine/Source/Programs/UnrealBuildTool/Platform/IOS/IOSExports.cs b/Engine/Source/Programs/UnrealBuildTool/Platform/IOS/IOSExports.cs
--- a/Engine/Source/Programs/UnrealBuildTool/Platform/IOS/IOSExports.cs
+++ b/Engine/Source/Programs/UnrealBuildTool/Platform/IOS/IOSExports.cs
@@ -290,12 +290,13 @@ namespace UnrealBuildTool
 		/// <param name="Platform">THe platform to make the .app for</param>
 		/// <param name="SchemeName">The name of the scheme (basically the target on the .xcworkspace)</param>
 		/// <param name="Configuration">Which configuration to make (Debug, etc)</param>
+		/// <param name="Action">Action (build, archive, etc)</param>
 		/// <param name="ExtraOptions">Extra options to pass to xcodebuild</param>
 		/// <param name="bForDistribution">True if this is making a bild for uploading to app store</param>
 		/// <param name="Logger">Logging object</param>
-		public static void FinalizeAppWithModernXcode(DirectoryReference XcodeProject, UnrealTargetPlatform Platform, string SchemeName, string Configuration, string ExtraOptions, bool bForDistribution, ILogger Logger)
+		public static void FinalizeAppWithModernXcode(DirectoryReference XcodeProject, UnrealTargetPlatform Platform, string SchemeName, string Configuration, string Action, string ExtraOptions, bool bForDistribution, ILogger Logger)
 		{
-			FinalizeAppWithXcode(XcodeProject, Platform, SchemeName, Configuration, null, null, null, ExtraOptions, false, bForDistribution, bUseModernXcode: true, Logger);
+			FinalizeAppWithXcode(XcodeProject, Platform, SchemeName, Configuration, null, null, null, ExtraOptions, false, bForDistribution, bUseModernXcode: true, Logger, Action);
 		}
 
 		/// <summary>
@@ -313,19 +314,24 @@ namespace UnrealBuildTool
 		/// <param name="bForDistribution">True if this is making a bild for uploading to app store</param>
 		/// <param name="bUseModernXcode">True if the project was made with modern xcode mode</param>
 		/// <param name="Logger">Logging object</param>
+		/// <param name="Action">Action (build, archive, etc). Defaults to build</param>
 		/// <returns></returns>
 		public static int FinalizeAppWithXcode(DirectoryReference XcodeProject, UnrealTargetPlatform Platform, string SchemeName, string Configuration,
-			string? Provision, string? Certificate, string? Team, string ExtraOptions, bool bAutomaticSigning, bool bForDistribution, bool bUseModernXcode, ILogger Logger)
+			string? Provision, string? Certificate, string? Team, string ExtraOptions, bool bAutomaticSigning, bool bForDistribution, bool bUseModernXcode, ILogger Logger, 
+			string Action="build")
 		{
+			int ForDistro = bForDistribution ? 1 : 0;
+
 			List<string> Arguments = new()
 			{
 				"UBT_NO_POST_DEPLOY=true",
 				new IOSToolChainSettings(Logger).XcodeDeveloperDir + "usr/bin/xcodebuild",
-				"build",
+				Action,
 				$"-workspace \"{XcodeProject.FullName}\"",
 				$"-scheme \"{SchemeName}\"",
 				$"-configuration \"{Configuration}\"",
 				$"-destination generic/platform=" + (Platform == UnrealTargetPlatform.TVOS ? "tvOS" : Platform == UnrealTargetPlatform.Mac ? "macOS" : "iOS"),
+				"-hideShellScriptEnvironment",
 				ExtraOptions,
 				//$"-sdk {SDKName}",
 			};
@@ -381,10 +387,7 @@ namespace UnrealBuildTool
 				}
 			}
 
-			int ReturnCode;
-			string Output = Utils.RunLocalProcessAndReturnStdOut("/usr/bin/env", string.Join(" ", Arguments), null, out ReturnCode);
-			Logger.LogDebug(Output);
-			DirectoryReference.Delete(XcodeProject, true);
+			int ReturnCode = Utils.RunLocalProcessAndLogOutput("/usr/bin/env", string.Join(" ", Arguments), Logger);
 			return ReturnCode;
 		}
 
diff --git a/Engine/Source/Programs/UnrealBuildTool/Platform/IOS/IOSToolChain.cs b/Engine/Source/Programs/UnrealBuildTool/Platform/IOS/IOSToolChain.cs
--- a/Engine/Source/Programs/UnrealBuildTool/Platform/IOS/IOSToolChain.cs
+++ b/Engine/Source/Programs/UnrealBuildTool/Platform/IOS/IOSToolChain.cs
@@ -1518,38 +1518,8 @@ namespace UnrealBuildTool
 
 			if (UseModernXcode(Target.ProjectFile))
 			{
-				// generate a run-only project file for codesigning, etc
-				DirectoryReference? GeneratedProjectFile;
-				IOSExports.GenerateRunOnlyXcodeProject(Target.ProjectFile, Target.Platform, Target.TargetName, Target.bForDistribution, Logger, out GeneratedProjectFile);
+				throw new BuildException("Do we get here?");
 
-				// @todo - should we move this in to FinalizeAppWithXcode?
-				string ConfigName = Target.Configuration.ToString();
-				if (Target.TargetType != TargetType.Game && Target.TargetType != TargetType.Program)
-				{
-					ConfigName += " " + Target.TargetType.ToString();
-				}
-
-				string SchemeName;
-				if (AppName == "UnrealGame" || AppName == "UnrealClient")
-				{
-					if (Target.bBuildAsFramework)
-					{
-						SchemeName = "UnrealGame";
-					}
-					else
-					{
-						SchemeName = "UE5";
-					}
-				}
-				else
-				{
-					SchemeName = Target.ProjectFile!.GetFileNameWithoutExtension();
-				}
-
-				// run xcodebuild on the generated project to make the .app
-				IOSExports.FinalizeAppWithModernXcode(GeneratedProjectFile!, Target.Platform, SchemeName, ConfigName, "", Target.bForDistribution, Logger);
-
-				return;
 			}
 
 
diff --git a/Engine/Source/Programs/UnrealBuildTool/Platform/Mac/MacExports.cs b/Engine/Source/Programs/UnrealBuildTool/Platform/Mac/MacExports.cs
--- a/Engine/Source/Programs/UnrealBuildTool/Platform/Mac/MacExports.cs
+++ b/Engine/Source/Programs/UnrealBuildTool/Platform/Mac/MacExports.cs
@@ -91,6 +91,35 @@ namespace UnrealBuildTool
 			return bUseModernXcode;
 		}
 
+		/// <summary>
+		/// Generates a stub xcode project for the given project/platform/target combo, then builds or archives it
+		/// </summary>
+		/// <param name="ProjectFile"></param>
+		/// <param name="Platform"></param>
+		/// <param name="Configuration"></param>
+		/// <param name="TargetName"></param>
+		/// <param name="bArchiveForDistro"></param>
+		/// <param name="Logger"></param>
+		/// <returns></returns>
+		public static bool BuildWithModernXcode(FileReference? ProjectFile, UnrealTargetPlatform Platform, UnrealTargetConfiguration Configuration, 
+			string TargetName, bool bArchiveForDistro, ILogger Logger)
+		{
+			DirectoryReference? GeneratedProjectFile;
+			// we don't use distro flag when making a modern project
+			IOSExports.GenerateRunOnlyXcodeProject(ProjectFile, Platform, TargetName, bForDistribution:false, Logger, out GeneratedProjectFile);
+
+			if (GeneratedProjectFile == null)
+			{
+				return false;
+			}
+
+			// run xcodebuild on the generated project to make the .app
+			IOSExports.FinalizeAppWithModernXcode(GeneratedProjectFile!, Platform, TargetName, Configuration.ToString(),
+				bArchiveForDistro ? "archive" : "build", "", bArchiveForDistro, Logger);
+
+			return true;
+		}
+
 		/// <summary>
 		/// Strips symbols from a file
 		/// </summary>
diff --git a/Engine/Source/Programs/UnrealBuildTool/ProjectFiles/Xcode/XcodeProject.cs b/Engine/Source/Programs/UnrealBuildTool/ProjectFiles/Xcode/XcodeProject.cs
--- a/Engine/Source/Programs/UnrealBuildTool/ProjectFiles/Xcode/XcodeProject.cs
+++ b/Engine/Source/Programs/UnrealBuildTool/ProjectFiles/Xcode/XcodeProject.cs
@@ -1249,7 +1249,7 @@ namespace UnrealBuildTool.XcodeProjectXcconfig
 			// NOTE: We read from hardcoded location where UBT writes to, but we write to CONFIGURATION_BUILD_DIR because
 			// when Archiving, the .app is somewhere else
 			string SourceEnvVar = (Project.Platform == UnrealTargetPlatform.Mac) ? "EXECUTABLE_PATH" : "EXECUTABLE_NAME";
-			string EngineOrProject = Project.UnrealData.UProjectFileLocation == null ? "Engine" : Project.UnrealData.ProductName;
+			string EngineOrProject = Project.UnrealData.UProjectFileLocation == null ? "Engine" : Project.UnrealData.UProjectFileLocation.GetFileNameWithoutAnyExtensions();
 
 			CopyScript.AddRange(new string[]
 			{
@@ -1319,7 +1319,7 @@ namespace UnrealBuildTool.XcodeProjectXcconfig
 			{
 				"set -e",
 				"",
-				"# Run the wrapper dsym geneerator",
+				"# Run the wrapper dsym generator",
 				"\\\"${UE_ENGINE_DIR}/Build/BatchFiles/Mac/GenerateUniversalDSYM.sh\\\" \\\"${CONFIGURATION_BUILD_DIR}/${EXECUTABLE_PATH}\\\" \\\"${DWARF_DSYM_FOLDER_PATH}/${DWARF_DSYM_FILE_NAME}\\\"",
 				"strip -no_code_signature_warning -D \\\"${CONFIGURATION_BUILD_DIR}/${EXECUTABLE_PATH}\\\"",
 				"",
diff --git a/Engine/Source/Programs/UnrealBuildTool/ToolChain/AppleToolChain.cs b/Engine/Source/Programs/UnrealBuildTool/ToolChain/AppleToolChain.cs
--- a/Engine/Source/Programs/UnrealBuildTool/ToolChain/AppleToolChain.cs
+++ b/Engine/Source/Programs/UnrealBuildTool/ToolChain/AppleToolChain.cs
@@ -587,21 +587,9 @@ namespace UnrealBuildTool
 
 		private void PostBuildSync(ApplePostBuildSyncTarget Target, ILogger Logger)
 		{
-			// generate a run-only project file for codesigning, etc
-			DirectoryReference? GeneratedProjectFile;
-			Logger.LogInformation("Generating temporary Xcode project...");
-			DateTime Start = DateTime.Now;
-			IOSExports.GenerateRunOnlyXcodeProject(Target.ProjectFile, Target.Platform, Target.TargetName, false, Logger, out GeneratedProjectFile);
-
-			// run xcodebuild on the generated project to make the .app
-			bool bForDistribution = false;
-			Logger.LogInformation($"Finalizing .app via Xcode... [generation took {(DateTime.Now - Start).TotalMilliseconds / 1000.0}s]");
-			IOSExports.FinalizeAppWithModernXcode(GeneratedProjectFile!, Target.Platform, Target.TargetName, Target.Configuration.ToString(), "", bForDistribution, Logger);
+			MacExports.BuildWithModernXcode(Target.ProjectFile, Target.Platform, Target.Configuration, Target.TargetName, bArchiveForDistro: false, Logger);
 		}
 
-
-
-
 		private static FileItem GetPostBuildOutputFile(FileReference Executable, string TargetName, UnrealTargetPlatform Platform)
 		{
 			FileReference StagedExe;
