- Fixed mac (and maybe IOS) archiving in legacy mode

diff --git a/Engine/Source/Programs/AutomationTool/Apple/ApplePlatform.Automation.cs b/Engine/Source/Programs/AutomationTool/Apple/ApplePlatform.Automation.cs
--- a/Engine/Source/Programs/AutomationTool/Apple/ApplePlatform.Automation.cs
+++ b/Engine/Source/Programs/AutomationTool/Apple/ApplePlatform.Automation.cs
@@ -177,6 +177,7 @@ public abstract class ApplePlatform : Platform
 	{
 		if (!MacExports.UseModernXcode(Params.RawProjectPath))
 		{
+			base.GetFilesToArchive(Params, SC);
 			return;
 		}
 
