Gauntlet fix Regex in GetTestExitCode for RequestExitWithStatus

diff --git a/Engine/Source/Programs/AutomationTool/Gauntlet/Unreal/Utils/Gauntlet.UnrealLogParser.cs b/Engine/Source/Programs/AutomationTool/Gauntlet/Unreal/Utils/Gauntlet.UnrealLogParser.cs
--- a/Engine/Source/Programs/AutomationTool/Gauntlet/Unreal/Utils/Gauntlet.UnrealLogParser.cs
+++ b/Engine/Source/Programs/AutomationTool/Gauntlet/Unreal/Utils/Gauntlet.UnrealLogParser.cs
@@ -923,7 +923,7 @@ namespace Gauntlet
 				return true;
 			}
 
-			Reg = new Regex(@"RequestExitWithStatus\(\d+,\s*(\d+)\)");
+			Reg = new Regex(@"RequestExitWithStatus\(\d+,\s*(\d+).*\)");
 
 			M = Reg.Match(Content);
 
