* Ensure any processes (TM2, NintendoSDKDaemon, RunOnTarget) started by [REDACTED_PERSON] utilities (ControlTarget, RunOnTarget) are shut down when we're done.
* Create UsedPlatform HashSet as we add target devices in case their unregistered before we get to dispose().

diff --git a/Engine/Source/Programs/AutomationTool/Gauntlet/Framework/Devices/Gauntlet.DevicePool.cs b/Engine/Source/Programs/AutomationTool/Gauntlet/Framework/Devices/Gauntlet.DevicePool.cs
--- a/Engine/Source/Programs/AutomationTool/Gauntlet/Framework/Devices/Gauntlet.DevicePool.cs
+++ b/Engine/Source/Programs/AutomationTool/Gauntlet/Framework/Devices/Gauntlet.DevicePool.cs
@@ -218,6 +218,11 @@ namespace Gauntlet
 		/// </summary>
 		List<ITargetDevice> ReservedDevices = new List<ITargetDevice>();
 
+		/// <summary>
+		/// List of platforms we've had devices for
+		/// </summary>
+		HashSet<UnrealTargetPlatform?> UsedPlatforms = new HashSet<UnrealTargetPlatform?>();
+
 		/// <summary>
 		/// Device reservation service URL 
 		/// </summary>
@@ -373,7 +378,6 @@ namespace Gauntlet
 					}
 
 					// Additional devices cleanup
-					IEnumerable<UnrealTargetPlatform?> UsedPlatforms = ReservedDevices.Select(D => D.Platform).Distinct();
 					CleanupDevices(UsedPlatforms);
 
 					AvailableDevices.Clear();
@@ -592,6 +596,7 @@ namespace Gauntlet
 				AvailableDevices = AvailableDevices.Where(D => DeviceList.Contains(D) == false).ToList();
 
 				ReservedDevices.AddRange(DeviceList);
+				DeviceList.All(D => UsedPlatforms.Add(D.Platform));
 			}
 			return true;
 		}
@@ -766,6 +771,7 @@ namespace Gauntlet
 					else
 					{						
 						ReservedDevices.Add(TargetDevice);
+						UsedPlatforms.Add(TargetDevice.Platform);
 					}
 
 					ServiceDeviceInfo[TargetDevice] = Def;
