BuildGraph: Set branch, changelist and code changelist from environment variables even when P4 is not enabled.

diff --git a/Engine/Source/Programs/AutomationTool/BuildGraph/BuildGraph.cs b/Engine/Source/Programs/AutomationTool/BuildGraph/BuildGraph.cs
--- a/Engine/Source/Programs/AutomationTool/BuildGraph/BuildGraph.cs
+++ b/Engine/Source/Programs/AutomationTool/BuildGraph/BuildGraph.cs
@@ -264,13 +264,18 @@ namespace AutomationTool
 				}
 			}
 
+			// Get the standard P4 properties, defaulting to the environment variables if not set. This allows setting p4-like properties without having a P4 connection.
+			string Branch = P4Enabled ? P4Env.Branch : GetEnvVarOrNull(EnvVarNames.BuildRootP4);
+			string Change = P4Enabled ? P4Env.Changelist.ToString() : GetEnvVarOrNull(EnvVarNames.Changelist);
+			string CodeChange = P4Enabled ? P4Env.CodeChangelist.ToString() : GetEnvVarOrNull(EnvVarNames.CodeChangelist);
+
 			// Set up the standard properties which build scripts might need
 			Dictionary<string, string> DefaultProperties = new Dictionary<string,string>(StringComparer.InvariantCultureIgnoreCase);
-			DefaultProperties["Branch"] = P4Enabled ? P4Env.Branch : "Unknown";
-			DefaultProperties["Depot"] = P4Enabled ? DefaultProperties["Branch"].Substring(2).Split('/').First() : "Unknown";
-			DefaultProperties["EscapedBranch"] = P4Enabled ? CommandUtils.EscapePath(P4Env.Branch) : "Unknown";
-			DefaultProperties["Change"] = P4Enabled ? P4Env.Changelist.ToString() : "0";
-			DefaultProperties["CodeChange"] = P4Enabled ? P4Env.CodeChangelist.ToString() : "0";
+			DefaultProperties["Branch"] = Branch ?? "Unknown";
+			DefaultProperties["Depot"] = (Branch != null && Branch.StartsWith("//", StringComparison.Ordinal))? Branch.Substring(2).Split('/').First() : "Unknown";
+			DefaultProperties["EscapedBranch"] = String.IsNullOrEmpty(Branch) ? "Unknown" : CommandUtils.EscapePath(P4Env.Branch);
+			DefaultProperties["Change"] = Change ?? "0";
+			DefaultProperties["CodeChange"] = CodeChange ?? "0";
 			DefaultProperties["IsBuildMachine"] = IsBuildMachine ? "true" : "false";
 			DefaultProperties["HostPlatform"] = HostPlatform.Current.HostEditorPlatform.ToString();
 			DefaultProperties["RestrictedFolderNames"] = String.Join(";", RestrictedFolder.GetNames());
@@ -1173,6 +1178,15 @@ namespace AutomationTool
 			return true;
 		}
 
+		/// <summary>
+		/// Gets an environment variable, returning null if it's not set or empty.
+		/// </summary>
+		static string GetEnvVarOrNull(string Name)
+		{
+			string EnvVar = Environment.GetEnvironmentVariable(Name);
+			return String.IsNullOrEmpty(EnvVar) ? null : EnvVar;
+		}
+
 		/// <summary>
 		/// Generate HTML documentation for all the tasks
 		/// </summary>
