Missing file

diff --git a/Engine/Source/Programs/AutomationTool/AutomationUtils/CommandEnvironment.cs b/Engine/Source/Programs/AutomationTool/AutomationUtils/CommandEnvironment.cs
--- a/Engine/Source/Programs/AutomationTool/AutomationUtils/CommandEnvironment.cs
+++ b/Engine/Source/Programs/AutomationTool/AutomationUtils/CommandEnvironment.cs
@@ -17,7 +17,7 @@ namespace AutomationTool
 	/// <summary>
 	/// Defines the environment variable names that will be used to setup the environment.
 	/// </summary>
-	static class EnvVarNames
+	public static class EnvVarNames
 	{
 		// Command Environment
 		public const string LocalRoot = "uebp_LOCAL_ROOT";
