BuildGraph: Set the correct branch/change values for native buildgraph scripts.

diff --git a/Engine/Source/Programs/AutomationTool/BuildGraph/BuildGraph.cs b/Engine/Source/Programs/AutomationTool/BuildGraph/BuildGraph.cs
--- a/Engine/Source/Programs/AutomationTool/BuildGraph/BuildGraph.cs
+++ b/Engine/Source/Programs/AutomationTool/BuildGraph/BuildGraph.cs
@@ -138,12 +138,14 @@ namespace AutomationTool
 		/// <summary>
 		/// Constructor
 		/// </summary>
-		/// <param name="P4Env"></param>
-		internal BgEnvironment(P4Environment P4Env)
+		/// <param name="Branch">The current branch</param>
+		/// <param name="Change">Changelist being built</param>
+		/// <param name="CodeChange">Code changelist being built</param>
+		internal BgEnvironment(string Branch, int? Change, int? CodeChange)
 		{
-			this.Stream = P4Env?.Branch ?? "Unknown";
-			this.Change = P4Env?.Changelist ?? 0;
-			this.CodeChange = P4Env?.CodeChangelist ?? 0;
+			this.Stream = Branch ?? "Unknown";
+			this.Change = Change ?? 0;
+			this.CodeChange = CodeChange ?? 0;
 			this.IsBuildMachine = CommandUtils.IsBuildMachine;
 
 			ReadOnlyBuildVersion Version = ReadOnlyBuildVersion.Current;
@@ -266,16 +268,16 @@ namespace AutomationTool
 
 			// Get the standard P4 properties, defaulting to the environment variables if not set. This allows setting p4-like properties without having a P4 connection.
 			string Branch = P4Enabled ? P4Env.Branch : GetEnvVarOrNull(EnvVarNames.BuildRootP4);
-			string Change = P4Enabled ? P4Env.Changelist.ToString() : GetEnvVarOrNull(EnvVarNames.Changelist);
-			string CodeChange = P4Enabled ? P4Env.CodeChangelist.ToString() : GetEnvVarOrNull(EnvVarNames.CodeChangelist);
+			int? Change = P4Enabled ? P4Env.Changelist : GetEnvVarIntOrNull(EnvVarNames.Changelist);
+			int? CodeChange = P4Enabled ? P4Env.CodeChangelist : GetEnvVarIntOrNull(EnvVarNames.CodeChangelist);
 
 			// Set up the standard properties which build scripts might need
 			Dictionary<string, string> DefaultProperties = new Dictionary<string,string>(StringComparer.InvariantCultureIgnoreCase);
 			DefaultProperties["Branch"] = Branch ?? "Unknown";
 			DefaultProperties["Depot"] = (Branch != null && Branch.StartsWith("//", StringComparison.Ordinal))? Branch.Substring(2).Split('/').First() : "Unknown";
-			DefaultProperties["EscapedBranch"] = String.IsNullOrEmpty(Branch) ? "Unknown" : CommandUtils.EscapePath(P4Env.Branch);
-			DefaultProperties["Change"] = Change ?? "0";
-			DefaultProperties["CodeChange"] = CodeChange ?? "0";
+			DefaultProperties["EscapedBranch"] = String.IsNullOrEmpty(Branch) ? "Unknown" : CommandUtils.EscapePath(Branch);
+			DefaultProperties["Change"] = (Change ?? 0).ToString();
+			DefaultProperties["CodeChange"] = (CodeChange ?? 0).ToString();
 			DefaultProperties["IsBuildMachine"] = IsBuildMachine ? "true" : "false";
 			DefaultProperties["HostPlatform"] = HostPlatform.Current.HostEditorPlatform.ToString();
 			DefaultProperties["RestrictedFolderNames"] = String.Join(";", RestrictedFolder.GetNames());
@@ -397,7 +399,7 @@ namespace AutomationTool
 				}
 
 				BgGraphBuilder Builder = (BgGraphBuilder)Activator.CreateInstance(BuilderType);
-				BgGraph GraphSpec = Builder.CreateGraph(new BgEnvironment(P4Enabled ? P4Env : null));
+				BgGraph GraphSpec = Builder.CreateGraph(new BgEnvironment(Branch, Change, CodeChange));
 
 				(byte[] Data, BgThunkDef[] Methods) = BgCompiler.Compile(GraphSpec);
 
@@ -1187,6 +1189,19 @@ namespace AutomationTool
 			return String.IsNullOrEmpty(EnvVar) ? null : EnvVar;
 		}
 
+		/// <summary>
+		/// Gets an environment variable as an integer, returning null if it's not set or empty.
+		/// </summary>
+		static int? GetEnvVarIntOrNull(string Name)
+		{
+			string EnvVar = GetEnvVarOrNull(Name);
+			if (EnvVar != null && Int32.TryParse(EnvVar, out int Value))
+			{
+				return Value;
+			}
+			return null;
+		}
+
 		/// <summary>
 		/// Generate HTML documentation for all the tasks
 		/// </summary>
