Add temp storage timestamp exclusion for tbbmalloc.pdb.

diff --git a/Engine/Source/Programs/AutomationTool/BuildGraph/TempStorage.cs b/Engine/Source/Programs/AutomationTool/BuildGraph/TempStorage.cs
--- a/Engine/Source/Programs/AutomationTool/BuildGraph/TempStorage.cs
+++ b/Engine/Source/Programs/AutomationTool/BuildGraph/TempStorage.cs
@@ -1282,7 +1282,7 @@ namespace AutomationTool
 			{
 				return true;
 			}
-			if (FileName.Equals("tbbmalloc.dll", StringComparison.OrdinalIgnoreCase) || FileName.Equals("libtbbmalloc.dylib", StringComparison.OrdinalIgnoreCase))
+			if (FileName.Equals("tbbmalloc.dll", StringComparison.OrdinalIgnoreCase) || FileName.Equals("tbbmalloc.pdb", StringComparison.OrdinalIgnoreCase) || FileName.Equals("libtbbmalloc.dylib", StringComparison.OrdinalIgnoreCase))
 			{
 				return true;
 			}
