Change PerfEditorPIETest to use an isolated warming pass and to use a cached asset registry.
#rnx

diff --git a/Engine/Source/Programs/AutomationTool/Gauntlet/Unreal/RunUnreal.cs b/Engine/Source/Programs/AutomationTool/Gauntlet/Unreal/RunUnreal.cs
--- a/Engine/Source/Programs/AutomationTool/Gauntlet/Unreal/RunUnreal.cs
+++ b/Engine/Source/Programs/AutomationTool/Gauntlet/Unreal/RunUnreal.cs
@@ -215,12 +215,10 @@ namespace Gauntlet
 					}
 
 					// look for-args= and then -clientargs= and -editorargs etc
-					Role.ExtraArgs = string.Join(' ', Globals.Params.ParseValue("Args", string.Empty));
-					string ExtraRoleArgs = Globals.Params.ParseValue(Type.ToString() + "Args", string.Empty);
-					if (!string.IsNullOrEmpty(ExtraRoleArgs))
-					{
-						Role.ExtraArgs += " " + ExtraRoleArgs;
-					}
+					List<string> ArgsParams = Globals.Params.ParseValues("Args", false /* bCommaSeparated */);
+					List<string> RoleArgsParams = Globals.Params.ParseValues(Type.ToString() + "Args", false /* bCommaSeparated */);
+					ArgsParams.AddRange(RoleArgsParams);
+					Role.ExtraArgs = string.Join(' ', ArgsParams);
 
 					// look for -clientexeccmds=, -editorexeccmds= etc, these are separate from clientargs for sanity
 					string ExecCmds = Globals.Params.ParseValue("ExecCmds", "");
