AutomationTool: Fix condition in .csproj (via UDN)

#rnx

diff --git a/Engine/Source/Programs/AutomationTool/AutomationTool.csproj b/Engine/Source/Programs/AutomationTool/AutomationTool.csproj
--- a/Engine/Source/Programs/AutomationTool/AutomationTool.csproj
+++ b/Engine/Source/Programs/AutomationTool/AutomationTool.csproj
@@ -135,5 +135,5 @@
   </ItemGroup>
   
   <!-- Reference file is generated by GenerateProjectFiles scripts (calling UnrealBuildTool). It contains a list of all project files known to UBT based on how it was run. -->
-  <Import Project="..\..\..\Intermediate\ProjectFiles\AutomationTool.csproj.References" Condition="Exists('..\..\..\Intermediate\ProjectFiles\AutomationTool.csproj.References') And '$(VisualStudioDir)' != '' And '$(AutomationToolProjectOnly)' == ''" />
+  <Import Project="..\..\..\Intermediate\ProjectFiles\AutomationTool.csproj.References" Condition="Exists('..\..\..\Intermediate\ProjectFiles\AutomationTool.csproj.References') And '$(AutomationToolProjectOnly)' == ''" />
 </Project>
\ No newline at end of file
