[Automation] Passthrough BUILDMACHINE to relevant processes when running on a buildmachine.

diff --git a/Engine/Source/Programs/AutomationTool/Gauntlet/Unreal/Base/Gauntlet.UnrealTestConfiguration.cs b/Engine/Source/Programs/AutomationTool/Gauntlet/Unreal/Base/Gauntlet.UnrealTestConfiguration.cs
--- a/Engine/Source/Programs/AutomationTool/Gauntlet/Unreal/Base/Gauntlet.UnrealTestConfiguration.cs
+++ b/Engine/Source/Programs/AutomationTool/Gauntlet/Unreal/Base/Gauntlet.UnrealTestConfiguration.cs
@@ -942,7 +942,12 @@ namespace Gauntlet
 					AppConfig.CommandLineParams.GameMap = MapChoice;
 				}
 			}
-		}			
+
+			if(CommandUtils.IsBuildMachine)
+			{
+				AppConfig.CommandLineParams.AddUnique("BUILDMACHINE");
+			}
+		}
 	}
 
 }
