Improve Gauntlet warning message during Device resource starvation

#rnx

diff --git a/Engine/Source/Programs/AutomationTool/Gauntlet/Framework/Devices/Gauntlet.DevicePool.cs b/Engine/Source/Programs/AutomationTool/Gauntlet/Framework/Devices/Gauntlet.DevicePool.cs
--- a/Engine/Source/Programs/AutomationTool/Gauntlet/Framework/Devices/Gauntlet.DevicePool.cs
+++ b/Engine/Source/Programs/AutomationTool/Gauntlet/Framework/Devices/Gauntlet.DevicePool.cs
@@ -11,6 +11,7 @@ using System.Linq;
 using System.Text.Json;
 using System.Text.Json.Serialization;
 using EpicGames.Core;
+using System.Data;
 
 namespace Gauntlet
 {
@@ -174,6 +175,15 @@ namespace Gauntlet
 		{
 			return ToString().GetHashCode();
 		}
+
+		/// <summary>
+		/// Format to string the device constraint including its identify
+		/// </summary>
+		/// <returns></returns>
+		public string FormatWithIdentifier()
+		{
+			return string.Format("{0}:{1}", Platform, Model == string.Empty ? PerfSpec.ToString() : Model);
+		}
 	}
 
 	/// <summary>
@@ -191,6 +201,150 @@ namespace Gauntlet
 		public UnrealTargetPlatform Platform;
 	}
 
+	/// <summary>
+	/// Device reservation states
+	/// </summary>
+	public class DeviceReservationStates
+	{
+		public DeviceReservationStates()
+		{
+			ResetRequirements();
+			ProblemCollection = new Dictionary<string, int>();
+		}
+
+		public void ResetRequirements()
+		{
+			RequiredCollection = new Dictionary<string, ResevationState>();
+		}
+
+		public void AddRequired(string DeviceKey, int Count = 1)
+		{
+			if(RequiredCollection.ContainsKey(DeviceKey))
+			{
+				var ReservationState = RequiredCollection[DeviceKey];
+				ReservationState.NumRequired += Count;
+				ReservationState.NumLacking += Count;
+			}
+			else
+			{
+				RequiredCollection.Add(DeviceKey, new ResevationState(Count));
+			}
+		}
+
+		public void MarkReserved(string DeviceKey, int Count = 1)
+		{
+			var ReservationState = FindClosestState(DeviceKey);
+			if (ReservationState != null && ReservationState.NumLacking > 0)
+			{
+				ReservationState.NumLacking -= Count;
+			}
+		}
+
+		public void MarkReserved(Device InDevice, int Count = 1)
+		{
+			MarkReserved(KeyFromDevice(InDevice), Count);
+		}
+
+		public void MarkProblem(string DeviceKey, int Count = 1)
+		{
+			if (DeviceKey.Contains(":"))
+			{
+				DeviceKey = DeviceKey.Substring(0, DeviceKey.IndexOf(":"));
+			}
+
+			if (!ProblemCollection.ContainsKey(DeviceKey))
+			{
+				ProblemCollection[DeviceKey] = 0;
+			}
+
+			ProblemCollection[DeviceKey] += Count;
+		}
+
+		public void MarkProblem(Device InDevice, int Count = 1)
+		{
+			MarkProblem(KeyFromDevice(InDevice), Count);
+		}
+
+		/// <summary>
+		/// Return the list of requirements for reservation
+		/// </summary>
+		/// <returns></returns>
+		public string[] GetRequirements()
+		{
+			return RequiredCollection.SelectMany(P => Enumerable.Repeat(P.Key, P.Value.NumRequired)).ToArray();
+		}
+
+		/// <summary>
+		/// Return the list of devices and their reservation status as an array
+		/// </summary>
+		/// <returns></returns>
+		public string[] GetStatuses()
+		{
+			return RequiredCollection.Select(P => $"{P.Key.Replace($":{EPerfSpec.Unspecified}", "")} " +
+											$"[" +
+												$"lacking:{P.Value.NumLacking}, " +
+												$"reserved:{P.Value.NumRequired - P.Value.NumLacking}" +
+											$"]").ToArray();
+		}
+
+		/// <summary>
+		/// Return the list of devices marked with problem
+		/// </summary>
+		/// <returns></returns>
+		public string[] GetProblems()
+		{
+			return ProblemCollection.Select(P => $"{P.Value} {P.Key}").ToArray();
+		}
+
+		public int RequiredCount()
+		{
+			return RequiredCollection.Sum(C => C.Value.NumRequired);
+		}
+
+		public int ProblemCount()
+		{
+			return ProblemCollection.Sum(C => C.Value);
+		}
+
+		private class ResevationState
+		{
+			public ResevationState(int Count = 1)
+			{
+				NumRequired = Count;
+				NumLacking = Count;
+			}
+
+			public int NumRequired;
+			public int NumLacking;
+		}
+
+		private Dictionary<string, ResevationState> RequiredCollection;
+		private Dictionary<string, int> ProblemCollection;
+
+		private ResevationState FindClosestState(string DeviceKey)
+		{
+			if (RequiredCollection.ContainsKey(DeviceKey))
+			{
+				return RequiredCollection[DeviceKey];
+			}
+			else if (DeviceKey.Contains(":"))
+			{
+				string DeviceUndefinedKey = DeviceKey.Substring(0, DeviceKey.IndexOf(":") + 1) + EPerfSpec.Unspecified.ToString();
+				if (RequiredCollection.ContainsKey(DeviceUndefinedKey))
+				{
+					return RequiredCollection[DeviceUndefinedKey];
+				}
+			}
+
+			return null;
+		}
+
+		private string KeyFromDevice(Device InDevice)
+		{
+			return InDevice.Type.Replace("-DevKit", "", StringComparison.OrdinalIgnoreCase) + ":" + (InDevice.PerfSpec == string.Empty ? InDevice.Model : InDevice.PerfSpec);
+		}
+	}
+
 	/// <summary>
 	/// Singleton class that's responsible for providing a list of devices and reserving them. Code should call
 	/// EnumerateDevices to build a list of desired devices, which must then be reserved by calling ReserveDevices.
@@ -223,6 +377,11 @@ namespace Gauntlet
 		/// </summary>
 		HashSet<UnrealTargetPlatform?> UsedPlatforms = new HashSet<UnrealTargetPlatform?>();
 
+		/// <summary>
+		/// Device reservation left to fulfill
+		/// </summary>
+		private DeviceReservationStates ReservationStates;
+
 		/// <summary>
 		/// Device reservation service URL 
 		/// </summary>
@@ -291,6 +450,8 @@ namespace Gauntlet
 
 			// create two local devices by default?
 			AddLocalDevices(2);
+
+			ReservationStates = new DeviceReservationStates();
 		}
 
 		~DevicePool()
@@ -513,6 +674,24 @@ namespace Gauntlet
 			Reservation.ReportDeviceError(DeviceURL, ServiceDeviceName, ErrorMessage);
 		}
 
+		/// <summary>
+		/// Get the reservation status as an array
+		/// </summary>
+		/// <returns></returns>
+		public string[] GetReservationStatuses()
+		{
+			return ReservationStates != null? ReservationStates.GetStatuses() : null;
+		}
+
+		/// <summary>
+		/// Get the reserved device marked with problems as an array
+		/// </summary>
+		/// <returns></returns>
+		public string[] GetDeviceProblems()
+		{
+			return ReservationStates != null ? ReservationStates.GetProblems() : null;
+		}
+
 		/// <summary>
 		/// Force clean the cache path by setting all of the files in the cache path
 		/// to have "normal" attributes (i.e. not read-only) before attempting to delete.
@@ -657,6 +836,8 @@ namespace Gauntlet
 				return false;
 			}
 
+			ReservationStates.ResetRequirements();
+
 			Dictionary<UnrealTargetPlatform, string> DeviceMap = new Dictionary<UnrealTargetPlatform, string>();
 
 			foreach (string Platform in UnrealTargetPlatform.GetValidPlatformNames())
@@ -664,8 +845,6 @@ namespace Gauntlet
 				DeviceMap.Add(UnrealTargetPlatform.Parse(Platform), Platform);
 			}
 
-			List<string> Devices = new List<string>();
-
 			// convert devices to request list
 			foreach (KeyValuePair<UnrealDeviceTargetConstraint, int> Entry in DeviceTypes)
 			{
@@ -695,12 +874,7 @@ namespace Gauntlet
 				for (int i = 0; i < Entry.Value; i++)
 				{
 					// @todo: if any additional reservation requirements, encode constraint into json
-					string Constraint = Entry.Key.PerfSpec.ToString();
-					if (!string.IsNullOrEmpty(Entry.Key.Model))
-					{
-						Constraint = Entry.Key.Model;
-					}
-					Devices.Add(DeviceMap[Entry.Key.Platform.Value] + ":" + Constraint);
+					ReservationStates.AddRequired(Entry.Key.FormatWithIdentifier(), Entry.Value);
 				}				
 			}
 
@@ -714,15 +888,24 @@ namespace Gauntlet
 
 				try
 				{
-					DeviceReservation = new DeviceReservationAutoRenew(DeviceURL, 0, PoolID, Devices.ToArray());
+					DeviceReservation = new DeviceReservationAutoRenew(DeviceURL, 0, PoolID, ReservationStates.GetRequirements());
 				}
 				catch (Exception Ex)
 				{
-					Log.Info(KnownLogEvents.Gauntlet_DeviceEvent, "Unable to make device registration: {Exception}", Ex);
+					Log.Info(KnownLogEvents.Gauntlet_DeviceEvent, "Unable to make device registration with constraint {Constraint}: {Exception}", string.Join(", ", ReservationStates.GetRequirements()), Ex);
 					return false;
 				}
 
-				if (DeviceReservation == null || DeviceReservation.Devices.Count != Devices.Count())
+				if (DeviceReservation != null)
+				{
+					// Update RequiredContraints list with what the reservation satisfied.
+					foreach (var Device in DeviceReservation.Devices)
+					{
+						ReservationStates.MarkReserved(Device);
+					}
+				}
+
+				if (DeviceReservation == null || DeviceReservation.Devices.Count != ReservationStates.RequiredCount())
 				{
 					return false;
 				}
@@ -753,6 +936,7 @@ namespace Gauntlet
 					{
 
 						ReportDeviceError(Device.Name, "CreateDeviceError");
+						ReservationStates.MarkProblem(Device);
 
 						// If some devices from reservation have been created, release them which will also dispose of reservation
 						if (ReservedDevices.Count > 0)
@@ -766,6 +950,7 @@ namespace Gauntlet
 						}
 
 						Log.Info(KnownLogEvents.Gauntlet_DeviceEvent, "Unable to make device registration: device registration failed for {Platform}:{Name}", Def.Platform, Def.Name);
+						Log.Verbose(KnownLogEvents.Gauntlet_DeviceEvent, " Device Reservation status from {Pool} pool: {Status}", PoolID, string.Join(", ", ReservationStates.GetStatuses()));
 						return false;
 					}
 					else
@@ -779,13 +964,13 @@ namespace Gauntlet
 				}
 
 				Log.Info(KnownLogEvents.Gauntlet_DeviceEvent, "Successfully reserved service devices");
-				Devices.ForEach(Device => Log.Verbose(KnownLogEvents.Gauntlet_DeviceEvent, "    Device: {Name}", Device));
-
+				Log.Verbose(KnownLogEvents.Gauntlet_DeviceEvent, " Device Reservation status from {Pool} pool: {Status}", PoolID, string.Join(", ", ReservationStates.GetStatuses()));
+				ReservationStates = null;
 				return true;
 			}
 
 			Log.Info(KnownLogEvents.Gauntlet_DeviceEvent, "Unable to reserve service devices:");
-			Devices.ForEach(Device => Log.Info(KnownLogEvents.Gauntlet_DeviceEvent, "    Device: {Name}", Device));
+			Log.Verbose(KnownLogEvents.Gauntlet_DeviceEvent, " Device Reservation status from {Pool} pool: {Status}", Globals.DevicePoolId, string.Join(", ", ReservationStates.GetStatuses()));
 			return false;
 		}
 
diff --git a/Engine/Source/Programs/AutomationTool/Gauntlet/Framework/Gauntlet.TestExecutor.cs b/Engine/Source/Programs/AutomationTool/Gauntlet/Framework/Gauntlet.TestExecutor.cs
--- a/Engine/Source/Programs/AutomationTool/Gauntlet/Framework/Gauntlet.TestExecutor.cs
+++ b/Engine/Source/Programs/AutomationTool/Gauntlet/Framework/Gauntlet.TestExecutor.cs
@@ -44,7 +44,8 @@ namespace Gauntlet
 				NotStarted,
 				TimedOut,
 				Passed,
-				Failed
+				Failed,
+				Skipped
 			}
 
 			public ITestNode		TestNode;
@@ -258,7 +259,6 @@ namespace Gauntlet
 							{
 								// track the time that this test should have been able to run due to no other tests
 								// consuming resources (at least locally...)
-								// TODO - how can tests express resource requirements in a generic way?
 								// TODO - what about the situation where no tests can run so all FirstCheck times are set, but 
 								// then a test starts and consumes all resources?
 								if (RunningTests.Count() == 0 && StartingTests.Count() == 0)
@@ -271,11 +271,22 @@ namespace Gauntlet
 									double TimeWaiting = (DateTime.Now - NodeInfo.FirstReadyCheckTime).TotalSeconds;
 									if (TimeWaiting >= Options.Wait)
 									{
-										Log.Warning(KnownLogEvents.Gauntlet_TestEvent, "Test {Name} has been waiting to run resource-free for {Time:00} seconds. Removing from wait list", Node, TimeWaiting);
-										Node.AddTestEvent(new UnrealTestEvent(EventSeverity.Error, "Insufficient devices found", new List<string> {string.Format("Test {0} was unable to find enough devices after trying for {1:00} seconds.", Node, TimeWaiting), "This is not a test-related failure."}));
+										string PoolId = string.IsNullOrEmpty(Globals.DevicePoolId) ? "Local" : Globals.DevicePoolId;
+										Log.Warning(KnownLogEvents.Gauntlet_DeviceEvent, "Test {TestName} has been waiting to run resource-free for {Time:00} seconds. Removing from wait list", Node, TimeWaiting);
+										string[] DeviceReservationStatus = DevicePool.Instance.GetReservationStatuses();
+										if (DeviceReservationStatus != null && DeviceReservationStatus.Length > 0)
+										{
+											Log.Warning(KnownLogEvents.Gauntlet_DeviceEvent, "Device Reservation status from {Pool} pool: {Status}", PoolId, string.Join(", ", DeviceReservationStatus));
+										}
+										string[] DeviceProblemStatus = DevicePool.Instance.GetDeviceProblems();
+										if (DeviceProblemStatus != null && DeviceProblemStatus.Length > 0)
+										{
+											Log.Warning(KnownLogEvents.Gauntlet_DeviceEvent, "Reserved Devices marked with problems from {Pool} pool: {Status}", PoolId, string.Join(", ", DeviceProblemStatus));
+										}
+										Node.AddTestEvent(new UnrealTestEvent(EventSeverity.Warning, "Insufficient devices found", new List<string> {string.Format("Test {0} was unable to find enough devices after trying for {1:00} seconds.", Node, TimeWaiting), "This is not a test-related failure."}));
 										PendingTests[i] = null;
 										NodeInfo.TimeSetupBegan = NodeInfo.TimeSetupEnded = NodeInfo.TimeTestEnded = DateTime.Now;
-										NodeInfo.Result = TestExecutionInfo.ExecutionResult.TimedOut;
+										NodeInfo.Result = TestExecutionInfo.ExecutionResult.Skipped;
 										CompletedTests.Add(NodeInfo);
 									}
 								}
@@ -451,7 +462,7 @@ namespace Gauntlet
 						// status msg, kept uniform to avoid spam on notifiers (ie. don't include timestamps, etc) 
 						string Msg = string.Format("Test {0} {1}", T.TestNode, T.Result);
 
-						bool TestHadErrors = T.Result != TestExecutionInfo.ExecutionResult.Passed;
+						bool TestHadErrors = T.Result != TestExecutionInfo.ExecutionResult.Passed && T.Result != TestExecutionInfo.ExecutionResult.Skipped;
 						bool TestHadWarnings = T.TestNode.GetWarnings().Any();
 
 						if (TestHadErrors)
