Don't connect to host zenserver via uecommandline.txt if we are creating PAK files.
Original author: Dan Engelbrecht

diff --git a/Engine/Source/Programs/AutomationTool/Scripts/RunProjectCommand.Automation.cs b/Engine/Source/Programs/AutomationTool/Scripts/RunProjectCommand.Automation.cs
--- a/Engine/Source/Programs/AutomationTool/Scripts/RunProjectCommand.Automation.cs
+++ b/Engine/Source/Programs/AutomationTool/Scripts/RunProjectCommand.Automation.cs
@@ -823,7 +823,8 @@ namespace AutomationScripts
 		{
 			string FileHostParams = "";
 			FileReference ZenStoreMarkerFile = FindZenProjectStoreMarker(Params, SC);
-			if (!Params.CookOnTheFly && !Params.FileServer && ZenStoreMarkerFile == null)
+			bool UseZenServerHost = (ZenStoreMarkerFile != null) && !ShouldCreatePak(Params);
+			if (!Params.CookOnTheFly && !Params.FileServer && !UseZenServerHost)
 			{
 				return FileHostParams;
 			}
