UAT: Reduce number of decimal places when logging how long it took to execute an external program.

diff --git a/Engine/Source/Programs/AutomationTool/AutomationUtils/ProcessUtils.cs b/Engine/Source/Programs/AutomationTool/AutomationUtils/ProcessUtils.cs
--- a/Engine/Source/Programs/AutomationTool/AutomationUtils/ProcessUtils.cs
+++ b/Engine/Source/Programs/AutomationTool/AutomationUtils/ProcessUtils.cs
@@ -950,7 +950,7 @@ namespace AutomationTool
 				}
 				if (!Options.HasFlag(ERunOptions.NoLoggingOfRunCommand) || Options.HasFlag(ERunOptions.LoggingOfRunDuration))
 				{
-					LogWithVerbosity(SpewVerbosity, "Took {0}s to run {1}, ExitCode={2}", BuildDuration / 1000, Path.GetFileName(App), Result.ExitCode);
+					LogWithVerbosity(SpewVerbosity, "Took {0:n2}s to run {1}, ExitCode={2}", BuildDuration / 1000, Path.GetFileName(App), Result.ExitCode);
 				}
 				Result.OnProcessExited();
 				Result.DisposeProcess();
