Add ErrorIfNotFound option to MoveTask to allow for throwing an error if the file is not found. This option is disabled by default.

diff --git a/Engine/Source/Programs/AutomationTool/BuildGraph/Tasks/MoveTask.cs b/Engine/Source/Programs/AutomationTool/BuildGraph/Tasks/MoveTask.cs
--- a/Engine/Source/Programs/AutomationTool/BuildGraph/Tasks/MoveTask.cs
+++ b/Engine/Source/Programs/AutomationTool/BuildGraph/Tasks/MoveTask.cs
@@ -51,6 +51,12 @@ namespace AutomationTool.Tasks
 		/// </summary>
 		[TaskParameter(Optional = true, ValidationType = TaskParameterValidationType.TagList)]
 		public string Tag;
+
+		/// <summary>
+		/// Whether or not to throw an error if no files were found to copy
+		/// </summary>
+		[TaskParameter(Optional = true)]
+		public bool ErrorIfNotFound = false;
 	}
 
 	/// <summary>
@@ -101,7 +107,14 @@ namespace AutomationTool.Tasks
 			// Check we got some files
 			if(TargetFileToSourceFile.Count == 0)
 			{
-				Logger.LogInformation("No files found matching '{SourcePattern}'", SourcePattern);
+				if (Parameters.ErrorIfNotFound)
+				{
+					Logger.LogError("No files found matching '{SourcePattern}'", SourcePattern);
+				}
+				else
+				{
+					Logger.LogInformation("No files found matching '{SourcePattern}'", SourcePattern);
+				}
 				return Task.CompletedTask;
 			}
 
