LLT: collect log output option for certain platforms

diff --git a/Engine/Source/Programs/AutomationTool/LowLevelTests/RunLowLevelTests.cs b/Engine/Source/Programs/AutomationTool/LowLevelTests/RunLowLevelTests.cs
--- a/Engine/Source/Programs/AutomationTool/LowLevelTests/RunLowLevelTests.cs
+++ b/Engine/Source/Programs/AutomationTool/LowLevelTests/RunLowLevelTests.cs
@@ -180,6 +180,8 @@ namespace LowLevelTests
 
 		public bool LogReportContents;
 
+		public bool CaptureOutput;
+
 		public Type BuildSourceType { get; protected set; }
 
 		[AutoParam(UnrealTargetConfiguration.Development)]
@@ -221,6 +223,8 @@ namespace LowLevelTests
 
 			LogReportContents = Params.ParseParam("printreport");
 
+			CaptureOutput = Params.ParseParam("captureoutput");
+
 			Build = Params.ParseValue("build=", null).Replace('/', Path.DirectorySeparatorChar).Replace('\\', Path.DirectorySeparatorChar);
 			TestApp = Globals.Params.ParseValue("testapp=", "");
 
diff --git a/Engine/Source/Programs/AutomationTool/LowLevelTests/Tests/Gauntlet.LowLevelTestNode.cs b/Engine/Source/Programs/AutomationTool/LowLevelTests/Tests/Gauntlet.LowLevelTestNode.cs
--- a/Engine/Source/Programs/AutomationTool/LowLevelTests/Tests/Gauntlet.LowLevelTestNode.cs
+++ b/Engine/Source/Programs/AutomationTool/LowLevelTests/Tests/Gauntlet.LowLevelTestNode.cs
@@ -115,6 +115,15 @@ namespace LowLevelTests
 				{
 					ParseLowLevelTestsLog();
 
+					// Print stdout when -captureoutput, certain platforms don't always redirect stdout
+					if (CurrentProcessedLines != null && Context.Options.CaptureOutput)
+					{
+						foreach (string OutputLine in CurrentProcessedLines)
+						{
+							Console.WriteLine(OutputLine);
+						}
+					}
+
 					if (CheckForTimeout())
 					{
 						Log.Error("Timeout detected from application logged events, stopping.");
