[LaunchIt] Have a way for ParamHelpAttribute markup on class types to explicitly override the same key from a parent class.

diff --git a/Engine/Source/Programs/AutomationTool/AutomationUtils/HelpUtils.cs b/Engine/Source/Programs/AutomationTool/AutomationUtils/HelpUtils.cs
--- a/Engine/Source/Programs/AutomationTool/AutomationUtils/HelpUtils.cs
+++ b/Engine/Source/Programs/AutomationTool/AutomationUtils/HelpUtils.cs
@@ -83,7 +83,8 @@ namespace AutomationTool
 			Store,
 			Store_True,
 			Store_False,
-			Append
+			Append,
+			Override
 		}
 
 		private ParamHelpAttribute(string Description) : base(Description)
