Add ability to specify Run commands in the editor

diff --git a/Engine/Source/Programs/AutomationTool/Scripts/OpenEditor.Automation.cs b/Engine/Source/Programs/AutomationTool/Scripts/OpenEditor.Automation.cs
--- a/Engine/Source/Programs/AutomationTool/Scripts/OpenEditor.Automation.cs
+++ b/Engine/Source/Programs/AutomationTool/Scripts/OpenEditor.Automation.cs
@@ -17,6 +17,7 @@ namespace AutomationTool
 
 	[Help("Opens the specified project.")]
 	[Help("project=<QAGame>", "Project to open. Will search current path and paths in ueprojectdirs. If omitted will open vanilla UnrealEditor")]
+	[ParamHelp("Run", "The name of the commandlet to run")]
 	public class OpenEditor : BuildCommand
 	{
 		// exposed as a property so projects can derive and set this directly
