Don't dereference ReservationStates on successful reservation

#rnx

diff --git a/Engine/Source/Programs/AutomationTool/Gauntlet/Framework/Devices/Gauntlet.DevicePool.cs b/Engine/Source/Programs/AutomationTool/Gauntlet/Framework/Devices/Gauntlet.DevicePool.cs
--- a/Engine/Source/Programs/AutomationTool/Gauntlet/Framework/Devices/Gauntlet.DevicePool.cs
+++ b/Engine/Source/Programs/AutomationTool/Gauntlet/Framework/Devices/Gauntlet.DevicePool.cs
@@ -965,7 +965,6 @@ namespace Gauntlet
 
 				Log.Info(KnownLogEvents.Gauntlet_DeviceEvent, "Successfully reserved service devices");
 				Log.Verbose(KnownLogEvents.Gauntlet_DeviceEvent, " Device Reservation status from {Pool} pool: {Status}", PoolID, string.Join(", ", ReservationStates.GetStatuses()));
-				ReservationStates = null;
 				return true;
 			}
 
