Use UnRealPak to copy loose files using a special package enabling zenstore usage.

#precommit 6461fa95592448f08b9df249

diff --git a/Engine/Source/Developer/IoStoreUtilities/Private/IoStoreLooseFiles.cpp b/Engine/Source/Developer/IoStoreUtilities/Private/IoStoreLooseFiles.cpp
new file mode 100644
--- /dev/null
+++ b/Engine/Source/Developer/IoStoreUtilities/Private/IoStoreLooseFiles.cpp
@@ -0,0 +1,259 @@
+// Copyright Epic Games, Inc. All Rights Reserved.
+
+#include "IoStoreLooseFiles.h"
+
+#include "HAL/FileManager.h"
+#include "IO/IoStore.h"
+#include "Misc/FileHelper.h"
+#include "Misc/Paths.h"
+#include "Serialization/JsonReader.h"
+#include "Serialization/JsonSerializer.h"
+#include "Serialization/JsonWriter.h"
+#include "Tasks/Task.h"
+#include "Tasks/Pipe.h"
+
+namespace
+{
+
+using FJsonWriter = TSharedPtr<TJsonWriter<TCHAR, TPrettyJsonPrintPolicy<TCHAR>>>; 
+using FJsonWriterFactory = TJsonWriterFactory<TCHAR, TPrettyJsonPrintPolicy<TCHAR>>;
+
+class FLooseFilesIoStoreWriter final
+	: public IIoStoreWriter
+{
+	struct FWriteStats
+	{
+		std::atomic<uint32> PendingCount{0};
+		std::atomic<uint32> TotalCount{0};
+		std::atomic<uint64> TotalByteCount{0};
+	};
+
+	struct FPendingWrite
+	{
+		FPendingWrite(const FIoChunkId& Id, IIoStoreWriteRequest* Request, const FIoWriteOptions& Options)
+			: ChunkId(Id)
+			, WriteRequest(Request)
+			, WriteOptions(Options)
+			, TaskPipe(UE_SOURCE_LOCATION)
+		{ }
+		
+		FIoChunkId ChunkId;
+		TUniquePtr<IIoStoreWriteRequest> WriteRequest;
+		FIoWriteOptions WriteOptions;
+		UE::Tasks::FPipe TaskPipe;
+		FString ErrorText;
+		uint64 UncompressedSize{0};
+		uint64 CompressedSize{0};
+	};
+
+	const uint32 MaxConcurrentWrites;
+
+public:
+	FLooseFilesIoStoreWriter(const FLooseFilesWriterSettings& WriterSettings,
+		uint32 MaxConcurrentWrites)
+		: MaxConcurrentWrites(MaxConcurrentWrites)
+		, Settings(WriterSettings)
+	{
+		Settings.TocFilePath = FPaths::SetExtension(Settings.TocFilePath, TEXT("utocmanifest"));
+	}
+
+	virtual ~FLooseFilesIoStoreWriter()
+	{ }
+
+	virtual void SetReferenceChunkDatabase(TSharedPtr<IIoStoreWriterReferenceChunkDatabase> ReferenceChunkDatabase) override
+	{ }
+
+	virtual void SetHashDatabase(TSharedPtr<IIoStoreWriterHashDatabase> HashDatabase, bool bVerifyHashDatabase) override
+	{ }
+
+	virtual void EnableDiskLayoutOrdering(const TArray<TUniquePtr<FIoStoreReader>>& PatchSourceReaders = TArray<TUniquePtr<FIoStoreReader>>())
+	{ }
+
+	virtual void Append(const FIoChunkId& ChunkId, FIoBuffer Chunk, const FIoWriteOptions& WriteOptions, uint64 OrderHint = MAX_uint64) override
+	{
+		struct FWriteRequest
+			: IIoStoreWriteRequest
+		{
+			FWriteRequest(FIoBuffer InSourceBuffer, uint64 InOrderHint)
+				: OrderHint(InOrderHint)
+			{
+				SourceBuffer = InSourceBuffer;
+				SourceBuffer.MakeOwned();
+			}
+
+			virtual ~FWriteRequest() = default;
+
+			void PrepareSourceBufferAsync(FGraphEventRef CompletionEvent) override
+			{
+				CompletionEvent->DispatchSubsequents();
+			}
+
+			const FIoBuffer* GetSourceBuffer() override
+			{
+				return &SourceBuffer;
+			}
+
+			void FreeSourceBuffer() override
+			{
+			}
+
+			uint64 GetOrderHint() override
+			{
+				return OrderHint;
+			}
+
+			TArrayView<const FFileRegion> GetRegions()
+			{
+				return TArrayView<const FFileRegion>();
+			}
+
+			FIoBuffer SourceBuffer;
+			uint64 OrderHint;
+		};
+
+		Append(ChunkId, new FWriteRequest(Chunk, OrderHint), WriteOptions);
+	}
+
+	virtual void Append(const FIoChunkId& ChunkId, IIoStoreWriteRequest* Request, const FIoWriteOptions& WriteOptions) override
+	{
+		check(Request);
+
+		for(;;)
+		{
+			{
+				FScopeLock _(&PendingLock);
+				if (WriteStats.PendingCount.load() < MaxConcurrentWrites)
+				{
+					WriteStats.PendingCount++;
+					FPendingWrite* PendingWrite = PendingWrites.Add_GetRef(MakeUnique<FPendingWrite>(ChunkId, Request, WriteOptions)).Get();
+					
+					// Setup the task pipe when holding the lock to make it easy to flush all pending task pipe(s)
+
+					FGraphEventRef Event = FGraphEvent::CreateGraphEvent();
+					PendingWrite->WriteRequest->PrepareSourceBufferAsync(Event);
+					UE::Tasks::FTask ReadChunkTask = PendingWrite->TaskPipe.Launch(TEXT("ReadChunk"), [PendingWrite, Event]() mutable
+					{
+						Event->Wait();
+					});
+
+					UE::Tasks::FTask WriteChunkTask = PendingWrite->TaskPipe.Launch(TEXT("WriteChunk"), [this, PendingWrite]() mutable
+					{
+						if (const FIoBuffer* SourceBuffer = PendingWrite->WriteRequest->GetSourceBuffer())
+						{
+							FString FileName = PendingWrite->WriteOptions.FileName;
+							FString FilePath = FPaths::ConvertRelativePathToFull(Settings.TargetRootPath, FileName);
+							FString FileDirectory = FPaths::GetPath(FilePath);
+							if (IFileManager::Get().MakeDirectory(*FileDirectory, true))
+							{
+								if (TUniquePtr<FArchive> Ar(IFileManager::Get().CreateFileWriter(*FilePath)); Ar.IsValid())
+								{
+									const uint64 ChunkSize = SourceBuffer->GetSize();
+									const uint32 CurrentCount = WriteStats.TotalCount.fetch_add(1) + 1;
+									UE_CLOG((CurrentCount % 128 == 0), LogIoStore, Display,
+										TEXT("Writing loose file chunk #%u '%s' -> '%s' (%llu Bytes)"),
+										CurrentCount, *PendingWrite->WriteOptions.FileName, *FileName, ChunkSize);
+
+									Ar->Serialize((void*)SourceBuffer->GetData(), ChunkSize);
+									PendingWrite->UncompressedSize = PendingWrite->CompressedSize = ChunkSize;
+									PendingWrite->WriteRequest->FreeSourceBuffer();
+									WriteStats.TotalByteCount += ChunkSize;
+								}
+								else
+								{
+									PendingWrite->ErrorText = FString::Printf(TEXT("Failed to create directory '%s'"), *FileDirectory);
+								}
+							}
+							else
+							{
+								PendingWrite->ErrorText = FString::Printf(TEXT("Failed to create parent folder for file '%s'"), *FileName);
+							}
+						}
+						else
+						{
+							PendingWrite->ErrorText = TEXT("Invalid source buffer");
+						}
+
+						PendingWrite->WriteRequest.Reset();
+
+						WriteStats.PendingCount--;
+						WriteCompletedEvent->Trigger();
+					}, ReadChunkTask);
+					
+					break;
+				}
+			}
+
+			WriteCompletedEvent->Wait();
+		}
+	}
+
+	virtual TIoStatusOr<FIoStoreWriterResult> GetResult() override
+	{
+		FScopeLock _(&PendingLock);
+
+		for (TUniquePtr<FPendingWrite>& PendingWrite : PendingWrites)
+		{
+			PendingWrite->TaskPipe.WaitUntilEmpty();
+		}
+		
+		FIoStoreWriterResult WriteResult;
+		WriteResult.ContainerId = FIoContainerId::FromName(Settings.ContainerName);
+		WriteResult.ContainerName = Settings.ContainerName.ToString();
+		
+		FString Json;
+		FJsonWriter JsonWriter = FJsonWriterFactory::Create(&Json);
+		JsonWriter->WriteObjectStart();
+		JsonWriter->WriteValue(TEXT("ContainerName"), Settings.ContainerName.ToString());
+		JsonWriter->WriteArrayStart(TEXT("Chunks"));
+		
+		for (TUniquePtr<FPendingWrite>& PendingWrite : PendingWrites)
+		{
+			if (PendingWrite->ErrorText.IsEmpty() == false)
+			{
+				UE_LOG(LogIoStore, Display, TEXT("Failed to write on demand chunk '%s', reason '%s'"), *LexToString(PendingWrite->ChunkId), *PendingWrite->ErrorText);
+				continue;
+			}
+			JsonWriter->WriteObjectStart();
+			JsonWriter->WriteValue(TEXT("IoChunkId"), LexToString(PendingWrite->ChunkId));
+			JsonWriter->WriteValue(TEXT("FileName"), *PendingWrite->WriteOptions.FileName);
+			JsonWriter->WriteObjectEnd();
+			
+			WriteResult.UncompressedContainerSize += PendingWrite->UncompressedSize;
+			WriteResult.CompressedContainerSize += PendingWrite->CompressedSize;
+			WriteResult.TocEntryCount++;
+		}
+
+		JsonWriter->WriteArrayEnd();
+		JsonWriter->WriteObjectEnd();
+		JsonWriter->Close();
+
+		UE_LOG(LogIoStore, Display, TEXT("Saving loose file JSON manifest '%s'"), *Settings.TocFilePath);
+		if (!FFileHelper::SaveStringToFile(Json, *Settings.TocFilePath))
+		{
+			return FIoStatus(EIoErrorCode::CorruptToc, TEXTVIEW("Failed writing loose file JSON manifest")); 
+		}
+		
+		WriteResult.TocSize = sizeof(TCHAR) + Json.Len();
+
+		return WriteResult;
+	}
+
+	virtual void EnumerateChunks(TFunction<bool(FIoStoreTocChunkInfo&&)>&& Callback) const override
+	{
+	}
+
+private:
+	FLooseFilesWriterSettings Settings;
+	FWriteStats WriteStats;
+	FCriticalSection PendingLock;
+	TArray<TUniquePtr<FPendingWrite>> PendingWrites;
+	FEventRef WriteCompletedEvent;
+};
+
+} // namespace
+
+TSharedPtr<IIoStoreWriter> MakeLooseFilesIoStoreWriter(const FLooseFilesWriterSettings& WriterSettings,
+	uint32 MaxConcurrentWrites)
+{
+	return MakeShared<FLooseFilesIoStoreWriter>(WriterSettings, MaxConcurrentWrites);
+}
diff --git a/Engine/Source/Developer/IoStoreUtilities/Private/IoStoreLooseFiles.h b/Engine/Source/Developer/IoStoreUtilities/Private/IoStoreLooseFiles.h
new file mode 100644
--- /dev/null
+++ b/Engine/Source/Developer/IoStoreUtilities/Private/IoStoreLooseFiles.h
@@ -0,0 +1,19 @@
+// Copyright Epic Games, Inc. All Rights Reserved.
+
+#pragma once
+
+#include "Containers/UnrealString.h"
+#include "Templates/UniquePtr.h"
+#include "UObject/NameTypes.h"
+
+class IIoStoreWriter;
+
+struct FLooseFilesWriterSettings
+{
+	FString TocFilePath;
+	FString TargetRootPath;
+	FName ContainerName;
+};
+
+TSharedPtr<IIoStoreWriter> MakeLooseFilesIoStoreWriter(const FLooseFilesWriterSettings& WriterSettings,
+	uint32 MaxConcurrentWrites = 64);
diff --git a/Engine/Source/Developer/IoStoreUtilities/Private/IoStoreUtilities.cpp b/Engine/Source/Developer/IoStoreUtilities/Private/IoStoreUtilities.cpp
--- a/Engine/Source/Developer/IoStoreUtilities/Private/IoStoreUtilities.cpp
+++ b/Engine/Source/Developer/IoStoreUtilities/Private/IoStoreUtilities.cpp
@@ -2,6 +2,7 @@
 
 #include "IoStoreUtilities.h"
 
+#include "IoStoreLooseFiles.h"
 #include "Async/AsyncWork.h"
 #include "HAL/FileManager.h"
 #include "HAL/PlatformFileManager.h"
@@ -425,6 +426,7 @@ struct FContainerSourceSpec
 	FString OutputPath;
 	FString OptionalOutputPath;
 	FString OnDemandOutputPath;
+	FString StageLooseFileRootPath;
 	TArray<FContainerSourceFile> SourceFiles;
 	FString PatchTargetFile;
 	TArray<FString> PatchSourceContainerFiles;
@@ -936,6 +938,7 @@ struct FContainerTargetSpec
 	FString OutputPath;
 	FString OptionalSegmentOutputPath;
 	FString OnDemandOutputPath;
+	FString StageLooseFileRootPath;
 	TSharedPtr<IIoStoreWriter> IoStoreWriter;
 	TSharedPtr<IIoStoreWriter> OptionalSegmentIoStoreWriter;
 	TArray<FContainerTargetFile> TargetFiles;
@@ -2523,6 +2526,7 @@ void InitializeContainerTargetsAndPackages(
 		FContainerTargetSpec* ContainerTarget = AddContainer(ContainerSource.Name, ContainerTargets);
 		ContainerTarget->OutputPath = ContainerSource.OutputPath;
 		ContainerTarget->OnDemandOutputPath = ContainerSource.OnDemandOutputPath;
+		ContainerTarget->StageLooseFileRootPath = ContainerSource.StageLooseFileRootPath;
 		ContainerTarget->bGenerateDiffPatch = ContainerSource.bGenerateDiffPatch;
 		if (Arguments.bSign)
 		{
@@ -3745,7 +3749,14 @@ int32 CreateTarget(const FIoStoreArguments& Arguments, const FIoStoreWriterSetti
 				continue;
 			}
 
-			if (!ContainerTarget->OnDemandOutputPath.IsEmpty())
+			if (!ContainerTarget->StageLooseFileRootPath.IsEmpty())
+			{
+				FLooseFilesWriterSettings WriterSettings;
+				WriterSettings.TargetRootPath = ContainerTarget->StageLooseFileRootPath;
+				ContainerTarget->IoStoreWriter = MakeLooseFilesIoStoreWriter(WriterSettings);
+				IoStoreWriters.Add(ContainerTarget->IoStoreWriter);
+			}
+			else if (!ContainerTarget->OnDemandOutputPath.IsEmpty())
 			{
 				if (!OnDemandIoStoreWriter.IsValid())
 				{
@@ -7376,6 +7387,8 @@ bool ParseContainerGenerationArguments(FIoStoreArguments& Arguments, FIoStoreWri
 
 			FParse::Value(*Command, TEXT("OptionalOutput="), ContainerSpec.OptionalOutputPath);
 
+			FParse::Value(*Command, TEXT("StageLooseFileRootPath="), ContainerSpec.StageLooseFileRootPath);
+
 			FString ContainerName;
 			if (FParse::Value(*Command, TEXT("ContainerName="), ContainerName))
 			{
diff --git a/Engine/Source/Programs/AutomationTool/AutomationUtils/DeploymentContext.cs b/Engine/Source/Programs/AutomationTool/AutomationUtils/DeploymentContext.cs
--- a/Engine/Source/Programs/AutomationTool/AutomationUtils/DeploymentContext.cs
+++ b/Engine/Source/Programs/AutomationTool/AutomationUtils/DeploymentContext.cs
@@ -41,8 +41,6 @@ public class FilesToStage
 	/// </summary>
 	public Dictionary<StagedFileReference, FileReference> UFSFiles = new Dictionary<StagedFileReference, FileReference>();
 
-	public Dictionary<StagedFileReference, FileReference> LooseFiles = new Dictionary<StagedFileReference, FileReference>();
-
 	/// <summary>
 	/// After staging, this is a map from staged file to source file. These file are binaries, etc and can't go into a pak file.
 	/// </summary>
diff --git a/Engine/Source/Programs/AutomationTool/Scripts/CopyBuildToStagingDirectory.Automation.cs b/Engine/Source/Programs/AutomationTool/Scripts/CopyBuildToStagingDirectory.Automation.cs
--- a/Engine/Source/Programs/AutomationTool/Scripts/CopyBuildToStagingDirectory.Automation.cs
+++ b/Engine/Source/Programs/AutomationTool/Scripts/CopyBuildToStagingDirectory.Automation.cs
@@ -283,6 +283,7 @@ namespace AutomationScripts
 			DirectoryReference OptionalOutputLocation,
 			DirectoryReference OnDemandOutputLocation,
 			bool bCompressed,
+			DirectoryReference OptionalStageLooseFileRootPath,
 			EncryptionAndSigning.CryptoSettings CryptoSettings,
 			string EncryptionKeyGuid,
 			string PatchSourceContentPath,
@@ -315,6 +316,11 @@ namespace AutomationScripts
 				CmdLine.Append(" -GenerateDiffPatch");
 			}
 
+			if (OptionalStageLooseFileRootPath != null)
+			{
+				CmdLine.AppendFormat(" -StageLooseFileRootPath=\"{0}\"", OptionalStageLooseFileRootPath);
+			}
+
 			// Force encryption of ALL files if we're using specific encryption key. This should be made an option per encryption key in the settings, but for our initial
 			// implementation we will just assume that we require maximum security for this data.
 			bool bForceEncryption = !string.IsNullOrEmpty(EncryptionKeyGuid);
@@ -2698,6 +2704,11 @@ namespace AutomationScripts
 			/// </summary>
 			public bool bOnDemand;
 
+			/// <summary>
+			// Whether this pak file is to copy the source files as loose files
+			/// </summary>
+			public bool bStageLoose;
+
 			/// <summary>
 			/// GUID of the encryption key for this pak file
 			/// </summary>
@@ -2709,7 +2720,7 @@ namespace AutomationScripts
 			/// <param name="PakName">Path to the base output file for this pak file</param>
 			/// <param name="UnrealPakResponseFile">Map of files within the pak file to their source file on disk</param>
 			/// <param name="bCompressed">Whether to enable compression</param>
-			public CreatePakParams(string PakName, Dictionary<string, string> UnrealPakResponseFile, bool bCompressed, bool RehydrateAssets, string EncryptionKeyGuid, bool bOnDemand = false)
+			public CreatePakParams(string PakName, Dictionary<string, string> UnrealPakResponseFile, bool bCompressed, bool RehydrateAssets, string EncryptionKeyGuid, bool bOnDemand = false, bool bStageLoose = false)
 			{
 				this.PakName = PakName;
 				this.UnrealPakResponseFile = UnrealPakResponseFile;
@@ -2717,6 +2728,7 @@ namespace AutomationScripts
 				this.bRehydrateAssets = RehydrateAssets;
 				this.EncryptionKeyGuid = EncryptionKeyGuid;
 				this.bOnDemand = bOnDemand;
+				this.bStageLoose = bStageLoose;
 			}
 		}
 
@@ -3287,7 +3299,7 @@ namespace AutomationScripts
 
 				bool bCopiedExistingPak = false;
 
-				if (!PakParams.bOnDemand && SC.StageTargetPlatform != SC.CookSourcePlatform && !Params.IgnorePaksFromDifferentCookSource)
+				if (!PakParams.bOnDemand && !PakParams.bStageLoose && SC.StageTargetPlatform != SC.CookSourcePlatform && !Params.IgnorePaksFromDifferentCookSource)
 				{
 					// Check to see if we have an existing pak file we can use
 
@@ -3418,6 +3430,7 @@ namespace AutomationScripts
 							}
 							bool bGenerateDiffPatch = bShouldGeneratePatch && !ShouldSkipGeneratingPatch(PlatformGameConfig, PakParams.PakName);
 							bool bCompressContainers = PakParams.bCompressed;
+							DirectoryReference StageLooseFileRootPath = PakParams.bStageLoose ? DirectoryReference.Combine(SC.RuntimeRootDir, "Engine", "Binaries", SC.PlatformDir) : null;
 
 							IoStoreCommands.Add(GetIoStoreCommandArguments(
 								IoStoreResponseFile,
@@ -3426,6 +3439,7 @@ namespace AutomationScripts
 								SC.OptionalFileStageDirectory,
 								PakParams.bOnDemand ? OnDemandOutputLocation : null,
 								bCompressContainers,
+								StageLooseFileRootPath,
 								CryptoSettings,
 								PakParams.EncryptionKeyGuid,
 								ContainerPatchSourcePath,
@@ -3433,6 +3447,8 @@ namespace AutomationScripts
 								Params.HasDLCName));
 						}
 
+						if (!PakParams.bStageLoose)
+						{
 						string PakFileSpecificAdditionalArgs = "";
 						if (bShouldGeneratePatch && ShouldSkipGeneratingPatch(PlatformGameConfig, PakParams.PakName))
 						{
@@ -3448,7 +3464,7 @@ namespace AutomationScripts
 							Params.SkipEncryption ? null : CryptoSettings,
 							PatchSourceContentPath,
 							Params.SkipEncryption ? "" : PakParams.EncryptionKeyGuid));
-
+						}
 						LogNames.Add(OutputLocation.GetFileNameWithoutExtension());
 					}
 				}
@@ -3956,6 +3972,7 @@ namespace AutomationScripts
 				Manifest = null;
 				bCompressed = false;
 				bOnDemand = false;
+				bStageLoose = false;
 			}
 
 			// Name of pak file without extension, ie pakchunk0
@@ -3974,6 +3991,9 @@ namespace AutomationScripts
 
 			// Whether the chunk is used for content-on-demand content delivery or not
 			public bool bOnDemand;
+
+			// Whether the chunk is used for content-on-demand content delivery or not
+			public bool bStageLoose;
 		}
 
 		private static void AddChunkDefinition(ConcurrentDictionary<string, ChunkDefinition> ChunkSet, ChunkDefinition OriginalChunk, string ModifiedChunkName)
@@ -4133,7 +4153,8 @@ namespace AutomationScripts
 
 						ChunkDefinition Chunk = new ChunkDefinition(ChunkName);
 						Chunk.bOnDemand = Rule.bOnDemand;
-						Chunk.bCompressed = true;
+						Chunk.bStageLoose = Rule.bStageLoose;
+						Chunk.bCompressed = !Rule.bStageLoose;
 						Chunk.EncryptionKeyGuid = bIsEncryptionKeyValid ? Rule.EncryptionKeyGuid : null;
 						ChunkDefinitions.Add(Chunk);
 					}
@@ -4145,8 +4166,6 @@ namespace AutomationScripts
 					ChunkNameToDefinition.TryAdd(Chunk.ChunkName, Chunk);
 				}
 
-				ConcurrentDictionary<StagedFileReference, FileReference> LooseFiles = new ConcurrentDictionary<StagedFileReference, FileReference>();
-
 				Parallel.ForEach(StagingManifestResponseFile, StagingFile =>
 				{
 					bool bAddedToChunk = false;
@@ -4189,12 +4208,15 @@ namespace AutomationScripts
 
 					if (bStageLoose)
 					{
-						string StagingFileValue = StagingFile.Value;
-						if (StagingFileValue.StartsWith(SC.PakFileInternalRoot))
+						string StageLooseChunkName = "_stage_loose_files";
+						if (!AdditionalChunks.ContainsKey(StageLooseChunkName))
 						{
-							StagingFileValue = StagingFileValue.Substring(SC.PakFileInternalRoot.Length);
+							ChunkDefinition NewChunk = new ChunkDefinition(StageLooseChunkName);
+							NewChunk.bStageLoose = true;
+							AdditionalChunks.TryAdd(StageLooseChunkName, NewChunk);
 						}
-						LooseFiles.TryAdd(new StagedFileReference(StagingFileValue), new FileReference(StagingFile.Key));
+						ChunkDefinition TargetChunk = AdditionalChunks[StageLooseChunkName];
+						TargetChunk.ResponseFile.TryAdd(StagingFile.Key, StagingFile.Value);
 
 						// Loose files get uploaded through BPT as well as copied to "ondemand" pak files
 						//return;
@@ -4215,10 +4237,10 @@ namespace AutomationScripts
 							string OrigExt = Path.GetExtension(OriginalFilename);
 							if (bStageLoose)
 							{
-								// make a streaming pak
-								string ChunkName = Chunk.ChunkName + (OrigExt.Equals(OptionalBulkDataFileExtension) ? "ondemandoptional" : "ondemand");
-								AddChunkDefinition(AdditionalChunks, Chunk, ChunkName);
-								TargetChunk = AdditionalChunks[ChunkName];
+								// make a streaming pak as well
+								string StreamingChunkName = Chunk.ChunkName + (OrigExt.Equals(OptionalBulkDataFileExtension) ? "ondemandoptional" : "ondemand");
+								AddChunkDefinition(AdditionalChunks, Chunk, StreamingChunkName);
+								TargetChunk = AdditionalChunks[StreamingChunkName];
 							}
 							else if (OrigExt.Equals(OptionalBulkDataFileExtension))
 							{
@@ -4257,11 +4279,6 @@ namespace AutomationScripts
 				{
 					ChunkDefinitions.Add(OptionalChunkIt.Value);
 				}
-
-				foreach (var LooseFile in LooseFiles)
-				{
-					SC.FilesToStage.LooseFiles.Add(LooseFile.Key, LooseFile.Value);
-				}
 			}
 
 			// chunk downloader pak file is a minimal pak file which contains no content.  It can be provided with as a minimal download so that the game can download all the content from another source.
@@ -4339,7 +4356,7 @@ namespace AutomationScripts
 					string EncryptionKeyToUse = Params.SkipEncryption ? "" : Chunk.EncryptionKeyGuid;
 					PakInputs.Add(new CreatePakParams(Chunk.ChunkName,
 						Chunk.ResponseFile.ToDictionary(entry => entry.Key, entry => entry.Value),
-						Params.Compressed || Chunk.bCompressed, Params.RehydrateAssets, EncryptionKeyToUse, Chunk.bOnDemand));
+						Params.Compressed || Chunk.bCompressed, Params.RehydrateAssets, EncryptionKeyToUse, Chunk.bOnDemand, Chunk.bStageLoose));
 				}
 			}
 
@@ -4603,15 +4620,9 @@ namespace AutomationScripts
 				}
 			}
 
-		string BaseManifestFileName = CombinePaths(CmdEnv.LogFolder, "FinalCopy" + (SC.DedicatedServer ? "_Server" : "") + SC.StageTargetPlatform.PlatformType.ToString());
+			string BaseManifestFileName = CombinePaths(CmdEnv.LogFolder, "FinalCopy" + (SC.DedicatedServer ? "_Server" : "") + SC.StageTargetPlatform.PlatformType.ToString());
 			DumpManifest(SC.FilesToStage.UFSFiles, BaseManifestFileName + "_UFSFiles.txt");
 
-			if (!SC.Stage || Params.SkipStage)
-			{
-				CopyManifestFilesToStageDir(SC, SC.FilesToStage.LooseFiles, "LooseFiles");
-				return;
-			}
-
 			DumpManifest(SC.FilesToStage.NonUFSFiles, BaseManifestFileName + "_NonUFSFiles.txt");
 			DumpManifest(SC.FilesToStage.NonUFSDebugFiles, BaseManifestFileName + "_NonUFSFilesDebug.txt");
 
