- Fixed staging of programs under a Game directory, not engine
- #rb [REDACTED_PERSON] (tested)

diff --git a/Engine/Source/Programs/AutomationTool/Scripts/CopyBuildToStagingDirectory.Automation.cs b/Engine/Source/Programs/AutomationTool/Scripts/CopyBuildToStagingDirectory.Automation.cs
--- a/Engine/Source/Programs/AutomationTool/Scripts/CopyBuildToStagingDirectory.Automation.cs
+++ b/Engine/Source/Programs/AutomationTool/Scripts/CopyBuildToStagingDirectory.Automation.cs
@@ -4945,7 +4945,19 @@ namespace AutomationScripts
 					string Target = TargetAndConfig.Item1;
 					UnrealTargetConfiguration Config = TargetAndConfig.Item2;
 					UnrealArchitectures Architecture = TargetAndConfig.Item3;
-					DirectoryReference ReceiptBaseDir = (Params.IsCodeBasedProject && !Params.ProgramTargets.Contains(Target)) ? ProjectDir : EngineDir;
+
+					DirectoryReference ReceiptBaseDir = Params.IsCodeBasedProject ? ProjectDir : EngineDir;
+					// handle Engine programs that have a .uproject specially
+					// @todo Non-engine programs (unde a game directory) and also have a .uproject are not supported, because the target and executable are in 
+					// Game/Binaries, but here the ProjectDir is the program's directory, not the project. Currently not used, but will need work to fix
+					if (Params.ProgramTargets.Contains(Target))
+					{
+						// if we have a uproject file that is under the engine dir, then the target/exe will be in Engine/Binaries
+						if (ProjectDir.IsUnderDirectory(EngineDir))
+						{
+							ReceiptBaseDir = EngineDir;
+						}
+					}
 
 					Platform PlatformInstance = Platform.Platforms[StagePlatform];
 					UnrealTargetPlatform[] SubPlatformsToStage = PlatformInstance.GetStagePlatforms();
