AutomationTool: Fixed 'Launch failed' messages in the editor after terminating a successfully launched game.The problems was caused by us treating all non-zero results as errors while some platforms return different non-zero codes to provide additional information.

Resolved by adding IProcessResult::bExitCodeSuccess virtual property which may be overridden on the affected platforms to explicitly indicate which exit codes are allowed.

#rnx

diff --git a/Engine/Source/Programs/AutomationTool/Android/AndroidPlatform.Automation.cs b/Engine/Source/Programs/AutomationTool/Android/AndroidPlatform.Automation.cs
--- a/Engine/Source/Programs/AutomationTool/Android/AndroidPlatform.Automation.cs
+++ b/Engine/Source/Programs/AutomationTool/Android/AndroidPlatform.Automation.cs
@@ -118,6 +118,8 @@ public class AndroidPlatform : Platform
 			set { LogFileProcessExitCode = value; }
 		}
 
+		public bool bExitCodeSuccess => ExitCode == 0;
+
 		public string Output
 		{
 			get { return AdbLogProcess.Output; }
diff --git a/Engine/Source/Programs/AutomationTool/AutomationUtils/ProcessUtils.cs b/Engine/Source/Programs/AutomationTool/AutomationUtils/ProcessUtils.cs
--- a/Engine/Source/Programs/AutomationTool/AutomationUtils/ProcessUtils.cs
+++ b/Engine/Source/Programs/AutomationTool/AutomationUtils/ProcessUtils.cs
@@ -223,6 +223,7 @@ namespace AutomationTool
 		void StdOut(object sender, DataReceivedEventArgs e);
 		void StdErr(object sender, DataReceivedEventArgs e);
 		int ExitCode { get; set; }
+		bool bExitCodeSuccess { get; }
 		string Output { get; }
 		Process ProcessObject { get; }
 		string ToString();
@@ -392,6 +393,8 @@ namespace AutomationTool
 			set { ProcessExitCode = value; }
 		}
 
+		public bool bExitCodeSuccess => ExitCode == 0;
+
 		/// <summary>
 		/// Gets all std output the process generated.
 		/// </summary>
diff --git a/Engine/Source/Programs/AutomationTool/IOS/IOSPlatform.Automation.cs b/Engine/Source/Programs/AutomationTool/IOS/IOSPlatform.Automation.cs
--- a/Engine/Source/Programs/AutomationTool/IOS/IOSPlatform.Automation.cs
+++ b/Engine/Source/Programs/AutomationTool/IOS/IOSPlatform.Automation.cs
@@ -92,6 +92,8 @@ class IOSClientProcess : IProcessResult
 		set { childProcess.ExitCode = value; }
 	}
 
+	public bool bExitCodeSuccess => ExitCode == 0;
+
 	public string Output
 	{
 		get { return childProcess.Output; }
diff --git a/Engine/Source/Programs/AutomationTool/Scripts/RunProjectCommand.Automation.cs b/Engine/Source/Programs/AutomationTool/Scripts/RunProjectCommand.Automation.cs
--- a/Engine/Source/Programs/AutomationTool/Scripts/RunProjectCommand.Automation.cs
+++ b/Engine/Source/Programs/AutomationTool/Scripts/RunProjectCommand.Automation.cs
@@ -537,7 +537,7 @@ namespace AutomationScripts
 			{
 				// Any non-zero exit code should propagate an exception. The PostRunClient function above may have
 				// already thrown a more specific exception or given a more specific ErrorCode, but this catches the rest.
-				if (ClientProcess != null && ClientProcess.ExitCode != 0)
+				if (ClientProcess != null && !ClientProcess.bExitCodeSuccess)
 				{
 					throw new AutomationException("Client exited with error code: {0} (see {1} for more info)", ClientProcess.ExitCode, ClientLogFile);
 				}
