Adding hook for any platform-specific files that are needed when staging DLC.

#rnx

diff --git a/Engine/Source/Programs/AutomationTool/AutomationUtils/Platform.cs b/Engine/Source/Programs/AutomationTool/AutomationUtils/Platform.cs
--- a/Engine/Source/Programs/AutomationTool/AutomationUtils/Platform.cs
+++ b/Engine/Source/Programs/AutomationTool/AutomationUtils/Platform.cs
@@ -467,6 +467,14 @@ namespace AutomationTool
 			throw new AutomationException("{0} does not yet implement GetFilesToDeployOrStage.", PlatformType);
 		}
 
+		/// <summary>
+		/// Get additional platform specific files to stage when staging DLC
+		/// </summary>
+		/// <param name="SC">Deployment Context</param>
+		public virtual void GetFilesToStageForDLC(ProjectParams Params, DeploymentContext SC)
+		{
+		}
+
 		/// <summary>
 		/// Called after CopyUsingStagingManifest.  Does anything platform specific that requires a final list of staged files.
 		/// </summary>
diff --git a/Engine/Source/Programs/AutomationTool/Scripts/CopyBuildToStagingDirectory.Automation.cs b/Engine/Source/Programs/AutomationTool/Scripts/CopyBuildToStagingDirectory.Automation.cs
--- a/Engine/Source/Programs/AutomationTool/Scripts/CopyBuildToStagingDirectory.Automation.cs
+++ b/Engine/Source/Programs/AutomationTool/Scripts/CopyBuildToStagingDirectory.Automation.cs
@@ -1126,6 +1126,8 @@ namespace AutomationScripts
 						}
 					}
 				}
+				
+				ThisPlatform.GetFilesToStageForDLC(Params, SC);
 
 				bCreatePluginManifest = true;
 			}
