BuildGraph: Allow specifying the character sequence used to join lines in a property element body.

diff --git a/Engine/Source/Programs/AutomationTool/BuildGraph/BgSchema.cs b/Engine/Source/Programs/AutomationTool/BuildGraph/BgSchema.cs
--- a/Engine/Source/Programs/AutomationTool/BuildGraph/BgSchema.cs
+++ b/Engine/Source/Programs/AutomationTool/BuildGraph/BgSchema.cs
@@ -787,6 +787,7 @@ namespace AutomationTool
 			extension.BaseTypeName = s_stringTypeName;
 			extension.Attributes.Add(CreateSchemaAttribute("Name", ScriptSchemaStandardType.Name, XmlSchemaUse.Required));
 			extension.Attributes.Add(CreateSchemaAttribute("Value", s_stringTypeName, XmlSchemaUse.Optional));
+			extension.Attributes.Add(CreateSchemaAttribute("Separator", s_stringTypeName, XmlSchemaUse.Optional));
 			extension.Attributes.Add(CreateSchemaAttribute("If", ScriptSchemaStandardType.BalancedString, XmlSchemaUse.Optional));
 
 			XmlSchemaSimpleContent contentModel = new XmlSchemaSimpleContent();
diff --git a/Engine/Source/Programs/AutomationTool/BuildGraph/BgScriptReader.cs b/Engine/Source/Programs/AutomationTool/BuildGraph/BgScriptReader.cs
--- a/Engine/Source/Programs/AutomationTool/BuildGraph/BgScriptReader.cs
+++ b/Engine/Source/Programs/AutomationTool/BuildGraph/BgScriptReader.cs
@@ -773,6 +773,13 @@ namespace AutomationTool
 					string value = ReadAttribute(element, "Value");
 					if (element.HasChildNodes)
 					{
+						// Get the separator character
+						string separator = ";";
+						if (element.HasAttribute("Separator"))
+						{
+							separator = ReadAttribute(element, "Separator");
+						}
+
 						// Read the element content, and append each line to the value as a semicolon delimited list
 						StringBuilder builder = new StringBuilder(value);
 						foreach (string line in element.InnerText.Split('\n'))
@@ -782,7 +789,7 @@ namespace AutomationTool
 							{
 								if (builder.Length > 0)
 								{
-									builder.Append(";");
+									builder.Append(separator);
 								}
 								builder.Append(trimLine);
 							}
