BuildGraph: Add a task which can write an arbitary JSON value to a new or existing JSON file. Similar to the legacy ModifyJsonValue task, but much more comprehensive.

The property to modify is specified in a limited subset of JsonPath notation, supporting objects and array indices. Values may be any valid JSON values (string, object, array, number, boolean, or null).

diff --git a/Engine/Source/Programs/AutomationTool/BuildGraph/Tasks/WriteJsonValueTask.cs b/Engine/Source/Programs/AutomationTool/BuildGraph/Tasks/WriteJsonValueTask.cs
new file mode 100644
--- /dev/null
+++ b/Engine/Source/Programs/AutomationTool/BuildGraph/Tasks/WriteJsonValueTask.cs
@@ -0,0 +1,190 @@
+// Copyright Epic Games, Inc. All Rights Reserved.
+
+using EpicGames.Core;
+using IdentityModel.Client;
+using System;
+using System.Collections.Generic;
+using System.Linq;
+using System.Text;
+using System.Text.Json;
+using System.Text.Json.Nodes;
+using System.Text.RegularExpressions;
+using System.Threading.Tasks;
+using System.Xml;
+using UnrealBuildBase;
+
+#nullable enable
+
+namespace AutomationTool.Tasks
+{
+	using JsonObject = System.Text.Json.Nodes.JsonObject;
+
+	/// <summary>
+	/// Parameters for <see cref="WriteJsonValueTask"/> task
+	/// </summary>
+	public class WriteJsonValueTaskParameters
+	{
+		/// <summary>
+		/// Json file(s) which will be modified
+		/// </summary>
+		[TaskParameter(ValidationType = TaskParameterValidationType.FileSpec)]
+		public string File = null!;
+
+		/// <summary>
+		/// Json element to set in each file. Syntax for this string is a limited subset of JsonPath notation, and may support object properties and
+		/// array indices. Any array indices which are omitted or out of range will add a new element to the array (eg. '$.foo.bar[]' will add
+		/// an element to the 'bar' array in the 'foo' object).
+		/// </summary>
+		[TaskParameter(ValidationType = TaskParameterValidationType.Default)]
+		public string Key = null!;
+
+		/// <summary>
+		/// New value to set. May be any value JSON value (string, array, object, number, boolean or null).
+		/// </summary>
+		[TaskParameter(ValidationType = TaskParameterValidationType.Default)]
+		public string Value = null!;
+	}
+
+	/// <summary>
+	/// Modifies json files by setting a value specified in the key path
+	/// </summary>
+	[TaskElement("WriteJsonValue", typeof(WriteJsonValueTaskParameters))]
+	public class WriteJsonValueTask : BgTaskImpl
+	{
+		WriteJsonValueTaskParameters Parameters;
+
+		/// <summary>
+		/// Create a new ModifyJsonValue.
+		/// </summary>
+		/// <param name="InParameters">Parameters for this task.</param>
+		public WriteJsonValueTask(WriteJsonValueTaskParameters InParameters)
+		{
+			Parameters = InParameters;
+		}
+
+		/// <summary>
+		/// Placeholder comment
+		/// </summary>
+		public override async Task ExecuteAsync(JobContext Job, HashSet<FileReference> BuildProducts, Dictionary<string, HashSet<FileReference>> TagNameToFileSet)
+		{
+			HashSet<FileReference> Files = ResolveFilespec(Unreal.RootDirectory, Parameters.File, TagNameToFileSet);
+
+			JsonNode? ValueNode = String.IsNullOrEmpty(Parameters.Value) ? null : JsonNode.Parse(Parameters.Value);
+			foreach (FileReference JsonFile in Files)
+			{
+				string JsonText = FileReference.Exists(JsonFile) ? await FileReference.ReadAllTextAsync(JsonFile) : "{}";
+
+				if (!Parameters.Key.StartsWith("$", StringComparison.Ordinal))
+				{
+					throw new AutomationException("Key must be in JsonPath format (eg. $.Foo.Bar[123])");
+				}
+
+				JsonNode? RootNode = JsonNode.Parse(JsonText);
+				RootNode = MergeValue(Parameters.Key, 1, RootNode, ValueNode);
+
+				string NewJsonText = RootNode?.ToJsonString(new JsonSerializerOptions { WriteIndented = true }) ?? String.Empty;
+
+				DirectoryReference.CreateDirectory(JsonFile.Directory);
+				await FileReference.WriteAllTextAsync(JsonFile, NewJsonText);
+			}
+		}
+
+		static JsonNode? MergeValue(string Key, int MinIdx, JsonNode? PrevValue, JsonNode? Value)
+		{
+			if (MinIdx == Key.Length)
+			{
+				return Value;
+			}
+
+			// Find the length of the next token
+			int MaxIdx = MinIdx + 1;
+			while (MaxIdx < Key.Length && Key[MaxIdx] != '[' && Key[MaxIdx] != '.')
+			{
+				MaxIdx++;
+			}
+
+			// Handle different types of element
+			if (Key[MinIdx] == '.')
+			{
+				JsonObject? Obj = PrevValue as JsonObject;
+				if (Obj != null)
+				{
+					Obj = Obj.Deserialize<JsonObject>(); // Clone so we can reattach
+				}
+				Obj ??= new JsonObject();
+
+				string PropertyName = Key.Substring(MinIdx + 1, MaxIdx - (MinIdx + 1));
+
+				JsonNode? NextNode;
+				Obj.TryGetPropertyValue(PropertyName, out NextNode);
+				Obj[PropertyName] = MergeValue(Key, MaxIdx, NextNode, Value);
+
+				return Obj;
+			}
+			else if (Key[MinIdx] == '[')
+			{
+				if (Key[MaxIdx - 1] != ']')
+				{
+					throw new AutomationException("Missing ']' in array subscript in Json path expression '{Key}'");
+				}
+
+				string IndexStr = Key.Substring(MinIdx + 1, (MaxIdx - 1) - (MinIdx + 1)).Trim();
+
+				int Index = int.MaxValue;
+				if (IndexStr.Length > 0)
+				{
+					Index = int.Parse(IndexStr);
+				}
+
+				JsonArray? Array = PrevValue as JsonArray;
+				if (Array != null)
+				{
+					Array = Array.Deserialize<JsonArray>();
+				}
+				Array ??= new JsonArray();
+
+				if (Index < Array.Count)
+				{
+					Array[Index] = MergeValue(Key, MaxIdx, Array[Index], Value);
+				}
+				else
+				{
+					Array.Add(MergeValue(Key, MaxIdx, null, Value));
+				}
+
+				return Array;
+			}
+			else
+			{
+				throw new AutomationException($"Unable to parse JSON path after '{Key}'");
+			}
+		}
+
+		/// <summary>
+		/// Placeholder comment
+		/// </summary>
+		public override void Write(XmlWriter Writer)
+		{
+			Write(Writer, Parameters);
+		}
+
+		/// <summary>
+		/// Placeholder comment
+		/// </summary>
+		public override IEnumerable<string> FindConsumedTagNames()
+		{
+			foreach (string TagName in FindTagNamesFromFilespec(Parameters.File))
+			{
+				yield return TagName;
+			}
+		}
+
+		/// <summary>
+		/// Placeholder comment
+		/// </summary>
+		public override IEnumerable<string> FindProducedTagNames()
+		{
+			yield break;
+		}
+	}
+}
