- No longer ask user what directory they want to use for UE_SDKS_ROOT when syncing from p4 because we know the answer is where it synced to

diff --git a/Engine/Source/Programs/AutomationTool/Turnkey/SdkInfo.cs b/Engine/Source/Programs/AutomationTool/Turnkey/SdkInfo.cs
--- a/Engine/Source/Programs/AutomationTool/Turnkey/SdkInfo.cs
+++ b/Engine/Source/Programs/AutomationTool/Turnkey/SdkInfo.cs
@@ -253,15 +253,8 @@ namespace Turnkey
 				// now go one up to the parent of Host
 				AutoSdkSearch = AutoSdkSearch.Parent;
 
+				// we don't even ask the user in the p4 case because it doesn't make sense to use anything but where we downloaded to
 				string AutoSdkDir = AutoSdkSearch.FullName;
-				if (!bUnattended)
-				{
-					string Response = TurnkeyUtils.ReadInput("Enter directory for root of AutoSdks. Use detected value, or enter another:", AutoSdkSearch.FullName);
-					if (string.IsNullOrEmpty(Response))
-					{
-						return false;
-					}
-				}
 
 				// set the env var, globally
 				TurnkeyUtils.StartTrackingExternalEnvVarChanges();
