Horde: Update the DeployTool task to upload chunked data directly, rather than uploading an entire zip.

diff --git a/Engine/Source/Programs/AutomationTool/BuildGraph/Tasks/DeployToolTask.cs b/Engine/Source/Programs/AutomationTool/BuildGraph/Tasks/DeployToolTask.cs
--- a/Engine/Source/Programs/AutomationTool/BuildGraph/Tasks/DeployToolTask.cs
+++ b/Engine/Source/Programs/AutomationTool/BuildGraph/Tasks/DeployToolTask.cs
@@ -13,6 +13,11 @@ using System.Threading.Tasks;
 using System.Xml;
 using UnrealBuildBase;
 using Microsoft.Extensions.Logging;
+using EpicGames.Horde.Storage.Backends;
+using EpicGames.Horde.Storage;
+using EpicGames.Horde.Storage.Nodes;
+using System.Threading;
+using System.Data;
 
 #nullable enable
 
@@ -54,16 +59,10 @@ namespace AutomationTool.Tasks
 		public bool Paused = false;
 
 		/// <summary>
-		/// Zip file containing the data to upload
+		/// Directory to upload for the tool
 		/// </summary>
-		[TaskParameter(Optional = true)]
-		public string? File;
-
-		/// <summary>
-		/// Directory to zip and upload for the tool
-		/// </summary>
-		[TaskParameter(Optional = true)]
-		public string? Directory;
+		[TaskParameter]
+		public string Directory = null!;
 	}
 
 	/// <summary>
@@ -78,6 +77,17 @@ namespace AutomationTool.Tasks
 			public string? Token { get; set; }
 		}
 
+		/// <summary>
+		/// Options for a new deployment
+		/// </summary>
+		class CreateDeploymentRequest
+		{
+			public string Version { get; set; } = "Unknown";
+			public double? Duration { get; set; }
+			public bool? CreatePaused { get; set; }
+			public string? Node { get; set; }
+		}
+
 		/// <summary>
 		/// Parameters for this task
 		/// </summary>
@@ -100,22 +110,6 @@ namespace AutomationTool.Tasks
 		/// <param name="TagNameToFileSet">Mapping from tag names to the set of files they include</param>
 		public override async Task ExecuteAsync(JobContext Job, HashSet<FileReference> BuildProducts, Dictionary<string, HashSet<FileReference>> TagNameToFileSet)
 		{
-			FileReference? file;
-			if (Parameters.File != null && Parameters.Directory == null)
-			{
-				file = ResolveFile(Parameters.File);
-			}
-			else if (Parameters.File == null && Parameters.Directory != null)
-			{
-				file = FileReference.Combine(Unreal.EngineDirectory, "Intermediate", "DeployTool", "Tool.zip");
-				CommandUtils.DeleteFile(file);
-				CommandUtils.ZipFiles(file, ResolveDirectory(Parameters.Directory), new FileFilter(FileFilterType.Include));
-			}
-			else
-			{
-				throw new AutomationException("Either file or Directory must be specified to the DeployTool task (not both).");
-			}
-
 			FileReference settingsFile = ResolveFile(Parameters.Settings);
 			if (!FileReference.Exists(settingsFile))
 			{
@@ -135,50 +129,60 @@ namespace AutomationTool.Tasks
 				throw new AutomationException($"Missing 'server' key from {settingsFile}");
 			}
 
-			using (HttpClient httpClient = new HttpClient())
+			Uri serverUri = new Uri(settings.Server);
+
+			HttpClient CreateHttpClient()
 			{
-				httpClient.BaseAddress = new Uri(settings.Server);
-				using (HttpRequestMessage request = new HttpRequestMessage(HttpMethod.Post, $"api/v1/tools/{Parameters.Id}/deployments"))
+				HttpClient httpClient = new HttpClient();
+				httpClient.BaseAddress = new Uri(serverUri, $"api/v1/tools/{Parameters.Id}/");
+				if (settings?.Token != null)
 				{
-					if (settings.Token != null)
-					{
-						request.Headers.Authorization = new AuthenticationHeaderValue("Bearer", settings.Token);
-					}
+					httpClient.DefaultRequestHeaders.Authorization = new AuthenticationHeaderValue("Bearer", settings.Token);
+				}
+				return httpClient;
+			}
 
-					using (FileStream stream = FileReference.Open(file, FileMode.Open, FileAccess.Read))
-					{
-						StreamContent streamContent = new StreamContent(stream);
+			NodeHandle handle;
 
-						MultipartFormDataContent content = new MultipartFormDataContent();
-						content.Add(new StringContent(Parameters.Version), "version");
-						if (Parameters.Duration != 0)
-						{
-							content.Add(new StringContent(TimeSpan.FromMinutes(Parameters.Duration).ToString()), "duration");
-						}
-						if (Parameters.Paused)
+			HttpStorageClient storageClient = new HttpStorageClient(CreateHttpClient, () => new HttpClient(), Logger);
+			using (TreeWriter treeWriter = new TreeWriter(storageClient))
+			{
+				DirectoryInfo directoryInfo = new DirectoryInfo(Parameters.Directory!);
+
+				DirectoryNode sandbox = new DirectoryNode();
+				await sandbox.CopyFromDirectoryAsync(directoryInfo, new ChunkingOptions(), treeWriter, null);
+
+				handle = await treeWriter.FlushAsync(sandbox);
+			}
+
+			CreateDeploymentRequest request = new CreateDeploymentRequest();
+			request.Version = Parameters.Version;
+			if (Parameters.Duration != 0)
+			{
+				request.Duration = Parameters.Duration;
+			}
+			if (Parameters.Paused)
+			{
+				request.CreatePaused = true;
+			}
+			request.Node = handle.ToString();
+
+			using (HttpClient httpClient = CreateHttpClient())
+			{
+				using (HttpResponseMessage response = await httpClient.PostAsync<CreateDeploymentRequest>(new Uri(serverUri, $"api/v2/tools/{Parameters.Id}/deployments"), request, CancellationToken.None))
+				{
+					if (!response.IsSuccessStatusCode)
+					{
+						string? responseContent;
+						try
 						{
-							content.Add(new StringContent("true"), "paused");
+							responseContent = await response.Content.ReadAsStringAsync();
 						}
-						content.Add(streamContent, "file", file.GetFileName());
-						request.Content = content;
-
-						Logger.LogInformation("Uploading {File} to {Url}", file, request.RequestUri);
-						using (HttpResponseMessage response = await httpClient.SendAsync(request))
+						catch
 						{
-							if (!response.IsSuccessStatusCode)
-							{
-								string? responseContent;
-								try
-								{
-									responseContent = await response.Content.ReadAsStringAsync();
-								}
-								catch
-								{
-									responseContent = "(No message)";
-								}
-								throw new AutomationException($"Upload failed ({response.StatusCode}): {responseContent}");
-							}
+							responseContent = "(No message)";
 						}
+						throw new AutomationException($"Upload failed ({response.StatusCode}): {responseContent}");
 					}
 				}
 			}
