Allow case insensitive property names when reading NuGet license config files.

diff --git a/Engine/Source/Programs/AutomationTool/BuildGraph/Tasks/NuGetLicenseCheckTask.cs b/Engine/Source/Programs/AutomationTool/BuildGraph/Tasks/NuGetLicenseCheckTask.cs
--- a/Engine/Source/Programs/AutomationTool/BuildGraph/Tasks/NuGetLicenseCheckTask.cs
+++ b/Engine/Source/Programs/AutomationTool/BuildGraph/Tasks/NuGetLicenseCheckTask.cs
@@ -181,7 +181,7 @@ namespace AutomationTool.Tasks
 							if (File.HasExtension(".json"))
 							{
 								byte[] Data = await FileReference.ReadAllBytesAsync(File);
-								LicenseConfig Config = JsonSerializer.Deserialize<LicenseConfig>(Data);
+								LicenseConfig Config = JsonSerializer.Deserialize<LicenseConfig>(Data, new JsonSerializerOptions { PropertyNameCaseInsensitive = true, AllowTrailingCommas = true, ReadCommentHandling = JsonCommentHandling.Skip });
 								LicenseUrls.UnionWith(Config.Urls);
 							}
 							else if (File.HasExtension(".txt") || File.HasExtension(".html"))
