BuildGraph: Fix NuGet license scan not recognizing licenses that are only discovered in a second pass.

diff --git a/Engine/Source/Programs/AutomationTool/BuildGraph/Tasks/NuGetLicenseCheckTask.cs b/Engine/Source/Programs/AutomationTool/BuildGraph/Tasks/NuGetLicenseCheckTask.cs
--- a/Engine/Source/Programs/AutomationTool/BuildGraph/Tasks/NuGetLicenseCheckTask.cs
+++ b/Engine/Source/Programs/AutomationTool/BuildGraph/Tasks/NuGetLicenseCheckTask.cs
@@ -70,12 +70,23 @@ namespace AutomationTool.Tasks
 			Parameters = InParameters;
 		}
 
+		enum PackageState
+		{
+			None,
+			IgnoredViaArgs,
+			MissingPackageFolder,
+			MissingPackageDescriptor,
+			Valid,
+		}
+
 		class PackageInfo
 		{
 			public string Name;
 			public string Version;
 			public LicenseInfo License;
 			public string LicenseSource;
+			public PackageState State;
+			public FileReference Descriptor;
 		}
 
 		class LicenseInfo
@@ -206,32 +217,29 @@ namespace AutomationTool.Tasks
 			}
 
 			Dictionary<string, LicenseInfo> LicenseUrlToInfo = new Dictionary<string, LicenseInfo>(StringComparer.OrdinalIgnoreCase);
-
-			Logger.LogInformation("Referenced Packages:");
-			Logger.LogInformation("");
-			foreach (PackageInfo Info in Packages.Values.OrderBy(x => x.Name).ThenBy(x => x.Version))
+			foreach (PackageInfo Info in Packages.Values)
 			{
 				if (IgnorePackages.Contains(Info.Name) || IgnorePackages.Contains($"{Info.Name}@{Info.Version}"))
 				{
-					Logger.LogInformation("  {Name,-60} {Version,-10} Explicitly ignored via task arguments", Info.Name, Info.Version);
+					Info.State = PackageState.IgnoredViaArgs;
 					continue;
 				}
 
 				DirectoryReference PackageDir = NuGetPackageDirs.Select(x => DirectoryReference.Combine(x, Info.Name.ToLowerInvariant(), Info.Version.ToLowerInvariant())).FirstOrDefault(x => DirectoryReference.Exists(x));
 				if (PackageDir == null)
 				{
-					Logger.LogInformation("  {Name,-60} {Version,-10} NuGet package not found", Info.Name, Info.Version);
+					Info.State = PackageState.MissingPackageFolder;
 					continue;
 				}
 
-				FileReference NuSpecFile = FileReference.Combine(PackageDir, $"{Info.Name.ToLowerInvariant()}.nuspec");
-				if (!FileReference.Exists(NuSpecFile))
+				Info.Descriptor = FileReference.Combine(PackageDir, $"{Info.Name.ToLowerInvariant()}.nuspec");
+				if (!FileReference.Exists(Info.Descriptor))
 				{
-					Logger.LogWarning("Missing package descriptor: {NuSpecFile}", NuSpecFile);
+					Info.State = PackageState.MissingPackageDescriptor;
 					continue;
 				}
 
-				using (Stream Stream = FileReference.Open(NuSpecFile, FileMode.Open, FileAccess.Read, FileShare.Read))
+				using (Stream Stream = FileReference.Open(Info.Descriptor, FileMode.Open, FileAccess.Read, FileShare.Read))
 				{
 					XmlTextReader XmlReader = new XmlTextReader(Stream);
 					XmlReader.Namespaces = false;
@@ -289,18 +297,40 @@ namespace AutomationTool.Tasks
 						}
 					}
 				}
+			}
 
-				if (Info.License == null)
-				{
-					Logger.LogError("  {Name,-60} {Version,-10} No license metadata found", Info.Name, Info.Version);
-				}
-				else if (!Info.License.Approved)
-				{
-					Logger.LogWarning("  {Name,-60} {Version,-10} {Hash}", Info.Name, Info.Version, Info.License.Hash);
-				}
-				else
+			Logger.LogInformation("Referenced Packages:");
+			Logger.LogInformation("");
+			foreach (PackageInfo Info in Packages.Values.OrderBy(x => x.Name).ThenBy(x => x.Version))
+			{
+				switch (Info.State)
 				{
-					Logger.LogInformation("  {Name,-60} {Version,-10} {Hash}", Info.Name, Info.Version, Info.License.Hash);
+					case PackageState.IgnoredViaArgs:
+						Logger.LogInformation("  {Name,-60} {Version,-10} Explicitly ignored via task arguments", Info.Name, Info.Version);
+						break;
+					case PackageState.MissingPackageFolder:
+						Logger.LogInformation("  {Name,-60} {Version,-10} NuGet package not found", Info.Name, Info.Version);
+						break;
+					case PackageState.MissingPackageDescriptor:
+						Logger.LogWarning("  {Name,-60} {Version,-10} Missing package descriptor: {NuSpecFile}", Info.Name, Info.Version, Info.Descriptor);
+						break;
+					case PackageState.Valid:
+						if (Info.License == null)
+						{
+							Logger.LogError("  {Name,-60} {Version,-10} No license metadata found", Info.Name, Info.Version);
+						}
+						else if (!Info.License.Approved)
+						{
+							Logger.LogWarning("  {Name,-60} {Version,-10} {Hash}", Info.Name, Info.Version, Info.License.Hash);
+						}
+						else
+						{
+							Logger.LogInformation("  {Name,-60} {Version,-10} {Hash}", Info.Name, Info.Version, Info.License.Hash);
+						}
+						break;
+					default:
+						Logger.LogError("  {Name,-60} {Version,-10} Unhandled state: {State}", Info.Name, Info.Version, Info.State);
+						break;
 				}
 			}
 
