BuildGraph: Fix state value of NuGet packages after initial license scan.

diff --git a/Engine/Source/Programs/AutomationTool/BuildGraph/Tasks/NuGetLicenseCheckTask.cs b/Engine/Source/Programs/AutomationTool/BuildGraph/Tasks/NuGetLicenseCheckTask.cs
--- a/Engine/Source/Programs/AutomationTool/BuildGraph/Tasks/NuGetLicenseCheckTask.cs
+++ b/Engine/Source/Programs/AutomationTool/BuildGraph/Tasks/NuGetLicenseCheckTask.cs
@@ -297,6 +297,8 @@ namespace AutomationTool.Tasks
 						}
 					}
 				}
+
+				Info.State = PackageState.Valid;
 			}
 
 			Logger.LogInformation("Referenced Packages:");
