Horde: Add a custom error code for BuildGraph script compilation errors.

diff --git a/Engine/Source/Programs/AutomationTool/BuildGraph/BgScriptReader.cs b/Engine/Source/Programs/AutomationTool/BuildGraph/BgScriptReader.cs
--- a/Engine/Source/Programs/AutomationTool/BuildGraph/BgScriptReader.cs
+++ b/Engine/Source/Programs/AutomationTool/BuildGraph/BgScriptReader.cs
@@ -258,7 +258,7 @@ namespace AutomationTool
 			allArgs[0] = location.File;
 			allArgs[1] = location.LineNumber;
 			args.CopyTo(allArgs, 2);
-			logger.LogError($"{{Script}}({{Line}}): error: {format}", allArgs);
+			logger.LogError(KnownLogEvents.AutomationTool_BuildGraphScript, $"{{Script}}({{Line}}): error: {format}", allArgs);
 		}
 
 		/// <summary>
@@ -270,7 +270,7 @@ namespace AutomationTool
 			allArgs[0] = location.File;
 			allArgs[1] = location.LineNumber;
 			args.CopyTo(allArgs, 2);
-			logger.LogWarning($"{{Script}}({{Line}}): warning: {format}", allArgs);
+			logger.LogWarning(KnownLogEvents.AutomationTool_BuildGraphScript, $"{{Script}}({{Line}}): warning: {format}", allArgs);
 		}
 	}
 
@@ -512,7 +512,7 @@ namespace AutomationTool
 						await ReadAnnotationAsync(childElement);
 						break;
 					default:
-						LogError(childElement, "Invalid element '{0}'", childElement.Name);
+						LogError(childElement, "Invalid element '{ElementName}'", childElement.Name);
 						break;
 				}
 			}
@@ -554,7 +554,7 @@ namespace AutomationTool
 			// Make sure this property name was not already used in a child scope; it likely indicates an error.
 			if (_shadowProperties[scopeIdx].Contains(name))
 			{
-				LogError(element, "Property '{0}' was already used in a child scope. Move this definition before the previous usage if they are intended to share scope, or use a different name.", name);
+				LogError(element, "Property '{PropertyName}' was already used in a child scope. Move this definition before the previous usage if they are intended to share scope, or use a different name.", name);
 			}
 			else
 			{
@@ -721,7 +721,7 @@ namespace AutomationTool
 						string restrict = ReadAttribute(element, "Restrict");
 						if (!String.IsNullOrEmpty(restrict) && !Regex.IsMatch(existingValue, "^" + restrict + "$", RegexOptions.IgnoreCase))
 						{
-							LogError(element, "'{0} is already set to '{1}', which does not match the given restriction ('{2}')", name, existingValue, restrict);
+							LogError(element, "'{Name} is already set to '{ExistingValue}', which does not match the given restriction ('{Restrict}')", name, existingValue, restrict);
 						}
 					}
 					else
@@ -747,11 +747,11 @@ namespace AutomationTool
 							string pattern = "^(" + restrict + ")$";
 							if (!Regex.IsMatch(value, pattern, RegexOptions.IgnoreCase))
 							{
-								LogError(element, "'{0}' is not a valid value for '{1}' (required: '{2}')", value, name, restrict);
+								LogError(element, "'{Value}' is not a valid value for '{Name}' (required: '{Restrict}')", value, name, restrict);
 							}
 							if (option.DefaultValue != value && !Regex.IsMatch(option.DefaultValue, pattern, RegexOptions.IgnoreCase))
 							{
-								LogError(element, "Default value '{0}' is not valid for '{1}' (required: '{2}')", option.DefaultValue, name, restrict);
+								LogError(element, "Default value '{DefaultValue}' is not valid for '{Name}' (required: '{Restrict}')", option.DefaultValue, name, restrict);
 							}
 						}
 					}
@@ -826,7 +826,7 @@ namespace AutomationTool
 					// this includes the entire string match group as [0], so don't count that one.
 					if (captureNames.Length != groupNumbers.Count() - 1)
 					{
-						LogError(element, "MatchGroup count: {0} does not match the number of names specified: {1}", groupNumbers.Count() - 1, captureNames.Length);
+						LogError(element, "MatchGroup count: {Count} does not match the number of names specified: {NameCount}", groupNumbers.Count() - 1, captureNames.Length);
 					}
 					else
 					{
@@ -839,7 +839,7 @@ namespace AutomationTool
 						{
 							if (!optional)
 							{
-								LogError(element, "Regex {0} did not find a match against input string {1}", regexString, input);
+								LogError(element, "Regex {Regex} did not find a match against input string {Input}", regexString, input);
 							}
 						}
 						else
@@ -885,7 +885,7 @@ namespace AutomationTool
 				if (_macroNameToDefinition.TryGetValue(name, out originalDefinition))
 				{
 					BgScriptLocation location = originalDefinition.Elements[0].Location;
-					LogError(element, "Macro '{0}' has already been declared (see {1} line {2})", name, location.File, location.LineNumber);
+					LogError(element, "Macro '{Name}' has already been declared (see {File} line {Line})", name, location.File, location.LineNumber);
 				}
 				else
 				{
@@ -915,7 +915,7 @@ namespace AutomationTool
 				{
 					if (argumentNameToIndex.ContainsKey(argumentName))
 					{
-						LogWarning(element, "Argument '{0}' is listed multiple times", argumentName);
+						LogWarning(element, "Argument '{Name}' is listed multiple times", argumentName);
 					}
 					else
 					{
@@ -942,7 +942,7 @@ namespace AutomationTool
 				}
 				else
 				{
-					LogError(element, "Macro '{0}' has not been declared", name);
+					LogError(element, "Macro '{Name}' has not been declared", name);
 				}
 			}
 		}
@@ -967,7 +967,7 @@ namespace AutomationTool
 					{
 						if (types.Length != agent.PossibleTypes.Count || !types.SequenceEqual(agent.PossibleTypes, StringComparer.InvariantCultureIgnoreCase))
 						{
-							LogError(element, "Agent types ({0}) were different than previous agent definition with types ({1}). Must either be empty or match exactly.", String.Join(",", types), String.Join(",", agent.PossibleTypes));
+							LogError(element, "Agent types ({Types}) were different than previous agent definition with types ({PossibleTypes}). Must either be empty or match exactly.", String.Join(",", types), String.Join(",", agent.PossibleTypes));
 						}
 					}
 				}
@@ -975,7 +975,7 @@ namespace AutomationTool
 				{
 					if (types.Length == 0)
 					{
-						LogError(element, "Missing type for agent '{0}'", name);
+						LogError(element, "Missing type for agent '{Name}'", name);
 					}
 					agent = new BgAgentDef(name);
 					agent.PossibleTypes.AddRange(types);
@@ -1044,7 +1044,7 @@ namespace AutomationTool
 						await ReadAnnotationAsync(childElement);
 						break;
 					default:
-						LogError(childElement, "Unexpected element type '{0}'", childElement.Name);
+						LogError(childElement, "Unexpected element type '{ElementName}'", childElement.Name);
 						break;
 				}
 			}
@@ -1273,11 +1273,11 @@ namespace AutomationTool
 					BgNodeOutput? existingOutput;
 					if (_graph.TagNameToNodeOutput.TryGetValue(producesName, out existingOutput))
 					{
-						LogError(element, "Output tag '{0}' is already generated by node '{1}'", producesName, existingOutput.ProducingNode.Name);
+						LogError(element, "Output tag '{Tag}' is already generated by node '{Name}'", producesName, existingOutput.ProducingNode.Name);
 					}
 					else if (!producesName.StartsWith("#"))
 					{
-						LogError(element, "Output tag names must begin with a '#' character ('{0}')", producesName);
+						LogError(element, "Output tag names must begin with a '#' character ('{Name}')", producesName);
 					}
 					else
 					{
@@ -1301,7 +1301,7 @@ namespace AutomationTool
 				{
 					foreach (BgNodeDef node in _graph.Agents[idx].Nodes.Where(x => orderDependencies.Contains(x)))
 					{
-						LogError(element, "Node '{0}' has a dependency on '{1}', which was declared after the initial definition of '{2}'.", name, node.Name, _enclosingAgent!.Name);
+						LogError(element, "Node '{Name}' has a dependency on '{OtherName}', which was declared after the initial definition of '{AgentName}'.", name, node.Name, _enclosingAgent!.Name);
 					}
 				}
 
@@ -1324,7 +1324,7 @@ namespace AutomationTool
 						BgNodeOutput? existingOutput;
 						if (_graph.TagNameToNodeOutput.TryGetValue(output.TagName, out existingOutput))
 						{
-							LogError(element, "Node '{0}' already has an output called '{1}'", existingOutput.ProducingNode.Name, output.TagName);
+							LogError(element, "Node '{NodeName}' already has an output called '{TagName}'", existingOutput.ProducingNode.Name, output.TagName);
 						}
 						else
 						{
@@ -1433,7 +1433,7 @@ namespace AutomationTool
 				string separator = ReadAttribute(element, "Separator");
 				if (separator.Length > 1)
 				{
-					LogWarning(element, "Node {0}'s Separator attribute is more than one character ({1}). Defaulting to ;", name, separator);
+					LogWarning(element, "Node {Name}'s Separator attribute is more than one character ({Separator}). Defaulting to ;", name, separator);
 					separator = ";";
 				}
 				if (String.IsNullOrEmpty(separator))
@@ -1444,7 +1444,7 @@ namespace AutomationTool
 				{
 					if (ScopedProperties.Any(x => x.ContainsKey(name)))
 					{
-						LogError(element, "Loop variable '{0}' already exists as a local property in an outer scope", name);
+						LogError(element, "Loop variable '{Name}' already exists as a local property in an outer scope", name);
 					}
 					else
 					{
@@ -1476,7 +1476,7 @@ namespace AutomationTool
 					BgScriptMacro? macro;
 					if (!_macroNameToDefinition.TryGetValue(name, out macro))
 					{
-						LogError(element, "Macro '{0}' does not exist", name);
+						LogError(element, "Macro '{Name}' does not exist", name);
 					}
 					else
 					{
@@ -1493,7 +1493,7 @@ namespace AutomationTool
 								}
 								else
 								{
-									LogWarning(element, "Macro '{0}' does not take an argument '{1}'", name, attribute.Name);
+									LogWarning(element, "Macro '{Name}' does not take an argument '{ArgName}'", name, attribute.Name);
 								}
 							}
 						}
@@ -1504,7 +1504,7 @@ namespace AutomationTool
 						{
 							if (arguments[idx] == null)
 							{
-								LogWarning(element, "Macro '{0}' is missing argument '{1}'", macro.Name, macro.ArgumentNameToIndex.First(x => x.Value == idx).Key);
+								LogWarning(element, "Macro '{Name}' is missing argument '{ArgName}'", macro.Name, macro.ArgumentNameToIndex.First(x => x.Value == idx).Key);
 								bHasMissingArguments = true;
 							}
 						}
@@ -1641,7 +1641,7 @@ namespace AutomationTool
 						}
 						else
 						{
-							LogError(element, "Report '{0}' has not been defined", reportName);
+							LogError(element, "Report '{ReportName}' has not been defined", reportName);
 						}
 					}
 				}
@@ -1765,7 +1765,7 @@ namespace AutomationTool
 			// Get the nodes that it maps to
 			if (_graph.ContainsName(name))
 			{
-				LogError(element, "'{0}' is already defined; cannot add a second time", name);
+				LogError(element, "'{Name}' is already defined; cannot add a second time", name);
 				return false;
 			}
 			return true;
@@ -1789,11 +1789,11 @@ namespace AutomationTool
 				}
 				else if (!referenceName.StartsWith("#") && _graph.TagNameToNodeOutput.ContainsKey("#" + referenceName))
 				{
-					LogError(element, "Reference to '{0}' cannot be resolved; did you mean '#{0}'?", referenceName);
+					LogError(element, "Reference to '{Name}' cannot be resolved; did you mean '{PossibleName}'?", referenceName, $"#{referenceName}");
 				}
 				else
 				{
-					LogError(element, "Reference to '{0}' cannot be resolved; check it has been defined.", referenceName);
+					LogError(element, "Reference to '{Name}' cannot be resolved; check it has been defined.", referenceName);
 				}
 			}
 			return nodes;
@@ -1817,11 +1817,11 @@ namespace AutomationTool
 				}
 				else if (!referenceName.StartsWith("#") && _graph.TagNameToNodeOutput.ContainsKey("#" + referenceName))
 				{
-					LogError(element, "Reference to '{0}' cannot be resolved; did you mean '#{0}'?", referenceName);
+					LogError(element, "Reference to '{Name}' cannot be resolved; did you mean '{PossibleName}'?", referenceName, $"#{referenceName}");
 				}
 				else
 				{
-					LogError(element, "Reference to '{0}' cannot be resolved; check it has been defined.", referenceName);
+					LogError(element, "Reference to '{Name}' cannot be resolved; check it has been defined.", referenceName);
 				}
 			}
 			return inputs;
@@ -1847,12 +1847,12 @@ namespace AutomationTool
 			{
 				if (idx > 0 && name[idx] == ' ' && name[idx - 1] == ' ')
 				{
-					LogError(element, "Consecutive spaces in object name - '{0}'", name);
+					LogError(element, "Consecutive spaces in object name - '{Name}'", name);
 					return false;
 				}
 				if (Char.IsControl(name[idx]) || BgScriptSchema.IllegalNameCharacters.IndexOf(name[idx]) != -1)
 				{
-					LogError(element, "Invalid character in object name - '{0}'", name[idx]);
+					LogError(element, "Invalid character in object name - '{Name}'", name[idx]);
 					return false;
 				}
 			}
@@ -1878,7 +1878,7 @@ namespace AutomationTool
 			}
 			catch (ArgumentException invalidRegex)
 			{
-				LogError(element, "Could not construct the Regex, reason: {0}", invalidRegex.Message);
+				LogError(element, "Could not construct the Regex, reason: {Reason}", invalidRegex.Message);
 				return null;
 			}
 		}
@@ -1936,19 +1936,19 @@ namespace AutomationTool
 					int idx = pair.IndexOf('=');
 					if (idx < 0)
 					{
-						LogError(element, "Invalid annotation '{0}'", pair);
+						LogError(element, "Invalid annotation '{Pair}'", pair);
 						continue;
 					}
 
 					string key = pair.Substring(0, idx).Trim();
 					if (!Regex.IsMatch(key, @"[a-zA-Z0-9_\.]+"))
 					{
-						LogError(element, "Invalid annotation key '{0}'", pair);
+						LogError(element, "Invalid annotation key '{Pair}'", pair);
 						continue;
 					}
 					if (pairMap.ContainsKey(key))
 					{
-						LogError(element, "Annotation key '{0}' was specified twice", key);
+						LogError(element, "Annotation key '{Key}' was specified twice", key);
 						continue;
 					}
 
@@ -2008,7 +2008,7 @@ namespace AutomationTool
 				}
 				else
 				{
-					LogError(element, "Invalid integer value '{0}'", value);
+					LogError(element, "Invalid integer value '{Value}'", value);
 				}
 			}
 			return result;
@@ -2036,7 +2036,7 @@ namespace AutomationTool
 				}
 				else
 				{
-					LogError(element, "Invalid value '{0}' - expected {1}", value, String.Join("/", Enum.GetNames(typeof(T))));
+					LogError(element, "Invalid value '{Value}' - expected {PossibleValues}", value, String.Join("/", Enum.GetNames(typeof(T))));
 				}
 			}
 			return result;
@@ -2087,7 +2087,7 @@ namespace AutomationTool
 			}
 			catch (BgConditionException ex)
 			{
-				LogError(element, "Error in condition: {0}", ex.Message);
+				LogError(element, "Error in condition: {Message}", ex.Message);
 				return false;
 			}
 		}
@@ -2118,7 +2118,7 @@ namespace AutomationTool
 				string? value;
 				if (!TryGetPropertyValue(name, out value))
 				{
-					LogWarning(element, "Property '{0}' is not defined", name);
+					LogWarning(element, "Property '{Name}' is not defined", name);
 					value = "";
 				}
 
diff --git a/Engine/Source/Programs/Shared/EpicGames.Core/KnownLogEvents.cs b/Engine/Source/Programs/Shared/EpicGames.Core/KnownLogEvents.cs
--- a/Engine/Source/Programs/Shared/EpicGames.Core/KnownLogEvents.cs
+++ b/Engine/Source/Programs/Shared/EpicGames.Core/KnownLogEvents.cs
@@ -130,6 +130,11 @@ namespace EpicGames.Core
 		/// </summary>
 		public static EventId AutomationTool_UnacceptableWords { get; } = new EventId(406);
 
+		/// <summary>
+		/// UAT: BuildGraph script parsing error
+		/// </summary>
+		public static EventId AutomationTool_BuildGraphScript { get; } = new EventId(407);
+
 		/// <summary>
 		/// MSBuild: Generic error
 		/// </summary>
