[Gauntlet] If an unreal process produces a malformed log, don't skip saving all artifacts when throwing an encoding exception

diff --git a/Engine/Source/Programs/AutomationTool/Gauntlet/Unreal/Base/Gauntlet.UnrealSession.cs b/Engine/Source/Programs/AutomationTool/Gauntlet/Unreal/Base/Gauntlet.UnrealSession.cs
--- a/Engine/Source/Programs/AutomationTool/Gauntlet/Unreal/Base/Gauntlet.UnrealSession.cs
+++ b/Engine/Source/Programs/AutomationTool/Gauntlet/Unreal/Base/Gauntlet.UnrealSession.cs
@@ -1123,17 +1123,25 @@ namespace Gauntlet
 			// Write instance stdout stream
 			LogOut.Append(InRunningRole.AppInstance.StdOut);
 
-			File.WriteAllText(ArtifactLogPath, LogOut.ToString());
+			try
+			{
+				File.WriteAllText(ArtifactLogPath, LogOut.ToString());
 
-			// On build machines, copy all role logs to Horde.
-			if (CommandUtils.IsBuildMachine)
+				// On build machines, copy all role logs to Horde.
+				if (CommandUtils.IsBuildMachine)
+				{
+					string HordeLogFileName = Path.Combine(CommandUtils.CmdEnv.LogFolder, RoleName + "Output.log");
+					File.WriteAllText(HordeLogFileName, LogOut.ToString());
+				}
+
+				Log.Info("Wrote Log to {0}", ArtifactLogPath);
+			}
+			catch(Exception Ex)
 			{
-				string HordeLogFileName = Path.Combine(CommandUtils.CmdEnv.LogFolder, RoleName + "Output.log");
-				File.WriteAllText(HordeLogFileName, LogOut.ToString());
+				string Message = "Encountered an {0} when attempting to write the {1} process log. The log may contain malformed encoding and will not be present on horde. {2}";
+				Log.Warning(Message, Ex.GetType().Name, RoleName, Ex.Message);
 			}
 
-			Log.Info("Wrote Log to {0}", ArtifactLogPath);
-
 			if (IsServer == false)
 			{
 				// gif-ify and jpeg-ify any screenshots
