Fix DeployTool task no longer accepting a zip file argument.

diff --git a/Engine/Source/Programs/AutomationTool/BuildGraph/Tasks/DeployToolTask.cs b/Engine/Source/Programs/AutomationTool/BuildGraph/Tasks/DeployToolTask.cs
--- a/Engine/Source/Programs/AutomationTool/BuildGraph/Tasks/DeployToolTask.cs
+++ b/Engine/Source/Programs/AutomationTool/BuildGraph/Tasks/DeployToolTask.cs
@@ -58,11 +58,17 @@ namespace AutomationTool.Tasks
 		[TaskParameter(Optional = true)]
 		public bool Paused = false;
 
+		/// <summary>
+		/// Zip file containing files to upload
+		/// </summary>
+		[TaskParameter(Optional = true)]
+		public string? File = null!;
+
 		/// <summary>
 		/// Directory to upload for the tool
 		/// </summary>
-		[TaskParameter]
-		public string Directory = null!;
+		[TaskParameter(Optional = true)]
+		public string? Directory = null!;
 	}
 
 	/// <summary>
@@ -147,11 +153,21 @@ namespace AutomationTool.Tasks
 			HttpStorageClient storageClient = new HttpStorageClient(CreateHttpClient, () => new HttpClient(), Logger);
 			using (TreeWriter treeWriter = new TreeWriter(storageClient))
 			{
-				DirectoryInfo directoryInfo = new DirectoryInfo(Parameters.Directory!);
-
 				DirectoryNode sandbox = new DirectoryNode();
-				await sandbox.CopyFromDirectoryAsync(directoryInfo, new ChunkingOptions(), treeWriter, null);
-
+				if (Parameters.File != null)
+				{
+					using FileStream stream = FileReference.Open(ResolveFile(Parameters.File), FileMode.Open, FileAccess.Read);
+					await sandbox.CopyFromZipStreamAsync(stream, treeWriter);
+				}
+				else if (Parameters.Directory != null)
+				{
+					DirectoryInfo directoryInfo = ResolveDirectory(Parameters.Directory).ToDirectoryInfo();
+					await sandbox.CopyFromDirectoryAsync(directoryInfo, new ChunkingOptions(), treeWriter, null);
+				}
+				else
+				{
+					throw new AutomationException("Either File=... or Directory=... must be specified");
+				}
 				handle = await treeWriter.FlushAsync(sandbox);
 			}
 
diff --git a/Engine/Source/Programs/Shared/EpicGames.Horde/Storage/Nodes/DirectoryNode.cs b/Engine/Source/Programs/Shared/EpicGames.Horde/Storage/Nodes/DirectoryNode.cs
--- a/Engine/Source/Programs/Shared/EpicGames.Horde/Storage/Nodes/DirectoryNode.cs
+++ b/Engine/Source/Programs/Shared/EpicGames.Horde/Storage/Nodes/DirectoryNode.cs
@@ -610,7 +610,7 @@ namespace EpicGames.Horde.Storage.Nodes
 		/// <param name="stream">Input stream</param>
 		/// <param name="writer">Writer for new nodes</param>
 		/// <param name="cancellationToken">Cancellation token for the operation</param>
-		public async Task CopyFromZipStreamAsync(Stream stream, TreeWriter writer, CancellationToken cancellationToken)
+		public async Task CopyFromZipStreamAsync(Stream stream, TreeWriter writer, CancellationToken cancellationToken = default)
 		{
 			FileNodeWriter fileWriter = new FileNodeWriter(writer, new ChunkingOptions());
 			using (ZipArchive archive = new ZipArchive(stream, ZipArchiveMode.Read, true))
