- Fixed location of legacy xcode projects for where to find temporary workspaces

diff --git a/Engine/Source/Programs/AutomationTool/IOS/IOSPlatform.Automation.cs b/Engine/Source/Programs/AutomationTool/IOS/IOSPlatform.Automation.cs
--- a/Engine/Source/Programs/AutomationTool/IOS/IOSPlatform.Automation.cs
+++ b/Engine/Source/Programs/AutomationTool/IOS/IOSPlatform.Automation.cs
@@ -1077,13 +1077,13 @@ public class IOSPlatform : ApplePlatform
 	{
 		// first check for the .xcodeproj
 		bWasGenerated = false;
-		string RawProjectDir = RawProjectPath.Directory.FullName;
-		string XcodeProj = RawProjectPath.FullName.Replace(".uproject", "_" + PlatformName + ".xcworkspace");
-		if (!Directory.Exists(RawProjectDir + "/Source") && !Directory.Exists(RawProjectDir + "/Intermediate/Source"))
+		DirectoryReference RawProjectDir = RawProjectPath.Directory;
+		DirectoryReference XcodeProj = DirectoryReference.Combine(RawProjectDir, "Intermediate/ProjectFiles", $"{RawProjectPath.GetFileNameWithoutAnyExtensions()}_{PlatformName}.xcworkspace");
+		if (!DirectoryReference.Exists(DirectoryReference.Combine(RawProjectDir, "Source")) && !DirectoryReference.Exists(DirectoryReference.Combine(RawProjectDir, "/Intermediate/Source")))
 		{
-			XcodeProj = CombinePaths(CmdEnv.LocalRoot, "Engine", Path.GetFileName(XcodeProj));
+			XcodeProj = DirectoryReference.Combine(Unreal.EngineDirectory, "Intermediate/ProjectFiles/", XcodeProj.GetDirectoryName());
 		}
-		Console.WriteLine("Project: " + XcodeProj);
+		Console.WriteLine("Project: " + XcodeProj.FullName);
 		{
 			// project.xcodeproj doesn't exist, so generate temp project
 			string Arguments = "-project=\"" + RawProjectPath + "\"";
@@ -1106,14 +1106,14 @@ public class IOSPlatform : ApplePlatform
 
 			bWasGenerated = true;
 
-			if (!Directory.Exists(XcodeProj))
+			if (!DirectoryReference.Exists(XcodeProj))
 			{
 				// something very bad happened
-				throw new AutomationException("iOS couldn't find the appropriate Xcode Project " + XcodeProj);
+				throw new AutomationException("iOS couldn't find the appropriate Xcode Project " + XcodeProj.FullName);
 			}
 		}
 
-		return XcodeProj;
+		return XcodeProj.FullName;
 	}
 
 	private void CodeSign(string BaseDirectory, string GameName, FileReference RawProjectPath, UnrealTargetConfiguration TargetConfig, string LocalRoot, string ProjectName, string ProjectDirectory, bool IsCode, bool Distribution, string Provision, string Certificate, string Team, bool bAutomaticSigning, string SchemeName, string SchemeConfiguration)
diff --git a/Engine/Source/Programs/UnrealBuildTool/Platform/IOS/IOSToolChain.cs b/Engine/Source/Programs/UnrealBuildTool/Platform/IOS/IOSToolChain.cs
--- a/Engine/Source/Programs/UnrealBuildTool/Platform/IOS/IOSToolChain.cs
+++ b/Engine/Source/Programs/UnrealBuildTool/Platform/IOS/IOSToolChain.cs
@@ -1504,11 +1504,11 @@ namespace UnrealBuildTool
                 {
 					if (AppName == "UnrealGame" || AppName == "UnrealClient" || Target.ProjectFile == null || Target.ProjectFile.IsUnderDirectory(Unreal.EngineDirectory))
 					{
-						XcodeWorkspaceDir = DirectoryReference.Combine(Unreal.RootDirectory, String.Format("UE5_{0}.xcworkspace", (Target.Platform == UnrealTargetPlatform.IOS ? "IOS" : "TVOS")));
+						XcodeWorkspaceDir = DirectoryReference.Combine(Unreal.EngineDirectory, "Intermediate/ProjectFiles", String.Format("UE5_{0}.xcworkspace", (Target.Platform == UnrealTargetPlatform.IOS ? "IOS" : "TVOS")));
 					}
 					else
 					{
-						XcodeWorkspaceDir = DirectoryReference.Combine(Target.ProjectDirectory, String.Format("{0}_{1}.xcworkspace", Target.ProjectFile.GetFileNameWithoutExtension(), (Target.Platform == UnrealTargetPlatform.IOS ? "IOS" : "TVOS")));
+						XcodeWorkspaceDir = DirectoryReference.Combine(Target.ProjectDirectory, "Intermediate/ProjectFiles", String.Format("{0}_{1}.xcworkspace", Target.ProjectFile.GetFileNameWithoutExtension(), (Target.Platform == UnrealTargetPlatform.IOS ? "IOS" : "TVOS")));
 					}
                 }
 
