Better diagnostic for detecting use of zenstore when generating uecommand.txt

diff --git a/Engine/Source/Programs/AutomationTool/Scripts/RunProjectCommand.Automation.cs b/Engine/Source/Programs/AutomationTool/Scripts/RunProjectCommand.Automation.cs
--- a/Engine/Source/Programs/AutomationTool/Scripts/RunProjectCommand.Automation.cs
+++ b/Engine/Source/Programs/AutomationTool/Scripts/RunProjectCommand.Automation.cs
@@ -823,7 +823,7 @@ namespace AutomationScripts
 		{
 			string FileHostParams = "";
 			FileReference ZenStoreMarkerFile = FindZenProjectStoreMarker(Params, SC);
-			bool UseZenServerHost = (ZenStoreMarkerFile != null) && !ShouldCreatePak(Params);
+			bool UseZenServerHost = (ZenStoreMarkerFile != null) && !Params.UsePak(SC.StageTargetPlatform);
 			if (!Params.CookOnTheFly && !Params.FileServer && !UseZenServerHost)
 			{
 				return FileHostParams;
