LLT: Finding the correct executable for Development config when multiple different candidates exist in the same folder

diff --git a/Engine/Source/Programs/AutomationTool/LowLevelTests/RunLowLevelTests.cs b/Engine/Source/Programs/AutomationTool/LowLevelTests/RunLowLevelTests.cs
--- a/Engine/Source/Programs/AutomationTool/LowLevelTests/RunLowLevelTests.cs
+++ b/Engine/Source/Programs/AutomationTool/LowLevelTests/RunLowLevelTests.cs
@@ -509,16 +509,26 @@ namespace LowLevelTests
 					continue;
 				}
 
-				// Executable name must not contain any configuration or platform name
+				// Development executable does not contain configuration or platform name
 				Log.Verbose("Config type: {0}", InConfiguration);
 				if (InConfiguration == UnrealTargetConfiguration.Development)
 				{
-					if (BuildExecutableName.Equals(InTestApp, StringComparison.OrdinalIgnoreCase) 
+					if (BuildExecutableName.Equals(InTestApp, StringComparison.OrdinalIgnoreCase)
 						&& !BuildExecutableName.Contains(InPlatform.ToString(), StringComparison.OrdinalIgnoreCase))
 					{
 						Log.VeryVerbose("Output Development Executable: {0}", Path.GetRelativePath(InBuildPath, Executable));
 						return Path.GetRelativePath(InBuildPath, Executable);
 					}
+					else // move to the next executable until we find the Development one
+					{
+						continue;
+					}
+				}
+
+				// All executables other than Development contain the configuration in their name
+				if (!BuildExecutableName.Contains(InConfiguration.ToString()))
+				{
+					continue;
 				}
 
 				if (!(BuildExecutableName.Contains(InTestApp, StringComparison.OrdinalIgnoreCase) 
