AutomationTool: Fix reference

#rnx

diff --git a/Engine/Source/Programs/AutomationTool/AutomationTool.csproj b/Engine/Source/Programs/AutomationTool/AutomationTool.csproj
--- a/Engine/Source/Programs/AutomationTool/AutomationTool.csproj
+++ b/Engine/Source/Programs/AutomationTool/AutomationTool.csproj
@@ -110,19 +110,13 @@
   </ItemGroup>
 
   <Choose>
-    <When Condition="'$(IsWindows)'=='true'">
-      <ItemGroup Condition="Exists('..\..\..\Restricted\NotForLicensees\Binaries\Win64\Box\BoxManaged.dll')">
-        <Reference Include="..\..\..\Restricted\NotForLicensees\Binaries\Win64\Box\BoxManaged.dll" />
-        <Content Include="..\..\..\Restricted\NotForLicensees\Binaries\Win64\Box\BoxDetours.dll">
-          <CopyToOutputDirectory>PreserveNewest</CopyToOutputDirectory>
-        </Content>
-        <Content Include="..\..\..\Restricted\NotForLicensees\Binaries\Win64\Box\Ijwhost.dll">
-          <CopyToOutputDirectory>PreserveNewest</CopyToOutputDirectory>
-        </Content>
-      </ItemGroup>
-      <PropertyGroup Condition="Exists('..\..\..\Restricted\NotForLicensees\Binaries\Win64\Box\BoxManaged.dll')">
+    <When Condition="'$([System.Runtime.InteropServices.RuntimeInformation]::IsOSPlatform($([System.Runtime.InteropServices.OSPlatform]::Windows)))' == 'true' And '$([System.Runtime.InteropServices.RuntimeInformation]::OSArchitecture)' != 'Arm64'">
+      <PropertyGroup Condition="Exists('..\..\..\Restricted\NotForLicensees\Source\Programs\UnrealBuildTool\Executors\Experimental\BoxExecutor.cs')">
         <DefineConstants>$(DefineConstants);__BOXEXECUTOR_AVAILABLE__</DefineConstants>
       </PropertyGroup>
+      <PropertyGroup Condition="Exists('..\..\..\Restricted\NotForLicensees\Source\Programs\Shared\EpicGames.Box\BoxServer.cs')">
+        <DefineConstants>$(DefineConstants);__BOX_AVAILABLE__</DefineConstants>
+      </PropertyGroup>
     </When>
   </Choose>
 
