Build fix

diff --git a/Engine/Source/Programs/AutomationTool/BuildGraph/Tasks/DeployToolTask.cs b/Engine/Source/Programs/AutomationTool/BuildGraph/Tasks/DeployToolTask.cs
--- a/Engine/Source/Programs/AutomationTool/BuildGraph/Tasks/DeployToolTask.cs
+++ b/Engine/Source/Programs/AutomationTool/BuildGraph/Tasks/DeployToolTask.cs
@@ -148,7 +148,7 @@ namespace AutomationTool.Tasks
 				return httpClient;
 			}
 
-			NodeHandle handle;
+			HashedNodeHandle handle;
 
 			HttpStorageClient storageClient = new HttpStorageClient(CreateHttpClient, () => new HttpClient(), Logger);
 			using (TreeWriter treeWriter = new TreeWriter(storageClient))
@@ -181,7 +181,7 @@ namespace AutomationTool.Tasks
 			{
 				request.CreatePaused = true;
 			}
-			request.Node = handle.ToString();
+			request.Node = handle.Handle.ToString();
 
 			using (HttpClient httpClient = CreateHttpClient())
 			{
