Don't attempt to log into p4 on a builder and with no password supplied

#rnx

diff --git a/Engine/Source/Programs/AutomationTool/AutomationUtils/P4Utils.cs b/Engine/Source/Programs/AutomationTool/AutomationUtils/P4Utils.cs
--- a/Engine/Source/Programs/AutomationTool/AutomationUtils/P4Utils.cs
+++ b/Engine/Source/Programs/AutomationTool/AutomationUtils/P4Utils.cs
@@ -1540,6 +1540,11 @@ namespace AutomationTool
 
 			string Output;
             string P4Passwd = InternalUtils.GetEnvironmentVariable("uebp_PASS", "", true) + '\n';
+			if (Automation.IsBuildMachine && string.IsNullOrEmpty(P4Passwd))
+			{
+				return AuthenticationToken;
+			}
+
             P4Output(out Output, "", "login -a -p", P4Passwd);
 
 			// Validate output.
