Fixing password empty check

#rnx

diff --git a/Engine/Source/Programs/AutomationTool/AutomationUtils/P4Utils.cs b/Engine/Source/Programs/AutomationTool/AutomationUtils/P4Utils.cs
--- a/Engine/Source/Programs/AutomationTool/AutomationUtils/P4Utils.cs
+++ b/Engine/Source/Programs/AutomationTool/AutomationUtils/P4Utils.cs
@@ -1539,13 +1539,13 @@ namespace AutomationTool
 			string AuthenticationToken = null;
 
 			string Output;
-            string P4Passwd = InternalUtils.GetEnvironmentVariable("uebp_PASS", "", true) + '\n';
+            string P4Passwd = InternalUtils.GetEnvironmentVariable("uebp_PASS", "", true);
 			if (Automation.IsBuildMachine && string.IsNullOrEmpty(P4Passwd))
 			{
 				return AuthenticationToken;
 			}
 
-            P4Output(out Output, "", "login -a -p", P4Passwd);
+            P4Output(out Output, "", "login -a -p", P4Passwd + '\n');
 
 			// Validate output.
 			const string PasswordPromptString = "Enter password: \r\n";
