BuildGraph: Add an <OnExit> task which schedules a shell command to be executed once a step is finished. When run on Horde, the command is guaranteed to be executed even if UAT terminates.

diff --git a/Engine/Source/Programs/AutomationTool/BuildGraph/BuildGraph.cs b/Engine/Source/Programs/AutomationTool/BuildGraph/BuildGraph.cs
--- a/Engine/Source/Programs/AutomationTool/BuildGraph/BuildGraph.cs
+++ b/Engine/Source/Programs/AutomationTool/BuildGraph/BuildGraph.cs
@@ -22,6 +22,7 @@ using EpicGames.BuildGraph.Expressions;
 using AutomationTool.Tasks;
 
 using static AutomationTool.CommandUtils;
+using System.Runtime.InteropServices;
 
 namespace AutomationTool
 {
@@ -961,6 +962,49 @@ namespace AutomationTool
 			return true;
 		}
 
+		/// <summary>
+		/// Helper class to execute a cleanup script on termination
+		/// </summary>
+		class CleanupScriptRunner : IDisposable
+		{
+			readonly ILogger _logger;
+			readonly FileReference _scriptFile;
+
+			public CleanupScriptRunner(ILogger logger)
+			{
+				_logger = logger;
+
+				string CleanupScriptEnvVar = Environment.GetEnvironmentVariable(CustomTask.CleanupScriptEnvVarName);
+				if (String.IsNullOrEmpty(CleanupScriptEnvVar))
+				{
+					string extension = RuntimeInformation.IsOSPlatform(OSPlatform.Windows) ? "bat" : "sh";
+					_scriptFile = FileReference.Combine(Unreal.EngineDirectory, "Intermediate", $"OnExit.{extension}");
+					FileReference.Delete(_scriptFile);
+					Environment.SetEnvironmentVariable(CustomTask.CleanupScriptEnvVarName, _scriptFile.FullName);
+				}
+			}
+
+			public void Dispose()
+			{
+				if (_scriptFile != null)
+				{
+					if (FileReference.Exists(_scriptFile))
+					{
+						_logger.LogInformation("Executing cleanup commands from {ScriptFile}", _scriptFile);
+						if (RuntimeInformation.IsOSPlatform(OSPlatform.Windows))
+						{
+							CommandUtils.Run(BuildHostPlatform.Current.Shell.FullName, $"/C {CommandLineArguments.Quote(_scriptFile.FullName)}");
+						}
+						else
+						{
+							CommandUtils.Run(BuildHostPlatform.Current.Shell.FullName, CommandLineArguments.Quote(_scriptFile.FullName));
+						}
+					}
+					Environment.SetEnvironmentVariable(CustomTask.CleanupScriptEnvVarName, null);
+				}
+			}
+		}
+
 		/// <summary>
 		/// Build a node
 		/// </summary>
@@ -975,6 +1019,9 @@ namespace AutomationTool
 		{
 			DirectoryReference RootDir = new DirectoryReference(CommandUtils.CmdEnv.LocalRoot);
 
+			// Register something to execute cleanup commands
+			using CleanupScriptRunner CleanupRunner = new CleanupScriptRunner(Logger);
+
 			// Create the mapping of tag names to file sets
 			Dictionary<string, HashSet<FileReference>> TagNameToFileSet = new Dictionary<string,HashSet<FileReference>>();
 
diff --git a/Engine/Source/Programs/AutomationTool/BuildGraph/Task.cs b/Engine/Source/Programs/AutomationTool/BuildGraph/Task.cs
--- a/Engine/Source/Programs/AutomationTool/BuildGraph/Task.cs
+++ b/Engine/Source/Programs/AutomationTool/BuildGraph/Task.cs
@@ -435,12 +435,17 @@ namespace AutomationTool
 		}
 
 		/// <summary>
-		/// Add cleanup commands to run after the job completes
+		/// Name of the environment variable containing cleanup commands
+		/// </summary>
+		public const string CleanupScriptEnvVarName = "UE_HORDE_CLEANUP";
+
+		/// <summary>
+		/// Add cleanup commands to run after the step completes
 		/// </summary>
 		/// <param name="NewLines">Lines to add to the cleanup script</param>
 		public static async Task AddCleanupCommandsAsync(IEnumerable<string> NewLines)
 		{
-			string CleanupScriptEnvVar = Environment.GetEnvironmentVariable("UE_HORDE_CLEANUP");
+			string CleanupScriptEnvVar = Environment.GetEnvironmentVariable(CleanupScriptEnvVarName);
 			if (!String.IsNullOrEmpty(CleanupScriptEnvVar))
 			{
 				FileReference CleanupScript = new FileReference(CleanupScriptEnvVar);
diff --git a/Engine/Source/Programs/AutomationTool/BuildGraph/Tasks/OnExitTask.cs b/Engine/Source/Programs/AutomationTool/BuildGraph/Tasks/OnExitTask.cs
new file mode 100644
--- /dev/null
+++ b/Engine/Source/Programs/AutomationTool/BuildGraph/Tasks/OnExitTask.cs
@@ -0,0 +1,68 @@
+using EpicGames.Core;
+using System;
+using System.Collections.Generic;
+using System.Linq;
+using System.Text;
+using System.Threading.Tasks;
+using System.Xml;
+
+namespace AutomationTool.Tasks
+{
+	/// <summary>
+	/// Parameters for a spawn task
+	/// </summary>
+	public class OnExitTaskParameters
+	{
+		/// <summary>
+		/// Executable to spawn.
+		/// </summary>
+		[TaskParameter]
+		public string Command = String.Empty;
+	}
+
+	/// <summary>
+	/// Spawns an external executable and waits for it to complete.
+	/// </summary>
+	[TaskElement("OnExit", typeof(OnExitTaskParameters))]
+	public class OnExitTask : BgTaskImpl
+	{
+		/// <summary>
+		/// Parameters for this task
+		/// </summary>
+		OnExitTaskParameters Parameters;
+
+		/// <summary>
+		/// Construct a spawn task
+		/// </summary>
+		/// <param name="InParameters">Parameters for the task</param>
+		public OnExitTask(OnExitTaskParameters InParameters)
+		{
+			Parameters = InParameters;
+		}
+
+		/// <inheritdoc/>
+		public override async Task ExecuteAsync(JobContext Job, HashSet<FileReference> BuildProducts, Dictionary<string, HashSet<FileReference>> TagNameToFileSet)
+		{
+			string[] commands = Parameters.Command.Split('\n').Select(x => x.Trim()).ToArray();
+			await AddCleanupCommandsAsync(commands);
+		}
+
+		/// <inheritdoc/>
+		public override void Write(XmlWriter Writer)
+		{
+			Write(Writer, Parameters);
+		}
+
+		/// <inheritdoc/>
+		public override IEnumerable<string> FindConsumedTagNames()
+		{
+			yield break;
+		}
+
+		/// <inheritdoc/>
+		public override IEnumerable<string> FindProducedTagNames()
+		{
+			yield break;
+		}
+	}
+}
