Missing copyright notice.

diff --git a/Engine/Source/Programs/AutomationTool/BuildGraph/Tasks/OnExitTask.cs b/Engine/Source/Programs/AutomationTool/BuildGraph/Tasks/OnExitTask.cs
--- a/Engine/Source/Programs/AutomationTool/BuildGraph/Tasks/OnExitTask.cs
+++ b/Engine/Source/Programs/AutomationTool/BuildGraph/Tasks/OnExitTask.cs
@@ -1,3 +1,5 @@
+// Copyright Epic Games, Inc. All Rights Reserved.
+
 using EpicGames.Core;
 using System;
 using System.Collections.Generic;
