#10506: Handle altsync and noaltsync values added in recent P4D versions.

diff --git a/Engine/Source/Programs/AutomationTool/AutomationUtils/P4Utils.cs b/Engine/Source/Programs/AutomationTool/AutomationUtils/P4Utils.cs
--- a/Engine/Source/Programs/AutomationTool/AutomationUtils/P4Utils.cs
+++ b/Engine/Source/Programs/AutomationTool/AutomationUtils/P4Utils.cs
@@ -78,6 +78,8 @@ namespace AutomationTool
         Locked = 512,
         ModTime = 1024,
         RmDir = 2048,
+        NoAltSync = 4096,
+        AltSync = 8192,
     }
 
 	public class P4ClientInfo
