BuildGraph: fix task to write the correct extension

diff --git a/Engine/Source/Programs/AutomationTool/BuildGraph/Tasks/HordeCreateReportTask.cs b/Engine/Source/Programs/AutomationTool/BuildGraph/Tasks/HordeCreateReportTask.cs
--- a/Engine/Source/Programs/AutomationTool/BuildGraph/Tasks/HordeCreateReportTask.cs
+++ b/Engine/Source/Programs/AutomationTool/BuildGraph/Tasks/HordeCreateReportTask.cs
@@ -74,7 +74,7 @@ namespace AutomationTool.Tasks
 			FileReference ReportTextFile = FileReference.Combine(new DirectoryReference(CommandUtils.CmdEnv.LogFolder), $"{Parameters.Name}.md");
 			await FileReference.WriteAllTextAsync(ReportTextFile, Parameters.Text);
 
-			FileReference ReportJsonFile = FileReference.Combine(new DirectoryReference(CommandUtils.CmdEnv.LogFolder), $"{Parameters.Name}.json.md");
+			FileReference ReportJsonFile = FileReference.Combine(new DirectoryReference(CommandUtils.CmdEnv.LogFolder), $"{Parameters.Name}.report.json");
 			using (FileStream ReportJsonStream = FileReference.Open(ReportJsonFile, FileMode.Create, FileAccess.Write, FileShare.Read))
 			{
 				using (Utf8JsonWriter Writer = new Utf8JsonWriter(ReportJsonStream))
