Use -VerifyLogin when running UE.Automation and UE.BootTest test on clients

#rnx

diff --git a/Engine/Build/Graph/Tasks/Inc/GauntletSettings.xml b/Engine/Build/Graph/Tasks/Inc/GauntletSettings.xml
--- a/Engine/Build/Graph/Tasks/Inc/GauntletSettings.xml
+++ b/Engine/Build/Graph/Tasks/Inc/GauntletSettings.xml
@@ -20,6 +20,7 @@
 	<Option Name="GauntletDevicePool" DefaultValue="UE4" Description="Gauntlet Device Pool to use to run target platform tests"/>
 	<Option Name="AlternateGauntletDevicePool" DefaultValue="UE5" Description="Alternate Gauntlet Device Pool to use to run target platform tests"/>
 	<Option Name="AlternateDevicePoolPlatforms" DefaultValue="XSX;PS5" Description="Platforms that need to use the Alternate Gauntlet Device Pool"/>
+	<Option Name="VerifyLoginOnDevice" DefaultValue="$(IsBuildMachine)" Description="Verify a user is login on Devices that supports it, try to sign in if not."/>
 	<Option Name="ExtraGauntletCommonArgs" DefaultValue="" Description="Additional Common arguments passed to Gauntlet command line"/>
 	<Option Name="ExtraGauntletTargetTestArgs" DefaultValue="" Description="Additional Target test arguments passed to Gauntlet command line"/>
 	<Option Name="ExtraGauntletEditorTestArgs" DefaultValue="" Description="Additional Editor test arguments passed to Gauntlet command line"/>
@@ -29,6 +30,7 @@
 
 	<Property Name="GauntletDeviceArgs" Value="" />
 	<Property Name="GauntletDeviceArgs" Value="-deviceurl=https://horde.devtools.epicgames.com/ -reboot -cleardevices -sandbox=&quot;$(ProjectName)&quot;" If="$(IsBuildMachine)" />
+	<Property Name="GauntletDeviceArgs" Value="$(GauntletDeviceArgs) -VerifyLogin" If="$(VerifyLoginOnDevice)"/>
 	<Property Name="GauntletDeviceAltArgs" Value="$(GauntletDeviceArgs)" />
 	<Property Name="GauntletDeviceArgs" Value="$(GauntletDeviceArgs) -devicepool=&quot;$(GauntletDevicePool)&quot;" If="$(IsBuildMachine)"/>
 	<Property Name="GauntletDeviceAltArgs" Value="$(GauntletDeviceAltArgs) -devicepool=&quot;$(AlternateGauntletDevicePool)&quot;" If="$(IsBuildMachine)"/>
diff --git a/Engine/Source/Programs/AutomationTool/Gauntlet/Unreal/Automation/UE.Automation.cs b/Engine/Source/Programs/AutomationTool/Gauntlet/Unreal/Automation/UE.Automation.cs
--- a/Engine/Source/Programs/AutomationTool/Gauntlet/Unreal/Automation/UE.Automation.cs
+++ b/Engine/Source/Programs/AutomationTool/Gauntlet/Unreal/Automation/UE.Automation.cs
@@ -297,6 +297,11 @@ namespace UE
 				AppConfig.CommandLine += " -rhivalidation";
 			}
 
+			if (ConfigRole.RoleType.IsClient())
+			{
+				VerifyLogin = Globals.Params.ParseParam("VerifyLogin");
+			}
+
 			// Options specific to roles running under the editor
 			if (ConfigRole.RoleType.UsesEditor())
 			{
diff --git a/Engine/Source/Programs/AutomationTool/Gauntlet/Unreal/Automation/UE.BootTest.cs b/Engine/Source/Programs/AutomationTool/Gauntlet/Unreal/Automation/UE.BootTest.cs
--- a/Engine/Source/Programs/AutomationTool/Gauntlet/Unreal/Automation/UE.BootTest.cs
+++ b/Engine/Source/Programs/AutomationTool/Gauntlet/Unreal/Automation/UE.BootTest.cs
@@ -11,10 +11,31 @@ using System.Text.RegularExpressions;
 
 namespace UE
 {
+	/// <summary>
+	/// Define a Config class for UE.BootTest
+	/// </summary>
+	public class BootTestConfig : UnrealTestConfiguration
+	{
+		/// <summary>
+		/// Implement how the settings above are applied to different roles in the session. This is the main avenue for converting options
+		/// into command line parameters for the different roles
+		/// </summary>
+		/// <param name="AppConfig"></param>
+		/// <param name="ConfigRole"></param>
+		/// <param name="OtherRoles"></param>
+		public override void ApplyToConfig(UnrealAppConfig AppConfig, UnrealSessionRole ConfigRole, IEnumerable<UnrealSessionRole> OtherRoles)
+		{
+			base.ApplyToConfig(AppConfig, ConfigRole, OtherRoles);
+			if (ConfigRole.RoleType.IsClient())
+			{
+				VerifyLogin = Globals.Params.ParseParam("VerifyLogin");
+			}
+		}
+	}
 	/// <summary>
 	/// Test that waits for the client and server to get to the front-end then quits
 	/// </summary>
-	public class BootTest : UnrealTestNode<UnrealTestConfiguration>
+	public class BootTest : UnrealTestNode<BootTestConfig>
 	{
 		/// <summary>
 		/// Used to track progress via logging
@@ -44,9 +65,9 @@ namespace UE
 		/// Returns the configuration description for this test
 		/// </summary>
 		/// <returns></returns>
-		public override UnrealTestConfiguration GetConfiguration()
+		public override BootTestConfig GetConfiguration()
 		{
-			UnrealTestConfiguration Config = base.GetConfiguration();
+			BootTestConfig Config = base.GetConfiguration();
 
 			UnrealTestRole Client = Config.RequireRole(UnrealTargetRole.Client);
 
@@ -191,9 +212,9 @@ namespace UE
 		/// Returns the configuration description for this test
 		/// </summary>
 		/// <returns></returns>
-		public override UnrealTestConfiguration GetConfiguration()
+		public override BootTestConfig GetConfiguration()
 		{
-			UnrealTestConfiguration Config = base.GetConfiguration();
+			BootTestConfig Config = base.GetConfiguration();
 			// currently needed as BootTest isn't an abstract class. Can be changed for 4.27
 			Config.ClearRoles();
 			UnrealTestRole EditorRole = Config.RequireRole(Config.CookedEditor ? UnrealTargetRole.CookedEditor : UnrealTargetRole.Editor);
@@ -221,9 +242,9 @@ namespace UE
 		/// Returns the configuration description for this test
 		/// </summary>
 		/// <returns></returns>
-		public override UnrealTestConfiguration GetConfiguration()
+		public override BootTestConfig GetConfiguration()
 		{
-			UnrealTestConfiguration Config = base.GetConfiguration();
+			BootTestConfig Config = base.GetConfiguration();
 			Config.RequireRole(UnrealTargetRole.Client);
 			return Config;
 		}
