Report on device reservation also when test fails to start and no retry is left

#rnx

diff --git a/Engine/Source/Programs/AutomationTool/Gauntlet/Framework/Devices/Gauntlet.DevicePool.cs b/Engine/Source/Programs/AutomationTool/Gauntlet/Framework/Devices/Gauntlet.DevicePool.cs
--- a/Engine/Source/Programs/AutomationTool/Gauntlet/Framework/Devices/Gauntlet.DevicePool.cs
+++ b/Engine/Source/Programs/AutomationTool/Gauntlet/Framework/Devices/Gauntlet.DevicePool.cs
@@ -692,6 +692,24 @@ namespace Gauntlet
 			return ReservationStates != null ? ReservationStates.GetProblems() : null;
 		}
 
+		/// <summary>
+		/// Report on Device reservation state
+		/// </summary>
+		public void ReportDeviceReservationState()
+		{
+			string PoolId = string.IsNullOrEmpty(Globals.DevicePoolId) ? "Local" : Globals.DevicePoolId;
+			string[] DeviceReservationStatus = GetReservationStatuses();
+			if (DeviceReservationStatus != null && DeviceReservationStatus.Length > 0)
+			{
+				Log.Warning(KnownLogEvents.Gauntlet_DeviceEvent, "Device Reservation status from {Pool} pool: {Status}", PoolId, string.Join(", ", DeviceReservationStatus));
+			}
+			string[] DeviceProblemStatus = GetDeviceProblems();
+			if (DeviceProblemStatus != null && DeviceProblemStatus.Length > 0)
+			{
+				Log.Warning(KnownLogEvents.Gauntlet_DeviceEvent, "Reserved Devices marked with problems from {Pool} pool: {Status}", PoolId, string.Join(", ", DeviceProblemStatus));
+			}
+		}
+
 		/// <summary>
 		/// Force clean the cache path by setting all of the files in the cache path
 		/// to have "normal" attributes (i.e. not read-only) before attempting to delete.
diff --git a/Engine/Source/Programs/AutomationTool/Gauntlet/Framework/Gauntlet.TestExecutor.cs b/Engine/Source/Programs/AutomationTool/Gauntlet/Framework/Gauntlet.TestExecutor.cs
--- a/Engine/Source/Programs/AutomationTool/Gauntlet/Framework/Gauntlet.TestExecutor.cs
+++ b/Engine/Source/Programs/AutomationTool/Gauntlet/Framework/Gauntlet.TestExecutor.cs
@@ -271,18 +271,8 @@ namespace Gauntlet
 									double TimeWaiting = (DateTime.Now - NodeInfo.FirstReadyCheckTime).TotalSeconds;
 									if (TimeWaiting >= Options.Wait)
 									{
-										string PoolId = string.IsNullOrEmpty(Globals.DevicePoolId) ? "Local" : Globals.DevicePoolId;
 										Log.Warning(KnownLogEvents.Gauntlet_DeviceEvent, "Test {TestName} has been waiting to run resource-free for {Time:00} seconds. Removing from wait list", Node, TimeWaiting);
-										string[] DeviceReservationStatus = DevicePool.Instance.GetReservationStatuses();
-										if (DeviceReservationStatus != null && DeviceReservationStatus.Length > 0)
-										{
-											Log.Warning(KnownLogEvents.Gauntlet_DeviceEvent, "Device Reservation status from {Pool} pool: {Status}", PoolId, string.Join(", ", DeviceReservationStatus));
-										}
-										string[] DeviceProblemStatus = DevicePool.Instance.GetDeviceProblems();
-										if (DeviceProblemStatus != null && DeviceProblemStatus.Length > 0)
-										{
-											Log.Warning(KnownLogEvents.Gauntlet_DeviceEvent, "Reserved Devices marked with problems from {Pool} pool: {Status}", PoolId, string.Join(", ", DeviceProblemStatus));
-										}
+										DevicePool.Instance.ReportDeviceReservationState();
 										Node.AddTestEvent(new UnrealTestEvent(EventSeverity.Warning, "Insufficient devices found", new List<string> {string.Format("Test {0} was unable to find enough devices after trying for {1:00} seconds.", Node, TimeWaiting), "This is not a test-related failure."}));
 										PendingTests[i] = null;
 										NodeInfo.TimeSetupBegan = NodeInfo.TimeSetupEnded = NodeInfo.TimeTestEnded = DateTime.Now;
diff --git a/Engine/Source/Programs/AutomationTool/Gauntlet/Unreal/Base/Gauntlet.UnrealSession.cs b/Engine/Source/Programs/AutomationTool/Gauntlet/Unreal/Base/Gauntlet.UnrealSession.cs
--- a/Engine/Source/Programs/AutomationTool/Gauntlet/Unreal/Base/Gauntlet.UnrealSession.cs
+++ b/Engine/Source/Programs/AutomationTool/Gauntlet/Unreal/Base/Gauntlet.UnrealSession.cs
@@ -749,6 +749,7 @@ namespace Gauntlet
 					{
 						if (ReservationRetries == 0)
 						{
+							DevicePool.Instance.ReportDeviceReservationState();
 							throw new AutomationException("Unable to acquire all devices for test.");
 						}
 						Log.Info("\nUnable to find enough device(s). Waiting {0} secs (retries left={1})\n", ReservationRetryWait, --ReservationRetries);
