LLT: default test result to failed to account for device reservation failures that do not get to run any tests.

diff --git a/Engine/Source/Programs/AutomationTool/LowLevelTests/Tests/Gauntlet.LowLevelTestNode.cs b/Engine/Source/Programs/AutomationTool/LowLevelTests/Tests/Gauntlet.LowLevelTestNode.cs
--- a/Engine/Source/Programs/AutomationTool/LowLevelTests/Tests/Gauntlet.LowLevelTestNode.cs
+++ b/Engine/Source/Programs/AutomationTool/LowLevelTests/Tests/Gauntlet.LowLevelTestNode.cs
@@ -1,6 +1,7 @@
 // Copyright Epic Games, Inc. All Rights Reserved.
 
 using AutomationTool;
+using EpicGames.Core;
 using Gauntlet;
 using System;
 using System.IO;
@@ -8,6 +9,8 @@ using System.Linq;
 using UnrealBuildBase;
 using UnrealBuildTool;
 
+using Log = Gauntlet.Log;
+
 namespace LowLevelTests
 {
 	public class LowLevelTests : BaseTest
@@ -68,6 +71,14 @@ namespace LowLevelTests
 			LowLevelTestResult = testResult;
 		}
 
+		public override void AddTestEvent(UnrealTestEvent InEvent)
+		{
+			if (InEvent.Summary.Equals("Insufficient devices found"))
+			{
+				Log.Error(KnownLogEvents.Gauntlet_TestEvent, "Test didn't run due to insufficient devices.");
+			}
+		}
+
 		public override bool StartTest(int Pass, int NumPasses)
 		{
 			if (LowLevelTestsApp == null)
