Add the ability to export a build snapshot from ZenServer for farm builds.

diff --git a/Engine/Build/Graph/Tasks/BuildAndTestProject.xml b/Engine/Build/Graph/Tasks/BuildAndTestProject.xml
--- a/Engine/Build/Graph/Tasks/BuildAndTestProject.xml
+++ b/Engine/Build/Graph/Tasks/BuildAndTestProject.xml
@@ -2,6 +2,9 @@
 <BuildGraph xmlns="http://www.epicgames.com/BuildGraph" xmlns:xsi="http://www.w3.org/2001/XMLSchema-instance" xsi:schemaLocation="http://www.epicgames.com/BuildGraph ../Schema.xsd" >
 	<EnvVar Name="UE_HORDE_JOBID"/>
 	
+	<EnvVar Name="UE-CloudPublishHost"/>
+	<EnvVar Name="UE-CloudPublishNamespace"/>
+	
 	<!-- Project Options -->
 	<Option Name="ProjectName" DefaultValue="None" Restrict=".+" Description="Name of Project"/>
 	<Option Name="ProjectPath" DefaultValue="None" Restrict=".+" Description="Path to the folder that contains your project"/>
@@ -54,6 +57,8 @@
 	<Option Name="NetworkReportRootOverride" DefaultValue="" Description="Network report root override"/>
 	<Option Name="SymbolServerPath" DefaultValue="" Description="Symbol server path"/>
 
+	<Option Name="PublishCloudSnapshot" DefaultValue="false" Restrict="true|false" Description="Publish build to cloud host" />
+
 	<Option Name="EnableSymbolStore" DefaultValue="$(IsBuildMachine)" Restrict="true|false" Description="Enables the upload of debugging symbols" />
 	<Property Name="SymbolServerPath" Value="" If="$(EnableSymbolStore) == 'false'"/>
 
@@ -517,9 +522,19 @@
 					<Spawn Exe="$(RootDir)/Engine/Restricted/NotForLicensees/Binaries/Win64/unsync.exe" Arguments="hash &quot;$(ProjectOutputDirectory)/$(StagedPlatformFolder)&quot; -v --buffered-files"/>
 				</Do>
 			</Node>
+
+			<!-- (Optional) Publish to Cloud -->
+			<Property Name="PlatformPublishStageDependencyName" Value="$(PlatformHashNodeName)" />
+			<Do If="$(PublishBuilds) and $(PublishCloudSnapshot) and '$(UE-CloudPublishHost)' != '' and '$(UE-CloudPublishNamespace)' != ''">
+				<Property Name="PlatformCloudPublishNodeName" Value="$(PreNodeName)Cloud Publish $(TargetPlatform)" />
+				<Node Name="$(PlatformCloudPublishNodeName)" Requires="$(PlatformHashNodeName)">
+					<ZenExportSnapshot Project="$(TargetProject)" Platform="$(CookPlatform)" DestinationStorageType="Cloud" DestinationCloudHost="$(UE-CloudPublishHost)" DestinationCloudNamespace="$(UE-CloudPublishNamespace)" DestinationCloudIdentifier="$(EscapedBranch)/$(BuildNamePath)" SnapshotDescriptorFile="$(NetworkOutputDirectory)/$(StagedPlatformFolder)/SnapshotDescriptor.json"/>
+				</Node>
+				<Property Name="PlatformPublishStageDependencyName" Value="$(PlatformCloudPublishNodeName)" />
+			</Do>
 			
 			<!-- Publish (Staged) -->
-			<Node Name="$(PlatformPublishStagedNodeName)" Requires="$(PlatformHashNodeName)">
+			<Node Name="$(PlatformPublishStagedNodeName)" Requires="$(PlatformPublishStageDependencyName)">
 				<Copy Files="..." From="$(ProjectOutputDirectory)/$(StagedPlatformFolder)" To="$(NetworkOutputDirectory)/$(StagedPlatformFolder)/Staged"/>
 				<Copy Files="..." From="$(ShaderSymbolsDirectory)/*$(TargetPlatform)*" To="$(NetworkOutputDirectory)/$(StagedPlatformFolder)/ShaderSymbols"/>
 			</Node>
diff --git a/Engine/Source/Programs/AutomationTool/BuildGraph/Tasks/ZenExportSnapshotTask.cs b/Engine/Source/Programs/AutomationTool/BuildGraph/Tasks/ZenExportSnapshotTask.cs
new file mode 100644
--- /dev/null
+++ b/Engine/Source/Programs/AutomationTool/BuildGraph/Tasks/ZenExportSnapshotTask.cs
@@ -0,0 +1,292 @@
+// Copyright Epic Games, Inc. All Rights Reserved.
+
+using AutomationTool;
+using EpicGames.BuildGraph;
+using System;
+using System.Collections.Generic;
+using System.IO;
+using System.Linq;
+using System.Text;
+using System.Threading.Tasks;
+using System.Xml;
+using EpicGames.Core;
+using EpicGames.Serialization;
+using UnrealBuildTool;
+using UnrealBuildBase;
+using Microsoft.Extensions.Logging;
+
+using static AutomationTool.CommandUtils;
+
+namespace AutomationTool.Tasks
+{
+	enum SnapshotStorageType
+	{
+		Invalid,
+		Cloud,
+		Zen,
+		File,
+	}
+
+	/// <summary>
+	/// Parameters for a task that exports an snapshot from ZenServer
+	/// </summary>
+	public class ZenExportSnapshotTaskParameters
+	{
+		/// <summary>
+		/// The project from which to export the snapshot
+		/// </summary>
+		[TaskParameter(Optional = true)]
+		public FileReference Project;
+
+		/// <summary>
+		/// The target platform(s) to export the snapshot for
+		/// </summary>
+		[TaskParameter(Optional = true)]
+		public string Platform;
+
+		/// <summary>
+		/// A file to create with information about the snapshot that was exported
+		/// </summary>
+		[TaskParameter(Optional = true)]
+		public FileReference SnapshotDescriptorFile;
+
+		/// <summary>
+		/// The type of destination to export the snapshot to (cloud, ...)
+		/// </summary>
+		[TaskParameter]
+		public string DestinationStorageType;
+
+		/// <summary>
+		/// The host name to use when exporting to a cloud destination
+		/// </summary>
+		[TaskParameter(Optional = true)]
+		public string DestinationCloudHost;
+
+		/// <summary>
+		/// The namespace to use when exporting to a cloud destination
+		/// </summary>
+		[TaskParameter(Optional = true)]
+		public string DestinationCloudNamespace;
+
+		/// <summary>
+		/// The identifier to use when exporting to a cloud destination
+		/// </summary>
+		[TaskParameter(Optional = true)]
+		public string DestinationCloudIdentifier;
+
+		/// <summary>
+		/// A custom bucket name to use when exporting to a cloud destination
+		/// </summary>
+		[TaskParameter(Optional = true)]
+		public string DestinationCloudBucket;
+	}
+
+	/// <summary>
+	/// Exports an snapshot from Zen to a specified destination.
+	/// </summary>
+	[TaskElement("ZenExportSnapshot", typeof(ZenExportSnapshotTaskParameters))]
+	public class ZenExportSnapshotTask : BgTaskImpl
+	{
+		private class ExportSourceData
+		{
+			public bool IsLocalHost;
+			public string HostName;
+			public int HostPort;
+			public string ProjectId;
+			public string OplogId;
+			public string TargetPlatform;
+		}
+
+		/// <summary>
+		/// Parameters for the task
+		/// </summary>
+		ZenExportSnapshotTaskParameters Parameters;
+
+		/// <summary>
+		/// Constructor.
+		/// </summary>
+		/// <param name="InParameters">Parameters for this task</param>
+		public ZenExportSnapshotTask(ZenExportSnapshotTaskParameters InParameters)
+		{
+			Parameters = InParameters;
+		}
+
+		/// <summary>
+		/// Execute the task.
+		/// </summary>
+		/// <param name="Job">Information about the current job</param>
+		/// <param name="BuildProducts">Set of build products produced by this node.</param>
+		/// <param name="TagNameToFileSet">Mapping from tag names to the set of files they include</param>
+		public override Task ExecuteAsync(JobContext Job, HashSet<FileReference> BuildProducts, Dictionary<string, HashSet<FileReference>> TagNameToFileSet)
+		{
+			SnapshotStorageType DestinationStorageType = SnapshotStorageType.Invalid;
+			if (!string.IsNullOrEmpty(Parameters.DestinationStorageType))
+			{
+				DestinationStorageType = (SnapshotStorageType)Enum.Parse(typeof(SnapshotStorageType), Parameters.DestinationStorageType);
+			}
+
+			FileReference ProjectFile = Parameters.Project;
+			if(!FileReference.Exists(ProjectFile))
+			{
+				throw new AutomationException("Missing project file - {0}", ProjectFile.FullName);
+			}
+
+			List<ExportSourceData> ExportSources = new List<ExportSourceData>();
+			foreach (string Platform in Parameters.Platform.Split('+'))
+			{
+				DirectoryReference PlatformCookedDirectory = DirectoryReference.Combine(ProjectFile.Directory, "Saved", "Cooked", Platform);
+				if (!DirectoryReference.Exists(PlatformCookedDirectory))
+				{
+					throw new AutomationException("Cook output directory not found ({0})", PlatformCookedDirectory.FullName);
+				}
+
+				FileReference ProjectStoreFile = FileReference.Combine(PlatformCookedDirectory, ".projectstore");
+				if (!FileReference.Exists(ProjectStoreFile))
+				{
+					throw new AutomationException("Project store file not found ({0})", ProjectStoreFile.FullName);
+				}
+
+				byte[] ProjectStoreData = File.ReadAllBytes(ProjectStoreFile.FullName);
+				CbObject ProjectStoreObject = new CbField(ProjectStoreData).AsObject();
+				CbObject ZenServerObject = ProjectStoreObject["zenserver"].AsObject();
+				if (ZenServerObject != CbObject.Empty)
+				{
+					ExportSourceData NewExportSource = new ExportSourceData();
+					NewExportSource.IsLocalHost = ZenServerObject["islocalhost"].AsBool();
+					NewExportSource.HostName = ZenServerObject["hostname"].AsString("localhost");
+					NewExportSource.HostPort = ZenServerObject["hostport"].AsInt16(1337);
+					NewExportSource.ProjectId = ZenServerObject["projectid"].AsString();
+					NewExportSource.OplogId = ZenServerObject["oplogid"].AsString();
+					NewExportSource.TargetPlatform = Platform;
+
+					ExportSources.Add(NewExportSource);
+				}
+			}
+
+			// Get the executable path
+			FileReference ZenExe;
+			if(HostPlatform.Current.HostEditorPlatform == UnrealTargetPlatform.Win64)
+			{
+				ZenExe = ResolveFile("Engine/Binaries/Win64/zen.exe");
+			}
+			else
+			{
+				ZenExe = ResolveFile(String.Format("Engine/Binaries/{0}/zen", HostPlatform.Current.HostEditorPlatform.ToString()));
+			}
+
+			// Format the command line
+			StringBuilder CommandLine = new StringBuilder();
+			CommandLine.AppendFormat("oplog-export");
+
+			switch (DestinationStorageType)
+			{
+				case SnapshotStorageType.Cloud:
+					if (string.IsNullOrEmpty(Parameters.DestinationCloudHost))
+					{
+						throw new AutomationException("Missing destination cloud host");
+					}
+					if (string.IsNullOrEmpty(Parameters.DestinationCloudNamespace))
+					{
+						throw new AutomationException("Missing destination cloud namespace");
+					}
+					if (string.IsNullOrEmpty(Parameters.DestinationCloudIdentifier))
+					{
+						throw new AutomationException("Missing destination cloud identifier");
+					}
+
+					string BucketName = Parameters.DestinationCloudBucket;
+					string ProjectNameAsBucketName = ProjectFile.GetFileNameWithoutAnyExtensions().ToLowerInvariant();
+					if (string.IsNullOrEmpty(BucketName))
+					{
+						BucketName = ProjectNameAsBucketName;
+					}
+
+					CommandLine.AppendFormat(" --cloud {0} --namespace {1} --bucket {2}", Parameters.DestinationCloudHost, Parameters.DestinationCloudNamespace, BucketName);
+
+					string[] ExportKeyIds = new string[ExportSources.Count];
+					int ExportIndex = 0;
+					foreach (ExportSourceData ExportSource in ExportSources)
+					{
+						StringBuilder SingleSourceCommandline = new StringBuilder(CommandLine.Length);
+						SingleSourceCommandline.Append(CommandLine);
+
+						StringBuilder DestinationKeyBuilder = new StringBuilder();
+						DestinationKeyBuilder.AppendFormat("{0}.{1}", ProjectNameAsBucketName,Parameters.DestinationCloudIdentifier, ExportSource.OplogId);
+						ExportKeyIds[ExportIndex] = DestinationKeyBuilder.ToString().ToLowerInvariant();
+						IoHash DestinationKeyHash = IoHash.Compute(Encoding.UTF8.GetBytes(ExportKeyIds[ExportIndex]));
+
+						SingleSourceCommandline.AppendFormat(" --hosturl http://{0}:{1} --key {2} {3} {4}", ExportSource.HostName, ExportSource.HostPort, DestinationKeyHash.ToString().ToLowerInvariant(), ExportSource.ProjectId, ExportSource.OplogId);
+						Logger.LogInformation("Running '{Arg0} {Arg1}'", CommandUtils.MakePathSafeToUseWithCommandLine(ZenExe.FullName), SingleSourceCommandline.ToString());
+						CommandUtils.RunAndLog(CommandUtils.CmdEnv, ZenExe.FullName, SingleSourceCommandline.ToString(), Options: CommandUtils.ERunOptions.Default);
+
+						ExportIndex = ExportIndex + 1;
+					}
+					
+					if (Parameters.SnapshotDescriptorFile != null)
+					{
+						// Write out a snapshot descriptor
+						using (JsonWriter Writer = new JsonWriter(Parameters.SnapshotDescriptorFile))
+						{
+							Writer.WriteObjectStart();
+
+							Writer.WriteArrayStart("snapshots");
+							
+							ExportIndex = 0;
+							foreach (ExportSourceData ExportSource in ExportSources)
+							{
+								Writer.WriteObjectStart();
+								
+								Writer.WriteValue("host", Parameters.DestinationCloudHost);
+								Writer.WriteValue("namespace", Parameters.DestinationCloudNamespace);
+								Writer.WriteValue("bucket", BucketName);
+								Writer.WriteValue("targetplatform", ExportSource.TargetPlatform);
+								IoHash DestinationKeyHash = IoHash.Compute(Encoding.UTF8.GetBytes(ExportKeyIds[ExportIndex]));
+								Writer.WriteValue("key", DestinationKeyHash.ToString().ToLowerInvariant());
+								Writer.WriteValue("desc", ExportKeyIds[ExportIndex]);
+
+								Writer.WriteObjectEnd();
+
+								ExportIndex = ExportIndex + 1;
+							}
+
+							Writer.WriteArrayEnd();
+							Writer.WriteObjectEnd();
+						}
+					}
+
+					break;
+				default:
+					throw new AutomationException("Unknown/invalid/unimplemented destination storage type - {0}", Parameters.DestinationStorageType);
+			}
+
+
+			return Task.CompletedTask;
+		}
+
+		/// <summary>
+		/// Output this task out to an XML writer.
+		/// </summary>
+		public override void Write(XmlWriter Writer)
+		{
+			Write(Writer, Parameters);
+		}
+
+		/// <summary>
+		/// Find all the tags which are used as inputs to this task
+		/// </summary>
+		/// <returns>The tag names which are read by this task</returns>
+		public override IEnumerable<string> FindConsumedTagNames()
+		{
+			yield break;
+		}
+
+		/// <summary>
+		/// Find all the tags which are modified by this task
+		/// </summary>
+		/// <returns>The tag names which are modified by this task</returns>
+		public override IEnumerable<string> FindProducedTagNames()
+		{
+			yield break;
+		}
+	}
+}
