Fix localhost use case when exporting zen snapshots.

diff --git a/Engine/Source/Programs/AutomationTool/BuildGraph/Tasks/ZenExportSnapshotTask.cs b/Engine/Source/Programs/AutomationTool/BuildGraph/Tasks/ZenExportSnapshotTask.cs
--- a/Engine/Source/Programs/AutomationTool/BuildGraph/Tasks/ZenExportSnapshotTask.cs
+++ b/Engine/Source/Programs/AutomationTool/BuildGraph/Tasks/ZenExportSnapshotTask.cs
@@ -215,7 +215,11 @@ namespace AutomationTool.Tasks
 						ExportKeyIds[ExportIndex] = DestinationKeyBuilder.ToString().ToLowerInvariant();
 						IoHash DestinationKeyHash = IoHash.Compute(Encoding.UTF8.GetBytes(ExportKeyIds[ExportIndex]));
 
-						SingleSourceCommandline.AppendFormat(" --hosturl http://{0}:{1} --key {2} {3} {4}", ExportSource.HostName, ExportSource.HostPort, DestinationKeyHash.ToString().ToLowerInvariant(), ExportSource.ProjectId, ExportSource.OplogId);
+						if (!ExportSource.IsLocalHost)
+						{
+							SingleSourceCommandline.AppendFormat(" --hosturl http://{0}:{1}", ExportSource.HostName, ExportSource.HostPort);
+						}
+						SingleSourceCommandline.AppendFormat(" --key {0} {1} {2}", DestinationKeyHash.ToString().ToLowerInvariant(), ExportSource.ProjectId, ExportSource.OplogId);
 						Logger.LogInformation("Running '{Arg0} {Arg1}'", CommandUtils.MakePathSafeToUseWithCommandLine(ZenExe.FullName), SingleSourceCommandline.ToString());
 						CommandUtils.RunAndLog(CommandUtils.CmdEnv, ZenExe.FullName, SingleSourceCommandline.ToString(), Options: CommandUtils.ERunOptions.Default);
 
