Fix prepare to debug

diff --git a/Engine/Source/Programs/AutomationTool/IOS/IOSPlatform.Automation.cs b/Engine/Source/Programs/AutomationTool/IOS/IOSPlatform.Automation.cs
--- a/Engine/Source/Programs/AutomationTool/IOS/IOSPlatform.Automation.cs
+++ b/Engine/Source/Programs/AutomationTool/IOS/IOSPlatform.Automation.cs
@@ -2355,9 +2355,14 @@ public class IOSPlatform : ApplePlatform
 
 	public override void PrepareForDebugging(string SourcePackage, string ProjectFilePath, string ClientPlatform)
 	{
+		Logger.LogInformation("Preparing for Debug ...");
+		Logger.LogInformation("SourcePackage : {SourcePackage}", SourcePackage);
+		Logger.LogInformation("ProjectFilePath : {SourcePackage}", ProjectFilePath);
+		Logger.LogInformation("ClientPlatform : {ClientPlatform}", ClientPlatform);
+
 		if (HostPlatform.Current.HostEditorPlatform == UnrealTargetPlatform.Win64 || HostPlatform.Current.HostEditorPlatform == UnrealTargetPlatform.Mac)
 		{
-    		int StartPos = ProjectFilePath.LastIndexOf("/");
+			int StartPos = ProjectFilePath.LastIndexOf(Path.DirectorySeparatorChar);
     		int StringLength = ProjectFilePath.Length - 10; // 9 for .uproject, 1 for the /
     		string PackageName = ProjectFilePath.Substring(StartPos + 1, StringLength - StartPos);
     		string PackagePath = Path.Combine(Path.GetDirectoryName(ProjectFilePath), "Binaries", ClientPlatform);
@@ -2366,6 +2371,7 @@ public class IOSPlatform : ApplePlatform
     			SourcePackage = Path.Combine(PackagePath, PackageName + ".ipa");
     		}
 			
+			Logger.LogInformation("SourcePackage : {SourcePackage}", SourcePackage);
 			string[] IPAFiles;
 			if (!File.Exists(SourcePackage))
             {
@@ -2389,7 +2395,8 @@ public class IOSPlatform : ApplePlatform
 			string ZipFile = Path.ChangeExtension(SourcePackage, "zip");
 
 			string PayloadPath = SourcePackage;
-			PayloadPath = PayloadPath.Substring(0, PayloadPath.LastIndexOf('\\'));
+			Logger.LogInformation("PayloadPath : {PayloadPath}", PayloadPath);
+			PayloadPath = PayloadPath.Substring(0, PayloadPath.LastIndexOf('/'));
 			string CookedDataDirectory = Path.Combine(Path.GetDirectoryName(PayloadPath), ClientPlatform, "Payload", PackageName + ".app", "cookeddata");
 
 			Logger.LogInformation("ClientPlatform : {ClientPlatform}", ClientPlatform);
@@ -2441,7 +2448,7 @@ public class IOSPlatform : ApplePlatform
 	public void UnzipPackage(string PackageToUnzip)
 	{
 		string UnzipPath = PackageToUnzip;
-		UnzipPath = UnzipPath.Substring(0, UnzipPath.LastIndexOf('\\'));
+		UnzipPath = UnzipPath.Substring(0, UnzipPath.LastIndexOf(Path.DirectorySeparatorChar));
 		Logger.LogInformation("Unzipping to {UnzipPath}", UnzipPath);
 
 		using (Ionic.Zip.ZipFile Zip = new Ionic.Zip.ZipFile(PackageToUnzip))
