Raise the exit reason as an error when a role is failing
+ Use the Main role for test summary context.

#rnx

diff --git a/Engine/Source/Programs/AutomationTool/Gauntlet/Unreal/Base/Gauntlet.UnrealTestNode.cs b/Engine/Source/Programs/AutomationTool/Gauntlet/Unreal/Base/Gauntlet.UnrealTestNode.cs
--- a/Engine/Source/Programs/AutomationTool/Gauntlet/Unreal/Base/Gauntlet.UnrealTestNode.cs
+++ b/Engine/Source/Programs/AutomationTool/Gauntlet/Unreal/Base/Gauntlet.UnrealTestNode.cs
@@ -423,10 +423,16 @@ namespace Gauntlet
 				return Name;
 			}
 
+			return string.Format("{0} ({1})", Name, GetMainRoleContextString());
+		}
+		public string GetMainRoleContextString()
+		{
+			if (Context == null)
+			{
+				return string.Empty;
+			}
 			var Config = GetConfiguration();
-			string RoleContext = (Config is UnrealTestConfiguration) ? Context.GetRoleContext(Config.GetMainRequiredRole().Type).ToString() : Context.ToString();
-
-			return string.Format("{0} ({1})", Name, RoleContext);
+			return (Config is UnrealTestConfiguration) ? Context.GetRoleContext(Config.GetMainRequiredRole().Type).ToString() : Context.ToString();
 		}
 
 		/// <summary>
@@ -1343,7 +1349,7 @@ namespace Gauntlet
 		{
 			if (string.IsNullOrEmpty(GetCachedConfiguration().HordeTestDataKey))
 			{
-				GetCachedConfiguration().HordeTestDataKey = Name + " " + Context.ToString();
+				GetCachedConfiguration().HordeTestDataKey = Name + " " + GetMainRoleContextString();
 			}
 
 			string TestName = Name;
@@ -1356,7 +1362,7 @@ namespace Gauntlet
 			HordeTestReport.TestName = TestName;
 			HordeTestReport.ReportCreatedOn = DateTime.Now.ToString();
 			HordeTestReport.TotalDurationSeconds = (float) (DateTime.Now - SessionStartTime).TotalSeconds;
-			HordeTestReport.Description = Context.ToString();
+			HordeTestReport.Description = GetMainRoleContextString();
 			HordeTestReport.URLLink = GetURLLink();
 			HordeTestReport.Errors.AddRange(GetErrorsAndAbnornalExits());
 			if (!string.IsNullOrEmpty(CancellationReason))
@@ -1413,7 +1419,7 @@ namespace Gauntlet
 				Log.Verbose("Reading json Unreal Automated test report from {Path}", JsonReportPath);
 				UnrealAutomatedTestPassResults JsonTestPassResults = UnrealAutomatedTestPassResults.LoadFromJson(JsonReportPath);
 				var MainRole = GetCachedConfiguration().GetMainRequiredRole();
-				string HordeTestDataKey = string.IsNullOrEmpty(GetCachedConfiguration().HordeTestDataKey) ? Name + " " + Context.ToString() : GetCachedConfiguration().HordeTestDataKey;
+				string HordeTestDataKey = string.IsNullOrEmpty(GetCachedConfiguration().HordeTestDataKey) ? Name + " " + GetMainRoleContextString() : GetCachedConfiguration().HordeTestDataKey;
 				GetCachedConfiguration().HordeTestDataKey = HordeTestDataKey;
 				// Convert test results for Horde
 				HordeReport.AutomatedTestSessionData HordeTestPassResults = HordeReport.AutomatedTestSessionData.FromUnrealAutomatedTests(
@@ -1898,14 +1904,7 @@ namespace Gauntlet
 			bool HasFailed = InRoleResult.ExitCode != 0 && InRoleResult.LogSummary.HasAbnormalExit;
 			string RoleState = HasFailed ? "failed:" : "completed:";
 			string StatusMessage = string.Format(" #### {0} {1} {2} ({3}, ExitCode={4})", RoleArtifacts.SessionRole.RoleType, RoleState, InRoleResult.Summary, InRoleResult.ProcessResult, InRoleResult.ExitCode);
-			if(HasFailed)
-			{
-				Log.Error(KnownLogEvents.Gauntlet_TestEvent, StatusMessage);
-			}
-			else
-			{
-				Log.Info(StatusMessage);
-			}
+			Log.Info(StatusMessage);
 
 			// log command line up here for visibility
 
@@ -2261,8 +2260,15 @@ namespace Gauntlet
 			foreach (var RoleResult in SortedRoles)
 			{
 				string RoleName = RoleResult.Artifacts.SessionRole.RoleType.ToString();
-
-				Log.Info(string.Format(" {0} Role: {1} ({2}, ExitCode {3})", RoleName, RoleResult.Summary, RoleResult.ProcessResult, RoleResult.ExitCode));
+				string LogMessage = string.Format(" {0} Role: {1} ({2}, ExitCode={3})", RoleName, RoleResult.Summary, RoleResult.ProcessResult, RoleResult.ExitCode);
+				if (RoleResult.ExitCode != 0)
+				{
+					Log.Error(KnownLogEvents.Gauntlet_TestEvent, LogMessage);
+				}
+				else
+				{
+					Log.Info(LogMessage);
+				}
 
 				FatalErrors += RoleResult.LogSummary.FatalError != null ? 1 : 0;
 				Ensures += RoleResult.LogSummary.Ensures.Count();
@@ -2272,7 +2278,7 @@ namespace Gauntlet
 
 			Log.Info("");
 			Log.Info(string.Join("\n", new string[] {
-				string.Format(" * Context: {0}", Context.ToString()),
+				string.Format(" * Main Context: {0}", GetMainRoleContextString()),
 				FatalErrors > 0 ? string.Format(" * FatalErrors: {0}", FatalErrors) : null,
 				Ensures > 0 ? string.Format(" * Ensures: {0}", Ensures) : null,
 				Errors > 0 ? string.Format(" * Log Errors: {0}", Errors) : null,
