horde: fix json blob for report name

diff --git a/Engine/Source/Programs/AutomationTool/BuildGraph/Tasks/HordeCreateReportTask.cs b/Engine/Source/Programs/AutomationTool/BuildGraph/Tasks/HordeCreateReportTask.cs
--- a/Engine/Source/Programs/AutomationTool/BuildGraph/Tasks/HordeCreateReportTask.cs
+++ b/Engine/Source/Programs/AutomationTool/BuildGraph/Tasks/HordeCreateReportTask.cs
@@ -83,7 +83,7 @@ namespace AutomationTool.Tasks
 					Writer.WriteString("scope", Parameters.Scope);
 					Writer.WriteString("name", Parameters.Name);
 					Writer.WriteString("placement", Parameters.Placement);
-					Writer.WriteString("file", ReportTextFile.GetFileName());
+					Writer.WriteString("fileName", ReportTextFile.GetFileName());
 					Writer.WriteEndObject();
 				}
 			}
