Pass RunCommandline during launch to improve iteration times on mobile by not needing to rebuild .apk every time, because they are not passed via staged build anymore.

diff --git a/Engine/Source/Programs/AutomationTool/Android/AndroidPlatform.Automation.cs b/Engine/Source/Programs/AutomationTool/Android/AndroidPlatform.Automation.cs
--- a/Engine/Source/Programs/AutomationTool/Android/AndroidPlatform.Automation.cs
+++ b/Engine/Source/Programs/AutomationTool/Android/AndroidPlatform.Automation.cs
@@ -4021,7 +4021,11 @@ public class AndroidPlatform : Platform
 				throw new AutomationException(ExitCode.Error_FailureGettingPackageInfo, "Failed to get package name from " + ClientApp);
 			}
 
-			// push ClientCmdLine args as a file to the device to override the stage/apk command line
+			var canReadClientCmdLineViaAmStart = Params.ClientConfigsToBuild.Count > 0 &&
+			                                     Params.ClientConfigsToBuild.First() != UnrealTargetConfiguration.Shipping;
+
+			// push ClientCmdLine args as a file to the device to override the stage/apk command line if we can't push it via am start
+			if (!canReadClientCmdLineViaAmStart)
 			{
 				bool bAFSEnablePlugin;
 				string AFSToken;
@@ -4045,7 +4049,12 @@ public class AndroidPlatform : Platform
 			RunAdbCommand(DeviceName, "logcat -c");
 
 			// start the app on device!
-			string CommandLine = "shell am start -n " + PackageName + "/" + GetLaunchableActivityName();
+			var CommandLine = "shell am start -n " + PackageName + "/" + GetLaunchableActivityName();
+			if (canReadClientCmdLineViaAmStart)
+			{
+				var ClientSessionCmdLineEscaped = ClientCmdLine.Replace(" ", "\\ ").Replace("\"", "\\\\\\\"");
+				CommandLine += " --es cmdline \"" + ClientSessionCmdLineEscaped + "\"";
+			}
 			RunAdbCommand(DeviceName, CommandLine);
 
 			// wait before getting the process list with "adb shell ps" from AdbCreatedProcess
diff --git a/Engine/Source/Programs/AutomationTool/AutomationUtils/ProjectParams.cs b/Engine/Source/Programs/AutomationTool/AutomationUtils/ProjectParams.cs
--- a/Engine/Source/Programs/AutomationTool/AutomationUtils/ProjectParams.cs
+++ b/Engine/Source/Programs/AutomationTool/AutomationUtils/ProjectParams.cs
@@ -2157,7 +2157,7 @@ namespace AutomationTool
 		/// <summary>
 		/// Run: Additional command line arguments to pass to the program
 		/// </summary>
-		[Help("addcmdline", "Additional command line arguments for the program")]
+		[Help("addcmdline", "Additional command line arguments for the program, which will not be staged in UECommandLine.txt in most cases")]
 		public string RunCommandline;
 
         /// <summary>
diff --git a/Engine/Source/Programs/AutomationTool/IOS/IOSPlatform.Automation.cs b/Engine/Source/Programs/AutomationTool/IOS/IOSPlatform.Automation.cs
--- a/Engine/Source/Programs/AutomationTool/IOS/IOSPlatform.Automation.cs
+++ b/Engine/Source/Programs/AutomationTool/IOS/IOSPlatform.Automation.cs
@@ -2028,6 +2028,7 @@ public class IOSPlatform : ApplePlatform
 			Arguments += " --detach";
 			Arguments = GetLibimobileDeviceNetworkedArgument(Arguments, Params.DeviceNames[0]);
 			Arguments += " run '" + BundleIdentifier + "'";
+			Arguments += " " + ClientCmdLine;
 
 			IProcessResult ClientProcess = Run(Program, Arguments, null, ClientRunFlags);
 			if (ClientProcess.ExitCode == -1)
diff --git a/Engine/Source/Programs/AutomationTool/Scripts/CopyBuildToStagingDirectory.Automation.cs b/Engine/Source/Programs/AutomationTool/Scripts/CopyBuildToStagingDirectory.Automation.cs
--- a/Engine/Source/Programs/AutomationTool/Scripts/CopyBuildToStagingDirectory.Automation.cs
+++ b/Engine/Source/Programs/AutomationTool/Scripts/CopyBuildToStagingDirectory.Automation.cs
@@ -4667,11 +4667,6 @@ namespace AutomationScripts
 				CommandLineBuilder.AppendFormat("{0} ", Params.StageCommandline.Trim(new char[] { '\"' }));
 			}
 
-			if (!string.IsNullOrEmpty(Params.RunCommandline))
-			{
-				CommandLineBuilder.AppendFormat("{0} ", Params.RunCommandline.Trim(new char[] { '\"' }));
-			}
-
 			String FileHostCommandline = GetFileHostCommandline(Params, SC);
 			if (!string.IsNullOrEmpty(FileHostCommandline))
 			{
diff --git a/Engine/Source/Runtime/Launch/Private/Android/LaunchAndroid.cpp b/Engine/Source/Runtime/Launch/Private/Android/LaunchAndroid.cpp
--- a/Engine/Source/Runtime/Launch/Private/Android/LaunchAndroid.cpp
+++ b/Engine/Source/Runtime/Launch/Private/Android/LaunchAndroid.cpp
@@ -495,7 +495,7 @@ int32 AndroidMain(struct android_app* state)
 
 	// read the command line file
 	InitCommandLine();
-	FPlatformMisc::LowLevelOutputDebugStringf(TEXT("Final commandline: %s\n"), FCommandLine::Get());
+	FPlatformMisc::LowLevelOutputDebugStringf(TEXT("Final commandline: %s (len %i)\n"), FCommandLine::Get(), FCString::Strlen(FCommandLine::Get()));
 
 #if !(UE_BUILD_SHIPPING)
 	// If "-waitforattach" or "-WaitForDebugger" was specified, halt startup and wait for a debugger to attach before continuing
