BuildGraph: Allow scheduling tasks to run on lease cleanup, when executing on Horde.

diff --git a/Engine/Source/Programs/AutomationTool/BuildGraph/Task.cs b/Engine/Source/Programs/AutomationTool/BuildGraph/Task.cs
--- a/Engine/Source/Programs/AutomationTool/BuildGraph/Task.cs
+++ b/Engine/Source/Programs/AutomationTool/BuildGraph/Task.cs
@@ -439,13 +439,19 @@ namespace AutomationTool
 		/// </summary>
 		public const string CleanupScriptEnvVarName = "UE_HORDE_CLEANUP";
 
+		/// <summary>
+		/// Name of the environment variable containing lease cleanup commands
+		/// </summary>
+		public const string LeaseCleanupScriptEnvVarName = "UE_HORDE_LEASE_CLEANUP";
+
 		/// <summary>
 		/// Add cleanup commands to run after the step completes
 		/// </summary>
 		/// <param name="NewLines">Lines to add to the cleanup script</param>
-		public static async Task AddCleanupCommandsAsync(IEnumerable<string> NewLines)
+		/// <param name="Lease">Whether to add the commands to run on lease termination</param>
+		public static async Task AddCleanupCommandsAsync(IEnumerable<string> NewLines, bool Lease = false)
 		{
-			string CleanupScriptEnvVar = Environment.GetEnvironmentVariable(CleanupScriptEnvVarName);
+			string CleanupScriptEnvVar = Environment.GetEnvironmentVariable(Lease? LeaseCleanupScriptEnvVarName : CleanupScriptEnvVarName);
 			if (!String.IsNullOrEmpty(CleanupScriptEnvVar))
 			{
 				FileReference CleanupScript = new FileReference(CleanupScriptEnvVar);
diff --git a/Engine/Source/Programs/AutomationTool/BuildGraph/Tasks/OnExitTask.cs b/Engine/Source/Programs/AutomationTool/BuildGraph/Tasks/OnExitTask.cs
--- a/Engine/Source/Programs/AutomationTool/BuildGraph/Tasks/OnExitTask.cs
+++ b/Engine/Source/Programs/AutomationTool/BuildGraph/Tasks/OnExitTask.cs
@@ -20,6 +20,12 @@ namespace AutomationTool.Tasks
 		/// </summary>
 		[TaskParameter]
 		public string Command = String.Empty;
+
+		/// <summary>
+		/// Whether to execute on lease termination
+		/// </summary>
+		[TaskParameter(Optional = true)]
+		public bool Lease = false;
 	}
 
 	/// <summary>
@@ -46,7 +52,7 @@ namespace AutomationTool.Tasks
 		public override async Task ExecuteAsync(JobContext Job, HashSet<FileReference> BuildProducts, Dictionary<string, HashSet<FileReference>> TagNameToFileSet)
 		{
 			string[] commands = Parameters.Command.Split('\n').Select(x => x.Trim()).ToArray();
-			await AddCleanupCommandsAsync(commands);
+			await AddCleanupCommandsAsync(commands, Parameters.Lease);
 		}
 
 		/// <inheritdoc/>
