Adding buildgraph option to EnablePlugins for cooking and building. Enabling a plugin in the buildgraph will enable it for cook and compilation but will not enable it in the application that must be done via other means
Updating BuildCookRun commands so they pass the EnablePlugins command through the build system

diff --git a/Engine/Plugins/Experimental/GameFeatures/Source/GameFeatures/Public/GameFeaturesSubsystemSettings.h b/Engine/Plugins/Experimental/GameFeatures/Source/GameFeatures/Public/GameFeaturesSubsystemSettings.h
--- a/Engine/Plugins/Experimental/GameFeatures/Source/GameFeatures/Public/GameFeaturesSubsystemSettings.h
+++ b/Engine/Plugins/Experimental/GameFeatures/Source/GameFeatures/Public/GameFeaturesSubsystemSettings.h
@@ -24,6 +24,10 @@ public:
 	UPROPERTY(config, EditAnywhere, Category=DefaultClasses, meta=(MetaClass="/Script/GameFeatures.GameFeaturesProjectPolicies", DisplayName="Game Feature Project Policy Class", ConfigRestartRequired=true))
 	FSoftClassPath GameFeaturesManagerClassName;
 
+	/** List of plugins that are forcibly enabled (e.g., via a hotfix) */
+	UPROPERTY(config, EditAnywhere, Category = GameFeatures)
+	TArray<FString> EnabledPlugins;
+
 	/** List of plugins that are forcibly disabled (e.g., via a hotfix) */
 	UPROPERTY(config, EditAnywhere, Category=GameFeatures)
 	TArray<FString> DisabledPlugins;
diff --git a/Engine/Source/Programs/AutomationTool/AutomationUtils/ProjectParams.cs b/Engine/Source/Programs/AutomationTool/AutomationUtils/ProjectParams.cs
--- a/Engine/Source/Programs/AutomationTool/AutomationUtils/ProjectParams.cs
+++ b/Engine/Source/Programs/AutomationTool/AutomationUtils/ProjectParams.cs
@@ -313,6 +313,7 @@ namespace AutomationTool
 			this.DLCOverrideCookedSubDir = InParams.DLCOverrideCookedSubDir;
 			this.DLCOverrideStagedSubDir = InParams.DLCOverrideStagedSubDir;
 			this.DiffCookedContentPath = InParams.DiffCookedContentPath;
+            this.AdditionalBuildOptions = InParams.AdditionalBuildOptions;
             this.AdditionalCookerOptions = InParams.AdditionalCookerOptions;
 			this.ClientCookedTargets = InParams.ClientCookedTargets;
 			this.ServerCookedTargets = InParams.ServerCookedTargets;
@@ -1857,10 +1858,15 @@ namespace AutomationTool
 		/// </summary>
 		public string DiffCookedContentPath;
 
-        /// <summary>
-        /// Cook: Additional cooker options to include on the cooker commandline
-        /// </summary>
-        public string AdditionalCookerOptions;
+		/// <summary>
+		/// Build: Additional build options to include on the build commandline
+		/// </summary>
+		public string AdditionalBuildOptions;
+
+		/// <summary>
+		/// Cook: Additional cooker options to include on the cooker commandline
+		/// </summary>
+		public string AdditionalCookerOptions;
 
         /// <summary>
         /// Cook: List of cultures to cook.
diff --git a/Engine/Source/Programs/AutomationTool/Scripts/BuildCookRun.Automation.cs b/Engine/Source/Programs/AutomationTool/Scripts/BuildCookRun.Automation.cs
--- a/Engine/Source/Programs/AutomationTool/Scripts/BuildCookRun.Automation.cs
+++ b/Engine/Source/Programs/AutomationTool/Scripts/BuildCookRun.Automation.cs
@@ -108,6 +108,15 @@ public class BuildCookRun : BuildCommand, IProjectParamsHelpers
 			}
 		}
 
+		List<string> PluginsToEnable = new List<string>();
+		string EnablePlugins = ParseParamValue("EnablePlugins", null);
+		if (!string.IsNullOrEmpty(EnablePlugins))
+		{
+			PluginsToEnable.AddRange(EnablePlugins.Split(new[] { ',', '+' }, StringSplitOptions.RemoveEmptyEntries));
+		}
+		Params.AdditionalCookerOptions += " -EnablePlugins=\"" + string.Join(",", PluginsToEnable) + "\"";
+		Params.AdditionalBuildOptions += " -EnablePlugins=\"" + string.Join(",", PluginsToEnable) + "\"";
+
 		Params.ValidateAndLog();
 		return Params;
 	}
diff --git a/Engine/Source/Programs/AutomationTool/Scripts/BuildProjectCommand.Automation.cs b/Engine/Source/Programs/AutomationTool/Scripts/BuildProjectCommand.Automation.cs
--- a/Engine/Source/Programs/AutomationTool/Scripts/BuildProjectCommand.Automation.cs
+++ b/Engine/Source/Programs/AutomationTool/Scripts/BuildProjectCommand.Automation.cs
@@ -79,10 +79,10 @@ namespace AutomationScripts
 
 
 			Func<UnrealArchitectures, string> GetArchString = Arch => (Arch == null) ? "" : $"-architecture={Arch}";
-			string ServerArchitecture = GetArchString(Params.ServerArchitecture);
-			string EditorArchitecture = GetArchString(Params.EditorArchitecture);
-			string ClientArchitecture = GetArchString(Params.ClientArchitecture);
-			string ProgramArchitecture = GetArchString(Params.ProgramArchitecture);
+			string ServerBuildArgs = Params.AdditionalBuildOptions + GetArchString(Params.ServerArchitecture);
+			string EditorBuildArgs = Params.AdditionalBuildOptions + GetArchString(Params.EditorArchitecture);
+			string ClientBuildArgs = Params.AdditionalBuildOptions + GetArchString(Params.ClientArchitecture);
+			string ProgramBuildArgs = Params.AdditionalBuildOptions + GetArchString(Params.ProgramArchitecture);
 
 			// Setup editor targets
 			if (Params.HasEditorTargets && (!Params.SkipBuildEditor) && (TargetMask & ProjectBuildTargets.Editor) == ProjectBuildTargets.Editor)
@@ -91,18 +91,18 @@ namespace AutomationScripts
 				UnrealTargetPlatform EditorPlatform = HostPlatform.Current.HostEditorPlatform;
 				const UnrealTargetConfiguration EditorConfiguration = UnrealTargetConfiguration.Development;
 
-				Agenda.AddTargets(Params.EditorTargets.ToArray(), EditorPlatform, EditorConfiguration, Params.CodeBasedUprojectPath, InAddArgs: EditorArchitecture);
+				Agenda.AddTargets(Params.EditorTargets.ToArray(), EditorPlatform, EditorConfiguration, Params.CodeBasedUprojectPath, InAddArgs: EditorBuildArgs);
 
 				if (!Unreal.IsEngineInstalled())
 				{
 					CrashReportPlatforms.Add(EditorPlatform);
 					if (Params.EditorTargets.Contains("ShaderCompileWorker") == false)
 					{
-						Agenda.AddTargets(new string[] { "ShaderCompileWorker" }, EditorPlatform, EditorConfiguration, InAddArgs: ProgramArchitecture);
+						Agenda.AddTargets(new string[] { "ShaderCompileWorker" }, EditorPlatform, EditorConfiguration, InAddArgs: ProgramBuildArgs);
 					}
 					if (Params.FileServer && Params.EditorTargets.Contains("UnrealFileServer") == false)
 					{
-						Agenda.AddTargets(new string[] { "UnrealFileServer" }, EditorPlatform, EditorConfiguration, InAddArgs: ProgramArchitecture);
+						Agenda.AddTargets(new string[] { "UnrealFileServer" }, EditorPlatform, EditorConfiguration, InAddArgs: ProgramBuildArgs);
 					}
 				}
 			}
@@ -112,7 +112,7 @@ namespace AutomationScripts
 			{
 				if (!Params.HasEditorTargets || Params.EditorTargets.Contains("UnrealPak") == false)
 				{
-					Agenda.AddTargets(new string[] { "UnrealPak" }, HostPlatform.Current.HostEditorPlatform, UnrealTargetConfiguration.Development, Params.CodeBasedUprojectPath, InAddArgs: ProgramArchitecture);
+					Agenda.AddTargets(new string[] { "UnrealPak" }, HostPlatform.Current.HostEditorPlatform, UnrealTargetConfiguration.Development, Params.CodeBasedUprojectPath, InAddArgs: ProgramBuildArgs);
 				}
 			}
 
@@ -160,7 +160,7 @@ namespace AutomationScripts
 					{
 						UnrealTargetPlatform CrashReportPlatform = Platform.GetPlatform(ClientPlatformType).CrashReportPlatform ?? ClientPlatformType;
 						CrashReportPlatforms.Add(CrashReportPlatform);
-						string Arch = Params.IsProgramTarget ? ProgramArchitecture : ClientArchitecture;
+						string Arch = Params.IsProgramTarget ? ProgramBuildArgs : ClientBuildArgs;
 						Agenda.AddTargets(Params.ClientCookedTargets.ToArray(), ClientPlatformType, BuildConfig, Params.CodeBasedUprojectPath, 
 							InAddArgs: $" -remoteini=\"{Params.RawProjectPath.Directory}\" {AdditionalArgs} {Arch}");
 					}
@@ -177,7 +177,7 @@ namespace AutomationScripts
 						UnrealTargetPlatform CrashReportPlatform = Platform.GetPlatform(ServerPlatformType).CrashReportPlatform ?? ServerPlatformType;
 						CrashReportPlatforms.Add(CrashReportPlatform);
 						Agenda.AddTargets(Params.ServerCookedTargets.ToArray(), ServerPlatformType, BuildConfig, Params.CodeBasedUprojectPath, 
-							InAddArgs: $" -remoteini=\"{Params.RawProjectPath.Directory}\" {AdditionalArgs} {ServerArchitecture}");
+							InAddArgs: $" -remoteini=\"{Params.RawProjectPath.Directory}\" {AdditionalArgs} {ServerBuildArgs}");
 					}
 				}
 			}
@@ -188,7 +188,7 @@ namespace AutomationScripts
 				{
 					if (Params.ClientTargetPlatforms.Contains(new TargetPlatformDescriptor(BootstrapPackagedGamePlatformType)))
 					{
-						Agenda.AddTarget("BootstrapPackagedGame", BootstrapPackagedGamePlatformType, UnrealBuildTool.UnrealTargetConfiguration.Shipping, InAddArgs: ClientArchitecture);
+						Agenda.AddTarget("BootstrapPackagedGame", BootstrapPackagedGamePlatformType, UnrealBuildTool.UnrealTargetConfiguration.Shipping, InAddArgs: ClientBuildArgs);
 					}
 				}
 			}
@@ -199,7 +199,7 @@ namespace AutomationScripts
 					if (PlatformSupportsCrashReporter(CrashReportPlatform))
 					{
 						Agenda.AddTarget("CrashReportClient", CrashReportPlatform, UnrealTargetConfiguration.Shipping, 
-							InAddArgs: $" -remoteini=\"{Params.RawProjectPath.Directory}\" {ProgramArchitecture}");
+							InAddArgs: $" -remoteini=\"{Params.RawProjectPath.Directory}\" {ProgramBuildArgs}");
 					}
 				}
 			}
@@ -211,7 +211,7 @@ namespace AutomationScripts
 				{
 					foreach (var ClientPlatformType in UniquePlatformTypes)
 					{
-						Agenda.AddTargets(Params.ProgramTargets.ToArray(), ClientPlatformType, BuildConfig, Params.CodeBasedUprojectPath, ProgramArchitecture);
+						Agenda.AddTargets(Params.ProgramTargets.ToArray(), ClientPlatformType, BuildConfig, Params.CodeBasedUprojectPath, ProgramBuildArgs);
 					}
 				}
 			}
