Gauntlet - UE.Networking test is not bubbling up error from gauntlet controller

diff --git a/Engine/Source/Programs/AutomationTool/Gauntlet/Unreal/Automation/UE.Networking.cs b/Engine/Source/Programs/AutomationTool/Gauntlet/Unreal/Automation/UE.Networking.cs
--- a/Engine/Source/Programs/AutomationTool/Gauntlet/Unreal/Automation/UE.Networking.cs
+++ b/Engine/Source/Programs/AutomationTool/Gauntlet/Unreal/Automation/UE.Networking.cs
@@ -6,12 +6,15 @@ using Gauntlet;
 using EpicGame;
 using AutomationTool;
 using UnrealBuildTool;
+using EpicGames.Core;
+using Log = Gauntlet.Log;
+
 
 namespace UE
 {
 	namespace NetworkAutomation
 	{
-		public class DedicatedServer : UnrealTestNode<EpicGameTestConfig>
+		public class DedicatedServer : NetworkingNodeBase<EpicGameTestConfig>
 		{
 			public DedicatedServer(UnrealTestContext InContext) : base(InContext)
 			{
@@ -78,7 +81,7 @@ namespace UE
 			}
 		}
 
-		public class ListenServer : UnrealTestNode<ListenServerConfig>
+		public class ListenServer : NetworkingNodeBase<ListenServerConfig>
 		{
 			public ListenServer(UnrealTestContext InContext) : base(InContext)
 			{
@@ -99,4 +102,50 @@ namespace UE
 			}
 		}
 	}
-}
\ No newline at end of file
+
+	/// <summary>
+	/// Runs automated tests on a platform
+	/// </summary>
+	public abstract class NetworkingNodeBase<TConfigClass> : UnrealTestNode<TConfigClass>
+	where TConfigClass : UnrealGame.UnrealTestConfig, new()
+	{
+		private DateTime LastAutomationEntryTime = DateTime.MinValue;
+
+		public NetworkingNodeBase(Gauntlet.UnrealTestContext InContext)
+			: base(InContext)
+		{
+			LogWarningsAndErrorsAfterSummary = false;
+		}
+
+		/// <summary>
+		/// Override the summary report
+		/// </summary>
+		/// <returns></returns>
+		protected override void LogTestSummaryHeader()
+		{
+			base.LogTestSummaryHeader();
+
+			var FailingRoles = RoleResults.Where(R => R.ExitCode != 0);
+
+			if (FailingRoles.Any())
+			{
+				foreach (var Role in FailingRoles)
+				{
+					UnrealLogParser Parser = new UnrealLogParser(Role.LogSummary.FullLogContent);
+
+					var LogErrors = Role.LogSummary.Errors.Where(E => E.Level == UnrealLog.LogLevel.Error);
+
+					if (LogErrors.Any())
+					{
+						Log.Info(" * The following errors did appear in {RoleContext} log:", Role.Artifacts.SessionRole.ToString());
+
+						foreach (var LogError in LogErrors.Distinct())
+						{
+							Log.Error(KnownLogEvents.Gauntlet_UnrealEngineTestEvent, "    - " + LogError.Message);
+						}
+					}
+				}
+			}
+		}
+	}
+}
