Don't pass -EnablePlugins= when there are no plugins enabled to avoid issues with plugin manager argument parsing

diff --git a/Engine/Source/Programs/AutomationTool/Scripts/BuildCookRun.Automation.cs b/Engine/Source/Programs/AutomationTool/Scripts/BuildCookRun.Automation.cs
--- a/Engine/Source/Programs/AutomationTool/Scripts/BuildCookRun.Automation.cs
+++ b/Engine/Source/Programs/AutomationTool/Scripts/BuildCookRun.Automation.cs
@@ -114,8 +114,12 @@ public class BuildCookRun : BuildCommand, IProjectParamsHelpers
 		{
 			PluginsToEnable.AddRange(EnablePlugins.Split(new[] { ',', '+' }, StringSplitOptions.RemoveEmptyEntries));
 		}
-		Params.AdditionalCookerOptions += " -EnablePlugins=\"" + string.Join(",", PluginsToEnable) + "\"";
-		Params.AdditionalBuildOptions += " -EnablePlugins=\"" + string.Join(",", PluginsToEnable) + "\"";
+
+		if (PluginsToEnable.Count > 0)
+		{
+			Params.AdditionalCookerOptions += " -EnablePlugins=\"" + string.Join(",", PluginsToEnable) + "\"";
+			Params.AdditionalBuildOptions += " -EnablePlugins=\"" + string.Join(",", PluginsToEnable) + "\"";
+		}
 
 		Params.ValidateAndLog();
 		return Params;
