Build fix.

diff --git a/Engine/Source/Programs/AutomationTool/BuildGraph/Tasks/DeployToolTask.cs b/Engine/Source/Programs/AutomationTool/BuildGraph/Tasks/DeployToolTask.cs
--- a/Engine/Source/Programs/AutomationTool/BuildGraph/Tasks/DeployToolTask.cs
+++ b/Engine/Source/Programs/AutomationTool/BuildGraph/Tasks/DeployToolTask.cs
@@ -181,7 +181,7 @@ namespace AutomationTool.Tasks
 			{
 				request.CreatePaused = true;
 			}
-			request.Node = new HashedNodeLocator(handle.Hash, handle.GetLocator()).ToString();
+			request.Node = handle.GetLocator().ToString();
 
 			using (HttpClient httpClient = CreateHttpClient())
 			{
