Support for arbitrary number of OBB overflow files
Set current maximum to 10 in ini file

diff --git a/Engine/Source/Programs/AutomationTool/Android/AndroidPlatform.Automation.cs b/Engine/Source/Programs/AutomationTool/Android/AndroidPlatform.Automation.cs
--- a/Engine/Source/Programs/AutomationTool/Android/AndroidPlatform.Automation.cs
+++ b/Engine/Source/Programs/AutomationTool/Android/AndroidPlatform.Automation.cs
@@ -1037,6 +1037,17 @@ public class AndroidPlatform : Platform
 		}
 		return true;
 	}
+	private class OverflowFileInfo
+	{
+		public OverflowFileInfo(List<FileReference> InFilesForOverflow, Int64 InOverflowObbSize)
+		{
+			FilesForOverflow = InFilesForOverflow;
+			OverflowObbSize = InOverflowObbSize;
+		}
+
+		public List<FileReference> FilesForOverflow { get; set; }
+		public Int64 OverflowObbSize { get; set; }
+	}
 
 	public override void Package(ProjectParams Params, DeploymentContext SC, int WorkingCL)
 	{
@@ -1085,10 +1096,7 @@ public class AndroidPlatform : Platform
 
 		string LocalObbName = SC.StageDirectory.FullName+".obb";
 		string LocalPatchName = SC.StageDirectory.FullName + ".patch.obb";
-		string LocalOverflow1Name = SC.StageDirectory.FullName + ".overflow1.obb";
-		string LocalOverflow2Name = SC.StageDirectory.FullName + ".overflow2.obb";
-		string LocalOverflow3Name = SC.StageDirectory.FullName + ".overflow3.obb";
-		string LocalOverflow4Name = SC.StageDirectory.FullName + ".overflow4.obb";
+		string LocalOverflowNameTemplate = SC.StageDirectory.FullName + ".overflow{0}.obb";
 
 		FileFilter ObbFileFilter = new FileFilter(FileFilterType.Include);
 		ConfigHierarchy EngineIni = ConfigCache.ReadHierarchy(ConfigHierarchyType.Engine, DirectoryReference.FromFile(Params.RawProjectPath), UnrealTargetPlatform.Android);
@@ -1153,36 +1161,18 @@ public class AndroidPlatform : Platform
 				File.Delete(LocalPatchName);
 			}
 
-			// Always delete the target overflow1 OBB file if it exists
-			if (File.Exists(LocalOverflow1Name))
-			{
-				File.Delete(LocalOverflow1Name);
-			}
-
-			// Always delete the target overflow2 OBB file if it exists
-			if (File.Exists(LocalOverflow2Name))
+			// Always delete all target overflow OBB files if they exists
+			int OverflowIndex = 1;
+			while (File.Exists(string.Format(LocalOverflowNameTemplate, OverflowIndex)))
 			{
-				File.Delete(LocalOverflow2Name);
-			}
-
-			// Always delete the target overflow3 OBB file if it exists
-			if (File.Exists(LocalOverflow3Name))
-			{
-				File.Delete(LocalOverflow3Name);
-			}
-
-			// Always delete the target overflow4 OBB file if it exists
-			if (File.Exists(LocalOverflow4Name))
-			{
-				File.Delete(LocalOverflow4Name);
+				File.Delete(string.Format(LocalOverflowNameTemplate, OverflowIndex));
+				// move to next overflow
+				OverflowIndex++;
 			}
 
 			List<FileReference> FilesToObb = FilesForObb;
 			List<FileReference> FilesToPatch = new List<FileReference>();
-			List<FileReference> FilesToOverflow1 = new List<FileReference>();
-			List<FileReference> FilesToOverflow2 = new List<FileReference>();
-			List<FileReference> FilesToOverflow3 = new List<FileReference>();
-			List<FileReference> FilesToOverflow4 = new List<FileReference>();
+			List<OverflowFileInfo> OverflowInfos = new List<OverflowFileInfo>();
 
 			if (AllowPatchOBBFile(SC))
 			{
@@ -1195,18 +1185,15 @@ public class AndroidPlatform : Platform
 				Int64 MinimumObbSize = 22 + 10;		// EOCD with comment (store version)
 				Int64 MainObbSize = MinimumObbSize;
 				Int64 PatchObbSize = MinimumObbSize;
-				Int64 Overflow1ObbSize = MinimumObbSize;
-				Int64 Overflow2ObbSize = MinimumObbSize;
-				Int64 Overflow3ObbSize = MinimumObbSize;
-				Int64 Overflow4ObbSize = MinimumObbSize;
+
 				foreach (FileReference FileRef in FilesForObb)
 				{
 					FileInfo LocalFileInfo = new FileInfo(FileRef.FullName);
 					Int64 LocalFileLength = LocalFileInfo.Length;
 					Int64 FilenameLength = FileRef.FullName.Length - StagingDirLength - 1;
 
-					Int64 LocalOverhead = (30 + FilenameLength + 36);		// local file descriptor
-					Int64 GlobalOverhead = (46 + FilenameLength + 36);		// central directory cost
+					Int64 LocalOverhead = (30 + FilenameLength + 36);       // local file descriptor
+					Int64 GlobalOverhead = (46 + FilenameLength + 36);      // central directory cost
 					Int64 FileRequirements = LocalFileLength + LocalOverhead + GlobalOverhead;
 
 					if (MainObbSize + FileRequirements < OBBSizeAllowed)
@@ -1219,31 +1206,49 @@ public class AndroidPlatform : Platform
 						FilesToPatch.Add(FileRef);
 						PatchObbSize += FileRequirements;
 					}
-					else if (AllowOverflowOBBLimit >= 1 && Overflow1ObbSize + FileRequirements < OBBSizeAllowed)
+					else if (AllowOverflowOBBLimit > 0)
 					{
-						FilesToOverflow1.Add(FileRef);
-						Overflow1ObbSize += FileRequirements;
-					}
-					else if (AllowOverflowOBBLimit >= 2 && Overflow2ObbSize + FileRequirements < OBBSizeAllowed)
-					{
-						FilesToOverflow2.Add(FileRef);
-						Overflow2ObbSize += FileRequirements;
-					}
-					else if (AllowOverflowOBBLimit >= 3 && Overflow3ObbSize + FileRequirements < OBBSizeAllowed)
-					{
-						FilesToOverflow3.Add(FileRef);
-						Overflow3ObbSize += FileRequirements;
-					}
-					else if (AllowOverflowOBBLimit >= 4 && Overflow4ObbSize + FileRequirements < OBBSizeAllowed)
-					{
-						FilesToOverflow4.Add(FileRef);
-						Overflow4ObbSize += FileRequirements;
+						// find an overflow we can fit this file
+						bool bFoundOverflow = false;				
+						foreach (OverflowFileInfo OverflowRef in OverflowInfos)
+						{
+							if (OverflowRef.OverflowObbSize + FileRequirements < OBBSizeAllowed)
+							{
+								OverflowRef.FilesForOverflow.Add(FileRef);
+								OverflowRef.OverflowObbSize += FileRequirements;
+								bFoundOverflow = true;
+							}
+						}
+
+						if (!bFoundOverflow)
+						{
+							if (AllowOverflowOBBLimit >= OverflowInfos.Count)
+							{
+								// create a new overflow obb
+								if (MinimumObbSize + FileRequirements < OBBSizeAllowed)
+								{
+									List<FileReference> NewObb = new List<FileReference>();
+									NewObb.Add(FileRef);
+									OverflowInfos.Add(new OverflowFileInfo(NewObb, MinimumObbSize + FileRequirements));
+								}
+								else
+								{
+									Logger.LogInformation("{Text}", "Failed to add " + FileRef.FullName + " to a new overflow as it is bigger than the allowed OBB size ");
+									throw new AutomationException(ExitCode.Error_AndroidOBBError, "Stage Failed. Could not add {0} ({1} bytes) to an OBB as OBBs are limited to {1} bytes.", FileRequirements, FileRef.FullName, OBBSizeAllowed);
+								}
+							}
+							else
+							{
+								Logger.LogInformation("{Text}", "Failed to add required overflow OBB: " + LocalObbName);
+								throw new AutomationException(ExitCode.Error_AndroidOBBError, "Stage Failed. Overflow OBBs limited to a count of {0}. Contents are to big to fit.", OverflowInfos.Count);
+							}
+						}
 					}
 					else
 					{
-						// no room in either file
+						// no room in either file and no overflows allowed
 						Logger.LogInformation("{Text}", "Failed to build OBB: " + LocalObbName);
-						throw new AutomationException(ExitCode.Error_AndroidOBBError, "Stage Failed. Could not build OBB {0}. The file may be too big to fit in an OBB ({1} limit)", LocalObbName, LimitString);
+						throw new AutomationException(ExitCode.Error_AndroidOBBError, "Stage Failed. Could not build OBB {0}. The file may be too big to fit in an OBB ({1} limit and no overflows are permitted)", LocalObbName, LimitString);
 					}
 				}
 			}
@@ -1265,44 +1270,19 @@ public class AndroidPlatform : Platform
 				}
 			}
 
-			// Now create the overflow1 OBB as a ZIP archive if required.
-			if (FilesToOverflow1.Count() > 0)
-			{
-				if (!CreateOBBFile(SC, LocalOverflow1Name, FilesToOverflow1))
-				{
-					Logger.LogInformation("{Text}", "Failed to build OBB: " + LocalOverflow1Name);
-					throw new AutomationException(ExitCode.Error_AndroidOBBError, "Stage Failed. Could not build OBB {0}. The file may be too big to fit in an OBB ({1} limit)", LocalOverflow1Name, LimitString);
-				}
-			}
-
-			// Now create the overflow2 OBB as a ZIP archive if required.
-			if (FilesToOverflow2.Count() > 0)
+			OverflowIndex = 1;
+			foreach (OverflowFileInfo OverflowRef in OverflowInfos)
 			{
-				if (!CreateOBBFile(SC, LocalOverflow2Name, FilesToOverflow2))
+				if (OverflowRef.FilesForOverflow.Count() > 0)
 				{
-					Logger.LogInformation("{Text}", "Failed to build OBB: " + LocalOverflow2Name);
-					throw new AutomationException(ExitCode.Error_AndroidOBBError, "Stage Failed. Could not build OBB {0}. The file may be too big to fit in an OBB ({1} limit)", LocalOverflow2Name, LimitString);
-				}
-			}
-
-			// Now create the overflow3 OBB as a ZIP archive if required.
-			if (FilesToOverflow3.Count() > 0)
-			{
-				if (!CreateOBBFile(SC, LocalOverflow3Name, FilesToOverflow3))
-				{
-					Logger.LogInformation("{Text}", "Failed to build OBB: " + LocalOverflow3Name);
-					throw new AutomationException(ExitCode.Error_AndroidOBBError, "Stage Failed. Could not build OBB {0}. The file may be too big to fit in an OBB ({1} limit)", LocalOverflow3Name, LimitString);
-				}
-			}
-
-			// Now create the overflow4 OBB as a ZIP archive if required.
-			if (FilesToOverflow4.Count() > 0)
-			{
-				if (!CreateOBBFile(SC, LocalOverflow4Name, FilesToOverflow4))
-				{
-					Logger.LogInformation("{Text}", "Failed to build OBB: " + LocalOverflow4Name);
-					throw new AutomationException(ExitCode.Error_AndroidOBBError, "Stage Failed. Could not build OBB {0}. The file may be too big to fit in an OBB ({1} limit)", LocalOverflow4Name, LimitString);
+					string LocalOverflowName = string.Format(LocalOverflowNameTemplate, OverflowIndex);
+					if (!CreateOBBFile(SC, LocalOverflowName, OverflowRef.FilesForOverflow))
+					{
+						Logger.LogInformation("{Text}", "Failed to build OBB: " + LocalOverflowName);
+						throw new AutomationException(ExitCode.Error_AndroidOBBError, "Stage Failed. Could not build OBB {0}. The file may be too big to fit in an OBB ({1} limit)", LocalOverflowName, LimitString);
+					}
 				}
+				OverflowIndex++;
 			}
 		}
 
@@ -1311,7 +1291,7 @@ public class AndroidPlatform : Platform
 		Int64 ObbFileLength = OBBFileInfo.Length;
 		if (ObbFileLength > OBBSizeAllowed)
 		{
-			Logger.LogInformation("OBB exceeds " + LimitString + " limit: " + ObbFileLength + " bytes");
+			Logger.LogInformation("{Text}", "OBB exceeds " + LimitString + " limit: " + ObbFileLength + " bytes");
 			throw new AutomationException(ExitCode.Error_AndroidOBBError, "Stage Failed. OBB {0} exceeds {1} limit)", LocalObbName, LimitString);
 		}
 
@@ -1381,14 +1361,9 @@ public class AndroidPlatform : Platform
 			string ObbName = "";
 			string DevicePatchName = "";
 			string PatchName = "";
-			string DeviceOverflow1Name = "";
-			string Overflow1Name = "";
-			string DeviceOverflow2Name = "";
-			string Overflow2Name = "";
-			string DeviceOverflow3Name = "";
-			string Overflow3Name = "";
-			string DeviceOverflow4Name = "";
-			string Overflow4Name = "";
+
+			List<OverflowBatchInstallInfo> OverflowInfos = new List<OverflowBatchInstallInfo>();
+			
 			if (!bPackageDataInsideApk)
 			{
 				DeviceObbName = GetDeviceObbName(ApkName, SC);
@@ -1408,44 +1383,19 @@ public class AndroidPlatform : Platform
 					UpdateObbStoreVersion(PatchName);
 				}
 
-				if (File.Exists(LocalOverflow1Name))
+				int OverflowIndex = 1;
+				while (File.Exists(string.Format(LocalOverflowNameTemplate, OverflowIndex)))
 				{
-					DeviceOverflow1Name = GetDeviceOverflowName(ApkName, SC, 1);
-					Overflow1Name = GetFinalOverflowName(ApkName, SC, 1);
-					CopyFile(LocalOverflow1Name, Overflow1Name);
+					string DeviceOverflowName = GetDeviceOverflowName(ApkName, SC, OverflowIndex);
+					string OverflowName = GetFinalOverflowName(ApkName, SC, OverflowIndex);
 
-					// apply store version to OBB to make it unique for PlayStore upload
-					UpdateObbStoreVersion(Overflow1Name);
-				}
-
-				if (File.Exists(LocalOverflow2Name))
-				{
-					DeviceOverflow2Name = GetDeviceOverflowName(ApkName, SC, 2);
-					Overflow2Name = GetFinalOverflowName(ApkName, SC, 2);
-					CopyFile(LocalOverflow2Name, Overflow2Name);
-
-					// apply store version to OBB to make it unique for PlayStore upload
-					UpdateObbStoreVersion(Overflow2Name);
-				}
-
-				if (File.Exists(LocalOverflow3Name))
-				{
-					DeviceOverflow3Name = GetDeviceOverflowName(ApkName, SC, 3);
-					Overflow3Name = GetFinalOverflowName(ApkName, SC, 3);
-					CopyFile(LocalOverflow3Name, Overflow3Name);
+					OverflowInfos.Add(new OverflowBatchInstallInfo(OverflowIndex, DeviceOverflowName, OverflowName, false));
+					CopyFile(string.Format(LocalOverflowNameTemplate, OverflowIndex), OverflowName);
 
 					// apply store version to OBB to make it unique for PlayStore upload
-					UpdateObbStoreVersion(Overflow3Name);
-				}
+					UpdateObbStoreVersion(OverflowName);
 
-				if (File.Exists(LocalOverflow4Name))
-				{
-					DeviceOverflow4Name = GetDeviceOverflowName(ApkName, SC, 4);
-					Overflow4Name = GetFinalOverflowName(ApkName, SC, 4);
-					CopyFile(LocalOverflow4Name, Overflow4Name);
-
-					// apply store version to OBB to make it unique for PlayStore upload
-					UpdateObbStoreVersion(Overflow4Name);
+					OverflowIndex++;
 				}
 			}
 
@@ -1485,20 +1435,19 @@ public class AndroidPlatform : Platform
 				// Write install batch file(s).
 				string PackageName = GetPackageInfo(ApkName, SC, false);
 				string BatchName = GetFinalBatchName(ApkName, SC, false, EBatchType.Install, Target);
-				string[] BatchLines = GenerateInstallBatchFile(bPackageDataInsideApk, PackageName, ApkName, Params, ObbName, DeviceObbName, false, PatchName, DevicePatchName, false, 
-					Overflow1Name, DeviceOverflow1Name, false, Overflow2Name, DeviceOverflow2Name, false, Overflow3Name, DeviceOverflow3Name, false, Overflow4Name, DeviceOverflow4Name, false, 
+				List<string> InstallBatchLines = GenerateInstallBatchFile(bPackageDataInsideApk, PackageName, ApkName, Params, ObbName, DeviceObbName, false, PatchName, DevicePatchName, false, OverflowInfos,
 					bIsPC, Params.Distribution, TargetSDKVersion > 22, bDisablePerfHarden, bUseAFS, bUseAFSProject, AFSToken, Target);
 				if (bHaveAPK)
 				{
 					// make a batch file that can be used to install the .apk and .obb files
-					File.WriteAllText(BatchName, string.Join(LineEnding, BatchLines) + LineEnding);
+					File.WriteAllText(BatchName, string.Join(LineEnding, InstallBatchLines) + LineEnding);
 				}
 				// make a batch file that can be used to uninstall the .apk and .obb files
 				string UninstallBatchName = GetFinalBatchName(ApkName, SC, false, EBatchType.Uninstall, Target);
-				BatchLines = GenerateUninstallBatchFile(bPackageDataInsideApk, PackageName, ApkName, Params, bIsPC);
+				string[] UninstallBatchLines = GenerateUninstallBatchFile(bPackageDataInsideApk, PackageName, ApkName, Params, bIsPC);
 				if (bHaveAPK || bHaveUniversal)
 				{
-					File.WriteAllText(UninstallBatchName, string.Join(LineEnding, BatchLines) + LineEnding);
+					File.WriteAllText(UninstallBatchName, string.Join(LineEnding, UninstallBatchLines) + LineEnding);
 				}
 
 				string UniversalBatchName = "";
@@ -1506,8 +1455,7 @@ public class AndroidPlatform : Platform
 				{
 					UniversalBatchName = GetFinalBatchName(UniversalApkName, SC, false, EBatchType.Install, Target);
 					// make a batch file that can be used to install the .apk
-					string[] UniversalBatchLines = GenerateInstallBatchFile(bPackageDataInsideApk, PackageName, UniversalApkName, Params, ObbName, DeviceObbName, false, PatchName, DevicePatchName, false,
-						Overflow1Name, DeviceOverflow1Name, false, Overflow2Name, DeviceOverflow2Name, false, Overflow3Name, DeviceOverflow3Name, false, Overflow4Name, DeviceOverflow4Name, false, 
+					List<string> UniversalBatchLines = GenerateInstallBatchFile(bPackageDataInsideApk, PackageName, UniversalApkName, Params, ObbName, DeviceObbName, false, PatchName, DevicePatchName, false, OverflowInfos,
 						bIsPC, Params.Distribution, TargetSDKVersion > 22, bDisablePerfHarden, bUseAFS, bUseAFSProject, AFSToken, Target);
 					File.WriteAllText(UniversalBatchName, string.Join(LineEnding, UniversalBatchLines) + LineEnding);
 				}
@@ -1515,8 +1463,8 @@ public class AndroidPlatform : Platform
 				string SymbolizeBatchName = GetFinalBatchName(ApkName, SC, false, EBatchType.Symbolize, Target);
 				if(bBuildWithHiddenSymbolVisibility || bSaveSymbols)
 				{
-					BatchLines = GenerateSymbolizeBatchFile(Params, PackageName, ApkName, SC, Architecture, bIsPC);
-					File.WriteAllText(SymbolizeBatchName, string.Join(LineEnding, BatchLines) + LineEnding);
+					string[] UniversalSymbolizeBatchLines = GenerateSymbolizeBatchFile(Params, PackageName, ApkName, SC, Architecture, bIsPC);
+					File.WriteAllText(SymbolizeBatchName, string.Join(LineEnding, UniversalSymbolizeBatchLines) + LineEnding);
 				}
 
 				if (!RuntimePlatform.IsWindows)
@@ -1588,12 +1536,26 @@ public class AndroidPlatform : Platform
 		return "win-x64/UnrealAndroidFileTool.exe";
 	}
 
-	private string[] GenerateInstallBatchFile(bool bPackageDataInsideApk, string PackageName, string ApkName, ProjectParams Params, string ObbName, string DeviceObbName, bool bNoObbInstall,
-		string PatchName, string DevicePatchName, bool bNoPatchInstall, string Overflow1Name, string DeviceOverflow1Name, bool bNoOverflow1Install, string Overflow2Name, string DeviceOverflow2Name, bool bNoOverflow2Install,
-		string Overflow3Name, string DeviceOverflow3Name, bool bNoOverflow3Install, string Overflow4Name, string DeviceOverflow4Name, bool bNoOverflow4Install,
+	private class OverflowBatchInstallInfo
+	{
+		public OverflowBatchInstallInfo(int InOverflowIndex, string InDeviceOverflowName, string InOverflowName, bool InNoOverflowInstall)
+		{
+			OverflowIndex = InOverflowIndex;
+			DeviceOverflowName = InDeviceOverflowName;
+			OverflowName = InOverflowName;
+			bNoOverflowInstall = InNoOverflowInstall;
+		}
+
+		public int OverflowIndex { get; }
+		public string DeviceOverflowName { get; }
+		public string OverflowName { get; }
+		public bool bNoOverflowInstall { get; }
+	}
+	private List<string> GenerateInstallBatchFile(bool bPackageDataInsideApk, string PackageName, string ApkName, ProjectParams Params, string ObbName, string DeviceObbName, bool bNoObbInstall,
+		string PatchName, string DevicePatchName, bool bNoPatchInstall, List<OverflowBatchInstallInfo> OverflowInfo,
 		bool bIsPC, bool bIsDistribution, bool bRequireRuntimeStoragePermission, bool bDisablePerfHarden, bool bUseAFS, bool bUseAFSProject, string AFSToken, UnrealTargetPlatform Target)
     {
-        string[] BatchLines = null;
+		List<string> BatchLines = new List<string>();
         string ReadPermissionGrantCommand = "shell pm grant " + PackageName + " android.permission.READ_EXTERNAL_STORAGE";
         string WritePermissionGrantCommand = "shell pm grant " + PackageName + " android.permission.WRITE_EXTERNAL_STORAGE";
 		string DisablePerfHardenCommand = "shell setprop security.perf_harden 0";
@@ -1606,10 +1568,6 @@ public class AndroidPlatform : Platform
 		bool bDontMoveOBB = bUseAFS ? true : bPackageDataInsideApk;
 
 		bool bHavePatch = (PatchName != "");
-		bool bHaveOverflow1 = (Overflow1Name != "");
-		bool bHaveOverflow2 = (Overflow2Name != "");
-		bool bHaveOverflow3 = (Overflow3Name != "");
-		bool bHaveOverflow4 = (Overflow4Name != "");
 
 		string AFSExecutable = GetAFSExecutable(Target);
 		string AFSCommonArg = "-p " + PackageName;
@@ -1619,7 +1577,7 @@ public class AndroidPlatform : Platform
 		}
 
 		if (!bIsPC)
-        {
+		{
 			string APKInstallCommand = "$ADB $DEVICE install " + Path.GetFileName(ApkName);
 			string APKReinstallCommand = "";
 
@@ -1627,10 +1585,8 @@ public class AndroidPlatform : Platform
 			// Note that $STORAGE/Android/obb will be the folder that contains the obb if you download the app from playstore.
 			string OBBInstallCommand = bNoObbInstall ? "\t$ADB $DEVICE shell 'rm -r $EXTERNAL_STORAGE/" + DeviceObbName + "'" : "\t$ADB $DEVICE push " + Path.GetFileName(ObbName) + " " + TargetAndroidTemp + DeviceObbName;
 			string PatchInstallCommand = bNoPatchInstall ? "\t$ADB $DEVICE shell 'rm -r $EXTERNAL_STORAGE/" + DevicePatchName + "'" : "\t$ADB $DEVICE push " + Path.GetFileName(PatchName) + TargetAndroidTemp + DevicePatchName;
-			string Overflow1InstallCommand = bNoOverflow1Install ? "\t$ADB $DEVICE shell 'rm -r $EXTERNAL_STORAGE/" + DeviceOverflow1Name + "'" : "\t$ADB $DEVICE push " + Path.GetFileName(Overflow1Name) + " " + TargetAndroidTemp + DeviceOverflow1Name;
-			string Overflow2InstallCommand = bNoOverflow2Install ? "\t$ADB $DEVICE shell 'rm -r $EXTERNAL_STORAGE/" + DeviceOverflow2Name + "'" : "\t$ADB $DEVICE push " + Path.GetFileName(Overflow2Name) + " " + TargetAndroidTemp + DeviceOverflow2Name;
-			string Overflow3InstallCommand = bNoOverflow3Install ? "\t$ADB $DEVICE shell 'rm -r $EXTERNAL_STORAGE/" + DeviceOverflow3Name + "'" : "\t$ADB $DEVICE push " + Path.GetFileName(Overflow3Name) + " " + TargetAndroidTemp + DeviceOverflow3Name;
-			string Overflow4InstallCommand = bNoOverflow4Install ? "\t$ADB $DEVICE shell 'rm -r $EXTERNAL_STORAGE/" + DeviceOverflow4Name + "'" : "\t$ADB $DEVICE push " + Path.GetFileName(Overflow4Name) + " " + TargetAndroidTemp + DeviceOverflow4Name;
+
+			List<string> OverflowInstallCommands = new List<string>();
 
 			if (bUseAFS)
 			{
@@ -1644,17 +1600,33 @@ public class AndroidPlatform : Platform
 					// stop the fileserver (not needed on reinstall above
 					APKReinstallCommand = "$AFS $DEVICE " + AFSCommonArg + " stop-all";
 				}
-				string AFSCommand = "\t$AFS $DEVICE " + AFSCommonArg; 
+				string AFSCommand = "\t$AFS $DEVICE " + AFSCommonArg;
 				OBBInstallCommand = bNoObbInstall ? AFSCommand + " deletefile '^mainobb'" : AFSCommand + " push " + Path.GetFileName(ObbName) + " '^mainobb'";
 				PatchInstallCommand = bNoPatchInstall ? AFSCommand + " deletefile '^patchobb'" : AFSCommand + " push " + Path.GetFileName(PatchName) + " '^patchobb'";
-				Overflow1InstallCommand = bNoOverflow1Install ? AFSCommand + " deletefile '^overflow1obb'" : AFSCommand + " push " + Path.GetFileName(Overflow1Name) + " '^overflow1obb'";
-				Overflow2InstallCommand = bNoOverflow2Install ? AFSCommand + " deletefile '^overflow2obb'" : AFSCommand + " push " + Path.GetFileName(Overflow2Name) + " '^overflow2obb'";
-				Overflow3InstallCommand = bNoOverflow3Install ? AFSCommand + " deletefile '^overflow3obb'" : AFSCommand + " push " + Path.GetFileName(Overflow3Name) + " '^overflow3obb'";
-				Overflow4InstallCommand = bNoOverflow4Install ? AFSCommand + " deletefile '^overflow4obb'" : AFSCommand + " push " + Path.GetFileName(Overflow4Name) + " '^overflow4obb'";
+				if (!bPackageDataInsideApk)
+				{
+					foreach (OverflowBatchInstallInfo Overflow in OverflowInfo)
+					{
+						string AfsOverflowName = string.Format("^overflow{0}obb", Overflow.OverflowIndex);
+						OverflowInstallCommands.Add(Overflow.bNoOverflowInstall ? AFSCommand + " deletefile '" + AfsOverflowName + "'" : AFSCommand + " push " + Path.GetFileName(Overflow.OverflowName) + " '" + AfsOverflowName + "'");
+						OverflowInstallCommands.Add("if \"%ERRORLEVEL%\" NEQ \"0\" goto Error");
+					}
+				}
+			}
+			else
+			{
+				if (!bPackageDataInsideApk)
+				{
+					foreach (OverflowBatchInstallInfo Overflow in OverflowInfo)
+					{
+						OverflowInstallCommands.Add(Overflow.bNoOverflowInstall ? "\t$ADB $DEVICE shell 'rm -r $EXTERNAL_STORAGE/" + Overflow.DeviceOverflowName + "'" : "\t$ADB $DEVICE push " + Path.GetFileName(Overflow.OverflowName) + " " + TargetAndroidTemp + Overflow.DeviceOverflowName);
+						OverflowInstallCommands.Add("if \"%ERRORLEVEL%\" NEQ \"0\" goto Error");
+					}
+				}
 			}
 
 			Logger.LogInformation("Writing shell script for install with {Arg0}", bPackageDataInsideApk ? "data in APK" : "separate obb");
-            BatchLines = new string[] {
+			BatchLines.AddRange(new string[] {
 						"#!/bin/sh",
 						"cd \"`dirname \"$0\"`\"",
 						"AFS=./" + AFSExecutable,
@@ -1669,15 +1641,15 @@ public class AndroidPlatform : Platform
 						"echo Installing existing application. Failures here indicate a problem with the device \\(connection or storage permissions\\) and are fatal.",
 						APKInstallCommand,
 						"if [ $? -eq 0 ]; then",
-                        "\techo",
+						"\techo",
 						"\t$ADB $DEVICE shell pm list packages " + PackageName,
 						bNeedGrantStoragePermission ? "\techo Grant READ_EXTERNAL_STORAGE and WRITE_EXTERNAL_STORAGE to the apk for reading OBB or game file in external storage." : "",
 						bNeedGrantStoragePermission ? "\t$ADB $DEVICE " + ReadPermissionGrantCommand : "",
 						bNeedGrantStoragePermission ? "\t$ADB $DEVICE " + WritePermissionGrantCommand : "",
 						bDisablePerfHarden ? "\t$ADB $DEVICE " + DisablePerfHardenCommand : "",
-                        "\techo",
+						"\techo",
 						"\techo Removing old data. Failures here are usually fine - indicating the files were not on the device.",
-                        "\t$ADB $DEVICE shell 'rm -r $EXTERNAL_STORAGE/UnrealGame/" + Params.ShortProjectName + "'",
+						"\t$ADB $DEVICE shell 'rm -r $EXTERNAL_STORAGE/UnrealGame/" + Params.ShortProjectName + "'",
 						"\t$ADB $DEVICE shell 'rm -r $EXTERNAL_STORAGE/UnrealGame/UECommandLine.txt" + "'",
 						"\t$ADB $DEVICE shell 'rm -r $EXTERNAL_STORAGE/" + TargetAndroidLocation + PackageName + "'",
 						"\t$ADB $DEVICE shell 'rm -r $EXTERNAL_STORAGE/Android/" + TargetAndroidLocation + PackageName + "'",
@@ -1687,11 +1659,11 @@ public class AndroidPlatform : Platform
 						bPackageDataInsideApk ? "" : "\tSTORAGE=$(echo \"`$ADB $DEVICE shell 'echo $EXTERNAL_STORAGE'`\" | cat -v | tr -d '^M')",
 						bPackageDataInsideApk ? "" : OBBInstallCommand,
 						bPackageDataInsideApk ? "if [ 1 ]; then" : "\tif [ $? -eq 0 ]; then",
-						!bHavePatch ? "" : (bPackageDataInsideApk ? "" : PatchInstallCommand),
-						!bHaveOverflow1 ? "" : (bPackageDataInsideApk ? "" : Overflow1InstallCommand),
-						!bHaveOverflow2 ? "" : (bPackageDataInsideApk ? "" : Overflow2InstallCommand),
-						!bHaveOverflow3 ? "" : (bPackageDataInsideApk ? "" : Overflow3InstallCommand),
-						!bHaveOverflow4 ? "" : (bPackageDataInsideApk ? "" : Overflow4InstallCommand),
+						!bHavePatch ? "" : (bPackageDataInsideApk ? "" : PatchInstallCommand)});
+
+			BatchLines.AddRange(OverflowInstallCommands);
+
+			BatchLines.AddRange(new string[] {
 						bDontMoveOBB ? "" : "\t\t$ADB $DEVICE shell mkdir $STORAGE/Android/" + TargetAndroidLocation + PackageName, // don't check for error since installing may create the obb directory
 						bDontMoveOBB ? "" : "\t\t$ADB $DEVICE shell mv " + TargetAndroidTemp + TargetAndroidLocation + PackageName + " $STORAGE/Android/" + TargetAndroidLocation,
 						bDontMoveOBB ? "" : "\t\t$ADB $DEVICE shell rm -r " + TargetAndroidTemp + TargetAndroidLocation,
@@ -1709,19 +1681,17 @@ public class AndroidPlatform : Platform
 						"echo Make sure all Developer options look normal on the device",
 						"echo Check that the device has an SD card.",
 						"exit 1"
-					};
-        }
-        else
-        {
+					});
+		}
+		else
+		{
 			string APKInstallCommand = "%ADB% %DEVICE% install " + Path.GetFileName(ApkName);
 			string APKReinstallCommand = "";
 
 			string OBBInstallCommand = bNoObbInstall ? "%ADB% %DEVICE% shell rm -r %STORAGE%/" + DeviceObbName : "%ADB% %DEVICE% push " + Path.GetFileName(ObbName) + " " + TargetAndroidTemp + DeviceObbName;
 			string PatchInstallCommand = bNoPatchInstall ? "%ADB% %DEVICE% shell rm -r %STORAGE%/" + DevicePatchName : "%ADB% %DEVICE% push " + Path.GetFileName(PatchName) + " " + TargetAndroidTemp + DevicePatchName;
-			string Overflow1InstallCommand = bNoOverflow1Install ? "%ADB% %DEVICE% shell rm -r %STORAGE%/" + DeviceOverflow1Name : "%ADB% %DEVICE% push " + Path.GetFileName(Overflow1Name) + " " + TargetAndroidTemp + DeviceOverflow1Name;
-			string Overflow2InstallCommand = bNoOverflow2Install ? "%ADB% %DEVICE% shell rm -r %STORAGE%/" + DeviceOverflow2Name : "%ADB% %DEVICE% push " + Path.GetFileName(Overflow2Name) + " " + TargetAndroidTemp + DeviceOverflow2Name;
-			string Overflow3InstallCommand = bNoOverflow3Install ? "%ADB% %DEVICE% shell rm -r %STORAGE%/" + DeviceOverflow3Name : "%ADB% %DEVICE% push " + Path.GetFileName(Overflow3Name) + " " + TargetAndroidTemp + DeviceOverflow3Name;
-			string Overflow4InstallCommand = bNoOverflow4Install ? "%ADB% %DEVICE% shell rm -r %STORAGE%/" + DeviceOverflow4Name : "%ADB% %DEVICE% push " + Path.GetFileName(Overflow4Name) + " " + TargetAndroidTemp + DeviceOverflow4Name;
+
+			List<string> OverflowInstallCommands = new List<string>();
 
 			if (bUseAFS)
 			{
@@ -1738,23 +1708,41 @@ public class AndroidPlatform : Platform
 				string AFSCommand = "\t%AFS% %DEVICE% " + AFSCommonArg;
 				OBBInstallCommand = bNoObbInstall ? AFSCommand + " deletefile \"^mainobb\"" : AFSCommand + " push " + Path.GetFileName(ObbName) + " \"^mainobb\"";
 				PatchInstallCommand = bNoPatchInstall ? AFSCommand + " deletefile \"^patchobb\"" : AFSCommand + " push " + Path.GetFileName(PatchName) + " \"^patchobb\"";
-				Overflow1InstallCommand = bNoOverflow1Install ? AFSCommand + " deletefile \"^overflow1obb\"" : AFSCommand + " push " + Path.GetFileName(Overflow1Name) + " \"^overflow1obb\"";
-				Overflow2InstallCommand = bNoOverflow2Install ? AFSCommand + " deletefile \"^overflow2obb\"" : AFSCommand + " push " + Path.GetFileName(Overflow2Name) + " \"^overflow2obb\"";
-				Overflow3InstallCommand = bNoOverflow3Install ? AFSCommand + " deletefile \"^overflow3obb\"" : AFSCommand + " push " + Path.GetFileName(Overflow3Name) + " \"^overflow3obb\"";
-				Overflow4InstallCommand = bNoOverflow4Install ? AFSCommand + " deletefile \"^overflow4obb\"" : AFSCommand + " push " + Path.GetFileName(Overflow4Name) + " \"^overflow4obb\"";
+
+				if (!bPackageDataInsideApk)
+				{
+					foreach (OverflowBatchInstallInfo Overflow in OverflowInfo)
+					{
+						string AfsOverflowName = string.Format("^overflow{0}obb", Overflow.OverflowIndex);
+						OverflowInstallCommands.Add(Overflow.bNoOverflowInstall ? AFSCommand + " deletefile \"" + AfsOverflowName + "\"" : AFSCommand + " push " + Path.GetFileName(Overflow.OverflowName) + " \"" + AfsOverflowName + "\"");
+						OverflowInstallCommands.Add("if \"%ERRORLEVEL%\" NEQ \"0\" goto Error");
+					}
+				}
+
+			}
+			else
+			{
+				if (!bPackageDataInsideApk)
+				{
+					foreach (OverflowBatchInstallInfo Overflow in OverflowInfo)
+					{
+						OverflowInstallCommands.Add(Overflow.bNoOverflowInstall  ? "%ADB% %DEVICE% shell rm -r %STORAGE%/" + Overflow.DeviceOverflowName : "%ADB% %DEVICE% push " + Path.GetFileName(Overflow.OverflowName) + " " + TargetAndroidTemp + Overflow.DeviceOverflowName);
+						OverflowInstallCommands.Add("if \"%ERRORLEVEL%\" NEQ \"0\" goto Error");
+					}
+				}
 			}
 
 			Logger.LogInformation("Writing bat for install with {Arg0}", bPackageDataInsideApk ? "data in APK" : "separate OBB");
-            BatchLines = new string[] {
+			BatchLines.AddRange(new string[] {
 						"setlocal",
 						"if NOT \"%UE_SDKS_ROOT%\"==\"\" (call %UE_SDKS_ROOT%\\HostWin64\\Android\\SetupEnvironmentVars.bat)",
-						"set ANDROIDHOME=%ANDROID_HOME%",		
+						"set ANDROIDHOME=%ANDROID_HOME%",
 						"if \"%ANDROIDHOME%\"==\"\" set ANDROIDHOME="+Environment.GetEnvironmentVariable("ANDROID_HOME"),
 						"set ADB=%ANDROIDHOME%\\platform-tools\\adb.exe",
 						"set AFS=.\\" + AFSExecutable.Replace("/", "\\"),
 						"set DEVICE=",
-                        "if not \"%1\"==\"\" set DEVICE=-s %1",
-                        "for /f \"delims=\" %%A in ('%ADB% %DEVICE% " + GetStorageQueryCommand(true) +"') do @set STORAGE=%%A",
+						"if not \"%1\"==\"\" set DEVICE=-s %1",
+						"for /f \"delims=\" %%A in ('%ADB% %DEVICE% " + GetStorageQueryCommand(true) +"') do @set STORAGE=%%A",
 						"@echo.",
 						"@echo Uninstalling existing application. Failures here can almost always be ignored.",
 						"%ADB% %DEVICE% uninstall " + PackageName,
@@ -1773,15 +1761,11 @@ public class AndroidPlatform : Platform
 						bPackageDataInsideApk ? "" : OBBInstallCommand,
 						bPackageDataInsideApk ? "" : "if \"%ERRORLEVEL%\" NEQ \"0\" goto Error",
 						!bHavePatch ? "" : (bPackageDataInsideApk ? "" : PatchInstallCommand),
-						!bHavePatch ? "" : (bPackageDataInsideApk ? "" : "if \"%ERRORLEVEL%\" NEQ \"0\" goto Error"),
-						!bHaveOverflow1 ? "" : (bPackageDataInsideApk ? "" : Overflow1InstallCommand),
-						!bHaveOverflow1 ? "" : (bPackageDataInsideApk ? "" : "if \"%ERRORLEVEL%\" NEQ \"0\" goto Error"),
-						!bHaveOverflow2 ? "" : (bPackageDataInsideApk ? "" : Overflow2InstallCommand),
-						!bHaveOverflow2 ? "" : (bPackageDataInsideApk ? "" : "if \"%ERRORLEVEL%\" NEQ \"0\" goto Error"),
-						!bHaveOverflow3 ? "" : (bPackageDataInsideApk ? "" : Overflow3InstallCommand),
-						!bHaveOverflow3 ? "" : (bPackageDataInsideApk ? "" : "if \"%ERRORLEVEL%\" NEQ \"0\" goto Error"),
-						!bHaveOverflow4 ? "" : (bPackageDataInsideApk ? "" : Overflow4InstallCommand),
-						!bHaveOverflow4 ? "" : (bPackageDataInsideApk ? "" : "if \"%ERRORLEVEL%\" NEQ \"0\" goto Error"),
+						!bHavePatch ? "" : (bPackageDataInsideApk ? "" : "if \"%ERRORLEVEL%\" NEQ \"0\" goto Error")});
+
+					BatchLines.AddRange(OverflowInstallCommands);
+
+					BatchLines.AddRange(new string[] {
 						bDontMoveOBB ? "" : "%ADB% %DEVICE% shell mkdir %STORAGE%/Android/" + TargetAndroidLocation + PackageName, // don't check for error since installing may create the obb directory
 						bDontMoveOBB ? "" : "%ADB% %DEVICE% shell mv " + TargetAndroidTemp + TargetAndroidLocation + PackageName + " %STORAGE%/Android/" + TargetAndroidLocation,
 						bDontMoveOBB ? "" : "if \"%ERRORLEVEL%\" NEQ \"0\" goto Error",
@@ -1805,7 +1789,7 @@ public class AndroidPlatform : Platform
 						"@echo Make sure all Developer options look normal on the device",
 						"@echo Check that the device has an SD card.",
 						"@pause"
-					};
+					});
         }
         return BatchLines;
     }
@@ -1931,6 +1915,7 @@ public class AndroidPlatform : Platform
 		bool bAFSAllowExternalStartInShipping;
 		UsingAndroidFileServer(Params, SC, out bAFSEnablePlugin, out AFSToken, out bIsShipping, out bAFSIncludeInShipping, out bAFSAllowExternalStartInShipping);
 		bool bUseAFS = bAFSEnablePlugin && !bPackageDataInsideApk;
+		int AllowOverflowOBBLimit = AllowOverflowOBBFiles(SC);
 
 		List<string> AddedObbFiles = new List<string>();
 		foreach (UnrealArch Architecture in Architectures.Architectures)
@@ -1940,10 +1925,6 @@ public class AndroidPlatform : Platform
 			bool bHaveAPK = FileExists(ApkName);
 			string ObbName = GetFinalObbName(ApkName, SC);
 			string PatchName = GetFinalPatchName(ApkName, SC);
-			string Overflow1Name = GetFinalOverflowName(ApkName, SC, 1);
-			string Overflow2Name = GetFinalOverflowName(ApkName, SC, 2);
-			string Overflow3Name = GetFinalOverflowName(ApkName, SC, 3);
-			string Overflow4Name = GetFinalOverflowName(ApkName, SC, 4);
 			bool bBuildWithHiddenSymbolVisibility = BuildWithHiddenSymbolVisibility(SC);
 			bool bSaveSymbols = GetSaveSymbols(SC);
 
@@ -2057,21 +2038,18 @@ public class AndroidPlatform : Platform
 					{
 						SC.ArchiveFiles(Path.GetDirectoryName(PatchName), Path.GetFileName(PatchName));
 					}
-					if (FileExists(Overflow1Name))
-					{
-						SC.ArchiveFiles(Path.GetDirectoryName(Overflow1Name), Path.GetFileName(Overflow1Name));
-					}
-					if (FileExists(Overflow2Name))
-					{
-						SC.ArchiveFiles(Path.GetDirectoryName(Overflow2Name), Path.GetFileName(Overflow2Name));
-					}
-					if (FileExists(Overflow3Name))
-					{
-						SC.ArchiveFiles(Path.GetDirectoryName(Overflow3Name), Path.GetFileName(Overflow3Name));
-					}
-					if (FileExists(Overflow4Name))
+
+					for (int Index = 1; Index <= AllowOverflowOBBLimit; Index++)
 					{
-						SC.ArchiveFiles(Path.GetDirectoryName(Overflow4Name), Path.GetFileName(Overflow4Name));
+						string OverflowName = GetFinalOverflowName(ApkName, SC, Index);
+						if (FileExists(OverflowName))
+						{
+							SC.ArchiveFiles(Path.GetDirectoryName(OverflowName), Path.GetFileName(OverflowName));
+						}
+						else
+						{
+							break;
+						}
 					}
 				}
 			}
@@ -2456,6 +2434,7 @@ public class AndroidPlatform : Platform
 		EConnectionType ConnectionType = GetAndroidFileServerNetworkConfig(SC, out bUseCompression, out bLogFiles, out bReportStats, out bUseManualIPAddress, out ManualIPAddress);
 
 		AndroidFileClient.OptimalADB adb = new AndroidFileClient.OptimalADB();
+		int AllowOverflowOBBLimit = AllowOverflowOBBFiles(SC);
 
 		foreach (var DeviceName in Params.DeviceNames)
 		{
@@ -2495,10 +2474,6 @@ public class AndroidPlatform : Platform
 			string StorageLocation = ExternalStorage + "/Android";	// "mnt/sdcard/Android"
 			string DeviceObbName = StorageLocation + "/" + GetDeviceObbName(ApkName, SC);
 			string DevicePatchName = StorageLocation + "/" + GetDevicePatchName(ApkName, SC);
-			string DeviceOverflow1Name = StorageLocation + "/" + GetDeviceOverflowName(ApkName, SC, 1);
-			string DeviceOverflow2Name = StorageLocation + "/" + GetDeviceOverflowName(ApkName, SC, 2);
-			string DeviceOverflow3Name = StorageLocation + "/" + GetDeviceOverflowName(ApkName, SC, 3);
-			string DeviceOverflow4Name = StorageLocation + "/" + GetDeviceOverflowName(ApkName, SC, 4);
 			string RemoteDir = client.Query("^project", true);
 			if (RemoteDir == null)
 			{
@@ -2975,10 +2950,20 @@ public class AndroidPlatform : Platform
 				// delete the .obb file, since it will cause nothing we deploy to be used
 				client.FileDelete(DeviceObbName);
 				client.FileDelete(DevicePatchName);
-				client.FileDelete(DeviceOverflow1Name);
-				client.FileDelete(DeviceOverflow2Name);
-				client.FileDelete(DeviceOverflow3Name);
-				client.FileDelete(DeviceOverflow4Name);
+
+				// delete existing overflow files on device, one more than the number we have to ensure we do not try to mount an old extra one left on device
+				for (int Index = 1; Index <= AllowOverflowOBBLimit; Index++)
+				{
+					string DeviceOverflowName = StorageLocation + "/" + GetDeviceOverflowName(ApkName, SC, Index); 
+					client.FileDelete(DeviceOverflowName);
+
+					// stop if there is no staged overflow of this index
+					string OverflowPath = Path.Combine(SC.StageDirectory.FullName, GetFinalOverflowName(ApkName, SC, Index));
+					if (!File.Exists(OverflowPath))
+					{
+						break;
+					}
+				}
 
 				// We now have a minimal set of file & dir entries we need
 				// to deploy. Files we deploy will get individually copied
@@ -3004,32 +2989,18 @@ public class AndroidPlatform : Platform
 					client.FileWrite(PatchPath, DevicePatchName);
 				}
 
-				// deploy the overflow1 if there is one
-				string Overflow1Path = Path.Combine(SC.StageDirectory.FullName, GetFinalOverflowName(ApkName, SC, 1));
-				if (File.Exists(Overflow1Path))
+				for (int Index = 1; Index <= AllowOverflowOBBLimit; Index++)
 				{
-					client.FileWrite(Overflow1Path, DeviceOverflow1Name);
-				}
-
-				// deploy the overflow2 if there is one
-				string Overflow2Path = Path.Combine(SC.StageDirectory.FullName, GetFinalOverflowName(ApkName, SC, 2));
-				if (File.Exists(Overflow2Path))
-				{
-					client.FileWrite(Overflow2Path, DeviceOverflow2Name);
-				}
-
-				// deploy the overflow3 if there is one
-				string Overflow3Path = Path.Combine(SC.StageDirectory.FullName, GetFinalOverflowName(ApkName, SC, 3));
-				if (File.Exists(Overflow3Path))
-				{
-					client.FileWrite(Overflow3Path, DeviceOverflow3Name);
-				}
-
-				// deploy the overflow4 if there is one
-				string Overflow4Path = Path.Combine(SC.StageDirectory.FullName, GetFinalOverflowName(ApkName, SC, 4));
-				if (File.Exists(Overflow4Path))
-				{
-					client.FileWrite(Overflow4Path, DeviceOverflow4Name);
+					string OverflowPath = Path.Combine(SC.StageDirectory.FullName, GetFinalOverflowName(ApkName, SC, Index));
+					if (File.Exists(OverflowPath))
+					{
+						string DeviceOverflowName = StorageLocation + "/" + GetDeviceOverflowName(ApkName, SC, Index);
+						client.FileWrite(OverflowPath, DeviceOverflowName);
+					}
+					else
+					{
+						break;
+					}
 				}
 			}
 			else
@@ -3058,6 +3029,7 @@ public class AndroidPlatform : Platform
 		ConfigHierarchy Ini = ConfigCache.ReadHierarchy(ConfigHierarchyType.Engine, DirectoryReference.FromFile(Params.RawProjectPath), UnrealTargetPlatform.Android);
 		bool bDisablePerfHarden = false;
 		Ini.GetBool("/Script/AndroidRuntimeSettings.AndroidRuntimeSettings", "bEnableMaliPerfCounters", out bDisablePerfHarden);
+		int AllowOverflowOBBLimit = AllowOverflowOBBFiles(SC);
 
 		foreach (var DeviceName in Params.DeviceNames)
         {
@@ -3083,10 +3055,6 @@ public class AndroidPlatform : Platform
             String StorageLocation = Result.Output.Trim(); // "/mnt/sdcard";
             string DeviceObbName = StorageLocation + "/" + GetDeviceObbName(ApkName, SC);
 			string DevicePatchName = StorageLocation + "/" + GetDevicePatchName(ApkName, SC);
-			string DeviceOverflow1Name = StorageLocation + "/" + GetDeviceOverflowName(ApkName, SC, 1);
-			string DeviceOverflow2Name = StorageLocation + "/" + GetDeviceOverflowName(ApkName, SC, 2);
-			string DeviceOverflow3Name = StorageLocation + "/" + GetDeviceOverflowName(ApkName, SC, 3);
-			string DeviceOverflow4Name = StorageLocation + "/" + GetDeviceOverflowName(ApkName, SC, 4);
 			string RemoteDir = StorageLocation + "/UnrealGame/" + Params.ShortProjectName;
 
 			if (bDisablePerfHarden)
@@ -3185,57 +3153,57 @@ public class AndroidPlatform : Platform
             FileReference IntermediateCmdLineFile = FileReference.Combine(SC.StageDirectory, "UECommandLine.txt");
             Project.WriteStageCommandline(IntermediateCmdLineFile, Params, SC);
 
-            // copy files to device if we were staging
-            if (SC.Stage)
-            {
-                // cache some strings
-                string BaseCommandline = "push";
+			// copy files to device if we were staging
+			if (SC.Stage)
+			{
+				// cache some strings
+				string BaseCommandline = "push";
 
-                HashSet<string> EntriesToDeploy = new HashSet<string>();
+				HashSet<string> EntriesToDeploy = new HashSet<string>();
 
-                if (Params.IterativeDeploy)
-                {
-                    // always send UECommandLine.txt (it was written above after delta checks applied)
-                    EntriesToDeploy.Add(IntermediateCmdLineFile.FullName);
+				if (Params.IterativeDeploy)
+				{
+					// always send UECommandLine.txt (it was written above after delta checks applied)
+					EntriesToDeploy.Add(IntermediateCmdLineFile.FullName);
 
-                    // Add non UFS files if any to deploy
-                    String NonUFSManifestPath = SC.GetNonUFSDeploymentDeltaPath(DeviceName);
-                    if (File.Exists(NonUFSManifestPath))
-                    {
-                        string NonUFSFiles = File.ReadAllText(NonUFSManifestPath);
-                        foreach (string Filename in NonUFSFiles.Split('\n'))
-                        {
-                            if (!string.IsNullOrEmpty(Filename) && !string.IsNullOrWhiteSpace(Filename))
-                            {
-                                EntriesToDeploy.Add(CombinePaths(SC.StageDirectory.FullName, Filename.Trim()));
-                            }
-                        }
-                    }
+					// Add non UFS files if any to deploy
+					String NonUFSManifestPath = SC.GetNonUFSDeploymentDeltaPath(DeviceName);
+					if (File.Exists(NonUFSManifestPath))
+					{
+						string NonUFSFiles = File.ReadAllText(NonUFSManifestPath);
+						foreach (string Filename in NonUFSFiles.Split('\n'))
+						{
+							if (!string.IsNullOrEmpty(Filename) && !string.IsNullOrWhiteSpace(Filename))
+							{
+								EntriesToDeploy.Add(CombinePaths(SC.StageDirectory.FullName, Filename.Trim()));
+							}
+						}
+					}
 
-                    // Add UFS files if any to deploy
-                    String UFSManifestPath = SC.GetUFSDeploymentDeltaPath(DeviceName);
-                    if (File.Exists(UFSManifestPath))
-                    {
-                        string UFSFiles = File.ReadAllText(UFSManifestPath);
-                        foreach (string Filename in UFSFiles.Split('\n'))
-                        {
-                            if (!string.IsNullOrEmpty(Filename) && !string.IsNullOrWhiteSpace(Filename))
-                            {
-                                EntriesToDeploy.Add(CombinePaths(SC.StageDirectory.FullName, Filename.Trim()));
-                            }
-                        }
-                    }
+					// Add UFS files if any to deploy
+					String UFSManifestPath = SC.GetUFSDeploymentDeltaPath(DeviceName);
+					if (File.Exists(UFSManifestPath))
+					{
+						string UFSFiles = File.ReadAllText(UFSManifestPath);
+						foreach (string Filename in UFSFiles.Split('\n'))
+						{
+							if (!string.IsNullOrEmpty(Filename) && !string.IsNullOrWhiteSpace(Filename))
+							{
+								EntriesToDeploy.Add(CombinePaths(SC.StageDirectory.FullName, Filename.Trim()));
+							}
+						}
+					}
 
-                    // For now, if too many files may be better to just push them all
-                    if (EntriesToDeploy.Count > 500)
-                    {
-                        // make sure device is at a clean state
-                        RunAdbCommand(Params, DeviceName, "shell rm -r " + RemoteDir);
+					// For now, if too many files may be better to just push them all
+					if (EntriesToDeploy.Count > 500)
+					{
+						// make sure device is at a clean state
+						RunAdbCommand(Params, DeviceName, "shell rm -r " + RemoteDir);
 
-                        EntriesToDeploy.Clear();
-                        EntriesToDeploy.TrimExcess();
-                        EntriesToDeploy.Add(SC.StageDirectory.FullName);
-                    }
+						EntriesToDeploy.Clear();
+						EntriesToDeploy.TrimExcess();
+						EntriesToDeploy.Add(SC.StageDirectory.FullName);
+					}
 					else
 					{
 						// Discover & remove any files on device that are not in staging
@@ -3259,232 +3227,241 @@ public class AndroidPlatform : Platform
 						}
 						else
 						{
-						// listing output is of the form
-						// [Samsung]                 [Google]
-						//
-						// RemoteDir/RestOfPath:     RemoteDir/RestOfPath:
-						// - File1.png               File1.png
-						// - File2.txt               File2.txt
-						// d SubDir1                 SubDir1/
-						// d SubDir2                 Subdir2/
-						//
-						// RemoteDir/RestOfPath/SubDir1:
-
-						HashSet<string> DirsToDeleteFromDevice = new HashSet<string>();
-						List<string> FilesToDeleteFromDevice = new List<string>();
-
-						using (var reader = new StringReader(CommandResult.Output))
-						{
-							string ProjectSaved = Params.ShortProjectName + "/Saved";
-							string ProjectConfig = Params.ShortProjectName + "/Config";
-							const string EngineSaved = "Engine/Saved"; // is this safe to use, or should we use SC.EngineRoot.GetDirectoryName()?
-							const string EngineConfig = "Engine/Config";
-							Logger.LogWarning("Excluding {ProjectSaved} {ProjectConfig} {EngineSaved} {EngineConfig} from clean during deployment.", ProjectSaved, ProjectConfig, EngineSaved, EngineConfig);
-
-							string CurrentDir = "";
-							bool SkipFiles = false;
-							for (string Line = reader.ReadLine(); Line != null; Line = reader.ReadLine())
-							{
-								if (String.IsNullOrWhiteSpace(Line))
-								{
-									continue; // ignore blank lines
-								}
+							// listing output is of the form
+							// [Samsung]                 [Google]
+							//
+							// RemoteDir/RestOfPath:     RemoteDir/RestOfPath:
+							// - File1.png               File1.png
+							// - File2.txt               File2.txt
+							// d SubDir1                 SubDir1/
+							// d SubDir2                 Subdir2/
+							//
+							// RemoteDir/RestOfPath/SubDir1:
 
-								if (Line.EndsWith(":"))
-								{
-									// RemoteDir/RestOfPath:
-									//      keep ^--------^
-									CurrentDir = Line.Substring(RemoteDir.Length + 1, Math.Max(0, Line.Length - RemoteDir.Length - 2));
-									// Max is there for the case of base "RemoteDir:" --> ""
+							HashSet<string> DirsToDeleteFromDevice = new HashSet<string>();
+							List<string> FilesToDeleteFromDevice = new List<string>();
 
-									// We want to keep config & logs between deployments.
-									if (CurrentDir.StartsWith(ProjectSaved) || CurrentDir.StartsWith(ProjectConfig) || CurrentDir.StartsWith(EngineSaved) || CurrentDir.StartsWith(EngineConfig))
-									{
-										SkipFiles = true;
-										continue;
-									}
+							using (var reader = new StringReader(CommandResult.Output))
+							{
+								string ProjectSaved = Params.ShortProjectName + "/Saved";
+								string ProjectConfig = Params.ShortProjectName + "/Config";
+								const string EngineSaved = "Engine/Saved"; // is this safe to use, or should we use SC.EngineRoot.GetDirectoryName()?
+								const string EngineConfig = "Engine/Config";
+								Logger.LogWarning("Excluding {ProjectSaved} {ProjectConfig} {EngineSaved} {EngineConfig} from clean during deployment.", ProjectSaved, ProjectConfig, EngineSaved, EngineConfig);
 
-									bool DirExistsInStagingArea = Directory.Exists(Path.Combine(SC.StageDirectory.FullName, CurrentDir));
-									if (DirExistsInStagingArea)
-									{
-										SkipFiles = false;
-									}
-									else
-									{
-										// delete directory from device
-										SkipFiles = true;
-										DirsToDeleteFromDevice.Add(CurrentDir);
-									}
-								}
-								else
+								string CurrentDir = "";
+								bool SkipFiles = false;
+								for (string Line = reader.ReadLine(); Line != null; Line = reader.ReadLine())
 								{
-									if (SkipFiles)
+									if (String.IsNullOrWhiteSpace(Line))
 									{
-										continue;
+										continue; // ignore blank lines
 									}
 
-									string FileName = GetFileNameFromListing(Line);
-									if (FileName != null)
+									if (Line.EndsWith(":"))
 									{
-										bool FileExistsInStagingArea = File.Exists(Path.Combine(SC.StageDirectory.FullName, CurrentDir, FileName));
-										if (FileExistsInStagingArea)
+										// RemoteDir/RestOfPath:
+										//      keep ^--------^
+										CurrentDir = Line.Substring(RemoteDir.Length + 1, Math.Max(0, Line.Length - RemoteDir.Length - 2));
+										// Max is there for the case of base "RemoteDir:" --> ""
+
+										// We want to keep config & logs between deployments.
+										if (CurrentDir.StartsWith(ProjectSaved) || CurrentDir.StartsWith(ProjectConfig) || CurrentDir.StartsWith(EngineSaved) || CurrentDir.StartsWith(EngineConfig))
 										{
-											// keep or overwrite
+											SkipFiles = true;
+											continue;
 										}
-										else if (FileName == "APKFileStamp.txt")
+
+										bool DirExistsInStagingArea = Directory.Exists(Path.Combine(SC.StageDirectory.FullName, CurrentDir));
+										if (DirExistsInStagingArea)
 										{
-											// keep it
+											SkipFiles = false;
 										}
 										else
 										{
-											// delete file from device
-											string FilePath = CurrentDir.Length == 0 ? FileName : (CurrentDir + "/" + FileName); // use / for Android target, no matter the development system
-											Logger.LogWarning("Deleting {FilePath} from device; not found in staging area", FilePath);
-											FilesToDeleteFromDevice.Add(FilePath);
+											// delete directory from device
+											SkipFiles = true;
+											DirsToDeleteFromDevice.Add(CurrentDir);
 										}
 									}
-									// We ignore subdirs here as each will have its own "RemoteDir/CurrentDir/SubDir:" entry.
+									else
+									{
+										if (SkipFiles)
+										{
+											continue;
+										}
+
+										string FileName = GetFileNameFromListing(Line);
+										if (FileName != null)
+										{
+											bool FileExistsInStagingArea = File.Exists(Path.Combine(SC.StageDirectory.FullName, CurrentDir, FileName));
+											if (FileExistsInStagingArea)
+											{
+												// keep or overwrite
+											}
+											else if (FileName == "APKFileStamp.txt")
+											{
+												// keep it
+											}
+											else
+											{
+												// delete file from device
+												string FilePath = CurrentDir.Length == 0 ? FileName : (CurrentDir + "/" + FileName); // use / for Android target, no matter the development system
+												Logger.LogWarning("Deleting {FilePath} from device; not found in staging area", FilePath);
+												FilesToDeleteFromDevice.Add(FilePath);
+											}
+										}
+										// We ignore subdirs here as each will have its own "RemoteDir/CurrentDir/SubDir:" entry.
+									}
+								}
+							}
+
+							// delete directories
+							foreach (var DirToDelete in DirsToDeleteFromDevice)
+							{
+								// if a whole tree is to be deleted, don't spend extra commands deleting its branches
+								int FinalSlash = DirToDelete.LastIndexOf('/');
+								string ParentDir = FinalSlash >= 0 ? DirToDelete.Substring(0, FinalSlash) : "";
+								bool ParentMarkedForDeletion = DirsToDeleteFromDevice.Contains(ParentDir);
+								if (!ParentMarkedForDeletion)
+								{
+									Logger.LogWarning("Deleting {DirToDelete} and its contents from device; not found in staging area", DirToDelete);
+									RunAdbCommand(Params, DeviceName, "shell rm -r " + RemoteDir + "/" + DirToDelete);
 								}
 							}
+
+							// delete loose files
+							if (FilesToDeleteFromDevice.Count > 0)
+							{
+								// delete all stray files with one command
+								Commandline = String.Format("shell cd {0}; rm ", RemoteDir);
+								RunAdbCommand(Params, DeviceName, Commandline + String.Join(" ", FilesToDeleteFromDevice));
+							}
 						}
+					}
+				}
+				else
+				{
+					// make sure device is at a clean state
+					RunAdbCommand(Params, DeviceName, "shell rm -r " + RemoteDir);
+
+					// Copy UFS files..
+					string[] Files = Directory.GetFiles(SC.StageDirectory.FullName, "*", SearchOption.AllDirectories);
+					System.Array.Sort(Files);
 
-						// delete directories
-						foreach (var DirToDelete in DirsToDeleteFromDevice)
+					// Find all the files we exclude from copying. And include
+					// the directories we need to individually copy.
+					HashSet<string> ExcludedFiles = new HashSet<string>();
+					SortedSet<string> IndividualCopyDirectories
+						= new SortedSet<string>((IComparer<string>)new LongestFirst());
+					foreach (string Filename in Files)
+					{
+						bool Exclude = false;
+						// Don't push the apk, we install it
+						Exclude |= Path.GetExtension(Filename).Equals(".apk", StringComparison.InvariantCultureIgnoreCase);
+						// For excluded files we add the parent dirs to our
+						// tracking of stuff to individually copy.
+						if (Exclude)
 						{
-							// if a whole tree is to be deleted, don't spend extra commands deleting its branches
-							int FinalSlash = DirToDelete.LastIndexOf('/');
-							string ParentDir = FinalSlash >= 0 ? DirToDelete.Substring(0, FinalSlash) : "";
-							bool ParentMarkedForDeletion = DirsToDeleteFromDevice.Contains(ParentDir);
-							if (!ParentMarkedForDeletion)
+							ExcludedFiles.Add(Filename);
+							// We include all directories up to the stage root in having
+							// to individually copy the files.
+							for (string FileDirectory = Path.GetDirectoryName(Filename);
+								!FileDirectory.Equals(SC.StageDirectory);
+								FileDirectory = Path.GetDirectoryName(FileDirectory))
 							{
-								Logger.LogWarning("Deleting {DirToDelete} and its contents from device; not found in staging area", DirToDelete);
-								RunAdbCommand(Params, DeviceName, "shell rm -r " + RemoteDir + "/" + DirToDelete);
+								if (!IndividualCopyDirectories.Contains(FileDirectory))
+								{
+									IndividualCopyDirectories.Add(FileDirectory);
+								}
+							}
+							if (!IndividualCopyDirectories.Contains(SC.StageDirectory.FullName))
+							{
+								IndividualCopyDirectories.Add(SC.StageDirectory.FullName);
 							}
 						}
+					}
 
-						// delete loose files
-						if (FilesToDeleteFromDevice.Count > 0)
+					// The directories are sorted above in "deepest" first. We can
+					// therefore start copying those individual dirs which will
+					// recreate the tree. As the subtrees will get copied at each
+					// possible individual level.
+					foreach (string DirectoryName in IndividualCopyDirectories)
+					{
+						string[] Entries
+							= Directory.GetFileSystemEntries(DirectoryName, "*", SearchOption.TopDirectoryOnly);
+						foreach (string Entry in Entries)
 						{
-							// delete all stray files with one command
-							Commandline = String.Format("shell cd {0}; rm ", RemoteDir);
-							RunAdbCommand(Params, DeviceName, Commandline + String.Join(" ", FilesToDeleteFromDevice));
-						}
+							// We avoid excluded files and the individual copy dirs
+							// (the individual copy dirs will get handled as we iterate).
+							if (ExcludedFiles.Contains(Entry) || IndividualCopyDirectories.Contains(Entry))
+							{
+								continue;
+							}
+							else
+							{
+								EntriesToDeploy.Add(Entry);
+							}
 						}
 					}
-                }
-                else
-                {
-                    // make sure device is at a clean state
-                    RunAdbCommand(Params, DeviceName, "shell rm -r " + RemoteDir);
-
-                    // Copy UFS files..
-                    string[] Files = Directory.GetFiles(SC.StageDirectory.FullName, "*", SearchOption.AllDirectories);
-                    System.Array.Sort(Files);
-
-                    // Find all the files we exclude from copying. And include
-                    // the directories we need to individually copy.
-                    HashSet<string> ExcludedFiles = new HashSet<string>();
-                    SortedSet<string> IndividualCopyDirectories
-                        = new SortedSet<string>((IComparer<string>)new LongestFirst());
-                    foreach (string Filename in Files)
-                    {
-                        bool Exclude = false;
-                        // Don't push the apk, we install it
-                        Exclude |= Path.GetExtension(Filename).Equals(".apk", StringComparison.InvariantCultureIgnoreCase);
-                        // For excluded files we add the parent dirs to our
-                        // tracking of stuff to individually copy.
-                        if (Exclude)
-                        {
-                            ExcludedFiles.Add(Filename);
-                            // We include all directories up to the stage root in having
-                            // to individually copy the files.
-                            for (string FileDirectory = Path.GetDirectoryName(Filename);
-                                !FileDirectory.Equals(SC.StageDirectory);
-                                FileDirectory = Path.GetDirectoryName(FileDirectory))
-                            {
-                                if (!IndividualCopyDirectories.Contains(FileDirectory))
-                                {
-                                    IndividualCopyDirectories.Add(FileDirectory);
-                                }
-                            }
-                            if (!IndividualCopyDirectories.Contains(SC.StageDirectory.FullName))
-                            {
-                                IndividualCopyDirectories.Add(SC.StageDirectory.FullName);
-                            }
-                        }
-                    }
-
-                    // The directories are sorted above in "deepest" first. We can
-                    // therefore start copying those individual dirs which will
-                    // recreate the tree. As the subtrees will get copied at each
-                    // possible individual level.
-                    foreach (string DirectoryName in IndividualCopyDirectories)
-                    {
-                        string[] Entries
-                            = Directory.GetFileSystemEntries(DirectoryName, "*", SearchOption.TopDirectoryOnly);
-                        foreach (string Entry in Entries)
-                        {
-                            // We avoid excluded files and the individual copy dirs
-                            // (the individual copy dirs will get handled as we iterate).
-                            if (ExcludedFiles.Contains(Entry) || IndividualCopyDirectories.Contains(Entry))
-                            {
-                                continue;
-                            }
-                            else
-                            {
-                                EntriesToDeploy.Add(Entry);
-                            }
-                        }
-                    }
 
-                    if (EntriesToDeploy.Count == 0)
-                    {
-                        EntriesToDeploy.Add(SC.StageDirectory.FullName);
-                    }
-                }
+					if (EntriesToDeploy.Count == 0)
+					{
+						EntriesToDeploy.Add(SC.StageDirectory.FullName);
+					}
+				}
 
-                // We now have a minimal set of file & dir entries we need
-                // to deploy. Files we deploy will get individually copied
-                // and dirs will get the tree copies by default (that's
-                // what ADB does).
-                HashSet<IProcessResult> DeployCommands = new HashSet<IProcessResult>();
-                foreach (string Entry in EntriesToDeploy)
-                {
-                    string FinalRemoteDir = RemoteDir;
-                    string RemotePath = Entry.Replace(SC.StageDirectory.FullName, FinalRemoteDir).Replace("\\", "/");
-                    string Commandline = string.Format("{0} \"{1}\" \"{2}\"", BaseCommandline, Entry, RemotePath);
-                    // We run deploy commands in parallel to maximize the connection
-                    // throughput.
-                    DeployCommands.Add(
-                        RunAdbCommand(Params, DeviceName, Commandline, null,
-                            ERunOptions.Default | ERunOptions.NoWaitForExit));
-                    // But we limit the parallel commands to avoid overwhelming
-                    // memory resources.
-                    if (DeployCommands.Count == DeployMaxParallelCommands)
-                    {
-                        while (DeployCommands.Count > DeployMaxParallelCommands / 2)
-                        {
-                            Thread.Sleep(1);
-                            DeployCommands.RemoveWhere(
-                                delegate (IProcessResult r)
-                                {
-                                    return r.HasExited;
-                                });
-                        }
-                    }
-                }
-                foreach (IProcessResult deploy_result in DeployCommands)
-                {
-                    deploy_result.WaitForExit();
-                }
+				// We now have a minimal set of file & dir entries we need
+				// to deploy. Files we deploy will get individually copied
+				// and dirs will get the tree copies by default (that's
+				// what ADB does).
+				HashSet<IProcessResult> DeployCommands = new HashSet<IProcessResult>();
+				foreach (string Entry in EntriesToDeploy)
+				{
+					string FinalRemoteDir = RemoteDir;
+					string RemotePath = Entry.Replace(SC.StageDirectory.FullName, FinalRemoteDir).Replace("\\", "/");
+					string Commandline = string.Format("{0} \"{1}\" \"{2}\"", BaseCommandline, Entry, RemotePath);
+					// We run deploy commands in parallel to maximize the connection
+					// throughput.
+					DeployCommands.Add(
+						RunAdbCommand(Params, DeviceName, Commandline, null,
+							ERunOptions.Default | ERunOptions.NoWaitForExit));
+					// But we limit the parallel commands to avoid overwhelming
+					// memory resources.
+					if (DeployCommands.Count == DeployMaxParallelCommands)
+					{
+						while (DeployCommands.Count > DeployMaxParallelCommands / 2)
+						{
+							Thread.Sleep(1);
+							DeployCommands.RemoveWhere(
+								delegate (IProcessResult r)
+								{
+									return r.HasExited;
+								});
+						}
+					}
+				}
+				foreach (IProcessResult deploy_result in DeployCommands)
+				{
+					deploy_result.WaitForExit();
+				}
 
-                // delete the .obb file, since it will cause nothing we just deployed to be used
-                RunAdbCommand(Params, DeviceName, "shell rm " + DeviceObbName);
+				// delete the .obb file, since it will cause nothing we just deployed to be used
+				RunAdbCommand(Params, DeviceName, "shell rm " + DeviceObbName);
 				RunAdbCommand(Params, DeviceName, "shell rm " + DevicePatchName);
-				RunAdbCommand(Params, DeviceName, "shell rm " + DeviceOverflow1Name);
-				RunAdbCommand(Params, DeviceName, "shell rm " + DeviceOverflow2Name);
-				RunAdbCommand(Params, DeviceName, "shell rm " + DeviceOverflow3Name);
-				RunAdbCommand(Params, DeviceName, "shell rm " + DeviceOverflow4Name);
+
+				// delete existing overflow files on device, one more than the number we have to ensure we do not try to mount an old extra one left on device
+				for (int Index = 1; Index <= AllowOverflowOBBLimit; Index++)
+				{
+					string DeviceOverflowName = StorageLocation + "/" + GetDeviceOverflowName(ApkName, SC, Index);
+					RunAdbCommand(Params, DeviceName, "shell rm " + DeviceOverflowName);
+
+					string OverflowPath = Path.Combine(SC.StageDirectory.FullName, GetFinalOverflowName(ApkName, SC, Index));
+					if (!File.Exists(OverflowPath))
+					{
+						break;
+					}
+				}
 			}
 			else if (SC.Archive)
             {
@@ -3510,48 +3487,23 @@ public class AndroidPlatform : Platform
 					RunAdbCommand(Params, DeviceName, Commandline);
 				}
 
-				// deploy the overflow1 if there is one
-				string Overflow1Path = Path.Combine(SC.StageDirectory.FullName, GetFinalOverflowName(ApkName, SC, 1));
-				if (File.Exists(Overflow1Path))
-				{
-					// cache some strings
-					string BaseCommandline = "push";
-
-					string Commandline = string.Format("{0} \"{1}\" \"{2}\"", BaseCommandline, Overflow1Path, DeviceOverflow1Name);
-					RunAdbCommand(Params, DeviceName, Commandline);
-				}
-
-				// deploy the overflow2 if there is one
-				string Overflow2Path = Path.Combine(SC.StageDirectory.FullName, GetFinalOverflowName(ApkName, SC, 2));
-				if (File.Exists(Overflow2Path))
-				{
-					// cache some strings
-					string BaseCommandline = "push";
-
-					string Commandline = string.Format("{0} \"{1}\" \"{2}\"", BaseCommandline, Overflow2Path, DeviceOverflow2Name);
-					RunAdbCommand(Params, DeviceName, Commandline);
-				}
-
-				// deploy the overflow3 if there is one
-				string Overflow3Path = Path.Combine(SC.StageDirectory.FullName, GetFinalOverflowName(ApkName, SC, 3));
-				if (File.Exists(Overflow3Path))
-				{
-					// cache some strings
-					string BaseCommandline = "push";
-
-					string Commandline = string.Format("{0} \"{1}\" \"{2}\"", BaseCommandline, Overflow3Path, DeviceOverflow3Name);
-					RunAdbCommand(Params, DeviceName, Commandline);
-				}
-
-				// deploy the overflow4 if there is one
-				string Overflow4Path = Path.Combine(SC.StageDirectory.FullName, GetFinalOverflowName(ApkName, SC, 4));
-				if (File.Exists(Overflow4Path))
+				// deploy the overflows
+				for (int Index = 1; Index <= AllowOverflowOBBLimit; Index++)
 				{
-					// cache some strings
-					string BaseCommandline = "push";
+					string OverflowPath = Path.Combine(SC.StageDirectory.FullName, GetFinalOverflowName(ApkName, SC, Index));
+					if (File.Exists(OverflowPath))
+					{
+						// cache some strings
+						string BaseCommandline = "push";
+						string DeviceOverflowName = StorageLocation + "/" + GetDeviceOverflowName(ApkName, SC, Index);
 
-					string Commandline = string.Format("{0} \"{1}\" \"{2}\"", BaseCommandline, Overflow4Path, DeviceOverflow4Name);
-					RunAdbCommand(Params, DeviceName, Commandline);
+						string Commandline = string.Format("{0} \"{1}\" \"{2}\"", BaseCommandline, OverflowPath, DeviceOverflowName);
+						RunAdbCommand(Params, DeviceName, Commandline);
+					}
+					else
+					{
+						break;
+					}
 				}
 			}
 			else
diff --git a/Engine/Source/Runtime/Core/Private/Android/AndroidPlatformFile.cpp b/Engine/Source/Runtime/Core/Private/Android/AndroidPlatformFile.cpp
--- a/Engine/Source/Runtime/Core/Private/Android/AndroidPlatformFile.cpp
+++ b/Engine/Source/Runtime/Core/Private/Android/AndroidPlatformFile.cpp
@@ -94,14 +94,6 @@ FString GOBBFilePathBase;
 FString GOBBMainFilePath;
 // Obb Patch filepath
 FString GOBBPatchFilePath;
-// Obb Overflow1 filepath
-FString GOBBOverflow1FilePath;
-// Obb Overflow2 filepath
-FString GOBBOverflow2FilePath;
-// Obb Overflow3 filepath
-FString GOBBOverflow3FilePath;
-// Obb Overflow2 filepath
-FString GOBBOverflow4FilePath;
 // Internal File Direcory Path (for application) - setup during load
 FString GInternalFilePath;
 // External File Direcory Path (for application) - setup during load
@@ -1250,87 +1242,32 @@ public:
 
 			// Only check for overflow files if we found a patch file
 			if (bHavePatch)
-			{
-				FString Overflow1OBBName = FString::Printf(TEXT("overflow1.%d.%s.obb"), GAndroidPackageVersion, *GPackageName);
-				FString Overflow2OBBName = FString::Printf(TEXT("overflow2.%d.%s.obb"), GAndroidPackageVersion, *GPackageName);
-				FString Overflow3OBBName = FString::Printf(TEXT("overflow3.%d.%s.obb"), GAndroidPackageVersion, *GPackageName);
-				FString Overflow4OBBName = FString::Printf(TEXT("overflow4.%d.%s.obb"), GAndroidPackageVersion, *GPackageName);
-
-				if (!GOBBOverflow1FilePath.IsEmpty() && FileExists(*GOBBOverflow1FilePath, true))
-				{
-					MountOBB(*GOBBOverflow1FilePath);
-					FPlatformMisc::LowLevelOutputDebugStringf(TEXT("Mounted overflow1 OBB: %s"), *GOBBOverflow1FilePath);
-				}
-				else if (FileExists(*(OBBDir1 / Overflow1OBBName), true))
-				{
-					GOBBOverflow1FilePath = OBBDir1 / Overflow1OBBName;
-					MountOBB(*GOBBOverflow1FilePath);
-					FPlatformMisc::LowLevelOutputDebugStringf(TEXT("Mounted overflow1 OBB: %s"), *GOBBOverflow1FilePath);
-				}
-				else if (FileExists(*(OBBDir2 / Overflow1OBBName), true))
-				{
-					GOBBOverflow1FilePath = OBBDir2 / Overflow1OBBName;
-					MountOBB(*GOBBOverflow1FilePath);
-					FPlatformMisc::LowLevelOutputDebugStringf(TEXT("Mounted overflow1 OBB: %s"), *GOBBOverflow1FilePath);
-				}
-
-				if (!GOBBOverflow2FilePath.IsEmpty() && FileExists(*GOBBOverflow2FilePath, true))
+			{	
+				for (int32 OverflowIndex = 1; ; OverflowIndex++)
 				{
-					MountOBB(*GOBBOverflow2FilePath);
-					FPlatformMisc::LowLevelOutputDebugStringf(TEXT("Mounted overflow2 OBB: %s"), *GOBBOverflow2FilePath);
-				}
-				else if (FileExists(*(OBBDir1 / Overflow2OBBName), true))
-				{
-					GOBBOverflow2FilePath = OBBDir1 / Overflow2OBBName;
-					MountOBB(*GOBBOverflow2FilePath);
-					FPlatformMisc::LowLevelOutputDebugStringf(TEXT("Mounted overflow2 OBB: %s"), *GOBBOverflow2FilePath);
-				}
-				else if (FileExists(*(OBBDir2 / Overflow2OBBName), true))
-				{
-					GOBBOverflow2FilePath = OBBDir2 / Overflow2OBBName;
-					MountOBB(*GOBBOverflow2FilePath);
-					FPlatformMisc::LowLevelOutputDebugStringf(TEXT("Mounted overflow2 OBB: %s"), *GOBBOverflow2FilePath);
-				}
+					FString OverflowOBBName = FString::Printf(TEXT("overflow%d.%d.%s.obb"), OverflowIndex, GAndroidPackageVersion, *GPackageName);
 
-				if (!GOBBOverflow3FilePath.IsEmpty() && FileExists(*GOBBOverflow3FilePath, true))
-				{
-					MountOBB(*GOBBOverflow3FilePath);
-					FPlatformMisc::LowLevelOutputDebugStringf(TEXT("Mounted overflow3 OBB: %s"), *GOBBOverflow3FilePath);
-				}
-				else if (FileExists(*(OBBDir1 / Overflow3OBBName), true))
-				{
-					GOBBOverflow3FilePath = OBBDir1 / Overflow3OBBName;
-					MountOBB(*GOBBOverflow3FilePath);
-					FPlatformMisc::LowLevelOutputDebugStringf(TEXT("Mounted overflow3 OBB: %s"), *GOBBOverflow3FilePath);
-				}
-				else if (FileExists(*(OBBDir2 / Overflow3OBBName), true))
-				{
-					GOBBOverflow3FilePath = OBBDir2 / Overflow3OBBName;
-					MountOBB(*GOBBOverflow3FilePath);
-					FPlatformMisc::LowLevelOutputDebugStringf(TEXT("Mounted overflow2 OBB: %s"), *GOBBOverflow3FilePath);
-				}
-
-				if (!GOBBOverflow4FilePath.IsEmpty() && FileExists(*GOBBOverflow4FilePath, true))
-				{
-					MountOBB(*GOBBOverflow4FilePath);
-					FPlatformMisc::LowLevelOutputDebugStringf(TEXT("Mounted overflow4 OBB: %s"), *GOBBOverflow4FilePath);
-				}
-				else if (FileExists(*(OBBDir1 / Overflow4OBBName), true))
-				{
-					GOBBOverflow4FilePath = OBBDir1 / Overflow4OBBName;
-					MountOBB(*GOBBOverflow4FilePath);
-					FPlatformMisc::LowLevelOutputDebugStringf(TEXT("Mounted overflow4 OBB: %s"), *GOBBOverflow4FilePath);
-				}
-				else if (FileExists(*(OBBDir2 / Overflow4OBBName), true))
-				{
-					GOBBOverflow4FilePath = OBBDir2 / Overflow4OBBName;
-					MountOBB(*GOBBOverflow4FilePath);
-					FPlatformMisc::LowLevelOutputDebugStringf(TEXT("Mounted overflow4 OBB: %s"), *GOBBOverflow4FilePath);
+					if (FileExists(*(OBBDir1 / OverflowOBBName), true))
+					{
+						FString OBBOverflowFilePath = OBBDir1 / OverflowOBBName;
+						MountOBB(*OBBOverflowFilePath);
+						FPlatformMisc::LowLevelOutputDebugStringf(TEXT("Mounted overflow%d OBB: %s"), OverflowIndex, *OBBOverflowFilePath);
+					}
+					else if (FileExists(*(OBBDir2 / OverflowOBBName), true))
+					{
+						FString OBBOverflowFilePath = OBBDir2 / OverflowOBBName;
+						MountOBB(*OBBOverflowFilePath);
+						FPlatformMisc::LowLevelOutputDebugStringf(TEXT("Mounted overflow%d OBB: %s"), OverflowIndex, *OBBOverflowFilePath);
+					}
+					else
+					{
+						break;
+					}
 				}
-
 			}
 		}
 
+
 		// make sure the base path directory exists (UnrealGame and UnrealGame/ProjectName)
 		FString FileBaseDir = GFilePathBase + FString(FILEBASE_DIRECTORY);
 		mkdir(TCHAR_TO_UTF8(*FileBaseDir), 0777);
diff --git a/Engine/Source/Runtime/Launch/Private/Android/AndroidJNI.cpp b/Engine/Source/Runtime/Launch/Private/Android/AndroidJNI.cpp
--- a/Engine/Source/Runtime/Launch/Private/Android/AndroidJNI.cpp
+++ b/Engine/Source/Runtime/Launch/Private/Android/AndroidJNI.cpp
@@ -45,10 +45,6 @@ extern FString GFontPathBase;
 extern bool GOBBinAPK;
 extern bool GOverrideAndroidLogDir;
 extern FString GOBBFilePathBase;
-extern FString GOBBMainFilePath;
-extern FString GOBBPatchFilePath;
-extern FString GOBBOverflow1FilePath;
-extern FString GOBBOverflow2FilePath;
 extern FString GAPKFilename;
 
 FOnActivityResult FJavaWrapper::OnActivityResultDelegate;
@@ -1671,10 +1667,6 @@ JNIEXPORT jint JNI_OnLoad(JavaVM* InJavaVM, void* InReserved)
 	// Copy that somewhere safe
 	GFilePathBase = FJavaHelper::FStringFromLocalRef(Env, (jstring)Env->CallObjectMethod(*externalStoragePath, getFilePath, nullptr));
 	GOBBFilePathBase = GFilePathBase;
-	GOBBMainFilePath = TEXT("");
-	GOBBPatchFilePath = TEXT("");
-	GOBBOverflow1FilePath = TEXT("");
-	GOBBOverflow2FilePath = TEXT("");
 
 	// then release...
 	FPlatformMisc::LowLevelOutputDebugStringf(TEXT("Path found as '%s'\n"), *GFilePathBase);
@@ -1698,10 +1690,7 @@ JNIEXPORT jint JNI_OnLoad(JavaVM* InJavaVM, void* InReserved)
 //This function is declared in the Java-defined class, GameActivity.java: "public native void naativeSetObbFilePaths();"
 JNI_METHOD void Java_com_epicgames_unreal_GameActivity_nativeSetObbFilePaths(JNIEnv* jenv, jobject thiz, jstring OBBMainFilePath, jstring OBBPatchFilePath, jstring OBBOverflow1FilePath, jstring OBBOverflow2FilePath)
 {
-	GOBBMainFilePath = FJavaHelper::FStringFromParam(jenv, OBBMainFilePath);
-	GOBBPatchFilePath = FJavaHelper::FStringFromParam(jenv, OBBPatchFilePath);
-	GOBBOverflow1FilePath = FJavaHelper::FStringFromParam(jenv, OBBOverflow1FilePath);
-	GOBBOverflow2FilePath = FJavaHelper::FStringFromParam(jenv, OBBOverflow2FilePath);
+	// no longer supported
 }
 
 //This function is declared in the Java-defined class, GameActivity.java: "public native void nativeSetGlobalActivity();"
