Fix reversed logic for missing uproject file vs missing projectstore file.

diff --git a/Engine/Source/Programs/AutomationTool/BuildGraph/Tasks/ZenExportSnapshotTask.cs b/Engine/Source/Programs/AutomationTool/BuildGraph/Tasks/ZenExportSnapshotTask.cs
--- a/Engine/Source/Programs/AutomationTool/BuildGraph/Tasks/ZenExportSnapshotTask.cs
+++ b/Engine/Source/Programs/AutomationTool/BuildGraph/Tasks/ZenExportSnapshotTask.cs
@@ -128,7 +128,7 @@ namespace AutomationTool.Tasks
 			FileReference ProjectFile = Parameters.Project;
 			if(!FileReference.Exists(ProjectFile))
 			{
-				return Task.CompletedTask;
+				throw new AutomationException("Missing project file - {0}", ProjectFile.FullName);
 			}
 
 			List<ExportSourceData> ExportSources = new List<ExportSourceData>();
@@ -143,7 +143,7 @@ namespace AutomationTool.Tasks
 				FileReference ProjectStoreFile = FileReference.Combine(PlatformCookedDirectory, ".projectstore");
 				if (!FileReference.Exists(ProjectStoreFile))
 				{
-					throw new AutomationException("Project store file not found ({0})", ProjectStoreFile.FullName);
+					continue;
 				}
 
 				byte[] ProjectStoreData = File.ReadAllBytes(ProjectStoreFile.FullName);
