CitySample - Fix csvs not being found on some platforms.

+ Use PathUtils.FindRelevantPath to search for the FPSChartStats folder within the artifacts dir. Then copy out the newest csv each iteration.
+ Remove deletion of saved dir between iterations as it's no longer needed since we grab the latest csv file.
+ Pass in the build name as it's not reliably set in Gauntlet on all platforms.

#rnx

diff --git a/Engine/Source/Programs/AutomationTool/Gauntlet/Framework/Utils/Gauntlet.PathUtils.cs b/Engine/Source/Programs/AutomationTool/Gauntlet/Framework/Utils/Gauntlet.PathUtils.cs
--- a/Engine/Source/Programs/AutomationTool/Gauntlet/Framework/Utils/Gauntlet.PathUtils.cs
+++ b/Engine/Source/Programs/AutomationTool/Gauntlet/Framework/Utils/Gauntlet.PathUtils.cs
@@ -8,6 +8,31 @@ namespace Gauntlet
 {
 	public class PathUtils
 	{
+		/// <summary>
+		/// Converts a path to the long path format if it's not already.
+		/// </summary>
+		/// <param name="InPath">The path to convert.</param>
+		/// <returns>The converted path.</returns>
+		public static string MakeLongPath(string InPath)
+		{
+			string LongPathPrefix = @"\\?\";
+			if (!InPath.StartsWith(LongPathPrefix))
+			{
+				InPath = LongPathPrefix + InPath;
+			}
+			return InPath;
+		}
+
+		/// <summary>
+		/// Similar to Path.Combine but creates a long path instead.
+		/// </summary>
+		/// <param name="Paths">The parts of the path to combine.</param>
+		/// <returns>The full long path.</returns>
+		public static string CombineLongPath(params string[] Paths)
+		{
+			return MakeLongPath(Path.Combine(Paths));
+		}
+
 		/// <summary>
 		/// Find a path that includes specific directories in the given order, but could also include unknown intermediate directories.
 		/// </summary>
