Fix logging issue

diff --git a/Engine/Source/Programs/AutomationTool/Android/AndroidPlatform.Automation.cs b/Engine/Source/Programs/AutomationTool/Android/AndroidPlatform.Automation.cs
--- a/Engine/Source/Programs/AutomationTool/Android/AndroidPlatform.Automation.cs
+++ b/Engine/Source/Programs/AutomationTool/Android/AndroidPlatform.Automation.cs
@@ -1234,7 +1234,7 @@ public class AndroidPlatform : Platform
 								else
 								{
 									Logger.LogInformation("{Text}", "Failed to add " + FileRef.FullName + " to a new overflow as it is bigger than the allowed OBB size ");
-									throw new AutomationException(ExitCode.Error_AndroidOBBError, "Stage Failed. Could not add {0} ({1} bytes) to an OBB as OBBs are limited to {1} bytes.", FileRequirements, FileRef.FullName, OBBSizeAllowed);
+									throw new AutomationException(ExitCode.Error_AndroidOBBError, "Stage Failed. Could not add {1} ({0} bytes) to an OBB as OBBs are limited to {2} bytes.", FileRequirements, FileRef.FullName, OBBSizeAllowed);
 								}
 							}
 							else
