Add null check

diff --git a/Engine/Source/Programs/AutomationTool/AutomationUtils/HelpUtils.cs b/Engine/Source/Programs/AutomationTool/AutomationUtils/HelpUtils.cs
--- a/Engine/Source/Programs/AutomationTool/AutomationUtils/HelpUtils.cs
+++ b/Engine/Source/Programs/AutomationTool/AutomationUtils/HelpUtils.cs
@@ -137,7 +137,10 @@ namespace AutomationTool
 					ChoicesStr += ". Choices=[";
 					foreach (var Value in ValidChoices)
 					{
-						ChoicesStr += Value.ToString() + (MultiSelectSeparator != null ? MultiSelectSeparator : "|");
+						if (Value != null)
+						{
+							ChoicesStr += Value.ToString() + (MultiSelectSeparator != null ? MultiSelectSeparator : "|");
+						}
 					}
 					ChoicesStr = ChoicesStr.Remove(ChoicesStr.Length - 1);
 					ChoicesStr += "]";
