Build snapshot functionality:
Fix bug with the writing of a snapshot descriptor when the destination directory doesn't exist.
Ensure we snapshot loose files into the oplog before performing the export.

diff --git a/Engine/Source/Programs/AutomationTool/BuildGraph/Tasks/ZenExportSnapshotTask.cs b/Engine/Source/Programs/AutomationTool/BuildGraph/Tasks/ZenExportSnapshotTask.cs
--- a/Engine/Source/Programs/AutomationTool/BuildGraph/Tasks/ZenExportSnapshotTask.cs
+++ b/Engine/Source/Programs/AutomationTool/BuildGraph/Tasks/ZenExportSnapshotTask.cs
@@ -174,9 +174,11 @@ namespace AutomationTool.Tasks
 				ZenExe = ResolveFile(String.Format("Engine/Binaries/{0}/zen", HostPlatform.Current.HostEditorPlatform.ToString()));
 			}
 
-			// Format the command line
-			StringBuilder CommandLine = new StringBuilder();
-			CommandLine.AppendFormat("oplog-export");
+			// Format the command lines
+			StringBuilder OplogSnapshotCommandline = new StringBuilder();
+			OplogSnapshotCommandline.AppendFormat("oplog-snapshot");
+			StringBuilder OplogExportCommandline = new StringBuilder();
+			OplogExportCommandline.AppendFormat("oplog-export");
 
 			switch (DestinationStorageType)
 			{
@@ -201,37 +203,38 @@ namespace AutomationTool.Tasks
 						BucketName = ProjectNameAsBucketName;
 					}
 
-					CommandLine.AppendFormat(" --cloud {0} --namespace {1} --bucket {2}", Parameters.DestinationCloudHost, Parameters.DestinationCloudNamespace, BucketName);
+					OplogExportCommandline.AppendFormat(" --cloud {0} --namespace {1} --bucket {2}", Parameters.DestinationCloudHost, Parameters.DestinationCloudNamespace, BucketName);
 
 					string[] ExportKeyIds = new string[ExportSources.Count];
 					int ExportIndex = 0;
 					foreach (ExportSourceData ExportSource in ExportSources)
 					{
-						StringBuilder SingleSourceCommandline = new StringBuilder(CommandLine.Length);
-						SingleSourceCommandline.Append(CommandLine);
+						String HostUrlArg = string.Format("--hosturl http://{0}:{1}", ExportSource.IsLocalHost ? "localhost" : ExportSource.HostName, ExportSource.HostPort);
+
+						StringBuilder SnapshotSingleSourceCommandline = new StringBuilder(OplogSnapshotCommandline.Length);
+						SnapshotSingleSourceCommandline.Append(OplogSnapshotCommandline);
+						SnapshotSingleSourceCommandline.AppendFormat(" {0} {1} {2}", HostUrlArg, ExportSource.ProjectId, ExportSource.OplogId);
+						Logger.LogInformation("Running '{Arg0} {Arg1}'", CommandUtils.MakePathSafeToUseWithCommandLine(ZenExe.FullName), SnapshotSingleSourceCommandline.ToString());
+						CommandUtils.RunAndLog(CommandUtils.CmdEnv, ZenExe.FullName, SnapshotSingleSourceCommandline.ToString(), Options: CommandUtils.ERunOptions.Default);
+
+						StringBuilder ExportSingleSourceCommandline = new StringBuilder(OplogExportCommandline.Length);
+						ExportSingleSourceCommandline.Append(OplogExportCommandline);
 
 						StringBuilder DestinationKeyBuilder = new StringBuilder();
 						DestinationKeyBuilder.AppendFormat("{0}.{1}", ProjectNameAsBucketName,Parameters.DestinationCloudIdentifier, ExportSource.OplogId);
 						ExportKeyIds[ExportIndex] = DestinationKeyBuilder.ToString().ToLowerInvariant();
 						IoHash DestinationKeyHash = IoHash.Compute(Encoding.UTF8.GetBytes(ExportKeyIds[ExportIndex]));
 
-						if (ExportSource.IsLocalHost)
-						{
-							SingleSourceCommandline.AppendFormat(" --hosturl http://localhost:{0}", ExportSource.HostPort);
-						}
-						else
-						{
-							SingleSourceCommandline.AppendFormat(" --hosturl http://{0}:{1}", ExportSource.HostName, ExportSource.HostPort);
-						}
-						SingleSourceCommandline.AppendFormat(" --key {0} {1} {2}", DestinationKeyHash.ToString().ToLowerInvariant(), ExportSource.ProjectId, ExportSource.OplogId);
-						Logger.LogInformation("Running '{Arg0} {Arg1}'", CommandUtils.MakePathSafeToUseWithCommandLine(ZenExe.FullName), SingleSourceCommandline.ToString());
-						CommandUtils.RunAndLog(CommandUtils.CmdEnv, ZenExe.FullName, SingleSourceCommandline.ToString(), Options: CommandUtils.ERunOptions.Default);
+						ExportSingleSourceCommandline.AppendFormat(" {0} --key {1} {2} {3}", HostUrlArg, DestinationKeyHash.ToString().ToLowerInvariant(), ExportSource.ProjectId, ExportSource.OplogId);
+						Logger.LogInformation("Running '{Arg0} {Arg1}'", CommandUtils.MakePathSafeToUseWithCommandLine(ZenExe.FullName), ExportSingleSourceCommandline.ToString());
+						CommandUtils.RunAndLog(CommandUtils.CmdEnv, ZenExe.FullName, ExportSingleSourceCommandline.ToString(), Options: CommandUtils.ERunOptions.Default);
 
 						ExportIndex = ExportIndex + 1;
 					}
 					
 					if (Parameters.SnapshotDescriptorFile != null)
 					{
+						DirectoryReference.CreateDirectory(Parameters.SnapshotDescriptorFile.Directory);
 						// Write out a snapshot descriptor
 						using (JsonWriter Writer = new JsonWriter(Parameters.SnapshotDescriptorFile))
 						{
