Avoid trying to write the snapshot descriptor file if no exports happend or report a step failure if the export failed (temporary).

diff --git a/Engine/Source/Programs/AutomationTool/BuildGraph/Tasks/ZenExportSnapshotTask.cs b/Engine/Source/Programs/AutomationTool/BuildGraph/Tasks/ZenExportSnapshotTask.cs
--- a/Engine/Source/Programs/AutomationTool/BuildGraph/Tasks/ZenExportSnapshotTask.cs
+++ b/Engine/Source/Programs/AutomationTool/BuildGraph/Tasks/ZenExportSnapshotTask.cs
@@ -227,12 +227,12 @@ namespace AutomationTool.Tasks
 
 						ExportSingleSourceCommandline.AppendFormat(" {0} --key {1} {2} {3}", HostUrlArg, DestinationKeyHash.ToString().ToLowerInvariant(), ExportSource.ProjectId, ExportSource.OplogId);
 						Logger.LogInformation("Running '{Arg0} {Arg1}'", CommandUtils.MakePathSafeToUseWithCommandLine(ZenExe.FullName), ExportSingleSourceCommandline.ToString());
-						CommandUtils.RunAndLog(CommandUtils.CmdEnv, ZenExe.FullName, ExportSingleSourceCommandline.ToString(), Options: CommandUtils.ERunOptions.Default);
+						CommandUtils.RunAndLog(CommandUtils.CmdEnv, ZenExe.FullName, ExportSingleSourceCommandline.ToString(), MaxSuccessCode: 1, Options: CommandUtils.ERunOptions.Default);
 
 						ExportIndex = ExportIndex + 1;
 					}
 					
-					if (Parameters.SnapshotDescriptorFile != null)
+					if ((Parameters.SnapshotDescriptorFile != null) && ExportSources.Any())
 					{
 						DirectoryReference.CreateDirectory(Parameters.SnapshotDescriptorFile.Directory);
 						// Write out a snapshot descriptor
