Fixes for the composition of the key when storing cloud snapshots.

diff --git a/Engine/Source/Programs/AutomationTool/BuildGraph/Tasks/ZenExportSnapshotTask.cs b/Engine/Source/Programs/AutomationTool/BuildGraph/Tasks/ZenExportSnapshotTask.cs
--- a/Engine/Source/Programs/AutomationTool/BuildGraph/Tasks/ZenExportSnapshotTask.cs
+++ b/Engine/Source/Programs/AutomationTool/BuildGraph/Tasks/ZenExportSnapshotTask.cs
@@ -221,7 +221,7 @@ namespace AutomationTool.Tasks
 						ExportSingleSourceCommandline.Append(OplogExportCommandline);
 
 						StringBuilder DestinationKeyBuilder = new StringBuilder();
-						DestinationKeyBuilder.AppendFormat("{0}.{1}", ProjectNameAsBucketName,Parameters.DestinationCloudIdentifier, ExportSource.OplogId);
+						DestinationKeyBuilder.AppendFormat("{0}.{1}.{2}", ProjectNameAsBucketName, Parameters.DestinationCloudIdentifier, ExportSource.OplogId);
 						ExportKeyIds[ExportIndex] = DestinationKeyBuilder.ToString().ToLowerInvariant();
 						IoHash DestinationKeyHash = IoHash.Compute(Encoding.UTF8.GetBytes(ExportKeyIds[ExportIndex]));
 
@@ -247,13 +247,14 @@ namespace AutomationTool.Tasks
 							{
 								Writer.WriteObjectStart();
 								
+								IoHash DestinationKeyHash = IoHash.Compute(Encoding.UTF8.GetBytes(ExportKeyIds[ExportIndex]));
+								Writer.WriteValue("name", ExportKeyIds[ExportIndex]);
+								Writer.WriteValue("type", "cloud");
+								Writer.WriteValue("targetplatform", ExportSource.TargetPlatform);
 								Writer.WriteValue("host", Parameters.DestinationCloudHost);
 								Writer.WriteValue("namespace", Parameters.DestinationCloudNamespace);
 								Writer.WriteValue("bucket", BucketName);
-								Writer.WriteValue("targetplatform", ExportSource.TargetPlatform);
-								IoHash DestinationKeyHash = IoHash.Compute(Encoding.UTF8.GetBytes(ExportKeyIds[ExportIndex]));
 								Writer.WriteValue("key", DestinationKeyHash.ToString().ToLowerInvariant());
-								Writer.WriteValue("desc", ExportKeyIds[ExportIndex]);
 
 								Writer.WriteObjectEnd();
 
