[Gauntlet] Improve clarity for logging in DevicePool creation warnings
This should yield a call stack showing the explicit reason for failure to help track down intermittent issues.

diff --git a/Engine/Source/Programs/AutomationTool/Gauntlet/Framework/Devices/Gauntlet.DevicePool.cs b/Engine/Source/Programs/AutomationTool/Gauntlet/Framework/Devices/Gauntlet.DevicePool.cs
--- a/Engine/Source/Programs/AutomationTool/Gauntlet/Framework/Devices/Gauntlet.DevicePool.cs
+++ b/Engine/Source/Programs/AutomationTool/Gauntlet/Framework/Devices/Gauntlet.DevicePool.cs
@@ -1233,7 +1233,7 @@ namespace Gauntlet
 			}
 			catch (Exception Ex)
 			{
-				Log.Info(KnownLogEvents.Gauntlet_DeviceEvent, "Failed to create device {Name}. {Message}", Def.ToString(), Ex.Message);
+				Log.Warning(KnownLogEvents.Gauntlet_DeviceEvent, "Failed to create device {Name}. {Message}\n{StackTrace}", Def.ToString(), Ex.Message, Ex.StackTrace);
 			}
 
 			return NewDevice;
