Gauntlet - Add extra context when failing on finding build target

#rnx

diff --git a/Engine/Source/Programs/AutomationTool/Gauntlet/Unreal/BuildSource/Gauntlet.UnrealBuildSource.cs b/Engine/Source/Programs/AutomationTool/Gauntlet/Unreal/BuildSource/Gauntlet.UnrealBuildSource.cs
--- a/Engine/Source/Programs/AutomationTool/Gauntlet/Unreal/BuildSource/Gauntlet.UnrealBuildSource.cs
+++ b/Engine/Source/Programs/AutomationTool/Gauntlet/Unreal/BuildSource/Gauntlet.UnrealBuildSource.cs
@@ -655,9 +655,9 @@ namespace Gauntlet
 					{
 						EditorExe = ProjectUtils.GetProjectTarget(ProjectPath, UnrealBuildTool.TargetType.Editor, TargetPlatform, TargetConfiguration);
 					}
-					catch
+					catch(Exception Ex)
 					{
-						throw new AutomationException("No suitable editor build for {0} configuration", TargetConfiguration);
+						throw new AutomationException("No suitable editor build for {0} configuration.\n{1}", TargetConfiguration, Ex.ToString());
 					}
 				}
 
