- Fixed ThisManifestDir in Turnkey with UNC style paths

diff --git a/Engine/Source/Programs/AutomationTool/Turnkey/TurnkeyManifest.cs b/Engine/Source/Programs/AutomationTool/Turnkey/TurnkeyManifest.cs
--- a/Engine/Source/Programs/AutomationTool/Turnkey/TurnkeyManifest.cs
+++ b/Engine/Source/Programs/AutomationTool/Turnkey/TurnkeyManifest.cs
@@ -198,7 +198,10 @@ namespace Turnkey
 						continue;
 					}
 
-					string ThisManifestDir = Path.GetDirectoryName(Source).Replace("\\", "/");
+					// chop off the last path component to get the dir of the ManifestFile source
+					int LastSlash = Source.LastIndexOf('/');
+					int LastBackSlash = Source.LastIndexOf('\\');
+					string ThisManifestDir = Source.Substring(0, Math.Max(LastSlash, LastBackSlash));
 
 					// if a directory is returned, look for the standardized manifest name, as we have not much else we can do with a directory
 					if (LocalManifestPath.EndsWith("/") || LocalManifestPath.EndsWith("\\"))
@@ -213,7 +216,7 @@ namespace Turnkey
 						ThisManifestDir = Source;
 					}
 
-// 					TurnkeyUtils.Log("Info: Setting ManifestDir to {0}", ThisManifestDir);
+// 					Console.WriteLine("Info: Setting ManifestDir to {0}", ThisManifestDir);
 
 					// allow this manifest's directory to be used in further copy, but needs to be in a stack
 					string PrevManifestDir = TurnkeyUtils.SetVariable("ThisManifestDir", ThisManifestDir);
