[Gauntlet] Prevent OOM error when writing log files of substantial size
Prevent log files > 1gb in size from being written.
#rnx

diff --git a/Engine/Source/Programs/AutomationTool/Gauntlet/Unreal/Base/Gauntlet.UnrealSession.cs b/Engine/Source/Programs/AutomationTool/Gauntlet/Unreal/Base/Gauntlet.UnrealSession.cs
--- a/Engine/Source/Programs/AutomationTool/Gauntlet/Unreal/Base/Gauntlet.UnrealSession.cs
+++ b/Engine/Source/Programs/AutomationTool/Gauntlet/Unreal/Base/Gauntlet.UnrealSession.cs
@@ -1114,33 +1114,40 @@ namespace Gauntlet
 			// save the output from TTY
 			string ArtifactLogPath = Path.Combine(InDestArtifactPath, RoleName + "Output.log");
 
-			// Write a brief Gauntlet header to aid debugging
-			StringBuilder LogOut = new StringBuilder();
-			LogOut.AppendLine("------ Gauntlet Test ------");
-			LogOut.AppendFormat("Role: {0}\r\n", InRunningRole.Role);
-			LogOut.AppendFormat("Automation Command: {0}\r\n", Environment.CommandLine);
-			LogOut.AppendLine("---------------------------");
+			int MaxLogSize = 1024 * 1024 * 1024;
+			int LogSize = InRunningRole.AppInstance.StdOut.Length * sizeof(char);
+			bool bIgnoreMaxSize = Globals.Params.ParseParam("NoMaxLogSize");
 
-			// Write instance stdout stream
-			LogOut.Append(InRunningRole.AppInstance.StdOut);
-
-			try
+			if (!bIgnoreMaxSize && LogSize > MaxLogSize)
+			{
+				Log.Warning("The process log for Role {0} was over 1 GB in size. A log artifact will not be generated for this process.", InRunningRole.ToString());
+			}
+			else
 			{
-				File.WriteAllText(ArtifactLogPath, LogOut.ToString());
+				try
+				{
+					using (StreamWriter Writer = new StreamWriter(ArtifactLogPath))
+					{
+						Writer.WriteLine("------ Gauntlet Test ------");
+						Writer.WriteLine(string.Format("Role: {0}\r\n", InRunningRole.Role));
+						Writer.WriteLine(string.Format("Automation Command: {0}\r\n", Environment.CommandLine));
+						Writer.WriteLine("---------------------------");
+						Writer.Write(InRunningRole.AppInstance.StdOut);
+					}
+					Log.Info("Wrote Log to {0}", ArtifactLogPath);
 
-				// On build machines, copy all role logs to Horde.
-				if (CommandUtils.IsBuildMachine)
+					// On build machines, copy all role logs to Horde.
+					if (CommandUtils.IsBuildMachine)
+					{
+						string HordeLogFileName = Path.Combine(CommandUtils.CmdEnv.LogFolder, RoleName + "Output.log");
+						File.Copy(ArtifactLogPath, HordeLogFileName);
+					}
+				}
+				catch (Exception Ex)
 				{
-					string HordeLogFileName = Path.Combine(CommandUtils.CmdEnv.LogFolder, RoleName + "Output.log");
-					File.WriteAllText(HordeLogFileName, LogOut.ToString());
+					string Message = "Encountered an {0} when attempting to write the {1} process log. The log may contain malformed encoding and will not be present on horde. {2}";
+					Log.Warning(Message, Ex.GetType().Name, RoleName, Ex.Message);
 				}
-
-				Log.Info("Wrote Log to {0}", ArtifactLogPath);
-			}
-			catch(Exception Ex)
-			{
-				string Message = "Encountered an {0} when attempting to write the {1} process log. The log may contain malformed encoding and will not be present on horde. {2}";
-				Log.Warning(Message, Ex.GetType().Name, RoleName, Ex.Message);
 			}
 
 			if (IsServer == false)
