[Gauntlet] Handle log writing exception
#rnx

diff --git a/Engine/Source/Programs/AutomationTool/Gauntlet/Unreal/Base/Gauntlet.UnrealSession.cs b/Engine/Source/Programs/AutomationTool/Gauntlet/Unreal/Base/Gauntlet.UnrealSession.cs
--- a/Engine/Source/Programs/AutomationTool/Gauntlet/Unreal/Base/Gauntlet.UnrealSession.cs
+++ b/Engine/Source/Programs/AutomationTool/Gauntlet/Unreal/Base/Gauntlet.UnrealSession.cs
@@ -1126,7 +1126,7 @@ namespace Gauntlet
 			{
 				try
 				{
-					using (StreamWriter Writer = new StreamWriter(ArtifactLogPath))
+					using (StreamWriter Writer = new StreamWriter(ArtifactLogPath, false))
 					{
 						Writer.WriteLine("------ Gauntlet Test ------");
 						Writer.WriteLine(string.Format("Role: {0}\r\n", InRunningRole.Role));
@@ -1140,7 +1140,7 @@ namespace Gauntlet
 					if (CommandUtils.IsBuildMachine)
 					{
 						string HordeLogFileName = Path.Combine(CommandUtils.CmdEnv.LogFolder, RoleName + "Output.log");
-						File.Copy(ArtifactLogPath, HordeLogFileName);
+						File.Copy(ArtifactLogPath, HordeLogFileName, true);
 					}
 				}
 				catch (Exception Ex)
