Localization:
Update the Localize UAT command to better handle gathering for plugins.
- Introduced the -IncludePluginsDirectory parameter to include all plugins under a particular directory
- Introduce the ExcludePluginsDirectory parameter to exclude all plugins under a particular directory
- Introduce the EnableIncludedPlugins flag to add all  included plugins that pass the exclude filter to the -EnablePlugins editor command line parameter to ensure the plugins are enabled for gathering.
#test Used the new commands to perform gathers for directories of plugins and enabling plugins that are not on by default.

diff --git a/Engine/Source/Programs/AutomationTool/Localization/LocalizationUtilities.cs b/Engine/Source/Programs/AutomationTool/Localization/LocalizationUtilities.cs
new file mode 100644
--- /dev/null
+++ b/Engine/Source/Programs/AutomationTool/Localization/LocalizationUtilities.cs
@@ -0,0 +1,53 @@
+// Copyright Epic Games, Inc. All Rights Reserved.
+using System.Collections.Generic;
+using System.IO;
+using UnrealBuildTool;
+using EpicGames.Core;
+
+namespace EpicGames.Localization
+{
+	public static class LocalizationUtilities
+	{
+		public static HashSet<string> GetPluginNamesUnderDirectory(string PluginDirectoryPath, string PluginRoot, PluginType Type)
+		{
+			HashSet<string> ReturnList = new HashSet<string>();
+			if (!Directory.Exists(PluginDirectoryPath))
+			{
+				return ReturnList;
+			}
+			if (!Directory.Exists(PluginRoot))
+			{
+				return ReturnList;
+			}
+
+			DirectoryReference PluginDirectoryReference = new DirectoryReference(PluginDirectoryPath);
+			DirectoryReference PluginRootDirectoryReference = new DirectoryReference(PluginRoot);
+			bool bIsSubdirectory = PluginDirectoryReference.IsUnderDirectory(PluginRootDirectoryReference);
+			if (bIsSubdirectory)
+			{
+				string RelativePath = PluginDirectoryReference.MakeRelativeTo(PluginRootDirectoryReference);
+				IReadOnlyList<PluginInfo> PluginsUnderDirectory = Plugins.ReadPluginsFromDirectory(PluginRootDirectoryReference, RelativePath, Type);
+				foreach (var Plugin in PluginsUnderDirectory)
+				{
+					ReturnList.Add(Plugin.Name);
+				}
+			}
+
+			return ReturnList;
+		}
+
+		public static string[] GetModularConfigSuffixes()
+		{
+			string[] FileSuffixes = {
+				"_Gather",
+				"_Import",
+				"_Export",
+				"_Compile",
+				"_GenerateReports",
+			};
+
+			return FileSuffixes;
+		}
+
+	}
+}
diff --git a/Engine/Source/Programs/AutomationTool/Scripts/Localisation.Automation.cs b/Engine/Source/Programs/AutomationTool/Scripts/Localisation.Automation.cs
--- a/Engine/Source/Programs/AutomationTool/Scripts/Localisation.Automation.cs
+++ b/Engine/Source/Programs/AutomationTool/Scripts/Localisation.Automation.cs
@@ -1,7 +1,6 @@
 // Copyright Epic Games, Inc. All Rights Reserved.
 using System;
 using System.Collections.Generic;
-using System.Globalization;
 using System.Linq;
 using System.Text;
 using System.IO;
@@ -13,8 +12,6 @@ using EpicGames.Core;
 using System.Threading.Tasks;
 using Microsoft.Extensions.Logging;
 
-using static AutomationTool.CommandUtils;
-
 [Help("Updates the external localization data using the arguments provided.")]
 [Help("UEProjectRoot", "Optional root-path to the project we're gathering for (defaults to CmdEnv.LocalRoot if unset).")]
 [Help("UEProjectDirectory", "Sub-path to the project we're gathering for (relative to UEProjectRoot).")]
@@ -24,7 +21,10 @@ using static AutomationTool.CommandUtils;
 [Help("LocalizationProvider", "Optional localization provide override.")]
 [Help("LocalizationSteps", "Optional comma separated list of localization steps to perform [Download, Gather, Import, Export, Compile, GenerateReports, Upload] (default is all). Only valid for projects using a modular config.")]
 [Help("IncludePlugins", "Optional flag to include plugins from within the given UEProjectDirectory as part of the gather. This may optionally specify a comma separated list of the specific plugins to gather (otherwise all plugins will be gathered).")]
+[Help("IncludePluginsDirectory", "Optional parameter that is a relative path to a directory under UEProjectDirectory. All plugins under this directory will be gathered from (if not excluded).")]
 [Help("ExcludePlugins", "Optional comma separated list of plugins to exclude from the gather.")]
+[Help("ExcludePluginsDirectory", "Optional relative path to a directory under UEProjectDirectory. All plugins under this directory will be excluded from gather.")]
+[Help("EnableIncludedPlugins", "Optional flag that passes all included plugins that aren't excluded to the -EnablePlugins editor argument to ensure content and metadata for plugins are loaded for gathering.")]
 [Help("IncludePlatforms", "Optional flag to include platforms from within the given UEProjectDirectory as part of the gather.")]
 [Help("AdditionalCommandletArguments", "Optional arguments to pass to the gather process.")]
 [Help("ParallelGather", "Run the gather processes for a single batch in parallel rather than sequence.")]
@@ -55,7 +55,7 @@ class Localize : BuildCommand
 			RootWorkingDirectory = CombinePaths(InUEProjectRoot, Batch.UEProjectDirectory);
 			RootLocalizationTargetDirectory = CombinePaths(InUEProjectRoot, Batch.LocalizationTargetDirectory);
 
-			// Try and find our localization provider
+			//Try and find our localization provider
 			{
 				LocalizationProvider.LocalizationProviderArgs LocProviderArgs;
 				LocProviderArgs.RootWorkingDirectory = RootWorkingDirectory;
@@ -129,14 +129,32 @@ class Localize : BuildCommand
 			}
 			LocalizationStepNames.Add("Monolithic"); // Always allow the monolithic scripts to run as we don't know which steps they do
 		}
+		
+		bool bShouldGatherPlugins = ParseParam("IncludePlugins");
+		bool bShouldEnableIncludedPlugins = ParseParam("EnableIncludedPlugins");
 
-		var ShouldGatherPlugins = ParseParam("IncludePlugins");
-		var IncludePlugins = new List<string>();
+			var IncludePlugins = new List<string>();
 		var ExcludePlugins = new List<string>();
-		if (ShouldGatherPlugins)
+
+		string PluginsRootPath = CombinePaths(UEProjectRoot, UEProjectDirectory);
+		string IncludePluginsUnderDirectoryStr = ParseParamValue("IncludePluginsDirectory");
+		if (!string.IsNullOrEmpty(IncludePluginsUnderDirectoryStr))
+		{
+			bShouldGatherPlugins = true;
+			string AbsolutePathToIncludePluginsDirectory = Path.Combine(PluginsRootPath, IncludePluginsUnderDirectoryStr);
+			IncludePlugins.AddRange(LocalizationUtilities.GetPluginNamesUnderDirectory(AbsolutePathToIncludePluginsDirectory, PluginsRootPath, UEProjectName.Length == 0 ? PluginType.Engine : PluginType.Project));
+		}
+		string ExcludePluginsUnderDirectoryStr = ParseParamValue("ExcludePluginsDirectory");
+		if (!string.IsNullOrEmpty(ExcludePluginsUnderDirectoryStr))
+		{
+			string AbsolutePathToExcludePluginsDirectory = Path.Combine(PluginsRootPath, ExcludePluginsUnderDirectoryStr);
+			ExcludePlugins.AddRange(LocalizationUtilities.GetPluginNamesUnderDirectory(AbsolutePathToExcludePluginsDirectory, PluginsRootPath, UEProjectName.Length == 0 ? PluginType.Engine : PluginType.Project));
+		}
+
+		if (bShouldGatherPlugins)
 		{
 			var IncludePluginsStr = ParseParamValue("IncludePlugins");
-			if (IncludePluginsStr != null)
+			if (!string.IsNullOrEmpty(IncludePluginsStr))
 			{
 				foreach (var PluginName in IncludePluginsStr.Split(','))
 				{
@@ -161,6 +179,16 @@ class Localize : BuildCommand
 		{
 			AdditionalCommandletArguments = "";
 		}
+		// We remove any leading or trailing quotes from AdditionalCommandletArguments
+		else if (!String.IsNullOrEmpty(AdditionalCommandletArguments))
+		{
+			AdditionalCommandletArguments = AdditionalCommandletArguments.Trim();
+			if (AdditionalCommandletArguments.StartsWith("\"") && AdditionalCommandletArguments.EndsWith("\""))
+			{
+				// We subtract 2 to nuke the last " character
+				AdditionalCommandletArguments = AdditionalCommandletArguments[1..^1];
+			}
+		}
 
 		var EnableParallelGather = ParseParam("ParallelGather");
 
@@ -205,7 +233,7 @@ class Localize : BuildCommand
 		}
 
 		// Build up any additional batches needed for plugins
-		if (ShouldGatherPlugins)
+		if (bShouldGatherPlugins)
 		{
 			var PluginsRootDirectory = new DirectoryReference(CombinePaths(UEProjectRoot, UEProjectDirectory));
 			IReadOnlyList<PluginInfo> AllPlugins = Plugins.ReadPluginsFromDirectory(PluginsRootDirectory, "Plugins", UEProjectName.Length == 0 ? PluginType.Engine : PluginType.Project);
@@ -217,7 +245,8 @@ class Localize : BuildCommand
 				AvailablePluginNames.Add(PluginInfo.Name);
 
 				bool ShouldIncludePlugin = (IncludePlugins.Count == 0 || IncludePlugins.Contains(PluginInfo.Name)) && !ExcludePlugins.Contains(PluginInfo.Name);
-				if (ShouldIncludePlugin && PluginInfo.Descriptor.LocalizationTargets != null && PluginInfo.Descriptor.LocalizationTargets.Length > 0)
+				bool PluginHasLocalizationTarget = PluginInfo.Descriptor.LocalizationTargets != null && PluginInfo.Descriptor.LocalizationTargets.Length > 0;
+				if (ShouldIncludePlugin && PluginHasLocalizationTarget)
 				{
 					var RootRelativePluginPath = PluginInfo.Directory.MakeRelativeTo(new DirectoryReference(UEProjectRoot));
 					RootRelativePluginPath = RootRelativePluginPath.Replace('\\', '/'); // Make sure we use / as these paths are used with P4
@@ -317,6 +346,12 @@ class Localize : BuildCommand
 		// These can run in parallel when ParallelGather is enabled
 		{
 			var EditorExe = CombinePaths(CmdEnv.LocalRoot, @"Engine/Binaries/Win64/UnrealEditor-Cmd.exe");
+			if (!File.Exists(EditorExe))
+			{
+				// Try using the debug .exe instead 
+				EditorExe = CombinePaths(CmdEnv.LocalRoot, @"Engine/Binaries/Win64/UnrealEditor-Win64-Debug-Cmd.exe");
+			}
+
 			// Set the common basic editor arguments
 			var EditorArguments = P4Enabled 
 				? String.Format("-SCCProvider=Perforce -P4Port={0} -P4User={1} -P4Client={2} -P4Passwd={3} -P4Changelist={4} -EnableSCC -DisableSCCSubmit", P4Env.ServerAndPort, P4Env.User, P4Env.Client, P4.GetAuthenticationToken(), PendingChangeList)
@@ -332,7 +367,65 @@ class Localize : BuildCommand
 			{
 				EditorArguments += " -multiprocess";
 			}
-			if (!String.IsNullOrEmpty(AdditionalCommandletArguments))
+
+			// We append all the included plugins to -EnablePlugins if -EnableIncludedPlugins is enabled. This wil ensure that the plugin content and metadata will be loaded.
+			// @TODOLocalization: Ideally the enabling of plugins should be per batch, otherwise each instance of the editor is enabling a bunch of plugins it doesn't need 
+			if (!string.IsNullOrEmpty(AdditionalCommandletArguments) && bShouldEnableIncludedPlugins && IncludePlugins.Count > 0)
+			{
+				HashSet<string> PluginsToEnableSet = IncludePlugins.Except(ExcludePlugins).ToHashSet();
+
+				// It's possible that there are already specified values for -EnabledPlugins, we willneed to try and parse them first.
+				string EnablePluginsToken = "-EnablePlugins=";
+				string EnablePluginsNewValue = "";
+
+				int EnablePluginsTokenIndex = AdditionalCommandletArguments.IndexOf(EnablePluginsToken);
+				string EnablePluginsOldValue = "";
+				if (EnablePluginsTokenIndex > -1)
+				{
+					// -EnablePlugins token exists in the additional commandlet args. We need to process it 
+					int EnablePluginsValueStartIndex = EnablePluginsTokenIndex + EnablePluginsToken.Length;
+					// We try and find the end of the string where it's separated by a space  between the next token 
+					int EnablePluginsValueEndIndex = AdditionalCommandletArguments.IndexOf(" ", EnablePluginsValueStartIndex);
+					// We can't find a next space. THis means we're the last parameter in AdditionalCommandletArguments. The end index will be the length of the string 
+					if (EnablePluginsValueEndIndex== -1)
+					{
+						EnablePluginsValueEndIndex = AdditionalCommandletArguments.Length;
+					}
+					// Isolate the value of -EnablePlugins and add them to our list of plugins to enable 
+					EnablePluginsOldValue = AdditionalCommandletArguments.Substring(EnablePluginsValueStartIndex, EnablePluginsValueEndIndex - EnablePluginsValueStartIndex);
+					foreach (string Plugin in EnablePluginsOldValue.Split(','))
+					{
+						PluginsToEnableSet.Add(Plugin);
+					}
+				}
+
+				// Just a counter to help iterate through the set to build out the comma separated value 
+				int IterationCount = 0;
+				StringBuilder EnablePluginsBuilder = new StringBuilder();
+					foreach (string Plugin in PluginsToEnableSet)
+				{
+					EnablePluginsBuilder.Append(Plugin);
+					if (IterationCount< PluginsToEnableSet.Count - 1)
+					{
+						EnablePluginsBuilder.Append(",");
+					}
+					++IterationCount;
+				}
+				EnablePluginsNewValue = EnablePluginsBuilder.ToString();
+				Logger.LogInformation($"Appending following plugins to be enabled: {EnablePluginsNewValue}");
+					// if we already had a value of -EnablePlugins in AdditionalCommandletArguments, we'll need to replace that with the new values we've created.
+					if (EnablePluginsTokenIndex > -1)
+				{
+					AdditionalCommandletArguments = AdditionalCommandletArguments.Replace(EnablePluginsToken + EnablePluginsOldValue, EnablePluginsToken + EnablePluginsNewValue);
+				}
+					else
+				{
+					// The token doesn't exist, we'll add it to the end 
+					AdditionalCommandletArguments += " " + EnablePluginsToken + EnablePluginsNewValue;
+				}
+			}
+
+				if (!String.IsNullOrEmpty(AdditionalCommandletArguments))
 			{
 				EditorArguments += " " + AdditionalCommandletArguments;
 			}
