IAS - updated handling of on demand pak file rules

- on demand pak file rules now respect the bOverrideChunkManifest property which means
that if not set to true together with OverridePaks=<MyPakName>, on demand content
will be placed in the <basechunkname>+<optionalondemand/ondemand>.pak
- this change makes bStageLoose and bOnDemand create the same pakchunks unless overriden with bOverrideChunkManifest=true

#rnx

diff --git a/Engine/Source/Programs/AutomationTool/Scripts/CopyBuildToStagingDirectory.Automation.cs b/Engine/Source/Programs/AutomationTool/Scripts/CopyBuildToStagingDirectory.Automation.cs
--- a/Engine/Source/Programs/AutomationTool/Scripts/CopyBuildToStagingDirectory.Automation.cs
+++ b/Engine/Source/Programs/AutomationTool/Scripts/CopyBuildToStagingDirectory.Automation.cs
@@ -2417,59 +2417,37 @@ namespace AutomationScripts
 		/// </summary>
 		/// <param name="Params"></param>
 		/// <param name="SC"></param>
-		private static void ApplyPakFileRules(List<PakFileRules> RulesList, KeyValuePair<string, string> StagingFile, HashSet<ChunkDefinition> ModifyPakList, ConcurrentDictionary<string, ChunkDefinition> ChunkNameToDefinition, out bool bExcludeFromPaks, out bool bStageLoose)
+		private static void ApplyPakFileRules(
+			List<PakFileRules> RulesList,
+			KeyValuePair<string, string> StagingFile,
+			HashSet<ChunkDefinition> ModifyPakList,
+			ConcurrentDictionary<string, ChunkDefinition> ChunkNameToDefinition,
+			out bool bExcludeFromPaks,
+			out bool bStageLoose,
+			out bool bOnDemand,
+			out string EncryptionKeyOverrideGuid)
 		{
 			bExcludeFromPaks = false;
 			bStageLoose = false;
+			bOnDemand = false;
+			EncryptionKeyOverrideGuid = null;
 
 			if (RulesList == null)
 			{
 				return;
 			}
 
-			// Process content-on-demand rules first
-			if (ModifyPakList != null)
-			{
-				foreach (PakFileRules Rule in RulesList.Where(Rule => Rule.bOnDemand))
-				{
-					if (PakFileRules.IsMatch(Rule, StagingFile))
-					{
-						ModifyPakList.Clear();
-
-						string ChunkName = Rule.Name;
-						if (Rule.OverridePaks != null && Rule.OverridePaks.Count > 0)
-						{
-							string OverrideChunkName = Rule.OverridePaks[0];
-							if (ChunkNameToDefinition.ContainsKey(OverrideChunkName))
-							{
-								ChunkName = OverrideChunkName;
-							}
-							else
-							{
-								Logger.LogWarning("Undefined PAK override '{Arg0}' in PAK rule '{Arg1}'", OverrideChunkName, Rule.Name);
-							}
-						}
-
-						ChunkDefinition Chunk = null;
-						if (ChunkNameToDefinition.TryGetValue(ChunkName, out Chunk))
-						{
-							ModifyPakList.Add(Chunk);
-						}
-						else
-						{
-							Logger.LogWarning("Undefined chunk name '{Arg0}' in PAK rule '{Arg1}'", ChunkName, Rule.Name);
-						}
-						
-						return;
-					}
-				}
-			}
-
 			// Search in order, return on first match
 			foreach (var PakRules in RulesList)
 			{
 				if (PakFileRules.IsMatch(PakRules, StagingFile))
 				{
+					if (PakRules.bOnDemand)
+					{
+						bOnDemand = true;
+						EncryptionKeyOverrideGuid = PakRules.EncryptionKeyGuid;
+					}
+
 					if (PakRules.bStageLoose)
 					{
 						bStageLoose = true;
@@ -2486,7 +2464,7 @@ namespace AutomationScripts
 						}
 
 						bOverrideChunkAssignment = true;
-						if (PakRules.OverridePaks != null)
+						if (PakRules.OverridePaks != null && !PakRules.bOnDemand)
 						{
 							Logger.LogInformation("Overridding chunk assignment {Arg0} to {Arg1}, this can cause broken references", StagingFile.Key, string.Join(", ", PakRules.OverridePaks));
 						}
@@ -2550,7 +2528,9 @@ namespace AutomationScripts
 			{
 				bool bExcludeFromPaks = false;
 				bool bStageLoose = false;
-				ApplyPakFileRules(PakRulesList, StagingFile, null, null, out bExcludeFromPaks, out bStageLoose);
+				bool bOnDemand = false;
+				string EncryptionKeyOverrideGuid;
+				ApplyPakFileRules(PakRulesList, StagingFile, null, null, out bExcludeFromPaks, out bStageLoose, out bOnDemand, out EncryptionKeyOverrideGuid);
 
 				if (bExcludeFromPaks)
 				{
@@ -4178,6 +4158,8 @@ namespace AutomationScripts
 					bool bAddedToChunk = false;
 					bool bExcludeFromPaks = false;
 					bool bStageLoose = false;
+					bool bOnDemand = false;
+					string EncryptionKeyOverrideGuid;
 					HashSet<ChunkDefinition> PakList = new HashSet<ChunkDefinition>();
 
 					string OriginalFilename = StagingFile.Key;
@@ -4211,7 +4193,7 @@ namespace AutomationScripts
 					}
 
 					// Now run through the pak rules which may override things
-					ApplyPakFileRules(PakRulesList, StagingFile, PakList, ChunkNameToDefinition, out bExcludeFromPaks, out bStageLoose);
+					ApplyPakFileRules(PakRulesList, StagingFile, PakList, ChunkNameToDefinition, out bExcludeFromPaks, out bStageLoose, out bOnDemand, out EncryptionKeyOverrideGuid);
 
 					if (bStageLoose)
 					{
@@ -4239,23 +4221,32 @@ namespace AutomationScripts
 					{
 						ChunkDefinition TargetChunk = Chunk;
 
-						if (!TargetChunk.bOnDemand)
+						string OrigExt = Path.GetExtension(OriginalFilename);
+						if (bStageLoose || bOnDemand)
 						{
-							string OrigExt = Path.GetExtension(OriginalFilename);
-							if (bStageLoose)
+							// make a new separate pak if not already marked as on demand
+							if (!Chunk.bOnDemand)
 							{
-								// make a streaming pak as well
 								string StreamingChunkName = Chunk.ChunkName + (OrigExt.Equals(OptionalBulkDataFileExtension) ? "ondemandoptional" : "ondemand");
 								AddChunkDefinition(AdditionalChunks, Chunk, StreamingChunkName);
 								TargetChunk = AdditionalChunks[StreamingChunkName];
+
+								if (bOnDemand)
+								{
+									TargetChunk.bOnDemand = true;
+									if (!string.IsNullOrEmpty(EncryptionKeyOverrideGuid))
+									{
+										TargetChunk.EncryptionKeyGuid = EncryptionKeyOverrideGuid;
+									}
+								}
 							}
-							else if (OrigExt.Equals(OptionalBulkDataFileExtension))
-							{
-								// any optional files encountered we want to put in a separate pak file
-								string OptionalChunkName = Chunk.ChunkName + "optional";
-								AddChunkDefinition(AdditionalChunks, Chunk, OptionalChunkName);
-								TargetChunk = AdditionalChunks[OptionalChunkName];
-							}
+						}
+						else if (OrigExt.Equals(OptionalBulkDataFileExtension))
+						{
+							// any optional files encountered we want to put in a separate pak file
+							string OptionalChunkName = Chunk.ChunkName + "optional";
+							AddChunkDefinition(AdditionalChunks, Chunk, OptionalChunkName);
+							TargetChunk = AdditionalChunks[OptionalChunkName];
 						}
 
 						TargetChunk.ResponseFile.TryAdd(StagingFile.Key, StagingFile.Value);
