- Changed how the sdk.json files are staged in cooked editors

diff --git a/Engine/Source/Programs/AutomationTool/CookedEditor/CookedEditor.Automation.cs b/Engine/Source/Programs/AutomationTool/CookedEditor/CookedEditor.Automation.cs
--- a/Engine/Source/Programs/AutomationTool/CookedEditor/CookedEditor.Automation.cs
+++ b/Engine/Source/Programs/AutomationTool/CookedEditor/CookedEditor.Automation.cs
@@ -768,7 +768,7 @@ public class MakeCookedEditor : BuildCommand
 		// plugins are already handled in the Plugins staging code
 		List<string> RootFoldersToStrip = new List<string> { "source", "plugins" };//, "binaries" };
 		List<string> SubFoldersToStrip = new List<string> { "source", "intermediate", "tests", "binaries" + Path.DirectorySeparatorChar + HostPlatform.Current.HostEditorPlatform.ToString().ToLower() };
-		List<string> RootNonUFSFolders = new List<string> { "shaders", "binaries", "build", "extras", "config" };
+		List<string> RootNonUFSFolders = new List<string> { "shaders", "binaries", "build", "extras" };
 
 
 		if (!Context.bStageShaderDirs)
@@ -809,6 +809,11 @@ public class MakeCookedEditor : BuildCommand
 					FilesToStage.RemoveAll(x => x.ContainsAnyNames(SubFoldersToStrip, Subdir));
 					ContextFileList.AddRange(FilesToStage);
 				}
+
+				if (SubdirName == "config")
+				{
+					Context.NonUFSFilesToStage.AddRange(DirectoryReference.EnumerateFiles(Subdir, "*SDK.json", SearchOption.AllDirectories));
+				}
 			}
 		}
 	}
