Adding -Packaged param to RunUnreal

diff --git a/Engine/Source/Programs/AutomationTool/Gauntlet/Unreal/RunUnreal.cs b/Engine/Source/Programs/AutomationTool/Gauntlet/Unreal/RunUnreal.cs
--- a/Engine/Source/Programs/AutomationTool/Gauntlet/Unreal/RunUnreal.cs
+++ b/Engine/Source/Programs/AutomationTool/Gauntlet/Unreal/RunUnreal.cs
@@ -32,6 +32,7 @@ namespace Gauntlet
 	[ParamHelp("Build", "Reference to the build that is being tested")]
 	[ParamHelp("Configuration", "Configuration to perform tests on", Choices = new string[] { "Debug", "DebugGame", "Development", "Test", "Shipping"})]
 	[ParamHelp("Platform", "Platforms to perform tests on and their params")]
+	[ParamHelp("Packaged", "Run packaged build instead of staged", ParamType = typeof(bool))]
 	[ParamHelp("Dev", "Run in Dev mode", ParamType = typeof(bool))]
 	[ParamHelp("p4", "Enable p4v support", ParamType = typeof(bool))]
 	[ParamHelp("MaxDuration", "Maximum duration for test in sections", ParamType = typeof(int), DefaultValue = 3600)]
