Fix path issues on Windows for prepare for debug

diff --git a/Engine/Source/Programs/AutomationTool/IOS/IOSPlatform.Automation.cs b/Engine/Source/Programs/AutomationTool/IOS/IOSPlatform.Automation.cs
--- a/Engine/Source/Programs/AutomationTool/IOS/IOSPlatform.Automation.cs
+++ b/Engine/Source/Programs/AutomationTool/IOS/IOSPlatform.Automation.cs
@@ -2346,16 +2346,24 @@ public class IOSPlatform : ApplePlatform
 
 		if (HostPlatform.Current.HostEditorPlatform == UnrealTargetPlatform.Win64 || HostPlatform.Current.HostEditorPlatform == UnrealTargetPlatform.Mac)
 		{
-			int StartPos = ProjectFilePath.LastIndexOf(Path.DirectorySeparatorChar);
-    		int StringLength = ProjectFilePath.Length - 10; // 9 for .uproject, 1 for the /
+			int StartPos = 0;
+			if (HostPlatform.Current.HostEditorPlatform == UnrealTargetPlatform.Win64)
+			{
+				StartPos = ProjectFilePath.LastIndexOf('/');
+			}
+			else
+			{
+				StartPos = ProjectFilePath.LastIndexOf(Path.DirectorySeparatorChar);
+			}
+			int StringLength = ProjectFilePath.Length - 10; // 9 for .uproject, 1 for the /
     		string PackageName = ProjectFilePath.Substring(StartPos + 1, StringLength - StartPos);
-    		string PackagePath = Path.Combine(Path.GetDirectoryName(ProjectFilePath), "Binaries", ClientPlatform);
-    		if (string.IsNullOrEmpty(SourcePackage))
+			string PackagePath = Path.Combine(Path.GetDirectoryName(ProjectFilePath), "Binaries", ClientPlatform);
+			Logger.LogInformation("PackagePath : {PackagePath}", PackagePath);
+			if (string.IsNullOrEmpty(SourcePackage))
     		{
-    			SourcePackage = Path.Combine(PackagePath, PackageName + ".ipa");
-    		}
-			
-			Logger.LogInformation("SourcePackage : {SourcePackage}", SourcePackage);
+				SourcePackage = Path.Combine(PackagePath, PackageName + ".ipa");
+			}
+
 			string[] IPAFiles;
 			if (!File.Exists(SourcePackage))
             {
@@ -2380,7 +2388,14 @@ public class IOSPlatform : ApplePlatform
 
 			string PayloadPath = SourcePackage;
 			Logger.LogInformation("PayloadPath : {PayloadPath}", PayloadPath);
-			PayloadPath = PayloadPath.Substring(0, PayloadPath.LastIndexOf('/'));
+			if (HostPlatform.Current.HostEditorPlatform == UnrealTargetPlatform.Win64)
+			{
+				PayloadPath = PayloadPath.Substring(0, PayloadPath.LastIndexOf('\\'));
+			}
+			else
+			{
+				PayloadPath = PayloadPath.Substring(0, PayloadPath.LastIndexOf('/'));
+			}
 			string CookedDataDirectory = Path.Combine(Path.GetDirectoryName(PayloadPath), ClientPlatform, "Payload", PackageName + ".app", "cookeddata");
 
 			Logger.LogInformation("ClientPlatform : {ClientPlatform}", ClientPlatform);
@@ -2431,8 +2446,16 @@ public class IOSPlatform : ApplePlatform
 
 	public void UnzipPackage(string PackageToUnzip)
 	{
+		Logger.LogInformation("PackageToUnzip : {PackageToUnzip}", PackageToUnzip);
 		string UnzipPath = PackageToUnzip;
-		UnzipPath = UnzipPath.Substring(0, UnzipPath.LastIndexOf(Path.DirectorySeparatorChar));
+		if (HostPlatform.Current.HostEditorPlatform == UnrealTargetPlatform.Win64)
+		{
+			UnzipPath = UnzipPath.Substring(0, UnzipPath.LastIndexOf('\\'));
+		}
+		else
+		{
+			UnzipPath = UnzipPath.Substring(0, UnzipPath.LastIndexOf(Path.DirectorySeparatorChar));
+		}
 		Logger.LogInformation("Unzipping to {UnzipPath}", UnzipPath);
 
 		using (Ionic.Zip.ZipFile Zip = new Ionic.Zip.ZipFile(PackageToUnzip))
