IAS - set both encryption and requested encryption key when overriding on demand pakchunks

#rnx

diff --git a/Engine/Source/Programs/AutomationTool/Scripts/CopyBuildToStagingDirectory.Automation.cs b/Engine/Source/Programs/AutomationTool/Scripts/CopyBuildToStagingDirectory.Automation.cs
--- a/Engine/Source/Programs/AutomationTool/Scripts/CopyBuildToStagingDirectory.Automation.cs
+++ b/Engine/Source/Programs/AutomationTool/Scripts/CopyBuildToStagingDirectory.Automation.cs
@@ -4237,6 +4237,7 @@ namespace AutomationScripts
 									if (!string.IsNullOrEmpty(EncryptionKeyOverrideGuid))
 									{
 										TargetChunk.EncryptionKeyGuid = EncryptionKeyOverrideGuid;
+										TargetChunk.RequestedEncryptionKeyGuid = EncryptionKeyOverrideGuid;
 									}
 								}
 							}
