Allow Info.plist and CodeResource files inside .app bundles to be different during horde jobs

diff --git a/Engine/Source/Programs/AutomationTool/BuildGraph/TempStorage.cs b/Engine/Source/Programs/AutomationTool/BuildGraph/TempStorage.cs
--- a/Engine/Source/Programs/AutomationTool/BuildGraph/TempStorage.cs
+++ b/Engine/Source/Programs/AutomationTool/BuildGraph/TempStorage.cs
@@ -1307,6 +1307,11 @@ namespace AutomationTool
 			{
 				return true;
 			}
+			if ((FileName.Equals("info.plist", StringComparison.OrdinalIgnoreCase) || FileName.Equals("coderesources", StringComparison.OrdinalIgnoreCase)) && LocalFile.FullName.Contains(".app/"))
+			{
+				// xcode can generate plist files and coderesources differently in different stages of compile/cook/stage/package/etc. only allow ones inside a .app bundle
+				return true;
+			}
 			return false;
 		}
 		
