UAT: Explicitly set input encoding for processes #8439

diff --git a/Engine/Source/Programs/AutomationTool/AutomationUtils/ProcessUtils.cs b/Engine/Source/Programs/AutomationTool/AutomationUtils/ProcessUtils.cs
--- a/Engine/Source/Programs/AutomationTool/AutomationUtils/ProcessUtils.cs
+++ b/Engine/Source/Programs/AutomationTool/AutomationUtils/ProcessUtils.cs
@@ -908,7 +908,24 @@ namespace AutomationTool
 					Proc.OutputDataReceived += Result.StdOut;
 					Proc.ErrorDataReceived += Result.StdErr;
 				}
-				Proc.StartInfo.RedirectStandardInput = Input != null;
+
+				// By default the standard input stream uses the current terminal input encoding (`Console.InputEncoding`),
+				// so let's make sure to set an explicit known input encoding (that doesn't produce a BOM).
+				if (Input != null)
+				{
+					Proc.StartInfo.RedirectStandardInput = true;
+
+					// Assume that if the application produces UTF-16, it also consumes UTF-16.
+					if ((Options & ERunOptions.UTF16Output) == ERunOptions.UTF16Output)
+					{
+						Proc.StartInfo.StandardInputEncoding = new UnicodeEncoding(false, false, false);
+					}
+					else
+					{
+						Proc.StartInfo.StandardInputEncoding = new UTF8Encoding(false);
+					}
+				}
+
 				Proc.StartInfo.CreateNoWindow = (Options & ERunOptions.NoHideWindow) == 0;
 				if ((Options & ERunOptions.UTF8Output) == ERunOptions.UTF8Output)
 				{
