[Gauntlet] Adds group parameter parsing to GauntletCommandLine

diff --git a/Engine/Source/Programs/AutomationTool/Gauntlet/Unreal/Base/Gauntlet.UnrealTestConfiguration.cs b/Engine/Source/Programs/AutomationTool/Gauntlet/Unreal/Base/Gauntlet.UnrealTestConfiguration.cs
--- a/Engine/Source/Programs/AutomationTool/Gauntlet/Unreal/Base/Gauntlet.UnrealTestConfiguration.cs
+++ b/Engine/Source/Programs/AutomationTool/Gauntlet/Unreal/Base/Gauntlet.UnrealTestConfiguration.cs
@@ -213,11 +213,61 @@ namespace Gauntlet
 			return Params[ParamName];
 		}
 
+		/// <summary>
+		/// Get a collection of all sub paremeters
+		/// Useful for group arguments like ExecCmds
+		/// </summary>
+		/// <param name="GroupName"></param>
+		/// <param name="SubParamDelimeter">Which separator to use when distinguishing the group's sub parameters</param>
+		/// <param name="SubValueDelimeter">Which separator to use when splitting the sub parameter from the corresponding value</param>
+		/// <returns>A dictionary containing the group's sub parameters of parameter values</returns>
+		public Dictionary<string, string> GetGroupParamValues(string GroupName, string SubParamDelimeter = ",", string SubValueDelimeter = " ")
+		{
+			Dictionary<string, string> Group = new Dictionary<string, string>();
+
+			string FullArgument = GetParamValue(GroupName).ToString();
+			string[] SubParams = FullArgument.Split(SubParamDelimeter);
+			foreach (string SubParam in SubParams)
+			{
+				int Index = SubParam.IndexOf(SubValueDelimeter);
+
+				if (Index < 1)
+				{
+					// No value, this is just a boolean param
+					Group.Add(SubParam, string.Empty);
+				}
+				else
+				{
+					Group.Add(SubParam.Substring(0, Index), SubParam.Substring(Index + 1));
+				}
+			}
+
+			return Group;
+		}
+
+		/// <summary>
+		/// Checks if a parameter is present
+		/// </summary>
+		/// <param name="ParamName">The name of the parameter</param>
+		/// <returns>True if the parameter has already been added</returns>
 		public bool HasParam(string ParamName)
 		{
 			return (Params != null && Params.ContainsKey(ParamName));
 		}
 
+		/// <summary>
+		/// Checks if a parameter group has a subparameter
+		/// </summary>
+		/// <param name="GroupName">The name of the param group, like "ExecCmds"</param>
+		/// <param name="SubParamName">A component of the group, like "sg.TextureQuality"</param>
+		/// <param name="SubParamDelimeter">Which separator to use when distinguishing the group's sub parameters</param>
+		/// <param name="SubValueDelimeter">Which separator to use when splitting the sub parameter from the corresponding value</param>
+		/// <returns>True if a group parameter exists and contains a matching sub parameter</returns>
+		public bool HasGroupParam(string GroupName, string SubParamName, string SubParamDelimeter = ",", string SubValueDelimeter = " ")
+		{
+			return HasParam(GroupName) && GetGroupParamValues(GroupName, SubParamDelimeter, SubValueDelimeter).ContainsKey(SubParamName);
+		}
+
 		/// <summary>
 		/// Wipe out the entire set of passed in params.
 		/// </summary>
