Set the correct sdk to use for iOS device and iOS Simulator builds

diff --git a/Engine/Source/Programs/AutomationTool/Scripts/BuildCMakeLib.Automation.cs b/Engine/Source/Programs/AutomationTool/Scripts/BuildCMakeLib.Automation.cs
--- a/Engine/Source/Programs/AutomationTool/Scripts/BuildCMakeLib.Automation.cs
+++ b/Engine/Source/Programs/AutomationTool/Scripts/BuildCMakeLib.Automation.cs
@@ -926,7 +926,7 @@ class MakefileTargetPlatform_IOS : BuildCMakeLib.MakefileTargetPlatform
 		return SysRoot;
 	}
 
-	private string GetSdkName() => Architecture == "x86_64" || Architecture == "i386" ? "iphonesimulator" : "iphoneos";
+	private string GetSdkName() => Architecture == "x86_64" || Architecture == "iossimulator" ? "iphonesimulator" : "iphoneos";
 }
 
 class XcodeTargetPlatform_TVOS : BuildCMakeLib.XcodeTargetPlatform
diff --git a/Engine/Source/Programs/UnrealBuildTool/Platform/TVOS/TVOSToolChain.cs b/Engine/Source/Programs/UnrealBuildTool/Platform/TVOS/TVOSToolChain.cs
--- a/Engine/Source/Programs/UnrealBuildTool/Platform/TVOS/TVOSToolChain.cs
+++ b/Engine/Source/Programs/UnrealBuildTool/Platform/TVOS/TVOSToolChain.cs
@@ -18,8 +18,8 @@ namespace UnrealBuildTool
 		{
 		}
 
-		public override string GetXcodeMinVersionParam()
-		{
+		public override string GetXcodeMinVersionParam(UnrealArch Architecture)
+	{
 			return "tvos-version-min";
 		}
 	}
