Brought changes over from UE5/Main

- Fixing the name of the .app when staging/deploying a BP IOS project with modern
- Uses a new AppleExports function to call MakeBinaryFilename, so we can use standardized name format for the BP .app
- Fixed which .app is archived for BP projects (it was archiving the UnrealGame.app, not the finalized MyProject.app)

- Programs now no longer use Game/Client entitlements, but can set bUseEntitlementsForPrograms to true in the .ini if really needed
- Added an Unrestricted.entitlements, an empty entitlements file
- Changed Programs to find the .ini for codesigning, without needing a .uproject

- Fixed path to strip to add a missing slash
- Pass -legacyxcode along to the legacy PostBuildSync step to make sure it uses legacy when overriding the .ini setting
- Added -legacyxcode commandline handler for legacy PostbuildSync

diff --git a/Engine/Source/Programs/AutomationTool/Apple/ApplePlatform.Automation.cs b/Engine/Source/Programs/AutomationTool/Apple/ApplePlatform.Automation.cs
--- a/Engine/Source/Programs/AutomationTool/Apple/ApplePlatform.Automation.cs
+++ b/Engine/Source/Programs/AutomationTool/Apple/ApplePlatform.Automation.cs
@@ -143,6 +143,15 @@ public abstract class ApplePlatform : Platform
 
 	private DirectoryReference GetFinalAppPath(StageTarget Target, DeploymentContext SC)
 	{
+		// content only projects have to not use the executable, but make a path to a .app 
+		if (!SC.IsCodeBasedProject)
+		{
+			string AppBundleName = AppleExports.MakeBinaryFileName(SC.ShortProjectName, Target.Receipt.Platform, Target.Receipt.Configuration, 
+				Target.Receipt.Architectures, UnrealTargetConfiguration.Development, ".app");
+
+			return DirectoryReference.Combine(SC.ProjectRoot, "Binaries", Target.Receipt.Platform.ToString(), AppBundleName);
+		}
+
 		DirectoryReference AppDir;
 		// get the executable from the receipt
 		FileReference Executable = Target.Receipt.BuildProducts.First(x => x.Type == BuildProductType.Executable).Path;
