[Gauntlet] Disable uploading of PersistentDownloadDir

diff --git a/Engine/Source/Programs/AutomationTool/Gauntlet/Unreal/Base/Gauntlet.UnrealSession.cs b/Engine/Source/Programs/AutomationTool/Gauntlet/Unreal/Base/Gauntlet.UnrealSession.cs
--- a/Engine/Source/Programs/AutomationTool/Gauntlet/Unreal/Base/Gauntlet.UnrealSession.cs
+++ b/Engine/Source/Programs/AutomationTool/Gauntlet/Unreal/Base/Gauntlet.UnrealSession.cs
@@ -1,19 +1,15 @@
 // Copyright Epic Games, Inc. All Rights Reserved.
 
+using AutomationTool;
+using EpicGames.Core;
 using System;
 using System.Collections.Generic;
 using System.IO;
-using AutomationTool;
-using AutomationTool.DeviceReservation;
-using UnrealBuildTool;
-using System.Threading;
+using System.Linq;
 using System.Text;
 using System.Text.RegularExpressions;
-using System.Drawing;
-using System.Linq;
-using System.Collections;
-using System.Diagnostics;
-using EpicGames.Core;
+using System.Threading;
+using UnrealBuildTool;
 
 namespace Gauntlet
 {
@@ -1087,33 +1083,34 @@ namespace Gauntlet
 		/// <returns></returns>
 		public UnrealRoleArtifacts SaveRoleArtifacts(UnrealTestContext InContext, UnrealSessionInstance.RoleInstance InRunningRole, string InDestArtifactPath)
 		{
-
 			bool IsServer = InRunningRole.Role.RoleType.IsServer();
 			string RoleName = (InRunningRole.Role.IsDummy() ? "Dummy" : "") + InRunningRole.Role.RoleType.ToString();
 			UnrealTargetPlatform? Platform = InRunningRole.Role.Platform;
 			string RoleConfig = InRunningRole.Role.Configuration.ToString();
 
-			Directory.CreateDirectory(InDestArtifactPath);
+			if (!Directory.Exists(InDestArtifactPath))
+			{
+				Directory.CreateDirectory(InDestArtifactPath);
+			}
 
 			bool IsDevBuild = InContext.TestParams.ParseParam("dev");
 			bool IsEditorBuild = InRunningRole.Role.RoleType.UsesEditor();
+			bool IsBuildMachine = CommandUtils.IsBuildMachine;
 
-			// Unless there was a crash on a builder don't archive editor data (there can be a *lot* of stuff in their).
+			// Unless there was a crash on a builder don't archive editor data (there can be a *lot* of stuff in there).
 			bool SkipArchivingAssets = IsDevBuild ||
 										(IsEditorBuild && 
-											(CommandUtils.IsBuildMachine == false || InRunningRole.AppInstance.ExitCode == 0)
+											(IsBuildMachine == false || InRunningRole.AppInstance.ExitCode == 0)
 										);
 
-
-			string DestSavedDir = InDestArtifactPath;
-			string SourceSavedDir = "";
-
+			DirectoryInfo DestSavedDirInfo = new DirectoryInfo(InDestArtifactPath);
 			// save the contents of the saved directory
-			SourceSavedDir = InRunningRole.AppInstance.ArtifactPath;
-
-			// save the output from TTY
-			string ArtifactLogPath = Path.Combine(InDestArtifactPath, RoleName + "Output.log");
+			string SourceSavedDir = InRunningRole.AppInstance.ArtifactPath;
+			
+			string ArtifactLogFilePath = String.Empty;
+			string HordeLogFilePath = String.Empty;
 
+			// Get the size of the log to determine if it should be saved under a certain size
 			int MaxLogSize = 1024 * 1024 * 1024;
 			int LogSize = InRunningRole.AppInstance.StdOut.Length * sizeof(char);
 			bool bIgnoreMaxSize = Globals.Params.ParseParam("NoMaxLogSize");
@@ -1126,7 +1123,22 @@ namespace Gauntlet
 			{
 				try
 				{
-					using (StreamWriter Writer = new StreamWriter(ArtifactLogPath, false))
+					string StreamWriterOutputFilePath = String.Empty;
+
+					// On build machines, copy all role logs to Horde.
+					if (IsBuildMachine)
+					{
+						HordeLogFilePath = Path.Combine(CommandUtils.CmdEnv.LogFolder, RoleName + "Output.log");
+						StreamWriterOutputFilePath = HordeLogFilePath;
+					}
+					else
+					{
+						ArtifactLogFilePath = Path.Combine(DestSavedDirInfo.FullName, RoleName + "Output.log");
+						StreamWriterOutputFilePath = ArtifactLogFilePath;
+					}
+
+					// save the output from TTY
+					using (StreamWriter Writer = new StreamWriter(StreamWriterOutputFilePath, false))
 					{
 						Writer.WriteLine("------ Gauntlet Test ------");
 						Writer.WriteLine(string.Format("Role: {0}\r\n", InRunningRole.Role));
@@ -1134,14 +1146,7 @@ namespace Gauntlet
 						Writer.WriteLine("---------------------------");
 						Writer.Write(InRunningRole.AppInstance.StdOut);
 					}
-					Log.Info("Wrote Log to {0}", ArtifactLogPath);
-
-					// On build machines, copy all role logs to Horde.
-					if (CommandUtils.IsBuildMachine)
-					{
-						string HordeLogFileName = Path.Combine(CommandUtils.CmdEnv.LogFolder, RoleName + "Output.log");
-						File.Copy(ArtifactLogPath, HordeLogFileName, true);
-					}
+					Log.Info($"Wrote {RoleName} Log to {StreamWriterOutputFilePath}");
 				}
 				catch (Exception Ex)
 				{
@@ -1155,23 +1160,24 @@ namespace Gauntlet
 				// gif-ify and jpeg-ify any screenshots
 				try
 				{
-
 					string ScreenshotPath = Path.Combine(SourceSavedDir, "Screenshots", Platform.ToString()).ToLower();
-					string TempScreenshotPath = Path.GetTempPath();
-					if (!Directory.Exists(TempScreenshotPath))
-					{
-						Log.Info("Creating temp directory {0}", TempScreenshotPath);
-						Directory.CreateDirectory(TempScreenshotPath);
-					}
 
-					if (Directory.Exists(ScreenshotPath) && Directory.GetFiles(ScreenshotPath).Length > 0)
+					// Check as early as possible for screenshots before creating a temp folder to copy
+					if (Directory.Exists(ScreenshotPath) && Directory.GetFiles(ScreenshotPath).Any())
 					{
+						string TempScreenshotPath = Path.GetTempPath();
+						if (!Directory.Exists(TempScreenshotPath))
+						{
+							Log.Info("Creating temp directory {0}", TempScreenshotPath);
+							Directory.CreateDirectory(TempScreenshotPath);
+						}
+
 						Log.Info("Downsizing and gifying session images at {0}", ScreenshotPath);
 
 						// downsize first so gif-step is quicker and takes less resoruces.
 						Utils.Image.ConvertImages(ScreenshotPath, TempScreenshotPath, "jpg", true);
 
-						string GifPath = Path.Combine(InDestArtifactPath, RoleName + "Test.gif");
+						string GifPath = Path.Combine(DestSavedDirInfo.FullName, RoleName + "Test.gif");
 						if (Utils.Image.SaveImagesAsGif(TempScreenshotPath, GifPath))
 						{
 							Log.Info("Saved gif to {0}", GifPath);
@@ -1184,31 +1190,6 @@ namespace Gauntlet
 				}
 			}
 
-			Dictionary<string, string> LongCrashReporterStringToIndex = new Dictionary<string, string>();
-
-			// Filter that truncates some overly long file paths like CrashReporter files. These are particularly problematic with 
-			// testflights which append a long random name to the destination folder, easily pushing past 260 chars.
-			Func<string, string> TruncateLongPathFilter= (X =>
-			{
-				// /UECC-Windows-F0DD9BB04C3C9250FAF39D8AB4A88556/
-				Match M = Regex.Match(X, @"UECC-.+-([\dA-Fa-f]+)");
-
-				if (M.Success)
-				{
-					string LongString = M.Groups[1].ToString();
-
-					if (!LongCrashReporterStringToIndex.ContainsKey(LongString))
-					{
-						LongCrashReporterStringToIndex[LongString] = LongCrashReporterStringToIndex.Keys.Count.ToString("D2");
-					}
-
-					string ShortSring = LongCrashReporterStringToIndex[LongString];
-					X = X.Replace(LongString, ShortSring);
-				}				
-
-				return X;
-			});
-
 			// don't archive data in dev mode, because peoples saved data could be huuuuuuuge!
 			if (SkipArchivingAssets)
 			{
@@ -1222,14 +1203,53 @@ namespace Gauntlet
 				}
 			}
 			else
-			{ 
+			{
 				LogLevel OldLevel = Log.Level;
 				Log.Level = LogLevel.Normal;
 
 				if (Directory.Exists(SourceSavedDir))
 				{
-					Utils.SystemHelpers.CopyDirectory(SourceSavedDir, Utils.SystemHelpers.GetFullyQualifiedPath(DestSavedDir), Utils.SystemHelpers.CopyOptions.Default, TruncateLongPathFilter);
-					Log.Info("Archived artifacts to to {0}", DestSavedDir);
+					// Only Copy the artifacts we want
+					foreach (string SubDirectory in Directory.EnumerateDirectories(SourceSavedDir))
+					{
+						DirectoryInfo DirInfo = new DirectoryInfo(SubDirectory);
+
+						if (DirInfo.Name == EIntendedBaseCopyDirectory.PersistentDownloadDir.ToString())
+						{
+							// Do not store the PersistentDownloadDir when possible
+							continue;
+						}
+						
+						// Don't copy an empty directory when possible
+						if (Directory.EnumerateDirectories(SubDirectory).Any() || Directory.EnumerateFiles(SubDirectory).Any())
+						{
+							string FullyQualifiedDestPath = Utils.SystemHelpers.GetFullyQualifiedPath(Path.Combine(DestSavedDirInfo.FullName, DirInfo.Name));
+							Utils.SystemHelpers.CopyDirectory(DirInfo.FullName, FullyQualifiedDestPath, Utils.SystemHelpers.CopyOptions.Default, TruncateLongPathFilter);
+							Log.Info($"Archived artifact \\{DirInfo.Parent.Name}\\{DirInfo.Name}\\ to {FullyQualifiedDestPath}");
+						}
+					}
+
+					// Copy any lose files from the source
+					foreach (string SourceFile in Directory.EnumerateFiles(SourceSavedDir))
+					{
+						FileInfo SourceFileInfo = new FileInfo(SourceFile);
+						string FullyQualifiedDestPath = Utils.SystemHelpers.GetFullyQualifiedPath(Path.Combine(DestSavedDirInfo.FullName, SourceFileInfo.Name));
+						FileInfo DestInfo = new FileInfo(FullyQualifiedDestPath);
+
+						try
+						{
+							DestInfo = SourceFileInfo.CopyTo(DestInfo.FullName, overwrite:true);
+
+							// Clear attributes and set last write time
+							DestInfo.Attributes = FileAttributes.Normal;
+							DestInfo.LastWriteTime = SourceFileInfo.LastWriteTime;
+							Log.Info($"Archived {SourceFileInfo.Name} to {DestInfo.FullName}");
+						}
+						catch (IOException ex)
+						{
+							Log.Warning($"Archive of {SourceFileInfo.Name} to {DestInfo.FullName} FAILED: Skipping\n{ex.Message}");
+						}
+					}
 				}
 				else
 				{
@@ -1247,88 +1267,31 @@ namespace Gauntlet
 					var DirToCopy = new DirectoryInfo(SourcePath);
 					if (DirToCopy.Exists)
 					{
-						// Grab the final dir name to copy everything into so everything's not just going into root artifact dir.
-						string IntendedCopyLocation = Path.Combine(InDestArtifactPath, DirToCopy.Name);
+						// Grab the final dir name to copy everything into, so everything does not go into the root artifact dir.
+						string IntendedCopyLocation = Path.Combine(DestSavedDirInfo.FullName, DirToCopy.Name);
 						Utils.SystemHelpers.CopyDirectory(SourcePath, Utils.SystemHelpers.GetFullyQualifiedPath(IntendedCopyLocation), Utils.SystemHelpers.CopyOptions.Default, TruncateLongPathFilter);
 					}
 				}
 			}
-			// TODO REMOVEME- this should go elsewhere, likely a util that can be called or inserted by relevant test nodes.
-			if (IsServer == false)
-			{
-				// Copy over PSOs
-				try
-				{
-					if (InContext.Options.LogPSO)
-					{
-						foreach (var ThisFile in CommandUtils.FindFiles_NoExceptions(true, "*.rec.upipelinecache", true, DestSavedDir))
-						{
-							bool Copied = false;
-							var JustFile = Path.GetFileName(ThisFile);
-							if (JustFile.StartsWith("++"))
-							{
 
-								var Parts = JustFile.Split(new Char[] { '+', '-' }).Where(A => A != "").ToArray();
-								if (Parts.Count() >= 2)
-								{
-									string ProjectName = Parts[0].ToString();
-									string BuildRoot = CommandUtils.CombinePaths(CommandUtils.RootBuildStorageDirectory());
-
-									string SrcBuildPath = CommandUtils.CombinePaths(BuildRoot, ProjectName);
-									string SrcBuildPath2 = CommandUtils.CombinePaths(BuildRoot, ProjectName.Replace("Game", "").Replace("game", ""));
-
-									if (!CommandUtils.DirectoryExists(SrcBuildPath))
-									{
-										SrcBuildPath = SrcBuildPath2;
-									}
-									if (CommandUtils.DirectoryExists(SrcBuildPath))
-									{
-										var JustBuildFolder = JustFile.Replace("-" + Parts.Last(), "");
-
-										string PlatformStr = Platform.ToString();
-										string SrcCLMetaPath = CommandUtils.CombinePaths(SrcBuildPath, JustBuildFolder, PlatformStr, "MetaData");
-										if (CommandUtils.DirectoryExists(SrcCLMetaPath))
-										{
-											string SrcCLMetaPathCollected = CommandUtils.CombinePaths(SrcCLMetaPath, "CollectedPSOs");
-											if (!CommandUtils.DirectoryExists(SrcCLMetaPathCollected))
-											{
-												Log.Info("Creating Directory {0}", SrcCLMetaPathCollected);
-												CommandUtils.CreateDirectory(SrcCLMetaPathCollected);
-											}
-											if (CommandUtils.DirectoryExists(SrcCLMetaPathCollected))
-											{
-												string DestFile = CommandUtils.CombinePaths(SrcCLMetaPathCollected, JustFile);
-												CommandUtils.CopyFile_NoExceptions(ThisFile, DestFile, true);
-												if (CommandUtils.FileExists(true, DestFile))
-												{
-													Log.Info("Deleting local file, copied to {0}", DestFile);
-													CommandUtils.DeleteFile_NoExceptions(ThisFile, true);
-													Copied = true;
-												}
-											}
-										}
-									}
-								}
-							}
-							if (!Copied)
-							{
-								Log.Warning("Could not find anywhere to put this file {0}", JustFile);
-							}
-						}
-					}
-				}
-				catch (Exception Ex)
-				{
-					Log.Info("Failed to copy upipelinecaches to the network {0}", Ex);
-				}
-
-			}
+			// TODO REMOVEME- this should go elsewhere, likely a util that can be called or inserted by relevant test nodes.
+			SavePSOs(InContext, InRunningRole, DestSavedDirInfo.FullName);
 			// END REMOVEME
 
-
 			UnrealLogParser LogParser = new UnrealLogParser(InRunningRole.AppInstance.StdOut);
+			UnrealRoleArtifacts Artifacts = null;
 
-			UnrealRoleArtifacts Artifacts = new UnrealRoleArtifacts(InRunningRole.Role, InRunningRole.AppInstance, InDestArtifactPath, ArtifactLogPath, LogParser);
+			// Only keep one copy of the StdOut where needed
+			if (IsBuildMachine)
+			{
+				// Save to Horde filepath
+				Artifacts = new UnrealRoleArtifacts(InRunningRole.Role, InRunningRole.AppInstance, DestSavedDirInfo.FullName, HordeLogFilePath, LogParser);
+			}
+			else
+			{
+				// Save the Artifact filepath
+				Artifacts = new UnrealRoleArtifacts(InRunningRole.Role, InRunningRole.AppInstance, DestSavedDirInfo.FullName, ArtifactLogFilePath, LogParser);
+			}
 
 			return Artifacts;
 		}
@@ -1381,7 +1344,7 @@ namespace Gauntlet
 
 				string DestPath = Path.Combine(OutputPath, FolderName);
 
-				if (!App.Role.IsNullRole() && !App.Role.InstallOnly )
+				if (!App.Role.IsNullRole() && !App.Role.InstallOnly)
 				{
 					Log.VeryVerbose("Calling SaveRoleArtifacts, Role: {0}  Artifact Path: {1}", App.ToString(), App.AppInstance.ArtifactPath);
 					UnrealRoleArtifacts Artifacts = null;
@@ -1425,5 +1388,124 @@ namespace Gauntlet
 
 			return AllArtifacts;
 		}
+
+		private void SavePSOs(UnrealTestContext InContext, UnrealSessionInstance.RoleInstance InRunningRole, string DestSavedDir)
+		{
+			if (InRunningRole.Role.RoleType.IsServer())
+			{
+				return;
+			}
+
+			if (!InContext.Options.LogPSO)
+			{
+				return;
+			}
+
+			if (!Directory.Exists(DestSavedDir))
+			{
+				try
+				{
+					Directory.CreateDirectory(DestSavedDir);
+				}
+				catch (Exception Ex)
+				{
+					Log.Info($"Archive path '{DestSavedDir}' was not found!\n{Ex.Message}");
+					return;
+				}
+			}
+
+			// Copy over PSOs
+			try
+			{
+				foreach (string ThisFile in CommandUtils.FindFiles_NoExceptions(true, "*.rec.upipelinecache", true, DestSavedDir))
+				{
+					bool Copied = false;
+					string JustFile = Path.GetFileName(ThisFile);
+					if (!JustFile.StartsWith("++"))
+					{
+						continue;
+					}
+
+					string[] Parts = JustFile.Split(new Char[] { '+', '-' }).Where(A => A != "").ToArray();
+					if (Parts.Count() >= 2)
+					{
+						string ProjectName = Parts[0].ToString();
+						string BuildRoot = CommandUtils.CombinePaths(CommandUtils.RootBuildStorageDirectory());
+
+						string SrcBuildPath = CommandUtils.CombinePaths(BuildRoot, ProjectName);
+						string SrcBuildPath2 = CommandUtils.CombinePaths(BuildRoot, ProjectName.Replace("Game", "").Replace("game", ""));
+
+						if (!CommandUtils.DirectoryExists(SrcBuildPath))
+						{
+							SrcBuildPath = SrcBuildPath2;
+						}
+						if (CommandUtils.DirectoryExists(SrcBuildPath))
+						{
+							var JustBuildFolder = JustFile.Replace("-" + Parts.Last(), "");
+
+							string PlatformStr = InRunningRole.Role.Platform.ToString();
+							string SrcCLMetaPath = CommandUtils.CombinePaths(SrcBuildPath, JustBuildFolder, PlatformStr, "MetaData");
+							if (CommandUtils.DirectoryExists(SrcCLMetaPath))
+							{
+								string SrcCLMetaPathCollected = CommandUtils.CombinePaths(SrcCLMetaPath, "CollectedPSOs");
+								if (!CommandUtils.DirectoryExists(SrcCLMetaPathCollected))
+								{
+									Log.Info("Creating Directory {0}", SrcCLMetaPathCollected);
+									CommandUtils.CreateDirectory(SrcCLMetaPathCollected);
+								}
+								if (CommandUtils.DirectoryExists(SrcCLMetaPathCollected))
+								{
+									string DestFile = CommandUtils.CombinePaths(SrcCLMetaPathCollected, JustFile);
+									CommandUtils.CopyFile_NoExceptions(ThisFile, DestFile, true);
+									if (CommandUtils.FileExists(true, DestFile))
+									{
+										Log.Info("Deleting local file, copied to {0}", DestFile);
+										CommandUtils.DeleteFile_NoExceptions(ThisFile, true);
+										Copied = true;
+									}
+								}
+							}
+						}
+					}
+					if (!Copied)
+					{
+						Log.Warning("Could not find anywhere to put this file {0}", JustFile);
+					}
+				}
+			}
+			catch (Exception Ex)
+			{
+				Log.Info("Failed to copy upipelinecaches to the network {0}", Ex);
+			}
+		}
+
+		/// <summary>
+		/// Filter that Truncate long file paths such as CrashReporter files. //UECC-Windows-F0DD9BB04C3C9250FAF39D8AB4A88556//
+		/// These are particularly problematic with testflights which append a long random name to the destination folder, 
+		/// easily pushing past 260 chars.
+		/// </summary>
+		/// <param name="LongFilePath"></param>
+		/// <returns></returns>
+		private string TruncateLongPathFilter(string LongFilePath)
+		{
+			Dictionary<string, string> LongCrashReporterStringToIndex = new Dictionary<string, string>();
+
+			Match RegexMatch = Regex.Match(LongFilePath, @"UECC-.+-([\dA-Fa-f]+)");
+
+			if (RegexMatch.Success)
+			{
+				string LongString = RegexMatch.Groups[1].ToString();
+
+				if (!LongCrashReporterStringToIndex.ContainsKey(LongString))
+				{
+					LongCrashReporterStringToIndex[LongString] = LongCrashReporterStringToIndex.Keys.Count.ToString("D2");
+				}
+
+				string ShortSring = LongCrashReporterStringToIndex[LongString];
+				LongFilePath = LongFilePath.Replace(LongString, ShortSring);
+			}
+
+			return LongFilePath;
+		}
 	}
 }
\ No newline at end of file
