Log when DumpManifest is called to help debug what is overwriting a manifest
#rnx

diff --git a/Engine/Source/Programs/AutomationTool/Scripts/CopyBuildToStagingDirectory.Automation.cs b/Engine/Source/Programs/AutomationTool/Scripts/CopyBuildToStagingDirectory.Automation.cs
--- a/Engine/Source/Programs/AutomationTool/Scripts/CopyBuildToStagingDirectory.Automation.cs
+++ b/Engine/Source/Programs/AutomationTool/Scripts/CopyBuildToStagingDirectory.Automation.cs
@@ -2122,6 +2122,9 @@ namespace AutomationScripts
 		{
 			if (Mapping.Count > 0)
 			{
+				int NumItems = Mapping.Count;
+				Logger.LogInformation("DumpManifest: {Filename}, NumItems: {NumItems}", Filename, NumItems);
+
 				List<string> Lines = new List<string>();
 				foreach (KeyValuePair<StagedFileReference, FileReference> Pair in Mapping)
 				{
