Horde: Split out report test name to a property that can be overridden
#rnx

diff --git a/Engine/Source/Programs/AutomationTool/Gauntlet/Unreal/Base/Gauntlet.UnrealTestNode.cs b/Engine/Source/Programs/AutomationTool/Gauntlet/Unreal/Base/Gauntlet.UnrealTestNode.cs
--- a/Engine/Source/Programs/AutomationTool/Gauntlet/Unreal/Base/Gauntlet.UnrealTestNode.cs
+++ b/Engine/Source/Programs/AutomationTool/Gauntlet/Unreal/Base/Gauntlet.UnrealTestNode.cs
@@ -1343,6 +1343,22 @@ namespace Gauntlet
 			return null;
 		}
 
+		/// <summary>
+		/// Optional override for the test report name, useful when a single test node has different testing modes
+		/// </summary>
+		protected virtual string HordeReportTestName { 
+			get 
+			{
+				string ReportName = Name;
+				if (ReportName.Split('.').Length > 1)
+				{
+					ReportName = ReportName.Split('.').Last();
+				}
+
+				return ReportName;				
+			}
+		}
+
 		/// <summary>
 		/// Generate a Simple Test Report from the results of this test
 		/// </summary>
@@ -1354,14 +1370,8 @@ namespace Gauntlet
 				GetCachedConfiguration().HordeTestDataKey = Name + " " + GetMainRoleContextString();
 			}
 
-			string TestName = Name;
-			if (TestName.Split('.').Length > 1)
-			{
-				TestName = TestName.Split('.').Last();
-			}
-
 			HordeReport.SimpleTestReport HordeTestReport = new HordeReport.SimpleTestReport(UnrealTestResult, Context, GetCachedConfiguration());
-			HordeTestReport.TestName = TestName;
+			HordeTestReport.TestName = HordeReportTestName;
 			HordeTestReport.ReportCreatedOn = DateTime.Now.ToString();
 			HordeTestReport.TotalDurationSeconds = (float) (DateTime.Now - SessionStartTime).TotalSeconds;
 			HordeTestReport.Description = GetMainRoleContextString();
