Adding AdditionalArgs LaunchIt field to RunUnreal
#rnx

diff --git a/Engine/Source/Programs/AutomationTool/Gauntlet/Unreal/RunUnreal.cs b/Engine/Source/Programs/AutomationTool/Gauntlet/Unreal/RunUnreal.cs
--- a/Engine/Source/Programs/AutomationTool/Gauntlet/Unreal/RunUnreal.cs
+++ b/Engine/Source/Programs/AutomationTool/Gauntlet/Unreal/RunUnreal.cs
@@ -61,6 +61,7 @@ namespace Gauntlet
 	[ParamHelp("ServerArgs", "Extra arguments passed to server role(s)")]
 	[ParamHelp("EditorArgs", "Extra arguments passed to editor role(s)")]
 	[ParamHelp("Namespaces", "Comma-separated list of namespaces to check for tests.", MultiSelectSeparator = ",")]
+	[ParamHelp("AdditionalArgs", "Any additional arguments to pass directly to the Gauntlet command line", IsArgument = true)]
 	public class RunUnreal : BuildCommand
 	{
 		/// <summary>
