Add new target for staging the crash reporter for Windows and Linux so it can be incorporated into the crash-reporter image
#RNX

diff --git a/Engine/Source/Programs/AutomationTool/BuildGraph/Tasks/StageTask.cs b/Engine/Source/Programs/AutomationTool/BuildGraph/Tasks/StageTask.cs
--- a/Engine/Source/Programs/AutomationTool/BuildGraph/Tasks/StageTask.cs
+++ b/Engine/Source/Programs/AutomationTool/BuildGraph/Tasks/StageTask.cs
@@ -110,7 +110,7 @@ namespace AutomationTool.Tasks
 			// Get the output directories. We flatten the directory structure on output.
 			DirectoryReference TargetDir = Parameters.ToDir;
 			DirectoryReference TargetEngineDir = DirectoryReference.Combine(TargetDir, "Engine");
-			DirectoryReference TargetProjectDir = DirectoryReference.Combine(TargetDir, ProjectFile.GetFileNameWithoutExtension());
+			DirectoryReference TargetProjectDir = (ProjectFile == null) ? TargetEngineDir : DirectoryReference.Combine(TargetDir, ProjectFile.GetFileNameWithoutExtension());
 
 			// Get the path to the receipt
 			FileReference ReceiptFileName = TargetReceipt.GetDefaultPath(SourceProjectDir, Parameters.Target, Parameters.Platform, Parameters.Configuration, UnrealArchitectures.FromString(Parameters.Architecture, Parameters.Platform));
