Removing mono/xbuild part 1

* Updating references to use dotnet
* Add new RunDotnet scripts
* Removing Mono sh scripts and references to them

#rnx

diff --git a/Engine/Build/BatchFiles/Linux/Build.sh b/Engine/Build/BatchFiles/Linux/Build.sh
--- a/Engine/Build/BatchFiles/Linux/Build.sh
+++ b/Engine/Build/BatchFiles/Linux/Build.sh
@@ -6,7 +6,7 @@ set -e
 
 cd "`dirname "$0"`/../../../.." 
 
-# Setup Environment and Mono
+# Setup Environment and dotnet
 source Engine/Build/BatchFiles/Linux/SetupEnvironment.sh -dotnet Engine/Build/BatchFiles/Linux
 
 # Skip UBT and SCW compile step if this is an installed build.
diff --git a/Engine/Build/BatchFiles/Linux/BuildThirdParty.sh b/Engine/Build/BatchFiles/Linux/BuildThirdParty.sh
--- a/Engine/Build/BatchFiles/Linux/BuildThirdParty.sh
+++ b/Engine/Build/BatchFiles/Linux/BuildThirdParty.sh
@@ -350,53 +350,6 @@ Buildcoremod()
   set +x
 }
 
-BuildMono()
-{
-  echo "building mono"
-  set -x
-
-  export UE_MONO_DIR=$TOP_DIR/Binaries/ThirdParty/Mono/Linux
-  mkdir -p $UE_MONO_DIR
- 
-  for file in $UE_MONO_DIR/$TARGET_ARCH/*
-  do
-    P4Open $file
-  done
-  P4Open $UE_MONO_DIR/lib/mono/4.5/mscorlib.dll
-
-  if [ ! -d mono ]; then
-	git clone -b mono-4.0.0-branch [REDACTED_PROFILE_URL]/mono.git
-  fi
-  cd mono
-  git reset --hard c99aa0c
-  set +e
-  # To add support for terminfo2
-  git cherry-pick 2c1f45f3791f274855e0f5fd2fb0af71c9a756f7
-  set -e
-  # This part of the change that needs resolving is not needed
-  git checkout c99aa0c -- configure.ac 
-
-  # This build is a lighter build intended to be shared with Engine/Binaries/ThirdParty/Mono/Mac
-  ./autogen.sh --prefix=$UE_MONO_DIR --enable-minimal=aot --disable-libraries --with-mcs-docs=no
-  make monolite_url=http://download.mono-project.com/monolite/monolite-117-latest.tar.gz get-monolite-latest
-  make XTERNAL_MCS=${PWD}/mcs/class/lib/monolite/basic.exe $MAKE_ARGS 
-  
-  mkdir -p $UE_MONO_DIR/$TARGET_ARCH/bin/
-  cp --remove-destination mono/mini/mono-sgen $UE_MONO_DIR/$TARGET_ARCH/bin/
-  cp --remove-destination mono/mini/mono-boehm $UE_MONO_DIR/$TARGET_ARCH/bin/
-  
-  mkdir -p $UE_MONO_DIR/$TARGET_ARCH/lib/
-  find support/.libs/ -name "*.so" -exec cp --remove-destination {} $UE_MONO_DIR/$TARGET_ARCH/lib/ \;
-  find ikvm-native/.libs/ -name "*.so" -exec cp --remove-destination {} $UE_MONO_DIR/$TARGET_ARCH/lib/ \;
-  find mcs/class/lib/net_4_5/ -name "*.so" -exec cp --remove-destination {} $UE_MONO_DIR/$TARGET_ARCH/lib/ \;
-  cp --remove-destination mcs/class/lib/net_4_5/mscorlib.dll $UE_MONO_DIR/lib/mono/4.5/
-
-  cd ..
-  rm -rf mono
-
-  set +x
-}
-
 BuildGoogleTest()
 {
   echo "building GoogleTest"
@@ -529,7 +482,6 @@ print_valid_build_opts() {
   echo "    coremod"
   echo "    Toolchain"
   echo "    ICU"
-  echo "    Mono"
   echo "    GoogleTest"
   echo "    LLVMCompilerRt"
   echo "    Strophe"
diff --git a/Engine/Build/BatchFiles/Linux/FixMonoFiles.sh b/Engine/Build/BatchFiles/Linux/FixMonoFiles.sh
deleted file mode 100755
diff --git a/Engine/Build/BatchFiles/Linux/GitDependencies.sh b/Engine/Build/BatchFiles/Linux/GitDependencies.sh
--- a/Engine/Build/BatchFiles/Linux/GitDependencies.sh
+++ b/Engine/Build/BatchFiles/Linux/GitDependencies.sh
@@ -1,9 +1,6 @@
 #!/bin/bash
 # Copyright Epic Games, Inc. All Rights Reserved.
 
-# Certain versions of mono can crash while running GitDependencies
-# This script checks if that happened and simply re-runs the tool.
-
 ARGS=$@
 
 SCRIPT_PATH=$0
diff --git a/Engine/Build/BatchFiles/Linux/RunMono.sh b/Engine/Build/BatchFiles/Linux/RunMono.sh
deleted file mode 100755
diff --git a/Engine/Build/BatchFiles/Linux/RunXBuild.sh b/Engine/Build/BatchFiles/Linux/RunXBuild.sh
deleted file mode 100755
diff --git a/Engine/Build/BatchFiles/Linux/Setup.sh b/Engine/Build/BatchFiles/Linux/Setup.sh
--- a/Engine/Build/BatchFiles/Linux/Setup.sh
+++ b/Engine/Build/BatchFiles/Linux/Setup.sh
@@ -23,7 +23,7 @@ if [ -e /etc/os-release ]; then
   # Ubuntu/Debian/Mint
   if [[ "$ID" == "ubuntu" ]] || [[ "$ID_LIKE" == "ubuntu" ]] || [[ "$ID" == "debian" ]] || [[ "$ID_LIKE" == "debian" ]] || [[ "$ID" == "tanglu" ]] || [[ "$ID_LIKE" == "tanglu" ]]; then
 
-    # Both mono and the compiler toolchain are now bundled.
+    # Both dotnet and the compiler toolchain are now bundled.
     # build-essential can still be useful for making sure the system has make
     DEPS="
       build-essential
diff --git a/Engine/Build/BatchFiles/Linux/SetupEnvironment.sh b/Engine/Build/BatchFiles/Linux/SetupEnvironment.sh
--- a/Engine/Build/BatchFiles/Linux/SetupEnvironment.sh
+++ b/Engine/Build/BatchFiles/Linux/SetupEnvironment.sh
@@ -3,7 +3,7 @@
 # This script is called by UE tools to source environment variables from the users .unrealrc script(s)
 # It expects to be called from the engine directory and the order is ~/.unrealrc then <workspace>/.unrealrc
 
-# To setup the environment and mono, use SetupEnvironment.sh -mono Engine/Build/BatchFiles/Mac
+# To setup the environment and dotnet, use SetupEnvironment.sh -dotnet Engine/Build/BatchFiles/Mac
 # (It's not clear why we can't just determine that path outselves, but for legacy reasons we'll keep it
 # as a param)
 
@@ -18,11 +18,6 @@ fi
 while test $# -gt 0
 do
     case "$1" in
-        -mono) 
-            # echo "setting up mono"
-            source "$2/SetupMono.sh" "$2"
-            shift
-            ;;
         -dotnet) 
             # echo "setting up dotnet"
             source "$2/SetupDotnet.sh" "$2"
diff --git a/Engine/Build/BatchFiles/Linux/SetupMono.sh b/Engine/Build/BatchFiles/Linux/SetupMono.sh
deleted file mode 100755
diff --git a/Engine/Build/BatchFiles/Mac/FixMonoFiles.sh b/Engine/Build/BatchFiles/Mac/FixMonoFiles.sh
deleted file mode 100755
diff --git a/Engine/Build/BatchFiles/Mac/RunMono.sh b/Engine/Build/BatchFiles/Mac/RunMono.sh
deleted file mode 100755
diff --git a/Engine/Build/BatchFiles/Mac/RunXBuild.sh b/Engine/Build/BatchFiles/Mac/RunXBuild.sh
deleted file mode 100755
diff --git a/Engine/Build/BatchFiles/Mac/SetupEnvironment.sh b/Engine/Build/BatchFiles/Mac/SetupEnvironment.sh
--- a/Engine/Build/BatchFiles/Mac/SetupEnvironment.sh
+++ b/Engine/Build/BatchFiles/Mac/SetupEnvironment.sh
@@ -3,7 +3,7 @@
 # This script is called by UE tools to source environment variables from the users .unrealrc script(s)
 # It expects to be called from the engine directory and the order is ~/.unrealrc then <workspace>/.unrealrc
 
-# To setup the environment and mono, use SetupEnvironment.sh -mono Engine/Build/BatchFiles/Mac
+# To setup the environment and dotnet, use SetupEnvironment.sh -dotnet Engine/Build/BatchFiles/Mac
 # (It's not clear why we can't just determine that path outselves, but for legacy reasons we'll keep it
 # as a param)
 
@@ -18,11 +18,6 @@ fi
 while test $# -gt 0
 do
     case "$1" in
-        -mono) 
-            # echo "setting up mono"
-            source "$2/SetupMono.sh" "$2"
-            shift
-            ;;
         -dotnet) 
             # echo "setting up dotnet"
             source "$2/SetupDotnet.sh" "$2"
diff --git a/Engine/Build/BatchFiles/Mac/SetupMono.sh b/Engine/Build/BatchFiles/Mac/SetupMono.sh
deleted file mode 100755
diff --git a/Engine/Build/BatchFiles/RunDotnet.sh b/Engine/Build/BatchFiles/RunDotnet.sh
new file mode 100755
--- /dev/null
+++ b/Engine/Build/BatchFiles/RunDotnet.sh
@@ -0,0 +1,37 @@
+#!/bin/bash
+
+## Unreal Engine dotnet setup script
+## Copyright Epic Games, Inc. All Rights Reserved.
+
+## This script is expecting to exist in the UE/Engine/Build/BatchFiles directory.  It will not work correctly
+## if you copy it to a different location and run it.
+
+EnvironmentType=-dotnet
+
+# If ran from somewhere other then the script location we'll have the full base path
+BASE_PATH="`dirname "$0"`"
+
+echo
+echo Running Dotnet...
+echo
+
+if [ "$(uname)" = "Darwin" ]; then
+	# Setup Environment
+	source "$BASE_PATH/Mac/SetupEnvironment.sh" $EnvironmentType "$BASE_PATH/Mac"
+fi
+
+if [ "$(uname)" = "Linux" ]; then
+	# Setup Environment
+	source "$BASE_PATH/Linux/SetupEnvironment.sh" $EnvironmentType "$BASE_PATH/Linux"
+fi
+
+# put ourselves into Engine directory (two up from location of this script)
+pushd "`dirname "$0"`/../.."
+
+if [ ! -f Build/BatchFiles/RunDotnet.sh ]; then
+	echo RunDotnet ERROR: The batch file does not appear to be located in the /Engine/Build/BatchFiles directory.  This script must be run from within that directory.
+	exit 1
+fi
+
+dotnet "$@"
+exit $?
diff --git a/Engine/Build/BatchFiles/RunDotnetMSBuild.sh b/Engine/Build/BatchFiles/RunDotnetMSBuild.sh
new file mode 100755
--- /dev/null
+++ b/Engine/Build/BatchFiles/RunDotnetMSBuild.sh
@@ -0,0 +1,28 @@
+#!/bin/bash
+
+## Unreal Engine dotnet setup script
+## Copyright Epic Games, Inc. All Rights Reserved.
+
+## This script is expecting to exist in the UE/Engine/Build/BatchFiles directory.  It will not work correctly
+## if you copy it to a different location and run it.
+
+echo
+echo Running dotnet msbuild...
+echo
+
+source "`dirname "$0"`/SetupEnvironment.sh" -dotnet "`dirname "$0"`"
+
+# put ourselves into Engine directory (two up from location of this script)
+pushd "`dirname "$0"`/../../.."
+
+if [ ! -f Build/BatchFiles/RunDotnetMSBuild.sh ]; then
+	echo RunDotnetMSBuild ERROR: The batch file does not appear to be located in the /Engine/Build/BatchFiles directory.  This script must be run from within that directory.
+	exit 1
+fi
+
+dotnet msbuild /verbosity:quiet /nologo "$@" |grep -wi error
+if [ $? -ne 1 ]; then
+	exit 1
+else
+	exit 0
+fi
diff --git a/Engine/Build/BatchFiles/RunUAT.sh b/Engine/Build/BatchFiles/RunUAT.sh
--- a/Engine/Build/BatchFiles/RunUAT.sh
+++ b/Engine/Build/BatchFiles/RunUAT.sh
@@ -81,13 +81,11 @@ UATDirectory=Binaries/DotNET/AutomationTool
 
 if [ "$(uname)" = "Darwin" ]; then
 	# Setup Environment
-	source "$SCRIPT_DIR/Mac/SetupEnvironment.sh" -mono "$SCRIPT_DIR/Mac" # ~0.4s :(
 	source "$SCRIPT_DIR/Mac/SetupEnvironment.sh" $EnvironmentType "$SCRIPT_DIR/Mac"
 fi
 
 if [ "$(uname)" = "Linux" ]; then
 	# Setup Environment
-	source "$SCRIPT_DIR/Linux/SetupEnvironment.sh" -mono "$SCRIPT_DIR/Linux"
 	source "$SCRIPT_DIR/Linux/SetupEnvironment.sh" $EnvironmentType "$SCRIPT_DIR/Linux"
 fi
 
diff --git a/Engine/Build/InstalledEngineFilters.xml b/Engine/Build/InstalledEngineFilters.xml
--- a/Engine/Build/InstalledEngineFilters.xml
+++ b/Engine/Build/InstalledEngineFilters.xml
@@ -334,14 +334,12 @@
 		Engine/Build/BatchFiles/Mac/FixDependencyFiles.sh
 		Engine/Build/BatchFiles/Mac/GenerateProjectFiles.sh
 		Engine/Build/BatchFiles/Mac/SetupEnvironment.sh
-		Engine/Build/BatchFiles/Mac/SetupMono.sh
 		Engine/Build/BatchFiles/Mac/SetupDotnet.sh
-		Engine/Build/BatchFiles/Mac/FixMonoFiles.sh
-		Engine/Build/BatchFiles/Mac/RunMono.sh
-		Engine/Build/BatchFiles/Mac/RunXBuild.sh
 		Engine/Build/BatchFiles/Mac/XcodeBuild.sh
-        Engine/Build/BatchFiles/Mac/GenerateUniversalDSYM.sh
-        Engine/Build/BatchFiles/Mac/UpdateVersionAfterBuild.sh
+		Engine/Build/BatchFiles/Mac/GenerateUniversalDSYM.sh
+		Engine/Build/BatchFiles/Mac/UpdateVersionAfterBuild.sh
+		Engine/Build/BatchFiles/RunDotnet.sh
+		Engine/Build/BatchFiles/RunDotnetMSBuild.sh
 		Engine/Build/BatchFiles/RunUAT.command
 		Engine/Build/BatchFiles/RunUAT.sh
 		Engine/Build/BatchFiles/RunUBT.sh
@@ -349,9 +347,6 @@
         <!-- Mac build resources (plist entitlement etc.) -->
         Engine/Build/Mac/Resources/...
 
-		<!-- Mono -->
-		Engine/Binaries/ThirdParty/Mono/Mac/...
-
 		<!-- DotNet -->
 		Engine/Binaries/ThirdParty/DotNet/6.0.302/mac-arm64/...
 		Engine/Binaries/ThirdParty/DotNet/6.0.302/mac-x64/...
@@ -390,11 +385,9 @@
 		Engine/Build/BatchFiles/Linux/QASmokeManual.sh
 		Engine/Build/BatchFiles/Linux/Setup.sh
 		Engine/Build/BatchFiles/Linux/SetupEnvironment.sh
-		Engine/Build/BatchFiles/Linux/SetupMono.sh
 		Engine/Build/BatchFiles/Linux/SetupDotnet.sh
-		Engine/Build/BatchFiles/Linux/FixMonoFiles.sh
-		Engine/Build/BatchFiles/Linux/RunMono.sh
-		Engine/Build/BatchFiles/Linux/RunXBuild.sh
+		Engine/Build/BatchFiles/RunDotnetMSBuild.sh
+		Engine/Build/BatchFiles/RunDotnet.sh
 		Engine/Build/BatchFiles/RunUAT.sh
 		Engine/Build/BatchFiles/RunUBT.sh
 
@@ -402,9 +395,6 @@
 		Engine/Binaries/Linux/dump_syms
 		Engine/Binaries/Linux/BreakpadSymbolEncoder
 
-		<!-- Mono -->
-		Engine/Binaries/ThirdParty/Mono/...
-
 		<!-- DotNet -->
 		Engine/Binaries/ThirdParty/DotNet/6.0.302/linux/...
 		Engine/Binaries/ThirdParty/VisualStudio/...
@@ -545,7 +535,6 @@
 	<Property Name="ExtraSignFilesIOS">
 		Engine/Binaries/ThirdParty/DotNet/6.0.302/mac-arm64/...
 		Engine/Binaries/ThirdParty/DotNet/6.0.302/mac-x64/...
-		Engine/Binaries/ThirdParty/Mono/Mac/...
 		Engine/Extras/ThirdPartyNotUE/libimobiledevice/x64/...
 	</Property>
 
@@ -585,7 +574,6 @@
 	<Property Name="ExtraSignFilesTVOS">
 		Engine/Binaries/ThirdParty/DotNet/6.0.302/mac-arm64/...
 		Engine/Binaries/ThirdParty/DotNet/6.0.302/mac-x64/...
-		Engine/Binaries/ThirdParty/Mono/Mac/...
 		Engine/Extras/ThirdPartyNotUE/libimobiledevice/x64/...
 	</Property>
 
@@ -614,7 +602,6 @@
 		Engine/Binaries/Linux/AgentInterface.dll
 		Engine/Binaries/Linux/dump_syms.exe
 		Engine/Binaries/Linux/BreakpadSymbolEncoder.exe
-		Engine/Binaries/ThirdParty/Mono/Linux/...
 		Engine/Binaries/Linux/UnrealTraceServer
 	</Property>
 
diff --git a/Engine/Build/Rsync/RsyncEngineScripts.txt b/Engine/Build/Rsync/RsyncEngineScripts.txt
--- a/Engine/Build/Rsync/RsyncEngineScripts.txt
+++ b/Engine/Build/Rsync/RsyncEngineScripts.txt
@@ -1,11 +1,11 @@
 # This file contains a list of shell scripts to be synced to the remote. These files will be locally converted to use Unix line endings, and will be
 # marked as executable on the remote. Wildcards are not supported.
 
-Binaries/ThirdParty/Mono/Mac/bin/xbuild
+Binaries/ThirdParty/DotNet/6.0.302/mac-x64/
+Binaries/ThirdParty/DotNet/6.0.302/mac-armx64/
+Build/BatchFiles/RunDotnet.sh
+Build/BatchFiles/RunDotnetMSBuild.sh
 Build/BatchFiles/Mac/Build.sh
 Build/BatchFiles/Mac/FixDependencyFiles.sh
-Build/BatchFiles/Mac/RunMono.sh
-Build/BatchFiles/Mac/FixMonoFiles.sh
 Build/BatchFiles/Mac/GenerateProjectFiles.sh
-Build/BatchFiles/Mac/SetupMono.sh
 Build/BatchFiles/Mac/SetupEnvironment.sh
diff --git a/Engine/Source/Developer/DesktopPlatform/Private/DesktopPlatformBase.cpp b/Engine/Source/Developer/DesktopPlatform/Private/DesktopPlatformBase.cpp
--- a/Engine/Source/Developer/DesktopPlatform/Private/DesktopPlatformBase.cpp
+++ b/Engine/Source/Developer/DesktopPlatform/Private/DesktopPlatformBase.cpp
@@ -1709,13 +1709,13 @@ bool FDesktopPlatformBase::BuildUnrealBuildTool(const FString& RootDir, FOutputD
 	}
 	else if (PLATFORM_MAC)
 	{
-		FString ScriptPath = FPaths::ConvertRelativePathToFull(RootDir / TEXT("Engine/Build/BatchFiles/Mac/RunXBuild.sh"));
+		FString ScriptPath = FPaths::ConvertRelativePathToFull(RootDir / TEXT("Engine/Build/BatchFiles/RunDotnetMSBuild.sh"));
 		CompilerExecutableFilename = TEXT("/bin/sh");
 		CmdLineParams = FString::Printf(TEXT("\"%s\" /property:Configuration=Development %s"), *ScriptPath, *CsProjLocation);
 	}
 	else if (PLATFORM_LINUX)
 	{
-		FString ScriptPath = FPaths::ConvertRelativePathToFull(RootDir / TEXT("Engine/Build/BatchFiles/Linux/RunXBuild.sh"));
+		FString ScriptPath = FPaths::ConvertRelativePathToFull(RootDir / TEXT("Engine/Build/BatchFiles/RunDotnetMSBuild.sh"));
 		CompilerExecutableFilename = TEXT("/bin/bash");
 		CmdLineParams = FString::Printf(TEXT("\"%s\" /property:Configuration=Development %s"), *ScriptPath, *CsProjLocation);
 	}
diff --git a/Engine/Source/Developer/IOS/IOSPlatformEditor/Private/IOSTargetSettingsCustomization.cpp b/Engine/Source/Developer/IOS/IOSPlatformEditor/Private/IOSTargetSettingsCustomization.cpp
--- a/Engine/Source/Developer/IOS/IOSPlatformEditor/Private/IOSTargetSettingsCustomization.cpp
+++ b/Engine/Source/Developer/IOS/IOSPlatformEditor/Private/IOSTargetSettingsCustomization.cpp
@@ -1355,7 +1355,7 @@ void FIOSTargetSettingsCustomization::FindRequiredFiles()
 	BundleIdentifier = BundleIdentifier.Replace(TEXT("_"), TEXT(""));
 #if PLATFORM_MAC
 	FString CmdExe = TEXT("/bin/sh");
-	FString ScriptPath = FPaths::ConvertRelativePathToFull(FPaths::EngineDir() / TEXT("Build/BatchFiles/Mac/RunMono.sh"));
+	FString ScriptPath = FPaths::ConvertRelativePathToFull(FPaths::EngineDir() / TEXT("Build/BatchFiles/RunDotnet.sh"));
 	FString IPPPath = FPaths::ConvertRelativePathToFull(FPaths::EngineDir() / TEXT("Binaries/DotNET/IOS/IPhonePackager.exe"));
 	FString CommandLine = FString::Printf(TEXT("\"%s\" \"%s\" certificates Engine -bundlename \"%s\""), *ScriptPath, *IPPPath, *(BundleIdentifier));
 #else
@@ -1435,7 +1435,7 @@ FReply FIOSTargetSettingsCustomization::OnInstallProvisionClicked()
 		BundleIdentifier = BundleIdentifier.Replace(TEXT("_"), TEXT(""));
 #if PLATFORM_MAC
 		FString CmdExe = TEXT("/bin/sh");
-		FString ScriptPath = FPaths::ConvertRelativePathToFull(FPaths::EngineDir() / TEXT("Build/BatchFiles/Mac/RunMono.sh"));
+		FString ScriptPath = FPaths::ConvertRelativePathToFull(FPaths::EngineDir() / TEXT("Build/BatchFiles/RunDotnet.sh"));
 		FString IPPPath = FPaths::ConvertRelativePathToFull(FPaths::EngineDir() / TEXT("Binaries/DotNET/IOS/IPhonePackager.exe"));
 		FString CommandLine = FString::Printf(TEXT("\"%s\" \"%s\" Install Engine -project \"%s\" -provision \"%s\" -bundlename \"%s\""), *ScriptPath, *IPPPath, *ProjectPath, *ProvisionPath, *BundleIdentifier);
 #else
@@ -1550,7 +1550,7 @@ FReply FIOSTargetSettingsCustomization::OnInstallCertificateClicked()
 		CertPath = FPaths::ConvertRelativePathToFull(OpenFilenames[0]);
 #if PLATFORM_MAC
 		FString CmdExe = TEXT("/bin/sh");
-		FString ScriptPath = FPaths::ConvertRelativePathToFull(FPaths::EngineDir() / TEXT("Build/BatchFiles/Mac/RunMono.sh"));
+		FString ScriptPath = FPaths::ConvertRelativePathToFull(FPaths::EngineDir() / TEXT("Build/BatchFiles/RunDotnet.sh"));
 		FString IPPPath = FPaths::ConvertRelativePathToFull(FPaths::EngineDir() / TEXT("Binaries/DotNET/IOS/IPhonePackager.exe"));
 		FString CommandLine = FString::Printf(TEXT("\"%s\" \"%s\" Install Engine -project \"%s\" -certificate \"%s\" -bundlename \"%s\""), *ScriptPath, *IPPPath, *ProjectPath, *CertPath, *BundleIdentifier);
 #else
diff --git a/Engine/Source/Developer/IOS/IOSTargetPlatform/Private/IOSTargetPlatform.cpp b/Engine/Source/Developer/IOS/IOSTargetPlatform/Private/IOSTargetPlatform.cpp
--- a/Engine/Source/Developer/IOS/IOSTargetPlatform/Private/IOSTargetPlatform.cpp
+++ b/Engine/Source/Developer/IOS/IOSTargetPlatform/Private/IOSTargetPlatform.cpp
@@ -276,7 +276,7 @@ int32 FIOSTargetPlatform::CheckRequirements(bool bProjectHasCode, EBuildConfigur
 	FString ProjectPath = FPaths::ConvertRelativePathToFull(FPaths::GetProjectFilePath());
 #if PLATFORM_MAC
     FString CmdExe = TEXT("/bin/sh");
-    FString ScriptPath = FPaths::ConvertRelativePathToFull(FPaths::EngineDir() / TEXT("Build/BatchFiles/Mac/RunMono.sh"));
+    FString ScriptPath = FPaths::ConvertRelativePathToFull(FPaths::EngineDir() / TEXT("Build/BatchFiles/RunDotnet.sh"));
     FString IPPPath = FPaths::ConvertRelativePathToFull(FPaths::EngineDir() / TEXT("Binaries/DotNET/IOS/IPhonePackager.exe"));
 	FString CommandLine = FString::Printf(TEXT("\"%s\" \"%s\" Validate Engine -project \"%s\" -bundlename \"%s\" -teamID \"%s\" %s %s"), *ScriptPath, *IPPPath, *ProjectPath, *(BundleIdentifier), *(TeamID), (bForDistribtion ? TEXT("-distribution") : TEXT("")), bAutomaticSigning ? TEXT("-autosigning") : TEXT(""));
 #else
diff --git a/Engine/Source/Programs/AutomationTool/AutomationUtils/HostPlatform.cs b/Engine/Source/Programs/AutomationTool/AutomationUtils/HostPlatform.cs
--- a/Engine/Source/Programs/AutomationTool/AutomationUtils/HostPlatform.cs
+++ b/Engine/Source/Programs/AutomationTool/AutomationUtils/HostPlatform.cs
@@ -38,7 +38,7 @@ namespace AutomationTool
 		}
 
 		/// <summary>
-		/// Gets the build executable filename for NET Framework projects e.g. msbuild, or xbuild
+		/// Gets the build executable filename for NET Framework projects e.g. msbuild
 		/// </summary>
 		/// <returns></returns>
 		abstract public string GetFrameworkMsbuildExe();
diff --git a/Engine/Source/Programs/AutomationTool/AutomationUtils/LinuxHostPlatform.cs b/Engine/Source/Programs/AutomationTool/AutomationUtils/LinuxHostPlatform.cs
--- a/Engine/Source/Programs/AutomationTool/AutomationUtils/LinuxHostPlatform.cs
+++ b/Engine/Source/Programs/AutomationTool/AutomationUtils/LinuxHostPlatform.cs
@@ -9,6 +9,7 @@ using System.IO;
 using UnrealBuildTool;
 using EpicGames.Core;
 using UnrealBuildBase;
+using Microsoft.Extensions.Logging;
 
 namespace AutomationTool
 {
@@ -18,21 +19,20 @@ namespace AutomationTool
 
 		public override string GetFrameworkMsbuildExe()
 		{
-			// As of 5.0 mono comes with msbuild which performs better. If that's installed then use it
+			// Look for dotnet, we only support dotnet.
 			if (string.IsNullOrEmpty(CachedFrameworkMsbuildTool))
 			{
-				int Value;
-				bool CanUseMsBuild = (int.TryParse(Environment.GetEnvironmentVariable("UE_USE_SYSTEM_MONO"), out Value) &&
-						Value != 0 &&
-						!string.IsNullOrEmpty(CommandUtils.WhichApp("msbuild")));
+				bool CanUseMsBuild = string.IsNullOrEmpty(CommandUtils.WhichApp("dotnet")) == false;
 
 				if (CanUseMsBuild)
 				{
-					CachedFrameworkMsbuildTool = "msbuild";
+					Logger.LogInformation($"using {CommandUtils.WhichApp("dotnet")}!");
+
+					CachedFrameworkMsbuildTool = "dotnet msbuild";
 				}
 				else
 				{
-					CachedFrameworkMsbuildTool = "xbuild";
+					throw new BuildException("Unable to find installation of dotnet.");
 				}
 			}
 
@@ -81,34 +81,15 @@ namespace AutomationTool
 		public override Process CreateProcess(string AppName)
 		{
 			var NewProcess = new Process();
-			if (AppName == "mono")
-			{
-				// Enable case-insensitive mode for Mono
-				if (!NewProcess.StartInfo.EnvironmentVariables.ContainsKey("MONO_IOMAP"))
-				{
-					NewProcess.StartInfo.EnvironmentVariables.Add("MONO_IOMAP", "case");
-				}
-			}
 			return NewProcess;
 		}
 
 		public override void SetupOptionsForRun(ref string AppName, ref CommandUtils.ERunOptions Options, ref string CommandLine)
 		{
-			if (AppName == "sh" || AppName == "xbuild" || AppName == "codesign")
+			if (AppName == "sh" || AppName == "codesign")
 			{
 				Options &= ~CommandUtils.ERunOptions.AppMustExist;
 			}
-			if (AppName == "xbuild")
-			{
-				AppName = "xbuild";
-				CommandLine = (String.IsNullOrEmpty(CommandLine) ? "" : CommandLine) + " /verbosity:quiet /nologo";
-				// Pass #define MONO to all the automation scripts
-				CommandLine += " /p:DefineConstants=MONO";
-				CommandLine += " /p:DefineConstants=__MonoCS__";
-				// Some projects have TargetFrameworkProfile=Client which causes warnings on Linux
-				// so force it to empty.
-				CommandLine += " /p:TargetFrameworkProfile=";
-			}
 			if (AppName.EndsWith(".exe") || ((AppName.Contains("/Binaries/Win64/") || AppName.Contains("/Binaries/Linux/")) && string.IsNullOrEmpty(Path.GetExtension(AppName))))
 			{
 				if (AppName.Contains("/Binaries/Win64/") || AppName.Contains("/Binaries/Linux/"))
@@ -118,16 +99,16 @@ namespace AutomationTool
 					AppName = AppName.Replace("-Cmd.exe", "");
 					AppName = AppName.Replace(".exe", "");
 				}
-				// some of our C# applications are converted to dotnet core, do not run those via mono
+				// some of our C# applications are converted to dotnet core, do not run those via dotnet
 				else if (AppName.Contains("UnrealBuildTool") || AppName.Contains("AutomationTool"))
 				{
 					Options &= ~CommandUtils.ERunOptions.AppMustExist;
 				}
 				else
 				{
-					// It's a C# app, so run it with Mono
+					// It's a C# app, so run it with dotnet
 					CommandLine = "\"" + AppName + "\" " + (String.IsNullOrEmpty(CommandLine) ? "" : CommandLine);
-					AppName = "mono";
+					AppName = "dotnet";
 					Options &= ~CommandUtils.ERunOptions.AppMustExist;
 				}
 			}
@@ -135,7 +116,7 @@ namespace AutomationTool
 
 		public override void SetConsoleCtrlHandler(ProcessManager.CtrlHandlerDelegate Handler)
 		{
-			// @todo: add mono support
+			// @todo: add dotnet support
 		}
 
 		public override UnrealTargetPlatform HostEditorPlatform
diff --git a/Engine/Source/Programs/AutomationTool/AutomationUtils/MacHostPlatform.cs b/Engine/Source/Programs/AutomationTool/AutomationUtils/MacHostPlatform.cs
--- a/Engine/Source/Programs/AutomationTool/AutomationUtils/MacHostPlatform.cs
+++ b/Engine/Source/Programs/AutomationTool/AutomationUtils/MacHostPlatform.cs
@@ -20,7 +20,7 @@ namespace AutomationTool
 
 		public override string GetFrameworkMsbuildExe()
 		{
-			// As of 5.0 mono comes with msbuild which performs better. If that's installed then use it
+			// Look for dotnet, we only support dotnet.
 			if (string.IsNullOrEmpty(CachedFrameworkMsbuildExe))
 			{
 				bool CanUseMsBuild = string.IsNullOrEmpty(CommandUtils.WhichApp("dotnet")) == false;
@@ -33,8 +33,7 @@ namespace AutomationTool
 				}
 				else
 				{
-					Logger.LogInformation("Using xbuild. Install Mono 5.0 or greater for faster builds!");
-					CachedFrameworkMsbuildExe = "xbuild";
+					throw new BuildException("Unable to find installation of dotnet.");
 				}
 			}
 
@@ -112,28 +111,15 @@ namespace AutomationTool
 		public override Process CreateProcess(string AppName)
 		{
 			var NewProcess = new Process();
-			if (AppName == "mono")
-			{
-				// Enable case-insensitive mode for Mono
-				if (!NewProcess.StartInfo.EnvironmentVariables.ContainsKey("MONO_IOMAP"))
-				{
-					NewProcess.StartInfo.EnvironmentVariables.Add("MONO_IOMAP", "case");
-				}
-			}
 			return NewProcess;
 		}
 
 		public override void SetupOptionsForRun(ref string AppName, ref CommandUtils.ERunOptions Options, ref string CommandLine)
 		{
-			if (AppName == "sh" || AppName == "xbuild" || AppName == "codesign")
+			if (AppName == "sh" || AppName == "codesign")
 			{
 				Options &= ~CommandUtils.ERunOptions.AppMustExist;
 			}
-			if (AppName == "xbuild")
-			{
-				AppName = "sh";
-				CommandLine = "-c 'xbuild " + (String.IsNullOrEmpty(CommandLine) ? "" : CommandLine) + " /p:DefineConstants=MONO /p:DefineConstants=__MonoCS__ /verbosity:quiet /nologo |grep -i error; if [ $? -ne 1 ]; then exit 1; else exit 0; fi'";
-			}
 			if (AppName.EndsWith(".exe") || ((AppName.Contains("/Binaries/Win64/") || AppName.Contains("/Binaries/Mac/")) && string.IsNullOrEmpty(Path.GetExtension(AppName))))
 			{
 				if (AppName.Contains("/Binaries/Win64/") || AppName.Contains("/Binaries/Mac/"))
@@ -149,16 +135,15 @@ namespace AutomationTool
 					}
 				}
 				// some of our C# applications are converted to dotnet core, do not run those via mono
-
 				else if (AppName.Contains("UnrealBuildTool") || AppName.Contains("AutomationTool"))
 				{
 					Options &= ~CommandUtils.ERunOptions.AppMustExist;
 				}
 				else
 				{
-					// It's a C# app, so run it with Mono
+					// It's a C# app, so run it with dotnet
 					CommandLine = "\"" + AppName + "\" " + (String.IsNullOrEmpty(CommandLine) ? "" : CommandLine);
-					AppName = "mono";
+					AppName = "dotnet";
 					Options &= ~CommandUtils.ERunOptions.AppMustExist;
 				}
 			}
@@ -166,7 +151,7 @@ namespace AutomationTool
 
 		public override void SetConsoleCtrlHandler(ProcessManager.CtrlHandlerDelegate Handler)
 		{
-			// @todo: add mono support
+			// @todo: add dotnet support
 		}
 
 		public override UnrealTargetPlatform HostEditorPlatform
diff --git a/Engine/Source/Programs/AutomationTool/BuildGraph/Tasks/CsCompileTask.cs b/Engine/Source/Programs/AutomationTool/BuildGraph/Tasks/CsCompileTask.cs
--- a/Engine/Source/Programs/AutomationTool/BuildGraph/Tasks/CsCompileTask.cs
+++ b/Engine/Source/Programs/AutomationTool/BuildGraph/Tasks/CsCompileTask.cs
@@ -168,11 +168,8 @@ namespace AutomationTool.Tasks
 				{
 					Arguments.Add(String.Format("/target:{0}", CommandUtils.MakePathSafeToUseWithCommandLine(Parameters.Target)));
 				}
-				if(!CommandUtils.CmdEnv.FrameworkMsbuildPath.Equals("xbuild"))
-				{
-					// not supported by xbuild
-					Arguments.Add("/restore");
-				}
+
+				Arguments.Add("/restore");
 				Arguments.Add("/verbosity:minimal");
 				Arguments.Add("/nologo");
 
diff --git a/Engine/Source/Programs/AutomationTool/CookedEditor/CookedEditor.Automation.cs b/Engine/Source/Programs/AutomationTool/CookedEditor/CookedEditor.Automation.cs
--- a/Engine/Source/Programs/AutomationTool/CookedEditor/CookedEditor.Automation.cs
+++ b/Engine/Source/Programs/AutomationTool/CookedEditor/CookedEditor.Automation.cs
@@ -613,9 +613,8 @@ public class MakeCookedEditor : BuildCommand
 
 		if (Context.IniPlatformName == "Linux")
 		{
-			// linux needs dotnet runtime, as well as mono (at least for now until we can break our dependence)
+			// linux needs dotnet runtime
 			SC.StageFiles(StagedFileType.NonUFS, Unreal.FindDotnetDirectoryForPlatform(RuntimePlatform.Type.Linux), StageFilesSearch.AllDirectories);
-			SC.StageFiles(StagedFileType.NonUFS, DirectoryReference.Combine(Context.EngineDirectory, "Binaries/ThirdParty/Mono/Linux"), StageFilesSearch.AllDirectories);
 		}
 
 		// not sure if we need this or not now
diff --git a/Engine/Source/Programs/AutomationTool/Scripts/CheckCsprojDotNetVersion.cs b/Engine/Source/Programs/AutomationTool/Scripts/CheckCsprojDotNetVersion.cs
--- a/Engine/Source/Programs/AutomationTool/Scripts/CheckCsprojDotNetVersion.cs
+++ b/Engine/Source/Programs/AutomationTool/Scripts/CheckCsprojDotNetVersion.cs
@@ -72,7 +72,7 @@ class CheckCsprojDotNetVersion : BuildCommand
 				Match = PossibleAppConfigRegex.Match(Contents);
 				if (!Match.Success)
 				{
-					Logger.LogInformation("No TargetFrameworkVersion or TargetFrameworkProfile found for project {CsProj}, is it a mono project? If not, does it compile properly?", CsProj);
+					Logger.LogInformation("No TargetFrameworkVersion or TargetFrameworkProfile found for project {CsProj}, does it compile properly?", CsProj);
 					continue;
 				}
 
diff --git a/Engine/Source/Programs/Enterprise/Datasmith/DatasmithRhinoExporter/DatasmithRhino6.Target.cs b/Engine/Source/Programs/Enterprise/Datasmith/DatasmithRhinoExporter/DatasmithRhino6.Target.cs
--- a/Engine/Source/Programs/Enterprise/Datasmith/DatasmithRhinoExporter/DatasmithRhino6.Target.cs
+++ b/Engine/Source/Programs/Enterprise/Datasmith/DatasmithRhinoExporter/DatasmithRhino6.Target.cs
@@ -47,6 +47,7 @@ public abstract class DatasmithRhinoBaseTarget : TargetRules
 
 		// Define post-build step
 		// Since the Datasmith Rhino Exporter is a C# project, build in batch the release configuration of the Visual Studio C# project file.
+		// The Rhino Exporter has to be compatible with older toolchains so it still requires mono to build, relying on the system to have mono installed.
 		string ProjectFile = Path.Combine(RhinoExporterPath, ProjectName, ProjectName+".csproj");
 		string CustomAssemblyInfoPath = GetCustomAssemblyInfoPath(Target);
 		string CompilerCommand = Target.Platform == UnrealTargetPlatform.Win64 ? @"$(EngineDir)\Build\BatchFiles\MSBuild.bat" : "xbuild";
diff --git a/Engine/Source/Programs/UnrealBuildTool/ProjectFiles/CodeLite/CodeLiteProject.cs b/Engine/Source/Programs/UnrealBuildTool/ProjectFiles/CodeLite/CodeLiteProject.cs
--- a/Engine/Source/Programs/UnrealBuildTool/ProjectFiles/CodeLite/CodeLiteProject.cs
+++ b/Engine/Source/Programs/UnrealBuildTool/ProjectFiles/CodeLite/CodeLiteProject.cs
@@ -472,7 +472,7 @@ namespace UnrealBuildTool
 					//
 					if (ProjectTargetType == TargetType.Game)
 					{
-						string CookGameCommandLine = "mono AutomationTool.exe BuildCookRun ";
+						string CookGameCommandLine = "dotnet AutomationTool.exe BuildCookRun ";
 
 						// Projects filename
 						if (OnlyGameProject != null)
diff --git a/Engine/Source/Programs/UnrealBuildTool/ProjectFiles/VisualStudioCode/VSCodeProjectFileGenerator.cs b/Engine/Source/Programs/UnrealBuildTool/ProjectFiles/VisualStudioCode/VSCodeProjectFileGenerator.cs
--- a/Engine/Source/Programs/UnrealBuildTool/ProjectFiles/VisualStudioCode/VSCodeProjectFileGenerator.cs
+++ b/Engine/Source/Programs/UnrealBuildTool/ProjectFiles/VisualStudioCode/VSCodeProjectFileGenerator.cs
@@ -1020,7 +1020,7 @@ namespace UnrealBuildTool
 							OutFile.AddField("group", "build");
 							if (!RuntimePlatform.IsWindows)
 							{
-								OutFile.AddField("command", MakePathString(FileReference.Combine(ProjectRoot, "Engine", "Build", "BatchFiles", HostPlatform.ToString(), "RunDotnet.sh")));
+								OutFile.AddField("command", MakePathString(FileReference.Combine(ProjectRoot, "Engine", "Build", "BatchFiles", "RunDotnet.sh")));
 							}
 							else
 							{
