Add automation flow for WebTests.

diff --git a/Engine/Source/Programs/AutomationTool/LowLevelTests/RunLowLevelTests.cs b/Engine/Source/Programs/AutomationTool/LowLevelTests/RunLowLevelTests.cs
--- a/Engine/Source/Programs/AutomationTool/LowLevelTests/RunLowLevelTests.cs
+++ b/Engine/Source/Programs/AutomationTool/LowLevelTests/RunLowLevelTests.cs
@@ -445,7 +445,7 @@ namespace LowLevelTests
 	}
 
 	/// <summary>
-	/// Platform and test specific extension that provides extra command line arguments.
+	/// Platform and test specific extension that provides extra command line arguments or logic.
 	/// </summary>
 	public interface ILowLevelTestsExtension
 	{
@@ -458,6 +458,16 @@ namespace LowLevelTests
 		/// Return extra command line arguments specific to a platform and/or test.
 		/// </summary>
 		string ExtraCommandLine(UnrealTargetPlatform InPlatform, string InTestApp, string InBuildPath);
+
+		/// <summary>
+		/// Run extra logic before running tests
+		/// </summary>
+		void PreRunTests();
+
+		/// <summary>
+		/// Run extra logic after finishing tests
+		/// </summary>
+		void PostRunTests();
 	}
 
 	/// <summary>
@@ -629,9 +639,7 @@ namespace LowLevelTests
 
 		private ILowLevelTestsBuildFactory LowLevelTestsBuildFactory;
 		private ILowLevelTestsReporting LowLevelTestsReporting;
-#nullable enable
-		private ILowLevelTestsExtension? LowLevelTestsExtension;
-#nullable disable
+		public ILowLevelTestsExtension[] LowLevelTestsExtensions { get; protected set; }
 
 		public UnrealTargetPlatform Platform { get; protected set; }
 		public UnrealTargetConfiguration Configuration { get; protected set; }
@@ -663,9 +671,9 @@ namespace LowLevelTests
 				.Where(B => B.CanSupportPlatform(InTargetPlatform))
 				.First();
 
-			LowLevelTestsExtension = Gauntlet.Utils.InterfaceHelpers.FindImplementations<ILowLevelTestsExtension>(true)
+			LowLevelTestsExtensions = Gauntlet.Utils.InterfaceHelpers.FindImplementations<ILowLevelTestsExtension>(true)
 				.Where(B => B.IsSupported(InTargetPlatform, InTestApp))
-				.FirstOrDefault();
+				.ToArray();
 		}
 
 		public UnrealAppConfig GetUnrealAppConfig(string InTags, int InSleep, bool InAttachToDebugger, string InReportType, int InPerTestTimeout = 0)
@@ -711,14 +719,24 @@ namespace LowLevelTests
 				{
 					CachedConfig.CommandLineParams.AddRawCommandline("--buildmachine");
 				}
-				if (LowLevelTestsExtension != null)
+
+				string ExtraCmd = "";
+
+				foreach (ILowLevelTestsExtension LowLevelTestsExtension in LowLevelTestsExtensions)
 				{
-					string ExtraCmd = LowLevelTestsExtension.ExtraCommandLine(Platform, TestApp, BuildPath);
-					if (!string.IsNullOrEmpty(ExtraCmd))
+					string ExtensionExtraCmd = LowLevelTestsExtension.ExtraCommandLine(Platform, TestApp, BuildPath);
+					if (!string.IsNullOrEmpty(ExtensionExtraCmd))
 					{
-						CachedConfig.CommandLineParams.AddRawCommandline(string.Format("--extra-args {0}", ExtraCmd));
+						ExtraCmd += " ";
+						ExtraCmd += ExtensionExtraCmd;
 					}
 				}
+
+				if (!string.IsNullOrEmpty(ExtraCmd))
+				{
+					CachedConfig.CommandLineParams.AddRawCommandline("--extra-args" + ExtraCmd);
+				}
+
 				CachedConfig.CanAlterCommandArgs = false; // No further changes by IAppInstall instances etc.
 			}
 			return CachedConfig;
diff --git a/Engine/Source/Programs/AutomationTool/LowLevelTests/Tests/Gauntlet.LowLevelTestNode.cs b/Engine/Source/Programs/AutomationTool/LowLevelTests/Tests/Gauntlet.LowLevelTestNode.cs
--- a/Engine/Source/Programs/AutomationTool/LowLevelTests/Tests/Gauntlet.LowLevelTestNode.cs
+++ b/Engine/Source/Programs/AutomationTool/LowLevelTests/Tests/Gauntlet.LowLevelTestNode.cs
@@ -90,6 +90,11 @@ namespace LowLevelTests
 			Log.Info("LowLevelTestNode.StartTest Creating artifacts path at {0}", ArtifactPath);
 			Directory.CreateDirectory(ArtifactPath);
 
+			foreach (ILowLevelTestsExtension LowLevelTestsExtension in Context.BuildInfo.LowLevelTestsExtensions)
+			{
+				LowLevelTestsExtension.PreRunTests();
+			}
+
 			TestInstance = LowLevelTestsApp.InstallAndRunNativeTestApp();
 			if (TestInstance != null)
 			{
@@ -168,142 +173,156 @@ namespace LowLevelTests
 
 		public override void StopTest(StopReason InReason)
 		{
-			base.StopTest(InReason);
-
-			if (TestInstance != null && !TestInstance.HasExited)
-			{
-				TestInstance.Kill();
-			}
-
-			string StdOut;
-			if (TestInstance is IWithUnfilteredStdOut)
-			{
-				StdOut = ((IWithUnfilteredStdOut)TestInstance).UnfilteredStdOut;
-			}
-			else
+			try
 			{
-				StdOut = TestInstance.StdOut;
-			}
-
-			string LogDir = Path.Combine(Unreal.EngineDirectory.FullName, "Programs", "AutomationTool", "Saved", "Logs");
+				base.StopTest(InReason);
 
-			if (StdOut == null || string.IsNullOrEmpty(StdOut.Trim()))
-			{
-				Log.Warning("No StdOut returned from low level test app.");
-			}
-			else // Save log artifact
-			{
-				const string ClientLogFile = "ClientOutput.log";
-				string ClientOutputLog = Path.Combine(ArtifactPath, ClientLogFile);
-
-				using (var ClientOutputWriter = File.CreateText(ClientOutputLog))
+				if (TestInstance != null && !TestInstance.HasExited)
 				{
-					ClientOutputWriter.Write(StdOut);
+					TestInstance.Kill();
 				}
 
-				string DestClientLogFile = Path.Combine(LogDir, ClientLogFile);
-				if (DestClientLogFile != ClientOutputLog)
+				string StdOut;
+				if (TestInstance is IWithUnfilteredStdOut)
 				{
-					File.Copy(ClientOutputLog, DestClientLogFile, true);
+					StdOut = ((IWithUnfilteredStdOut)TestInstance).UnfilteredStdOut;
+				}
+				else
+				{
+					StdOut = TestInstance.StdOut;
 				}
-			}
 
-			bool? ReportCopied = null;
-			string ReportPath = null;
-			if (!string.IsNullOrEmpty(Context.Options.ReportType))
-			{
-				ILowLevelTestsReporting LowLevelTestsReporting = Gauntlet.Utils.InterfaceHelpers.FindImplementations<ILowLevelTestsReporting>(true)
-					.Where(B => B.CanSupportPlatform(Context.Options.Platform))
-					.First();
+				string LogDir = Path.Combine(Unreal.EngineDirectory.FullName, "Programs", "AutomationTool", "Saved", "Logs");
 
-				try
+				if (StdOut == null || string.IsNullOrEmpty(StdOut.Trim()))
 				{
-					ReportPath = LowLevelTestsReporting.CopyDeviceReportTo(LowLevelTestsApp.Install, Context.Options.Platform, Context.Options.TestApp, Context.Options.Build, LogDir);
-					ReportCopied = true;
+					Log.Warning("No StdOut returned from low level test app.");
 				}
-				catch (Exception ex)
+				else // Save log artifact
 				{
-					ReportCopied = false;
-					Log.Error("Failed to copy report: {0}", ex.ToString());
+					const string ClientLogFile = "ClientOutput.log";
+					string ClientOutputLog = Path.Combine(ArtifactPath, ClientLogFile);
+
+					using (var ClientOutputWriter = File.CreateText(ClientOutputLog))
+					{
+						ClientOutputWriter.Write(StdOut);
+					}
+
+					string DestClientLogFile = Path.Combine(LogDir, ClientLogFile);
+					if (DestClientLogFile != ClientOutputLog)
+					{
+						File.Copy(ClientOutputLog, DestClientLogFile, true);
+					}
 				}
-			}
 
+				bool? ReportCopied = null;
+				string ReportPath = null;
+				if (!string.IsNullOrEmpty(Context.Options.ReportType))
+				{
+					ILowLevelTestsReporting LowLevelTestsReporting = Gauntlet.Utils.InterfaceHelpers.FindImplementations<ILowLevelTestsReporting>(true)
+						.Where(B => B.CanSupportPlatform(Context.Options.Platform))
+						.First();
 
-			string ExitReason = "";
-			if (TestInstance.WasKilled)
-			{
-				if (InReason == StopReason.MaxDuration || LowLevelTestResult == TestResult.TimedOut)
+					try
+					{
+						ReportPath = LowLevelTestsReporting.CopyDeviceReportTo(LowLevelTestsApp.Install, Context.Options.Platform, Context.Options.TestApp, Context.Options.Build, LogDir);
+						ReportCopied = true;
+					}
+					catch (Exception ex)
+					{
+						ReportCopied = false;
+						Log.Error("Failed to copy report: {0}", ex.ToString());
+					}
+				}
+
+
+				string ExitReason = "";
+				if (TestInstance.WasKilled)
 				{
-					LowLevelTestResult = TestResult.TimedOut;
-					ExitReason = "Timed Out";
+					if (InReason == StopReason.MaxDuration || LowLevelTestResult == TestResult.TimedOut)
+					{
+						LowLevelTestResult = TestResult.TimedOut;
+						ExitReason = "Timed Out";
+					}
+					else
+					{
+						LowLevelTestResult = TestResult.Failed;
+						ExitReason = $"Process was killed by Gauntlet with reason {InReason.ToString()}.";
+					}
 				}
-				else
+				else if (TestInstance.ExitCode != 0)
 				{
 					LowLevelTestResult = TestResult.Failed;
-					ExitReason = $"Process was killed by Gauntlet with reason {InReason.ToString()}.";
+					ExitReason = $"Process exited with exit code {TestInstance.ExitCode}";
 				}
-			}
-			else if (TestInstance.ExitCode != 0)
-			{
-				LowLevelTestResult = TestResult.Failed;
-				ExitReason = $"Process exited with exit code {TestInstance.ExitCode}";
-			}
-			else if (ReportCopied.HasValue && !ReportCopied.Value)
-			{
-				LowLevelTestResult = TestResult.Failed;
-				ExitReason = "Uabled to read test report";
-			}
-			else if (ReportPath != null)
-			{
-				string ReportContents = File.ReadAllText(ReportPath);
-				if (Context.Options.LogReportContents) // Some tests prefer to log report contents
+				else if (ReportCopied.HasValue && !ReportCopied.Value)
 				{
-					Log.Info(ReportContents);
+					LowLevelTestResult = TestResult.Failed;
+					ExitReason = "Uabled to read test report";
 				}
-				string ReportType = Context.Options.ReportType.ToLower();
-				if (ReportType == "console")
+				else if (ReportPath != null)
 				{
-					LowLevelTestsLogParser LowLevelTestsLogParser = new LowLevelTestsLogParser(ReportContents);
-					if (LowLevelTestsLogParser.GetCatchTestResults().Passed)
+					string ReportContents = File.ReadAllText(ReportPath);
+					if (Context.Options.LogReportContents) // Some tests prefer to log report contents
 					{
-						LowLevelTestResult = TestResult.Passed;
-						ExitReason = "Tests passed";
+						Log.Info(ReportContents);
 					}
-					else
+					string ReportType = Context.Options.ReportType.ToLower();
+					if (ReportType == "console")
 					{
-						LowLevelTestResult = TestResult.Failed;
-						ExitReason = "Tests failed";
+						LowLevelTestsLogParser LowLevelTestsLogParser = new LowLevelTestsLogParser(ReportContents);
+						if (LowLevelTestsLogParser.GetCatchTestResults().Passed)
+						{
+							LowLevelTestResult = TestResult.Passed;
+							ExitReason = "Tests passed";
+						}
+						else
+						{
+							LowLevelTestResult = TestResult.Failed;
+							ExitReason = "Tests failed";
+						}
+					}
+					else if (ReportType == "xml")
+					{
+						LowLevelTestsReportParser LowLevelTestsReportParser = new LowLevelTestsReportParser(ReportContents);
+						if (LowLevelTestsReportParser.HasPassed())
+						{
+							LowLevelTestResult = TestResult.Passed;
+							ExitReason = "Tests passed";
+						}
+						else
+						{
+							LowLevelTestResult = TestResult.Failed;
+							ExitReason = "Tests failed";
+						}
 					}
 				}
-				else if (ReportType == "xml")
+				else // ReportPath == null
 				{
-					LowLevelTestsReportParser LowLevelTestsReportParser = new LowLevelTestsReportParser(ReportContents);
-					if (LowLevelTestsReportParser.HasPassed())
+					if (TestInstance.ExitCode != 0)
 					{
-						LowLevelTestResult = TestResult.Passed;
-						ExitReason = "Tests passed";
+						LowLevelTestResult = TestResult.Failed;
+						ExitReason = "Tests failed (no report to parse)";
 					}
 					else
 					{
-						LowLevelTestResult = TestResult.Failed;
-						ExitReason = "Tests failed";
+						LowLevelTestResult = TestResult.Passed;
+						ExitReason = "Tests passed (no report to parse)";
 					}
 				}
+				Log.Info($"Low level test exited with code {TestInstance.ExitCode} and reason: {ExitReason}");
 			}
-			else // ReportPath == null
+			catch
 			{
-				if (TestInstance.ExitCode != 0)
-				{
-					LowLevelTestResult = TestResult.Failed;
-					ExitReason = "Tests failed (no report to parse)";
-				}
-				else
+				throw;
+			} 
+			finally 
+			{ 
+				foreach (ILowLevelTestsExtension LowLevelTestsExtension in Context.BuildInfo.LowLevelTestsExtensions)
 				{
-					LowLevelTestResult = TestResult.Passed;
-					ExitReason = "Tests passed (no report to parse)";
+					LowLevelTestsExtension.PostRunTests();
 				}
 			}
-			Log.Info($"Low level test exited with code {TestInstance.ExitCode} and reason: {ExitReason}");
 		}
 
 		public override void CleanupTest()
diff --git a/Engine/Source/Programs/AutomationTool/LowLevelTests/Tests/LowLevelTests.WebTestsLowLevelTestsExtension.cs b/Engine/Source/Programs/AutomationTool/LowLevelTests/Tests/LowLevelTests.WebTestsLowLevelTestsExtension.cs
new file mode 100644
--- /dev/null
+++ b/Engine/Source/Programs/AutomationTool/LowLevelTests/Tests/LowLevelTests.WebTestsLowLevelTestsExtension.cs
@@ -0,0 +1,56 @@
+// Copyright Epic Games, Inc. All Rights Reserved.
+
+using JetBrains.Annotations;
+using System;
+using System.Diagnostics;
+using System.IO;
+using EpicGames.Core;
+using AutomationTool;
+using UnrealBuildTool;
+using Gauntlet;
+using UnrealBuildBase;
+using EpicGames.Horde.Storage;
+
+namespace LowLevelTests
+{
+	public class WebTestsLowLevelTestsExtension : ILowLevelTestsExtension
+	{
+		private Process ServerProcess;
+
+		public bool IsSupported(UnrealTargetPlatform InPlatform, string InTestApp)
+		{
+			return InTestApp == "WebTests";
+		}
+
+		public string ExtraCommandLine(UnrealTargetPlatform InPlatform, string InTestApp, string InBuildPath)
+		{
+			return string.Format("--web_server_ip{0}", UnrealHelpers.GetHostIpAddress());
+		}
+
+		public void PreRunTests()
+		{
+			ProcessStartInfo StartInfo = new ProcessStartInfo();
+			StartInfo.WorkingDirectory = Path.Combine(Unreal.EngineDirectory.FullName, "Source", "Programs", "WebTests", "WebServer");
+			StartInfo.FileName = RuntimePlatform.IsWindows ? "cmd.exe" : "/bind/sh";
+			StartInfo.Arguments = RuntimePlatform.IsWindows ? "/c runserver.bat" : "-c 'runserver.sh'";
+			StartInfo.UseShellExecute = false;
+			StartInfo.CreateNoWindow = true;
+			StartInfo.RedirectStandardInput = true;
+			StartInfo.RedirectStandardOutput = true;
+			StartInfo.RedirectStandardError = true;
+
+			ServerProcess = new Process();
+			ServerProcess.StartInfo = StartInfo;
+			ServerProcess.Start();
+		}
+
+		public void PostRunTests()
+		{
+			if (ServerProcess != null)
+			{
+				ServerProcess.Kill();
+				ServerProcess = null;
+			}
+		}
+	}
+}
diff --git a/Engine/Source/Programs/WebTests/Private/TestHttp.cpp b/Engine/Source/Programs/WebTests/Private/TestHttp.cpp
--- a/Engine/Source/Programs/WebTests/Private/TestHttp.cpp
+++ b/Engine/Source/Programs/WebTests/Private/TestHttp.cpp
@@ -32,6 +32,8 @@ class FHttpModuleTestFixture
 public:
 	FHttpModuleTestFixture()
 		: WebServerIp(TEXT("127.0.0.1"))
+		, WebServerPort(8000)
+		, bRunHeavyTests(false)
 	{
 		ParseSettingsFromCommandLine();
 
@@ -50,16 +52,20 @@ public:
 	void ParseSettingsFromCommandLine()
 	{
 		FParse::Value(FCommandLine::Get(), TEXT("web_server_ip"), WebServerIp);
+		FParse::Bool(FCommandLine::Get(), TEXT("run_heavy_tests"), bRunHeavyTests);
 	}
 
 	const FString UrlWithInvalidPortToTestConnectTimeout() const { return FString::Format(TEXT("http://{0}:{1}"), { *WebServerIp, 8765 }); }
-	const FString UrlBase() const { return FString::Format(TEXT("http://{0}:{1}"), { *WebServerIp, 8000 }); }
+	const FString UrlBase() const { return FString::Format(TEXT("http://{0}:{1}"), { *WebServerIp, WebServerPort }); }
 	const FString UrlHttpTests() const { return FString::Format(TEXT("{0}/webtests/httptests"), { *UrlBase() }); }
 	const FString UrlToTestMethods() const { return FString::Format(TEXT("{0}/methods"), { *UrlHttpTests() }); }
 	const FString UrlStreamDownload(uint32 Chunks, uint32 ChunkSize) { return FString::Format(TEXT("{0}/streaming_download/{1}/{2}/"), { *UrlHttpTests(), Chunks, ChunkSize }); }
 
 	FString WebServerIp;
+	uint32 WebServerPort;
 	FHttpModule* HttpModule;
+
+	bool bRunHeavyTests;
 };
 
 TEST_CASE_METHOD(FHttpModuleTestFixture, "Shutdown http module without issue when there are ongoing http requests.", HTTP_TAG)
@@ -225,25 +231,26 @@ TEST_CASE_METHOD(FWaitUntilCompleteHttpFixture, "Get large response content with
 	HttpRequest->ProcessRequest();
 }
 
-TEST_CASE_METHOD(FWaitUntilCompleteHttpFixture, "Http request connect timeout", HTTP_TAG)
-{
-	TSharedRef<IHttpRequest> HttpRequest = HttpModule->CreateRequest();
-	HttpRequest->SetURL(UrlWithInvalidPortToTestConnectTimeout());
-	HttpRequest->SetVerb(TEXT("GET"));
-	HttpRequest->SetTimeout(7);
-	FDateTime StartTime = FDateTime::Now();
-	HttpRequest->OnProcessRequestComplete().BindLambda([StartTime](FHttpRequestPtr HttpRequest, FHttpResponsePtr HttpResponse, bool bSucceeded) {
-		CHECK(!bSucceeded);
-		CHECK(HttpResponse == nullptr);
-		// TODO: For now curl impl is using customized timeout instead of relying on native http timeout, 
-		// which doesn't get CURLE_COULDNT_CONNECT. Enable this after switching to native http timeout
-		//CHECK(HttpRequest->GetStatus() == EHttpRequestStatus::Failed_ConnectionError);
-		FTimespan Timespan = FDateTime::Now() - StartTime;
-		float DurationInSeconds = Timespan.GetTotalSeconds();
-		CHECK(FMath::IsNearlyEqual(DurationInSeconds, 7, HTTP_TIME_DIFF_TOLERANCE));
-	});
-	HttpRequest->ProcessRequest();
-}
+// TODO: Enable this after finding a more reliable way of simulate timeout
+//TEST_CASE_METHOD(FWaitUntilCompleteHttpFixture, "Http request connect timeout", HTTP_TAG)
+//{
+//	TSharedRef<IHttpRequest> HttpRequest = HttpModule->CreateRequest();
+//	HttpRequest->SetURL(UrlWithInvalidPortToTestConnectTimeout());
+//	HttpRequest->SetVerb(TEXT("GET"));
+//	HttpRequest->SetTimeout(7);
+//	FDateTime StartTime = FDateTime::Now();
+//	HttpRequest->OnProcessRequestComplete().BindLambda([StartTime](FHttpRequestPtr HttpRequest, FHttpResponsePtr HttpResponse, bool bSucceeded) {
+//		CHECK(!bSucceeded);
+//		CHECK(HttpResponse == nullptr);
+//		// TODO: For now curl impl is using customized timeout instead of relying on native http timeout, 
+//		// which doesn't get CURLE_COULDNT_CONNECT. Enable this after switching to native http timeout
+//		//CHECK(HttpRequest->GetStatus() == EHttpRequestStatus::Failed_ConnectionError);
+//		FTimespan Timespan = FDateTime::Now() - StartTime;
+//		float DurationInSeconds = Timespan.GetTotalSeconds();
+//		CHECK(FMath::IsNearlyEqual(DurationInSeconds, 7, HTTP_TIME_DIFF_TOLERANCE));
+//	});
+//	HttpRequest->ProcessRequest();
+//}
 
 TEST_CASE_METHOD(FWaitUntilCompleteHttpFixture, "Streaming http download", HTTP_TAG)
 {
@@ -405,6 +412,11 @@ TEST_CASE_METHOD(FWaitUntilCompleteHttpFixture, "Can run parallel stream downloa
 
 TEST_CASE_METHOD(FWaitUntilCompleteHttpFixture, "Can download big file exceeds 32 bits", HTTP_TAG)
 {
+	if (!bRunHeavyTests)
+	{
+		return;
+	}
+
 	// 5 * 1024 * 1024 * 1024 BYTES = 5368709120 BYTES = 5 GB
 	uint64 Chunks = 5 * 1024;
 	uint64 ChunkSize = 1024 * 1024;
@@ -514,6 +526,11 @@ public:
 
 TEST_CASE_METHOD(FWaitUntilCompleteHttpFixture, "Can upload big file exceeds 32 bits", HTTP_TAG)
 {
+	if (!bRunHeavyTests)
+	{
+		return;
+	}
+
 	// TODO: Back to check later. xCurl 2206.4.0.0 doesn't work with file bigger than 32 bits
 	// 5 * 1024 * 1024 * 1024 BYTES = 5368709120 BYTES = 5 GB
 	//const uint64 TotalSize = 5368709120;
diff --git a/Engine/Source/Programs/WebTests/WebServer/runserver.sh b/Engine/Source/Programs/WebTests/WebServer/runserver.sh
new file mode 100644
--- /dev/null
+++ b/Engine/Source/Programs/WebTests/WebServer/runserver.sh
@@ -0,0 +1,13 @@
+if [ -d "env" ]; then
+	source ./env/Scripts/activate.sh
+else
+	run ../../../../Binaries/ThirdParty/Python3/Mac/python3 -m venv env
+	source ./env/Scripts/activate.sh
+	run pip install -r requirements.txt
+	if [ $? -ne 0 ]; then
+		rm -rf env
+		exit $?
+	fi
+fi
+
+cmd python manage.py runserver 0.0.0.0:8000
