Gauntlet - Fix horde test name for simple report in EditorTests and UE.Automation

#rnx

diff --git a/Engine/Source/Programs/AutomationTool/Gauntlet/Unreal/Automation/UE.Automation.cs b/Engine/Source/Programs/AutomationTool/Gauntlet/Unreal/Automation/UE.Automation.cs
--- a/Engine/Source/Programs/AutomationTool/Gauntlet/Unreal/Automation/UE.Automation.cs
+++ b/Engine/Source/Programs/AutomationTool/Gauntlet/Unreal/Automation/UE.Automation.cs
@@ -460,6 +460,17 @@ namespace UE
 			}
 		}
 
+		/// <summary>
+		/// Override the HordeReportTestName in case a simple report is used
+		/// </summary>
+		protected override string HordeReportTestName
+		{
+			get
+			{
+				return Type;
+			}
+		}
+
 		/// <summary>
 		/// Override the Type
 		/// </summary>
diff --git a/Engine/Source/Programs/AutomationTool/Gauntlet/Unreal/Editor/EditorTests.cs b/Engine/Source/Programs/AutomationTool/Gauntlet/Unreal/Editor/EditorTests.cs
--- a/Engine/Source/Programs/AutomationTool/Gauntlet/Unreal/Editor/EditorTests.cs
+++ b/Engine/Source/Programs/AutomationTool/Gauntlet/Unreal/Editor/EditorTests.cs
@@ -150,5 +150,13 @@ namespace UnrealEditor
 			Config.RequireRole(UnrealTargetRole.Editor);	
 			return Config;
 		}
+
+		protected override string HordeReportTestName
+		{
+			get
+			{
+				return GetConfiguration().RunTest.Replace(".", " ");
+			}
+		}
 	}
 }
