Gauntlet - Introduce MaxEventsDisplayPerTest argument to UE.Automation to adjust how much errors are reported per test

#rnx

diff --git a/Engine/Source/Programs/AutomationTool/Gauntlet/Unreal/Automation/UE.Automation.cs b/Engine/Source/Programs/AutomationTool/Gauntlet/Unreal/Automation/UE.Automation.cs
--- a/Engine/Source/Programs/AutomationTool/Gauntlet/Unreal/Automation/UE.Automation.cs
+++ b/Engine/Source/Programs/AutomationTool/Gauntlet/Unreal/Automation/UE.Automation.cs
@@ -137,6 +137,12 @@ namespace UE
 		[AutoParam]
 		public virtual int LogIdleTimeout { get; set; } = 30 * 60;
 
+		/// <summary>
+		/// Maximum of events display per test
+		/// </summary>
+		[AutoParam]
+		public virtual int MaxEventsDisplayPerTest { get; set; } = 10;
+
 		/// <summary>
 		/// Used for having the editor and any client communicate
 		/// </summary>
@@ -869,7 +875,11 @@ namespace UE
 		/// <returns></returns>
 		protected override void LogTestSummaryHeader()
 		{
-			const int kMaxErrorsOrWarningsToDisplay = 10;
+			int kMaxErrorsOrWarningsToDisplay = 10;
+			if (GetConfiguration() is AutomationTestConfig TestConfig)
+			{
+				kMaxErrorsOrWarningsToDisplay = TestConfig.MaxEventsDisplayPerTest;
+			}
 
 			base.LogTestSummaryHeader();
 
@@ -894,6 +904,13 @@ namespace UE
 					}
 					return E;
 				};
+				Action<IEnumerable<UnrealAutomationEvent>> NotifyMoreIfNeeded = E =>
+				{
+					if (E.Count() > kMaxErrorsOrWarningsToDisplay)
+					{
+						Log.Info(" (and {Count} more)", E.Count() - kMaxErrorsOrWarningsToDisplay);
+					}
+				};
 				// If there were abnormal exits then look only at the failed and incomplete tests only to avoid confusing things.
 				if (GetRolesThatExitedAbnormally().Any())
 				{
@@ -918,11 +935,13 @@ namespace UE
 									EventId ErrorEventType = Error.IsCriticalFailure ? KnownLogEvents.Gauntlet_FatalEvent : KnownLogEvents.Gauntlet_UnrealEngineTestEvent;
 									Log.Error(ErrorEventType, "    " + Error.FormatToString());
 								}
+								NotifyMoreIfNeeded(Result.ErrorEvents);
 								var Warnings = CapErrorOrWarningList(Result.WarningEvents.Distinct());
 								foreach (var Warning in Warnings)
 								{
 									Log.Warning(KnownLogEvents.Gauntlet_UnrealEngineTestEvent, "    " + Warning.FormatToString());
 								}
+								NotifyMoreIfNeeded(Result.WarningEvents);
 								Log.Info("");
 							}
 						}
@@ -955,10 +974,12 @@ namespace UE
 							{
 								Log.Error(KnownLogEvents.Gauntlet_UnrealEngineTestEvent, "    " + Error.FormatToString());
 							}
+							NotifyMoreIfNeeded(Errors);
 							foreach (var Warning in CapErrorOrWarningList(Warnings))
 							{
 								Log.Warning(KnownLogEvents.Gauntlet_UnrealEngineTestEvent, "    " + Warning.FormatToString());
 							}
+							NotifyMoreIfNeeded(Warnings);
 							Log.Info("");
 						}
 						else
@@ -977,12 +998,12 @@ namespace UE
 							{
 								Log.Error(KnownLogEvents.Gauntlet_UnrealEngineTestEvent, " * Test '{Name}' failed.", Result.FullTestPath);
 								IEnumerable<UnrealAutomationEvent> Events = Result.ErrorEvents.Distinct();
-								Events = CapErrorOrWarningList(Events);
-								foreach (var Event in Events)
+								foreach (var Event in CapErrorOrWarningList(Events))
 								{
 									EventId ErrorEventType = Event.IsCriticalFailure ? KnownLogEvents.Gauntlet_FatalEvent : KnownLogEvents.Gauntlet_UnrealEngineTestEvent;
 									Log.Error(ErrorEventType, "    " + Event.FormatToString());
 								}
+								NotifyMoreIfNeeded(Events);
 								Log.Info("");
 							}
 						}
@@ -996,11 +1017,11 @@ namespace UE
 								Log.Warning(KnownLogEvents.Gauntlet_UnrealEngineTestEvent, " * Test '{Name}' completed with warnings.", Result.FullTestPath);
 								// only show the first N items
 								IEnumerable<UnrealAutomationEvent> WarningEvents = Result.WarningEvents.Distinct();
-								WarningEvents = CapErrorOrWarningList(WarningEvents);
-								foreach (var Event in WarningEvents)
+								foreach (var Event in CapErrorOrWarningList(WarningEvents))
 								{
 									Log.Warning(KnownLogEvents.Gauntlet_UnrealEngineTestEvent, "    " + Event.FormatToString());
 								}
+								NotifyMoreIfNeeded(WarningEvents);
 								Log.Info("");
 							}
 						}
@@ -1014,11 +1035,11 @@ namespace UE
 								Log.Error(KnownLogEvents.Gauntlet_UnrealEngineTestEvent, " * Test '{Name}' did not run or complete.", Result.FullTestPath);
 								// only show the first N items
 								IEnumerable<UnrealAutomationEvent> ErrorAndWarningEvents = Result.WarningAndErrorEvents.Distinct();
-								ErrorAndWarningEvents = CapErrorOrWarningList(ErrorAndWarningEvents);
-								foreach (var Event in ErrorAndWarningEvents)
+								foreach (var Event in CapErrorOrWarningList(ErrorAndWarningEvents))
 								{
 									Log.Error(KnownLogEvents.Gauntlet_UnrealEngineTestEvent, "    " + Event.FormatToString());
 								}
+								NotifyMoreIfNeeded(ErrorAndWarningEvents);
 							}
 							Log.Info("");
 						}
