Gauntlet - Move MaxEventsDisplayPerTest from test config to test node class

#rnx

diff --git a/Engine/Source/Programs/AutomationTool/Gauntlet/Unreal/Automation/UE.Automation.cs b/Engine/Source/Programs/AutomationTool/Gauntlet/Unreal/Automation/UE.Automation.cs
--- a/Engine/Source/Programs/AutomationTool/Gauntlet/Unreal/Automation/UE.Automation.cs
+++ b/Engine/Source/Programs/AutomationTool/Gauntlet/Unreal/Automation/UE.Automation.cs
@@ -137,12 +137,6 @@ namespace UE
 		[AutoParam]
 		public virtual int LogIdleTimeout { get; set; } = 30 * 60;
 
-		/// <summary>
-		/// Maximum of events display per test
-		/// </summary>
-		[AutoParam]
-		public virtual int MaxEventsDisplayPerTest { get; set; } = 10;
-
 		/// <summary>
 		/// Used for having the editor and any client communicate
 		/// </summary>
@@ -415,6 +409,9 @@ namespace UE
 
 		private DateTime LastAutomationEntryTime = DateTime.MinValue;
 
+		/// Maximum of events display per test
+		protected virtual int MaxEventsDisplayPerTest { get; set; } = 10;
+
 		public AutomationNodeBase(Gauntlet.UnrealTestContext InContext)
 			: base(InContext)
 		{
@@ -875,12 +872,6 @@ namespace UE
 		/// <returns></returns>
 		protected override void LogTestSummaryHeader()
 		{
-			int kMaxErrorsOrWarningsToDisplay = 10;
-			if (GetConfiguration() is AutomationTestConfig TestConfig)
-			{
-				kMaxErrorsOrWarningsToDisplay = TestConfig.MaxEventsDisplayPerTest;
-			}
-
 			base.LogTestSummaryHeader();
 
 			// Everything we need is in the editor artifacts
@@ -898,17 +889,17 @@ namespace UE
 
 				Func<IEnumerable<UnrealAutomationEvent>, IEnumerable<UnrealAutomationEvent>> CapErrorOrWarningList = (E) =>
 				{
-					if (E.Count() > kMaxErrorsOrWarningsToDisplay)
+					if (E.Count() > MaxEventsDisplayPerTest)
 					{
-						E = E.Take(kMaxErrorsOrWarningsToDisplay);
+						E = E.Take(MaxEventsDisplayPerTest);
 					}
 					return E;
 				};
 				Action<IEnumerable<UnrealAutomationEvent>> NotifyMoreIfNeeded = E =>
 				{
-					if (E.Count() > kMaxErrorsOrWarningsToDisplay)
+					if (E.Count() > MaxEventsDisplayPerTest)
 					{
-						Log.Info(" (and {Count} more)", E.Count() - kMaxErrorsOrWarningsToDisplay);
+						Log.Info(" (and {Count} more)", E.Count() - MaxEventsDisplayPerTest);
 					}
 				};
 				// If there were abnormal exits then look only at the failed and incomplete tests only to avoid confusing things.
