Avoid reverting to simple report if no json report was found

#rnx

diff --git a/Engine/Source/Programs/AutomationTool/Gauntlet/Unreal/Base/Gauntlet.UnrealTestNode.cs b/Engine/Source/Programs/AutomationTool/Gauntlet/Unreal/Base/Gauntlet.UnrealTestNode.cs
--- a/Engine/Source/Programs/AutomationTool/Gauntlet/Unreal/Base/Gauntlet.UnrealTestNode.cs
+++ b/Engine/Source/Programs/AutomationTool/Gauntlet/Unreal/Base/Gauntlet.UnrealTestNode.cs
@@ -1466,8 +1466,8 @@ namespace Gauntlet
 			}
 			else
 			{
-				Log.Warning("Could not find Unreal Automated test report at {FilePath}. Reverting to base report.", JsonReportPath);
-				return CreateSimpleReportForHorde(GetTestResult());
+				Log.Warning("Could not find Unreal Automated test report at {FilePath}. No Test report will be generated for Horde.", JsonReportPath);
+				return null;
 			}
 		}
 
