[Gauntlet] BugFix: Fix the ability for people to jump directly to the gauntlet log on the network.

diff --git a/Engine/Source/Programs/AutomationTool/Gauntlet/Unreal/Base/Gauntlet.UnrealSession.cs b/Engine/Source/Programs/AutomationTool/Gauntlet/Unreal/Base/Gauntlet.UnrealSession.cs
--- a/Engine/Source/Programs/AutomationTool/Gauntlet/Unreal/Base/Gauntlet.UnrealSession.cs
+++ b/Engine/Source/Programs/AutomationTool/Gauntlet/Unreal/Base/Gauntlet.UnrealSession.cs
@@ -1108,7 +1108,6 @@ namespace Gauntlet
 			string SourceSavedDir = InRunningRole.AppInstance.ArtifactPath;
 			
 			string ArtifactLogFilePath = String.Empty;
-			string HordeLogFilePath = String.Empty;
 
 			// Get the size of the log to determine if it should be saved under a certain size
 			int MaxLogSize = 1024 * 1024 * 1024;
@@ -1123,22 +1122,10 @@ namespace Gauntlet
 			{
 				try
 				{
-					string StreamWriterOutputFilePath = String.Empty;
-
-					// On build machines, copy all role logs to Horde.
-					if (IsBuildMachine)
-					{
-						HordeLogFilePath = Path.Combine(CommandUtils.CmdEnv.LogFolder, RoleName + "Output.log");
-						StreamWriterOutputFilePath = HordeLogFilePath;
-					}
-					else
-					{
-						ArtifactLogFilePath = Path.Combine(DestSavedDirInfo.FullName, RoleName + "Output.log");
-						StreamWriterOutputFilePath = ArtifactLogFilePath;
-					}
+					ArtifactLogFilePath = Path.Combine(DestSavedDirInfo.FullName, RoleName + "Output.log");
 
 					// save the output from TTY
-					using (StreamWriter Writer = new StreamWriter(StreamWriterOutputFilePath, false))
+					using (StreamWriter Writer = new(ArtifactLogFilePath, false))
 					{
 						Writer.WriteLine("------ Gauntlet Test ------");
 						Writer.WriteLine(string.Format("Role: {0}\r\n", InRunningRole.Role));
@@ -1146,7 +1133,13 @@ namespace Gauntlet
 						Writer.WriteLine("---------------------------");
 						Writer.Write(InRunningRole.AppInstance.StdOut);
 					}
-					Log.Info($"Wrote {RoleName} Log to {StreamWriterOutputFilePath}");
+					Log.Info($"Wrote {RoleName} Log to {ArtifactLogFilePath}");
+					// On build machines, copy all role logs to Horde.
+					if (IsBuildMachine)
+					{
+						string HordeLogFilePath = Path.Combine(CommandUtils.CmdEnv.LogFolder, RoleName + "Output.log");
+						File.Copy(ArtifactLogFilePath, HordeLogFilePath, true);
+					}
 				}
 				catch (Exception Ex)
 				{
@@ -1279,21 +1272,9 @@ namespace Gauntlet
 			// END REMOVEME
 
 			UnrealLogParser LogParser = new UnrealLogParser(InRunningRole.AppInstance.StdOut);
-			UnrealRoleArtifacts Artifacts = null;
-
-			// Only keep one copy of the StdOut where needed
-			if (IsBuildMachine)
-			{
-				// Save to Horde filepath
-				Artifacts = new UnrealRoleArtifacts(InRunningRole.Role, InRunningRole.AppInstance, DestSavedDirInfo.FullName, HordeLogFilePath, LogParser);
-			}
-			else
-			{
-				// Save the Artifact filepath
-				Artifacts = new UnrealRoleArtifacts(InRunningRole.Role, InRunningRole.AppInstance, DestSavedDirInfo.FullName, ArtifactLogFilePath, LogParser);
-			}
 
-			return Artifacts;
+			// Save the Artifact filepath
+			return new UnrealRoleArtifacts(InRunningRole.Role, InRunningRole.AppInstance, DestSavedDirInfo.FullName, ArtifactLogFilePath, LogParser);
 		}
 
 		/// <summary>
