Change UnrealBuildException to print as errors

* Give it a unique exit code
* And introduce a new output format MinimalError to give minimal info output, but as an error

#rnx

diff --git a/Engine/Source/Programs/AutomationTool/AutomationUtils/Automation.cs b/Engine/Source/Programs/AutomationTool/AutomationUtils/Automation.cs
--- a/Engine/Source/Programs/AutomationTool/AutomationUtils/Automation.cs
+++ b/Engine/Source/Programs/AutomationTool/AutomationUtils/Automation.cs
@@ -193,6 +193,11 @@ namespace AutomationTool
 					Logger.LogInformation(Ex, "{Message}", Ex.ToString().Replace("\n", "\n  "));
 					Logger.LogDebug(Ex, "{Details}", ExceptionUtils.FormatExceptionDetails(Ex));
 				}
+				else if (Ex.OutputFormat == AutomationExceptionOutputFormat.MinimalError)
+				{
+					Logger.LogError(Ex, "{Message}", Ex.ToString().Replace("\n", "\n  "));
+					Logger.LogDebug(Ex, "{Details}", ExceptionUtils.FormatExceptionDetails(Ex));
+				}
 				else
 				{
 					Log.WriteException(Ex, LogUtils.FinalLogFileName);
diff --git a/Engine/Source/Programs/AutomationTool/AutomationUtils/AutomationException.cs b/Engine/Source/Programs/AutomationTool/AutomationUtils/AutomationException.cs
--- a/Engine/Source/Programs/AutomationTool/AutomationUtils/AutomationException.cs
+++ b/Engine/Source/Programs/AutomationTool/AutomationUtils/AutomationException.cs
@@ -22,6 +22,11 @@ namespace AutomationTool
 		/// </summary>
 		Minimal,
 
+		/// <summary>
+		/// Output a brief summary to the console as an error, and the exception details to the log
+		/// </summary>
+		MinimalError,
+
 		/// <summary>
 		/// Treat the exception as an error
 		/// </summary>
diff --git a/Engine/Source/Programs/AutomationTool/AutomationUtils/UnrealBuild.cs b/Engine/Source/Programs/AutomationTool/AutomationUtils/UnrealBuild.cs
--- a/Engine/Source/Programs/AutomationTool/AutomationUtils/UnrealBuild.cs
+++ b/Engine/Source/Programs/AutomationTool/AutomationUtils/UnrealBuild.cs
@@ -23,9 +23,9 @@ namespace AutomationTool
 	class UnrealBuildException : AutomationException
 	{
 		public UnrealBuildException(string Message)
-			: base("BUILD FAILED: " + Message)
+			: base(ExitCode.Error_UBTFailure, "BUILD FAILED: " + Message)
 		{
-			OutputFormat = AutomationExceptionOutputFormat.Minimal;
+			OutputFormat = AutomationExceptionOutputFormat.MinimalError;
 		}
 
 		public UnrealBuildException(string Format, params object[] Args)
diff --git a/Engine/Source/Programs/Shared/EpicGames.Build/Automation/ExitCode.cs b/Engine/Source/Programs/Shared/EpicGames.Build/Automation/ExitCode.cs
--- a/Engine/Source/Programs/Shared/EpicGames.Build/Automation/ExitCode.cs
+++ b/Engine/Source/Programs/Shared/EpicGames.Build/Automation/ExitCode.cs
@@ -63,5 +63,6 @@ namespace UnrealBuildBase
 		Error_AndroidOBBError = 155,
 		Error_SDKInstallFailed = 200,
 		Error_DeviceUpdateFailed = 201,
+		Error_UBTFailure = 202
 	};
 }
