- Added an interface to StagedBuildSource to let a platform determine if a found Staged binary is valid
  - This fixes Gauntlet finding a partial .app in Mac staged directories that modern makes that is not a full .app

diff --git a/Engine/Source/Programs/AutomationTool/Gauntlet/Platform/Mac/Gauntlet.MacBuildSource.cs b/Engine/Source/Programs/AutomationTool/Gauntlet/Platform/Mac/Gauntlet.MacBuildSource.cs
--- a/Engine/Source/Programs/AutomationTool/Gauntlet/Platform/Mac/Gauntlet.MacBuildSource.cs
+++ b/Engine/Source/Programs/AutomationTool/Gauntlet/Platform/Mac/Gauntlet.MacBuildSource.cs
@@ -9,9 +9,43 @@ using System.Threading;
 using System.Text.RegularExpressions;
 using System.Linq;
 using EpicGames.Core;
+using System.Security.Cryptography.X509Certificates;
 
 namespace Gauntlet
 {
+	internal static class MacAppChecker
+	{
+		public static bool IsAppBundleUsable(string AppPath, string ProjectName)
+		{
+			string MacOSDir = Path.Combine(AppPath, "Contents", "MacOS");
+
+			// check both binary AND Engine content. Regular Unreal builds won't have the latter
+			if (Directory.Exists(MacOSDir)
+				&& Directory.Exists(Path.Combine(AppPath, "Contents", "UE")))
+			{
+				// if we don't have a project passed in, we have checked all we can, so we are valid now
+				if (ProjectName == null)
+				{
+					return true;
+				}
+
+				// check there's an executable with the right name 
+				string ShortName = Regex.Replace(ProjectName, "Game", "", RegexOptions.IgnoreCase);
+				FileInfo Executable = new DirectoryInfo(MacOSDir).GetFiles().Where(Fi => Fi.Name.StartsWith(ShortName, StringComparison.OrdinalIgnoreCase)).FirstOrDefault();
+
+				// make sure it was found
+				if (Executable != null)
+				{
+					return true;
+				}
+			}
+
+			// if we got here, we are not usable
+			return false;
+		}
+	}
+
+
 	/// <summary>
 	/// Represents a loose collection of files. For Mac the executable will point to the binary.app bundle and
 	/// not the inner executable file
@@ -23,6 +57,12 @@ namespace Gauntlet
 		public override UnrealTargetPlatform Platform { get { return UnrealTargetPlatform.Mac; } }
 
 		public override string PlatformFolderPrefix { get { return "Mac"; } }
+
+		public override bool ShouldMakeBuildAvailable(string AppPath)
+		{
+			return MacAppChecker.IsAppBundleUsable(AppPath, null);
+		}
+
 	}
 
 	/// <summary>
@@ -95,8 +135,6 @@ namespace Gauntlet
 				DirsToRecurse = DiscoveredDirs.Except(Packages);
 			}
 
-			string ShortName = Regex.Replace(InProjectName, "Game", "", RegexOptions.IgnoreCase);
-
 			// Now 
 			foreach (DirectoryInfo Di in AllPackages)
 			{
@@ -108,31 +146,17 @@ namespace Gauntlet
 
 				if (Config != UnrealTargetConfiguration.Unknown)
 				{
-					string MacOSDir = Path.Combine(Di.FullName, "Contents", "MacOS");
-
-					// check both binary AND Engine content. Regular Unreal builds won't have the latter
-					if (Directory.Exists(MacOSDir)
-						&& Directory.Exists(Path.Combine(Di.FullName, "Contents", "UE")))
+					if (MacAppChecker.IsAppBundleUsable(Di.FullName, InProjectName))
 					{
-						// Check there's an executable
-						FileInfo Executable = new DirectoryInfo(MacOSDir).GetFiles().Where(Fi => Fi.Name.StartsWith(ShortName, StringComparison.OrdinalIgnoreCase)).FirstOrDefault();
-
-						if (Executable == null)
-						{
-							Log.Warning("Found packaged build at {0} with no valid executable");
-						}
-						else
-						{
-							// We can use the base packaged build class. it has everything we need.
-							MacPackagedBuild Build = new MacPackagedBuild(
-								UnrealTargetPlatform.Mac,
-								Config,
-								Role,
-								Di.FullName,
-								BuildFlags.Packaged | BuildFlags.CanReplaceCommandLine | BuildFlags.CanReplaceExecutable);
-
-							Builds.Add(Build);
-						}
+						// We can use the base packaged build class. it has everything we need.
+						MacPackagedBuild Build = new MacPackagedBuild(
+							UnrealTargetPlatform.Mac,
+							Config,
+							Role,
+							Di.FullName,
+							BuildFlags.Packaged | BuildFlags.CanReplaceCommandLine | BuildFlags.CanReplaceExecutable);
+
+						Builds.Add(Build);
 					}
 				}
 			}
diff --git a/Engine/Source/Programs/AutomationTool/Gauntlet/Unreal/BuildSource/Gauntlet.StagedBuildSource.cs b/Engine/Source/Programs/AutomationTool/Gauntlet/Unreal/BuildSource/Gauntlet.StagedBuildSource.cs
--- a/Engine/Source/Programs/AutomationTool/Gauntlet/Unreal/BuildSource/Gauntlet.StagedBuildSource.cs
+++ b/Engine/Source/Programs/AutomationTool/Gauntlet/Unreal/BuildSource/Gauntlet.StagedBuildSource.cs
@@ -214,7 +214,7 @@ namespace Gauntlet
 			return DestPath;
 		}
 
-		public static IEnumerable<T> CreateFromPath<T>(UnrealTargetPlatform InPlatform, string InProjectName, string InPath, string InExecutableExtension)
+		public static IEnumerable<T> CreateFromPath<T>(UnrealTargetPlatform InPlatform, string InProjectName, string InPath, string InExecutableExtension, StagedBuildSource<T> VerificationBuildSource = null)
 			where T : StagedBuild
 		{
 			string BuildPath = InPath;
@@ -282,17 +282,22 @@ namespace Gauntlet
 					UnrealTargetRole Role = UnrealHelpers.GetRoleFromExecutableName(InProjectName, App.Name);
 					string Flavor = UnrealHelpers.GetBuildFlavorFromExecutableName(InProjectName, App.Name);
 
-					if (Config != UnrealTargetConfiguration.Unknown && Role != UnrealTargetRole.Unknown && !DiscoveredBuilds.Any(B => B.Configuration == Config && B.Role == Role && B.Flavor == Flavor))
-					{
-						// store the exe path as relative to the staged dir path
-						T NewBuild = Activator.CreateInstance(typeof(T), new object[] { InPlatform, Config, Role, InPath, Utils.SystemHelpers.MakePathRelative(App.FullName, InPath), Flavor }) as T;
 
-						if (App.Name.StartsWith("UnrealGame", StringComparison.OrdinalIgnoreCase))
+					// if we have a verificatio BuildSource object, verify the build is usable before we even create the build
+					if (VerificationBuildSource == null || VerificationBuildSource.ShouldMakeBuildAvailable(App.FullName))
+					{
+						if (Config != UnrealTargetConfiguration.Unknown && Role != UnrealTargetRole.Unknown && !DiscoveredBuilds.Any(B => B.Configuration == Config && B.Role == Role && B.Flavor == Flavor))
 						{
-							NewBuild.Flags |= BuildFlags.ContentOnlyProject;
-						}
+							// store the exe path as relative to the staged dir path
+							T NewBuild = Activator.CreateInstance(typeof(T), new object[] { InPlatform, Config, Role, InPath, Utils.SystemHelpers.MakePathRelative(App.FullName, InPath), Flavor }) as T;
+
+							if (App.Name.StartsWith("UnrealGame", StringComparison.OrdinalIgnoreCase))
+							{
+								NewBuild.Flags |= BuildFlags.ContentOnlyProject;
+							}
 
-						DiscoveredBuilds.Add(NewBuild);
+							DiscoveredBuilds.Add(NewBuild);
+						}
 					}
 				}
 			}
@@ -332,6 +337,11 @@ namespace Gauntlet
 			return InPlatform == Platform; 
 		}
 
+		virtual public bool ShouldMakeBuildAvailable(string AppPath)
+		{
+			return true;
+		}
+
 		virtual public string ExecutableExtension
 		{
 			get { return AutomationTool.Platform.GetExeExtension(Platform); }
@@ -387,7 +397,7 @@ namespace Gauntlet
 
 				foreach (DirectoryInfo Di in AllDirs)
 				{
-					IEnumerable<IBuild> FoundBuilds = StagedBuild.CreateFromPath<T>(Platform, InProjectName, Di.FullName, ExecutableExtension);
+					IEnumerable<IBuild> FoundBuilds = StagedBuild.CreateFromPath<T>(Platform, InProjectName, Di.FullName, ExecutableExtension, this);
 
 					if (FoundBuilds != null)
 					{
