Fix bug where the "silent" mode for project cooking could still emit an exception if the return code from the commandlet was outside the positive signed integer range.

diff --git a/Engine/Source/Programs/AutomationTool/AutomationUtils/CommandletUtils.cs b/Engine/Source/Programs/AutomationTool/AutomationUtils/CommandletUtils.cs
--- a/Engine/Source/Programs/AutomationTool/AutomationUtils/CommandletUtils.cs
+++ b/Engine/Source/Programs/AutomationTool/AutomationUtils/CommandletUtils.cs
@@ -241,12 +241,26 @@ namespace AutomationTool
 		/// <param name="Commandlet">Commandlet name.</param>
 		/// <param name="Parameters">Command line parameters (without -run=)</param>
 		/// <param name="ErrorLevel">The minimum exit code, which is treated as an error.</param>
-		public static void RunCommandlet(FileReference ProjectName, string UnrealExe, string Commandlet, string Parameters = null, int ErrorLevel = 1)
+		public static void RunCommandlet(FileReference ProjectName, string UnrealExe, string Commandlet, string Parameters = null, uint ErrorLevel = 1)
 		{
 			string LogFile;
 			RunCommandlet(ProjectName, UnrealExe, Commandlet, Parameters, out LogFile, ErrorLevel);
 		}
 
+		/// <summary>
+		/// Runs a commandlet using Engine/Binaries/Win64/UnrealEditor-Cmd.exe.
+		/// </summary>
+		/// <param name="ProjectName">Project name.</param>
+		/// <param name="UnrealExe">The name of the Unreal Editor executable to use.</param>
+		/// <param name="Commandlet">Commandlet name.</param>
+		/// <param name="Parameters">Command line parameters (without -run=)</param>
+		/// <param name="ErrorLevel">The minimum exit code, which is treated as an error.</param>
+		public static void RunCommandlet(FileReference ProjectName, string UnrealExe, string Commandlet, string Parameters, int ErrorLevel)
+		{
+			string LogFile;
+			RunCommandlet(ProjectName, UnrealExe, Commandlet, Parameters, out LogFile, (uint)ErrorLevel);
+		}
+
 		/// <summary>
 		/// Runs a commandlet using Engine/Binaries/Win64/UnrealEditor-Cmd.exe.
 		/// </summary>
@@ -256,7 +270,7 @@ namespace AutomationTool
 		/// <param name="Parameters">Command line parameters (without -run=)</param>
 		/// <param name="DestLogFile">Log file after completion</param>
 		/// <param name="ErrorLevel">The minimum exit code, which is treated as an error.</param>
-		public static void RunCommandlet(FileReference ProjectName, string UnrealExe, string Commandlet, string Parameters, out string DestLogFile, int ErrorLevel = 1)
+		public static void RunCommandlet(FileReference ProjectName, string UnrealExe, string Commandlet, string Parameters, out string DestLogFile, uint ErrorLevel = 1)
 		{
 			string LocalLogFile;
 			IProcessResult RunResult;
@@ -267,6 +281,20 @@ namespace AutomationTool
 			FinishRunCommandlet(ProjectName, Commandlet, StartTime, RunResult, LocalLogFile, out DestLogFile, ErrorLevel);
 		}
 
+		/// <summary>
+		/// Runs a commandlet using Engine/Binaries/Win64/UnrealEditor-Cmd.exe.
+		/// </summary>
+		/// <param name="ProjectName">Project name.</param>
+		/// <param name="UnrealExe">The name of the Unreal Editor executable to use.</param>
+		/// <param name="Commandlet">Commandlet name.</param>
+		/// <param name="Parameters">Command line parameters (without -run=)</param>
+		/// <param name="DestLogFile">Log file after completion</param>
+		/// <param name="ErrorLevel">The minimum exit code, which is treated as an error.</param>
+		public static void RunCommandlet(FileReference ProjectName, string UnrealExe, string Commandlet, string Parameters, out string DestLogFile, int ErrorLevel)
+		{
+			RunCommandlet(ProjectName, UnrealExe, Commandlet, Parameters, out DestLogFile, (uint)ErrorLevel);
+		}
+
 		public static void StartRunCommandlet(FileReference ProjectName, string UnrealExe, string Commandlet, string Parameters, ERunOptions RunOptions, out string LocalLogFile, out IProcessResult RunResult, ProcessResult.SpewFilterCallbackType SpewFilterCallback=null)
 		{
 			Logger.LogInformation("Running UnrealEditor {Commandlet} for project {ProjectName}", Commandlet, ProjectName);
@@ -303,7 +331,7 @@ namespace AutomationTool
 			PopDir();
 		}
 
-		public static void FinishRunCommandlet(FileReference ProjectName, string Commandlet, DateTime StartTime, IProcessResult RunResult, string LocalLogFile, out string DestLogFile, int ErrorLevel = 1)
+		public static void FinishRunCommandlet(FileReference ProjectName, string Commandlet, DateTime StartTime, IProcessResult RunResult, string LocalLogFile, out string DestLogFile, uint ErrorLevel = 1)
 		{
 			// If we're running on a Windows build machine, copy any crash dumps into the log folder
 			if(HostPlatform.Current.HostEditorPlatform == UnrealTargetPlatform.Win64 && IsBuildMachine)
@@ -427,7 +455,7 @@ namespace AutomationTool
 			CommandUtils.DeleteFile_NoExceptions(LocalLogFile);
 
 			// Throw an exception if the execution failed. Draw attention to signal exit codes on Posix systems, rather than just printing the exit code
-			if (RunResult.ExitCode != 0 && (uint)RunResult.ExitCode >= (uint)ErrorLevel)
+			if (RunResult.ExitCode != 0 && (uint)RunResult.ExitCode >= ErrorLevel)
 			{
 				string ExitCodeDesc = "";
 				if(RunResult.ExitCode > 128 && RunResult.ExitCode < 128 + 32)
@@ -444,6 +472,11 @@ namespace AutomationTool
 				throw new CommandletException(DestLogFile, RunResult.ExitCode, "Editor terminated with exit code {0}{1} while running {2}{3}; see log {4}", RunResult.ExitCode, ExitCodeDesc, Commandlet, (ProjectName == null)? "" : String.Format(" for {0}", ProjectName), DestLogFile) { OutputFormat = AutomationExceptionOutputFormat.Minimal };
 			}
 		}
+		
+		public static void FinishRunCommandlet(FileReference ProjectName, string Commandlet, DateTime StartTime, IProcessResult RunResult, string LocalLogFile, out string DestLogFile, int ErrorLevel)
+		{
+			FinishRunCommandlet(ProjectName, Commandlet, StartTime, RunResult, LocalLogFile, out DestLogFile, (uint)ErrorLevel);
+		}
 
 		/// <summary>
 		/// Returns the default path of the editor executable to use for running commandlets.
