Gauntlet - Fix a crash when no UnrealDeviceReservation instance was initiated

#rnx

diff --git a/Engine/Source/Programs/AutomationTool/Gauntlet/Unreal/Base/Gauntlet.UnrealSession.cs b/Engine/Source/Programs/AutomationTool/Gauntlet/Unreal/Base/Gauntlet.UnrealSession.cs
--- a/Engine/Source/Programs/AutomationTool/Gauntlet/Unreal/Base/Gauntlet.UnrealSession.cs
+++ b/Engine/Source/Programs/AutomationTool/Gauntlet/Unreal/Base/Gauntlet.UnrealSession.cs
@@ -1067,7 +1067,10 @@ namespace Gauntlet
 
 		public void ReleaseSessionDevices()
 		{
-			UnrealDeviceReservation.ReleaseDevices();
+			if (UnrealDeviceReservation != null)
+			{
+				UnrealDeviceReservation.ReleaseDevices();
+			}
 			if (RolesToInstalls != null)
 			{
 				RolesToInstalls.Clear();
