[Automation] Improve csv build version validation checks to catch all build label suffixes.

diff --git a/Engine/Source/Programs/AutomationTool/Gauntlet/Framework/Utils/Gauntlet.ReportGenUtils.cs b/Engine/Source/Programs/AutomationTool/Gauntlet/Framework/Utils/Gauntlet.ReportGenUtils.cs
--- a/Engine/Source/Programs/AutomationTool/Gauntlet/Framework/Utils/Gauntlet.ReportGenUtils.cs
+++ b/Engine/Source/Programs/AutomationTool/Gauntlet/Framework/Utils/Gauntlet.ReportGenUtils.cs
@@ -347,25 +347,13 @@ namespace Gauntlet
 				Gauntlet.Log.Error("No valid CSV metadata found!");
 				return false;
 			}
-			string BuildVersion = InBuildVersion;
 
-			// Remove suffixes for comparison purposes. The executable is not aware of these
-			string[] SuffixesToStrip = { "-PF-", "-LWC", "-VK" };
-			foreach (string Suffix in SuffixesToStrip)
+			if (InBuildVersion.Contains(CsvMetadata["buildversion"], StringComparison.OrdinalIgnoreCase))
 			{
-				int Index = BuildVersion.IndexOf(Suffix);
-				if (Index != -1)
-				{
-					BuildVersion = BuildVersion.Substring(0, Index);
-				}
-			}
-
-			if (CsvMetadata["buildversion"].ToLower() == BuildVersion.ToLower())
-			{
-				Gauntlet.Log.Info("Test build version matches CSV metadata: " + BuildVersion);
+				Gauntlet.Log.Info("Test build version matches CSV metadata: " + InBuildVersion);
 				return true;
 			}
-			Gauntlet.Log.Error("Test build version (" + BuildVersion + ") doesn't match CSV metadata (" + CsvMetadata["buildversion"] + "). Gauntlet likely failed to deploy the build correctly. Please report to automation team");
+			Gauntlet.Log.Error("Test build version (" + InBuildVersion + ") doesn't match CSV metadata (" + CsvMetadata["buildversion"] + "). Gauntlet likely failed to deploy the build correctly. Please report to automation team");
 			return false;
 		}
 
