[Gauntlet] Adds UnrealRoleConfiguration feature

IUnrealRoleConfigurations are a type that let you modularize functionality for any role in a gauntlet test. It's primarily designed to generate command line arguments, but can do anything that's desired to the target Role.
IUnrealRoleConfigurations supply two interface functions

-ApplyConfigToRole - Core implementation of the configuration, this is the "doing"
-VerifyRoleConfig - Verification implementation. This gives you a chance to verify the "Apply" step had the desired result.

diff --git a/Engine/Source/Programs/AutomationTool/Gauntlet/Unreal/Base/Gauntlet.UnrealRoleConfiguration.cs b/Engine/Source/Programs/AutomationTool/Gauntlet/Unreal/Base/Gauntlet.UnrealRoleConfiguration.cs
new file mode 100644
--- /dev/null
+++ b/Engine/Source/Programs/AutomationTool/Gauntlet/Unreal/Base/Gauntlet.UnrealRoleConfiguration.cs
@@ -0,0 +1,20 @@
+// Copyright Epic Games, Inc. All Rights Reserved.
+
+using System;
+using AutomationTool;
+
+namespace Gauntlet
+{
+	/// <summary>
+	/// UnrealRoleConfigurations can be added to an UnrealTestRole to modularize the creation of commandline arguments.
+	/// A common use case is generating args for a the CSV Profiler. You would normally need to manually create the arguments for every role,
+	/// but with an UnrealRoleConfiguration you can set a couple properties and the commandline will be generated and applied automatically.
+	///
+	/// Experimental feature! Subject to rapid change.
+	/// </summary>
+	public interface IUnrealRoleConfiguration
+	{
+		void ApplyConfigToRole(UnrealTestRole Role);
+		void VerifyRoleConfig(UnrealTestRole Role);
+	}
+}
\ No newline at end of file
diff --git a/Engine/Source/Programs/AutomationTool/Gauntlet/Unreal/Base/Gauntlet.UnrealTestConfiguration.cs b/Engine/Source/Programs/AutomationTool/Gauntlet/Unreal/Base/Gauntlet.UnrealTestConfiguration.cs
--- a/Engine/Source/Programs/AutomationTool/Gauntlet/Unreal/Base/Gauntlet.UnrealTestConfiguration.cs
+++ b/Engine/Source/Programs/AutomationTool/Gauntlet/Unreal/Base/Gauntlet.UnrealTestConfiguration.cs
@@ -363,9 +363,9 @@ namespace Gauntlet
 	/// <summary>
 	/// This class represents a process-role in a test and defines the type, command line,
 	/// and controllers that are needed.
-	/// 
+	///
 	/// TODO - can this be removed and UnrealSessionRole used directly?
-	/// 
+	///
 	/// </summary>
 	public class UnrealTestRole
 	{
@@ -383,6 +383,7 @@ namespace Gauntlet
 			ExplicitClientCommandLine = string.Empty;
 			Controllers = new List<string>();
 			FilesToCopy = new List<UnrealFileToCopy>();
+			RoleConfigurations = new List<IUnrealRoleConfiguration>();
 			AdditionalArtifactDirectories = new List<EIntendedBaseCopyDirectory>();
 			RoleType = ERoleModifier.None;
 			InstallOnly = false;
@@ -438,7 +439,6 @@ namespace Gauntlet
 			}
 		}
 
-
 		/// <summary>
 		/// Dictionary of commandline arguments that are turned into a commandline at the end.
 		/// For flags, leave the value set to null. Created and then passed through to the Session Role's Commandline Object
@@ -451,6 +451,11 @@ namespace Gauntlet
 		/// </summary>
 		public List<string> Controllers { get; set; }
 
+		/// <summary>
+		/// Collection of modular configurations applied to this role
+		/// </summary>
+		public List<IUnrealRoleConfiguration> RoleConfigurations { get; set; }
+
 		/// <summary>
 		/// Explicit command line for this role. If this is set no other
 		/// options from above or configs will be applied!
diff --git a/Engine/Source/Programs/AutomationTool/Gauntlet/Unreal/Base/Gauntlet.UnrealTestNode.cs b/Engine/Source/Programs/AutomationTool/Gauntlet/Unreal/Base/Gauntlet.UnrealTestNode.cs
--- a/Engine/Source/Programs/AutomationTool/Gauntlet/Unreal/Base/Gauntlet.UnrealTestNode.cs
+++ b/Engine/Source/Programs/AutomationTool/Gauntlet/Unreal/Base/Gauntlet.UnrealTestNode.cs
@@ -616,6 +616,18 @@ namespace Gauntlet
 					// important, use the type from the ContextRolke because Server may have been mapped to EditorServer etc
 					UnrealTargetPlatform SessionPlatform = TestRole.PlatformOverride ?? RoleContext.Platform;
 
+					// Apply all role configurations
+					foreach(IUnrealRoleConfiguration RoleConfiguration in TestRole.RoleConfigurations)
+					{
+						RoleConfiguration.ApplyConfigToRole(TestRole);
+					}
+
+					// Verify all role configurations
+					foreach (IUnrealRoleConfiguration RoleConfiguration in TestRole.RoleConfigurations)
+					{
+						RoleConfiguration.VerifyRoleConfig(TestRole);
+					}
+
 					UnrealSessionRole SessionRole = new UnrealSessionRole(RoleContext.Type, SessionPlatform, RoleContext.Configuration, TestRole.CommandLine);
 					SessionRole.InstallOnly = TestRole.InstallOnly;
 					SessionRole.DeferredLaunch = TestRole.DeferredLaunch;
