Add UseDetailedInstallSizes to generate size metadata even if UseChunkDBs is false.  UseDetailedInstallSizes defaults to true.

diff --git a/Engine/Source/Programs/AutomationTool/AutomationUtils/BundleUtils.cs b/Engine/Source/Programs/AutomationTool/AutomationUtils/BundleUtils.cs
--- a/Engine/Source/Programs/AutomationTool/AutomationUtils/BundleUtils.cs
+++ b/Engine/Source/Programs/AutomationTool/AutomationUtils/BundleUtils.cs
@@ -24,6 +24,7 @@ namespace AutomationUtils.Automation
 			public bool bContainsShaderLibrary { get; set; }
 			public int Order { get; set; }
 			public bool	UseChunkDBs { get; set; }
+			public bool UseDetailedInstallSizes { get; set; } // default is true
 			public string ExecFileName { get; set; } // TODO: We never used this.  Clean this up.
 		}
 
@@ -116,7 +117,6 @@ namespace AutomationUtils.Automation
 						Bundle.bContainsShaderLibrary = false;
 					}
 				}
-
 				{
 					bool bUseChunkDBs;
 					if (!BundleConfig.GetBool(SectionName, "UseChunkDBs", out bUseChunkDBs))
@@ -128,6 +128,18 @@ namespace AutomationUtils.Automation
 						Bundle.UseChunkDBs = bUseChunkDBs;
 					}
 				}
+				{
+					bool bUseDetailedInstallSizes;
+					if (!BundleConfig.GetBool(SectionName, "UseDetailedInstallSizes", out bUseDetailedInstallSizes))
+					{
+						Bundle.UseDetailedInstallSizes = true; // default to true
+					}
+					else
+					{
+						Bundle.UseDetailedInstallSizes = bUseDetailedInstallSizes;
+					}
+				}
+
 				GetPlatformSettings(Bundle, BundleConfig, BundleDefinitionPrefix + Bundle.Name);
 
 				Results.Add(Bundle);
