WorldPartition - WorldPartitionBuilder.Automation.cs: Test if SubmitTags is null before accessing

diff --git a/Engine/Source/Programs/AutomationTool/Scripts/WorldPartitionBuilder.Automation.cs b/Engine/Source/Programs/AutomationTool/Scripts/WorldPartitionBuilder.Automation.cs
--- a/Engine/Source/Programs/AutomationTool/Scripts/WorldPartitionBuilder.Automation.cs
+++ b/Engine/Source/Programs/AutomationTool/Scripts/WorldPartitionBuilder.Automation.cs
@@ -40,7 +40,7 @@ namespace AutomationScripts.Automation
 			{
 				CommandletArgs += " -AutoSubmit";
 
-				if (SubmitTags.Length != 0)
+				if (!String.IsNullOrEmpty(SubmitTags))
 				{
 					CommandletArgs += " -AutoSubmitTags=\"" + SubmitTags + "\"";
 				}
