UAT - extended PakFileRules with list of allowed chunk names when generating on demand chunks

#rnx

diff --git a/Engine/Source/Programs/AutomationTool/Scripts/CopyBuildToStagingDirectory.Automation.cs b/Engine/Source/Programs/AutomationTool/Scripts/CopyBuildToStagingDirectory.Automation.cs
--- a/Engine/Source/Programs/AutomationTool/Scripts/CopyBuildToStagingDirectory.Automation.cs
+++ b/Engine/Source/Programs/AutomationTool/Scripts/CopyBuildToStagingDirectory.Automation.cs
@@ -2215,6 +2215,9 @@ namespace AutomationScripts
 			// Encryption key
 			public string EncryptionKeyGuid;
 
+			// List of allowed chunk names when generating on demand chunk(s) 
+			public HashSet<string> OnDemandAllowedChunks;
+
 			public static bool IsMatch(PakFileRules PakRules, KeyValuePair<string, string> StagingFile)
 			{
 				bool bMatched = !PakRules.bDisabled &&
@@ -2389,6 +2392,16 @@ namespace AutomationScripts
 					continue;
 				}
 
+				IReadOnlyList<string> AllowedChunks;
+				if (PakRulesConfig.TryGetValues(SectionName, "OnDemandAllowedChunks", out AllowedChunks))
+				{
+					PakRules.OnDemandAllowedChunks = new HashSet<string>();
+					foreach (var ChunkName in AllowedChunks)
+					{
+						PakRules.OnDemandAllowedChunks.Add(ChunkName);
+					}
+				}
+
 				IReadOnlyList<string> FilesEnumberable;
 				if (PakRulesConfig.TryGetValues(SectionName, "Files", out FilesEnumberable))
 				{
@@ -2420,24 +2433,18 @@ namespace AutomationScripts
 		/// </summary>
 		/// <param name="Params"></param>
 		/// <param name="SC"></param>
-		private static void ApplyPakFileRules(
+		private static PakFileRules? ApplyPakFileRules(
 			List<PakFileRules> RulesList,
 			KeyValuePair<string, string> StagingFile,
 			HashSet<ChunkDefinition> ModifyPakList,
 			ConcurrentDictionary<string, ChunkDefinition> ChunkNameToDefinition,
-			out bool bExcludeFromPaks,
-			out bool bStageLoose,
-			out bool bOnDemand,
-			out string EncryptionKeyOverrideGuid)
+			out bool bExcludeFromPaks)
 		{
 			bExcludeFromPaks = false;
-			bStageLoose = false;
-			bOnDemand = false;
-			EncryptionKeyOverrideGuid = null;
 
 			if (RulesList == null)
 			{
-				return;
+				return null;
 			}
 
 			// Search in order, return on first match
@@ -2445,16 +2452,9 @@ namespace AutomationScripts
 			{
 				if (PakFileRules.IsMatch(PakRules, StagingFile))
 				{
-					if (PakRules.bOnDemand)
-					{
-						bOnDemand = true;
-						EncryptionKeyOverrideGuid = PakRules.EncryptionKeyGuid;
-					}
-
 					if (PakRules.bStageLoose)
 					{
-						bStageLoose = true;
-						return;
+						return PakRules;
 					}
 
 					bool bOverrideChunkAssignment = false;
@@ -2463,7 +2463,7 @@ namespace AutomationScripts
 						// Only override the existing list if bOverrideChunkManifest is set
 						if (!PakRules.bOverrideChunkManifest)
 						{
-							return;
+							return PakRules;
 						}
 
 						bOverrideChunkAssignment = true;
@@ -2504,9 +2504,11 @@ namespace AutomationScripts
 						}
 					}
 
-					return;
+					return PakRules;
 				}
 			}
+
+			return null;
 		}
 
 		/// <summary>
@@ -2530,10 +2532,7 @@ namespace AutomationScripts
 			Parallel.ForEach(UnrealPakResponseFile, StagingFile =>
 			{
 				bool bExcludeFromPaks = false;
-				bool bStageLoose = false;
-				bool bOnDemand = false;
-				string EncryptionKeyOverrideGuid;
-				ApplyPakFileRules(PakRulesList, StagingFile, null, null, out bExcludeFromPaks, out bStageLoose, out bOnDemand, out EncryptionKeyOverrideGuid);
+				ApplyPakFileRules(PakRulesList, StagingFile, null, null, out bExcludeFromPaks);
 
 				if (bExcludeFromPaks)
 				{
@@ -4160,9 +4159,6 @@ namespace AutomationScripts
 				{
 					bool bAddedToChunk = false;
 					bool bExcludeFromPaks = false;
-					bool bStageLoose = false;
-					bool bOnDemand = false;
-					string EncryptionKeyOverrideGuid;
 					HashSet<ChunkDefinition> PakList = new HashSet<ChunkDefinition>();
 
 					string OriginalFilename = StagingFile.Key;
@@ -4196,8 +4192,9 @@ namespace AutomationScripts
 					}
 
 					// Now run through the pak rules which may override things
-					ApplyPakFileRules(PakRulesList, StagingFile, PakList, ChunkNameToDefinition, out bExcludeFromPaks, out bStageLoose, out bOnDemand, out EncryptionKeyOverrideGuid);
+					PakFileRules? MatchingRule = ApplyPakFileRules(PakRulesList, StagingFile, PakList, ChunkNameToDefinition, out bExcludeFromPaks);
 
+					bool bStageLoose = MatchingRule.HasValue && MatchingRule.Value.bStageLoose;
 					if (bStageLoose)
 					{
 						string StageLooseChunkName = "_stage_loose_files";
@@ -4225,6 +4222,17 @@ namespace AutomationScripts
 						ChunkDefinition TargetChunk = Chunk;
 
 						string OrigExt = Path.GetExtension(OriginalFilename);
+
+						bool bOnDemand = false;
+						if (MatchingRule.HasValue && MatchingRule.Value.bOnDemand)
+						{
+							bOnDemand = true;
+							if (MatchingRule.Value.OnDemandAllowedChunks != null)
+							{
+								bOnDemand = MatchingRule.Value.OnDemandAllowedChunks.Contains(TargetChunk.ChunkName);
+							}
+						}
+
 						if (bStageLoose || bOnDemand)
 						{
 							// make a new separate pak if not already marked as on demand
@@ -4237,10 +4245,10 @@ namespace AutomationScripts
 								if (bOnDemand)
 								{
 									TargetChunk.bOnDemand = true;
-									if (!string.IsNullOrEmpty(EncryptionKeyOverrideGuid))
+									if (!string.IsNullOrEmpty(MatchingRule.Value.EncryptionKeyGuid))
 									{
-										TargetChunk.EncryptionKeyGuid = EncryptionKeyOverrideGuid;
-										TargetChunk.RequestedEncryptionKeyGuid = EncryptionKeyOverrideGuid;
+										TargetChunk.EncryptionKeyGuid = MatchingRule.Value.EncryptionKeyGuid;
+										TargetChunk.RequestedEncryptionKeyGuid = MatchingRule.Value.EncryptionKeyGuid;
 									}
 								}
 							}
