Suppress output from the zen.exe command when performing oplog export in a ZenExportSnapshotTask so that we don't have builds failing because of a server error logged to output being parsed as an error.

diff --git a/Engine/Source/Programs/AutomationTool/BuildGraph/Tasks/ZenExportSnapshotTask.cs b/Engine/Source/Programs/AutomationTool/BuildGraph/Tasks/ZenExportSnapshotTask.cs
--- a/Engine/Source/Programs/AutomationTool/BuildGraph/Tasks/ZenExportSnapshotTask.cs
+++ b/Engine/Source/Programs/AutomationTool/BuildGraph/Tasks/ZenExportSnapshotTask.cs
@@ -225,9 +225,13 @@ namespace AutomationTool.Tasks
 						ExportKeyIds[ExportIndex] = DestinationKeyBuilder.ToString().ToLowerInvariant();
 						IoHash DestinationKeyHash = IoHash.Compute(Encoding.UTF8.GetBytes(ExportKeyIds[ExportIndex]));
 
+						ProcessResult.SpewFilterCallbackType SilentOutputFilter = new ProcessResult.SpewFilterCallbackType(Line =>
+							{
+								return null;
+							});
 						ExportSingleSourceCommandline.AppendFormat(" {0} --key {1} {2} {3}", HostUrlArg, DestinationKeyHash.ToString().ToLowerInvariant(), ExportSource.ProjectId, ExportSource.OplogId);
 						Logger.LogInformation("Running '{Arg0} {Arg1}'", CommandUtils.MakePathSafeToUseWithCommandLine(ZenExe.FullName), ExportSingleSourceCommandline.ToString());
-						CommandUtils.RunAndLog(CommandUtils.CmdEnv, ZenExe.FullName, ExportSingleSourceCommandline.ToString(), MaxSuccessCode: 1, Options: CommandUtils.ERunOptions.Default);
+						CommandUtils.RunAndLog(CommandUtils.CmdEnv, ZenExe.FullName, ExportSingleSourceCommandline.ToString(), MaxSuccessCode: 1, Options: CommandUtils.ERunOptions.Default, SpewFilterCallback: SilentOutputFilter);
 
 						ExportIndex = ExportIndex + 1;
 					}
