Increase maximum return code tolerated from the zen.exe command when performing oplog export in a ZenExportSnapshotTask so that we don't have build steps failing because of it while issues are debugged in its operation.

diff --git a/Engine/Source/Programs/AutomationTool/BuildGraph/Tasks/ZenExportSnapshotTask.cs b/Engine/Source/Programs/AutomationTool/BuildGraph/Tasks/ZenExportSnapshotTask.cs
--- a/Engine/Source/Programs/AutomationTool/BuildGraph/Tasks/ZenExportSnapshotTask.cs
+++ b/Engine/Source/Programs/AutomationTool/BuildGraph/Tasks/ZenExportSnapshotTask.cs
@@ -231,7 +231,7 @@ namespace AutomationTool.Tasks
 							});
 						ExportSingleSourceCommandline.AppendFormat(" {0} --key {1} {2} {3}", HostUrlArg, DestinationKeyHash.ToString().ToLowerInvariant(), ExportSource.ProjectId, ExportSource.OplogId);
 						Logger.LogInformation("Running '{Arg0} {Arg1}'", CommandUtils.MakePathSafeToUseWithCommandLine(ZenExe.FullName), ExportSingleSourceCommandline.ToString());
-						CommandUtils.RunAndLog(CommandUtils.CmdEnv, ZenExe.FullName, ExportSingleSourceCommandline.ToString(), MaxSuccessCode: 1, Options: CommandUtils.ERunOptions.Default, SpewFilterCallback: SilentOutputFilter);
+						CommandUtils.RunAndLog(CommandUtils.CmdEnv, ZenExe.FullName, ExportSingleSourceCommandline.ToString(), MaxSuccessCode: int.MaxValue, Options: CommandUtils.ERunOptions.Default, SpewFilterCallback: SilentOutputFilter);
 
 						ExportIndex = ExportIndex + 1;
 					}
