Archive additional APK products from Gradle

#android

diff --git a/Engine/Source/Programs/AutomationTool/Android/AndroidPlatform.Automation.cs b/Engine/Source/Programs/AutomationTool/Android/AndroidPlatform.Automation.cs
--- a/Engine/Source/Programs/AutomationTool/Android/AndroidPlatform.Automation.cs
+++ b/Engine/Source/Programs/AutomationTool/Android/AndroidPlatform.Automation.cs
@@ -1994,6 +1994,16 @@ public class AndroidPlatform : Platform
 				SC.ArchiveFiles(APKDirectory, Path.GetFileName(AFSApkName));
 			}
 
+			// add any other APKs with a prefix
+			IEnumerable<string> files = Directory.EnumerateFiles(APKDirectory, "*_" + APKNameWithoutExtension + ".apk", SearchOption.TopDirectoryOnly);
+			foreach (string filename in files)
+			{
+				if (filename != AFSApkName)
+				{
+					SC.ArchiveFiles(APKDirectory, Path.GetFileName(filename));
+				}
+			}
+
 			// verify the files exist
 			if (!FileExists(ApkName))
 			{
