[Gauntlet] Fix Android device match regex to match emulated devices

diff --git a/Engine/Source/Programs/AutomationTool/Gauntlet/Platform/Android/Gauntlet.TargetDeviceAndroid.cs b/Engine/Source/Programs/AutomationTool/Gauntlet/Platform/Android/Gauntlet.TargetDeviceAndroid.cs
--- a/Engine/Source/Programs/AutomationTool/Gauntlet/Platform/Android/Gauntlet.TargetDeviceAndroid.cs
+++ b/Engine/Source/Programs/AutomationTool/Gauntlet/Platform/Android/Gauntlet.TargetDeviceAndroid.cs
@@ -613,7 +613,7 @@ namespace Gauntlet
 		{
            var Result = RunAdbGlobalCommand("devices");
 
-            MatchCollection DeviceMatches = Regex.Matches(Result.Output, @"^([\d\w\.\:]{6,32})\s+(\w+)", RegexOptions.Multiline);
+            MatchCollection DeviceMatches = Regex.Matches(Result.Output, @"^([\d\w\.\:\-]{6,32})\s+(\w+)", RegexOptions.Multiline);
 
             var DeviceList = DeviceMatches.Cast<Match>().ToDictionary(
                 M => M.Groups[1].ToString(),
