- Fixed some issues with codesigning BP projects with modern Xcode, when BaseEngine.ini has no signing info, and the signing info is in the projects settings
  - The in-process project generation was not able to get the .uproject location
  - The making and using of the stub xcode projects were using UnrealGame target/scheme, not MyProjectGame (which Xcode makes since it can handle content-only uprojects)
  - UnrealGame.app, when just compiling, has no uproject, so it now use the dummy codesigning I added for stub generation, so no projects settings are needed

diff --git a/Engine/Source/Programs/AutomationTool/Apple/ApplePlatform.Automation.cs b/Engine/Source/Programs/AutomationTool/Apple/ApplePlatform.Automation.cs
--- a/Engine/Source/Programs/AutomationTool/Apple/ApplePlatform.Automation.cs
+++ b/Engine/Source/Programs/AutomationTool/Apple/ApplePlatform.Automation.cs
@@ -80,6 +80,15 @@ public abstract class ApplePlatform : Platform
 		base.PreBuildAgenda(Build, Agenda, Params);
 	}
 
+	private string MakeContentOnlyTargetName(TargetReceipt Target, string GameName)
+	{
+		// for a BP project, with a uproject, where the TargetName is one of the Unreal targets, the xcode project actually makes
+		// a <ProjectName><Type> scheme/target, so we need to use it for the -SingleTarget param, and the scheme we xcodebuild
+		// note: this must match the code in ProjectFileGenerator.cs / AddProjectsForAllTargets, like this:
+		// ProjectFile = FindOrAddProject(GetProjectLocation($"{ProjectName}{EngineTarget.TargetRules!.Type}"), ContentOnlyGameProject.Directory, ...
+		// in the "if (bAllowContentOnlyProjects)" block
+		return $"{GameName}{Target.TargetType}";
+	}
 
 	public override void PostStagingFileCopy(ProjectParams Params, DeploymentContext SC)
 	{
@@ -96,14 +105,21 @@ public abstract class ApplePlatform : Platform
 					$"EFFECTIVE_PLATFORM_NAME={SC.StageDirectory.GetDirectoryName()} " +
 					// set where the stage directory is
 					$"UE_OVERRIDE_STAGE_DIR=\"{SC.StageDirectory}\"";
+
+				string TargetName = Target.TargetName;
 				if (!Params.IsCodeBasedProject)
 				{
 					// instead of staging an UnrealGame.app, stage something with the project name, like say MyProjectClient-IOS-Shipping.app
 					string ProductName = AppleExports.MakeBinaryFileName(SC.ShortProjectName, Target.Platform, Target.Configuration, Target.Architectures, UnrealTargetConfiguration.Development, null);
 					ExtraOptions += $" PRODUCT_NAME={ProductName}";
+
+					if (Params.RawProjectPath != null)
+					{
+						TargetName = MakeContentOnlyTargetName(Target, Params.ShortProjectName);
+					}
 				}
 
-				AppleExports.BuildWithStubXcodeProject(SC.RawProjectPath, Target.Platform, Target.Architectures, Target.Configuration, Target.TargetName, AppleExports.XcodeBuildMode.Stage, Logger, ExtraOptions);
+				AppleExports.BuildWithStubXcodeProject(SC.RawProjectPath, Target.Platform, Target.Architectures, Target.Configuration, TargetName, AppleExports.XcodeBuildMode.Stage, Logger, ExtraOptions);
 			}
 		}
 	}
@@ -182,6 +198,7 @@ public abstract class ApplePlatform : Platform
 			foreach (StageTarget Target in SC.StageTargets)
 			{
 				TargetReceipt Receipt = Target.Receipt;
+				string TargetName = Receipt.TargetName;
 
 				string ExtraOptions =
 					// set where the stage directory is
@@ -191,6 +208,11 @@ public abstract class ApplePlatform : Platform
 					// override where the .app will be located and named
 					ExtraOptions += $" SYMROOT=\"{SC.ProjectRoot}/Binaries\"";
 					ExtraOptions += $" PRODUCT_NAME={Params.ShortProjectName}";
+
+					if (Params.RawProjectPath != null)
+					{
+						TargetName = MakeContentOnlyTargetName(Receipt, Params.ShortProjectName);
+					}
 				}
 
 
@@ -198,7 +220,7 @@ public abstract class ApplePlatform : Platform
 				// the archive will be created in the standard Archives location accessible via Xcode. Using -archive will copy it out into
 				// the specified location for use as needed
 				AppleExports.XcodeBuildMode BuildMode = Params.Distribution ? AppleExports.XcodeBuildMode.Distribute : AppleExports.XcodeBuildMode.Package;
-				if (AppleExports.BuildWithStubXcodeProject(Params.RawProjectPath, Receipt.Platform, Receipt.Architectures, Receipt.Configuration, Receipt.TargetName, BuildMode, Logger, ExtraOptions) == 0)
+				if (AppleExports.BuildWithStubXcodeProject(Params.RawProjectPath, Receipt.Platform, Receipt.Architectures, Receipt.Configuration, TargetName, BuildMode, Logger, ExtraOptions) == 0)
 				{
 					Logger.LogInformation("=====================================================================================");
 					if (Params.Distribution)
diff --git a/Engine/Source/Programs/UnrealBuildTool/ToolChain/AppleToolChain.cs b/Engine/Source/Programs/UnrealBuildTool/ToolChain/AppleToolChain.cs
--- a/Engine/Source/Programs/UnrealBuildTool/ToolChain/AppleToolChain.cs
+++ b/Engine/Source/Programs/UnrealBuildTool/ToolChain/AppleToolChain.cs
@@ -815,7 +815,9 @@ namespace UnrealBuildTool
 				// this could be improved if ProjectFileGenerator constructor took a Arguments param, and it could parse it there instead of the GenerateProjectFilesMode
 				Generator.SingleTargetName = SingleTargetName;
 				// don't need the editor data since these are stub projects
+				ProjectFileGenerator.Current = Generator;
 				bool bSucces = Generator.GenerateProjectFiles(PlatformProjectGenerators, Arguments, bCacheDataForEditor: false, Logger);
+				ProjectFileGenerator.Current = null;
 				XcodeProjectFile = Generator.XCWorkspace;
 				return bSucces;
 			}
@@ -1011,8 +1013,14 @@ namespace UnrealBuildTool
 				return 0;
 			}
 
+			// for mobile builds (which need real codesigning to be able to run), we use dummy codesigning when making a .stub (which will
+			// be sent to Windows an re-codesigned), or when making UnrealGame.app without a .uproject (we will always make a .app
+			// again with a .uproject on the commandline to be able to get the staged data that will be pulled in to the app - we can't run 
+			// without a Staged directory, so this dummy codesigned .app won't be used directly)
+			bool bUseDummySigning = Target.Platform != UnrealTargetPlatform.Mac && (Target.bCreateStubIPA || Target.ProjectFile == null);
+
 			int ExitCode = AppleExports.BuildWithStubXcodeProject(Target.ProjectFile, Target.Platform, Target.Architectures, Target.Configuration, Target.TargetName, 
-				AppleExports.XcodeBuildMode.PostBuildSync, Logger, "", bForceDummySigning:Target.bCreateStubIPA);
+				AppleExports.XcodeBuildMode.PostBuildSync, Logger, "", bForceDummySigning: bUseDummySigning);
 			if (ExitCode != 0)
 			{
 				Logger.LogError("ERROR: Failed to finalize the .app with Xcode. Check the log for more information");
