- Tag everything inside .app's after a CompileTask from BuildGraph. This is because with modern Xcode, Xcode will copy files into the .app that UBT doesn't know about

diff --git a/Engine/Source/Programs/AutomationTool/BuildGraph/Tasks/CompileTask.cs b/Engine/Source/Programs/AutomationTool/BuildGraph/Tasks/CompileTask.cs
--- a/Engine/Source/Programs/AutomationTool/BuildGraph/Tasks/CompileTask.cs
+++ b/Engine/Source/Programs/AutomationTool/BuildGraph/Tasks/CompileTask.cs
@@ -191,10 +191,34 @@ namespace AutomationTool.Tasks
 					throw new AutomationException("Missing manifest for target {0} {1} {2}", TargetTagName.Key.TargetName, TargetTagName.Key.Platform, TargetTagName.Key.Config);
 				}
 
-				foreach(string TagName in CustomTask.SplitDelimitedList(TargetTagName.Value))
+				HashSet<FileReference> ManifestBuildProducts = Manifest.BuildProducts.Select(x => new FileReference(x)).ToHashSet();
+
+				// when we make a Mac/IOS build, Xcode will finalize the .app directory, adding files that UBT has no idea about, so now we recursively add any files in the .app
+				// as BuildProducts. look for any .apps that we have any files as BuildProducts, and expand to include all files in the .app
+				if (BuildHostPlatform.Current.Platform == UnrealTargetPlatform.Mac)
+				{
+					HashSet<string> AppBundleLocations = new();
+					foreach (FileReference File in ManifestBuildProducts)
+					{
+						// look for a ".app/" portion and chop off anything after it
+						int AppLocation = File.FullName.IndexOf(".app/", StringComparison.InvariantCultureIgnoreCase);
+						if (AppLocation > 0)
+						{
+							AppBundleLocations.Add(File.FullName.Substring(0, AppLocation + 4));
+						}
+					}
+
+					// now with a unique set of app bundles, add all files in them
+					foreach (string AppBundleLocation in AppBundleLocations)
+					{
+						ManifestBuildProducts.UnionWith(DirectoryReference.EnumerateFiles(new DirectoryReference(AppBundleLocation), "*", System.IO.SearchOption.AllDirectories));
+					}
+				}
+
+				foreach (string TagName in CustomTask.SplitDelimitedList(TargetTagName.Value))
 				{
 					HashSet<FileReference> FileSet = CustomTask.FindOrAddTagSet(TagNameToFileSet, TagName);
-					FileSet.UnionWith(Manifest.BuildProducts.Select(x => new FileReference(x)));
+					FileSet.UnionWith(ManifestBuildProducts);
 				}
 			}
 
@@ -202,7 +226,6 @@ namespace AutomationTool.Tasks
 			BuildProducts.UnionWith(Builder.BuildProductFiles.Select(x => new FileReference(x)));
 			return Task.CompletedTask;
 		}
-	}
 
 	/// <summary>
 	/// Compiles a target with UnrealBuildTool.
