- Fixed running Mac games from Staged dir (like when doing Launch On to Mac)

diff --git a/Engine/Source/Programs/AutomationTool/Mac/MacPlatform.Automation.cs b/Engine/Source/Programs/AutomationTool/Mac/MacPlatform.Automation.cs
--- a/Engine/Source/Programs/AutomationTool/Mac/MacPlatform.Automation.cs
+++ b/Engine/Source/Programs/AutomationTool/Mac/MacPlatform.Automation.cs
@@ -587,7 +587,15 @@ public class MacPlatform : ApplePlatform
 
 	public override IProcessResult RunClient(ERunOptions ClientRunFlags, string ClientApp, string ClientCmdLine, ProjectParams Params)
 	{
-		if (!File.Exists(ClientApp))
+		if (AppleExports.UseModernXcode(Params.RawProjectPath))
+		{
+			// moden creates a full .app in the root of the Staged dir, but ClientApp as passed in is in Staged/Project/Binaries/Mac, which exists, but is not a full app
+			string ExeName = Path.GetFileNameWithoutExtension(ClientApp);
+			Int32 BaseDirLen = Params.BaseStageDirectory.Length;
+			string StageSubDir = ClientApp.Substring(BaseDirLen, ClientApp.IndexOf("/", BaseDirLen + 1) - BaseDirLen);
+			ClientApp = CombinePaths(Params.BaseStageDirectory, StageSubDir, $"{ExeName}.app/Contents/MacOS/{ExeName}");
+		}
+		else if (!File.Exists(ClientApp))
 		{
 			if (Directory.Exists(ClientApp + ".app"))
 			{
