Fixing compile error

#rnx

diff --git a/Engine/Source/Programs/AutomationTool/BuildGraph/Tasks/CompileTask.cs b/Engine/Source/Programs/AutomationTool/BuildGraph/Tasks/CompileTask.cs
--- a/Engine/Source/Programs/AutomationTool/BuildGraph/Tasks/CompileTask.cs
+++ b/Engine/Source/Programs/AutomationTool/BuildGraph/Tasks/CompileTask.cs
@@ -226,6 +226,7 @@ namespace AutomationTool.Tasks
 			BuildProducts.UnionWith(Builder.BuildProductFiles.Select(x => new FileReference(x)));
 			return Task.CompletedTask;
 		}
+	}
 
 	/// <summary>
 	/// Compiles a target with UnrealBuildTool.
