Allow loading bundle config from editor INI file.

diff --git a/Engine/Source/Programs/AutomationTool/AutomationUtils/BundleUtils.cs b/Engine/Source/Programs/AutomationTool/AutomationUtils/BundleUtils.cs
--- a/Engine/Source/Programs/AutomationTool/AutomationUtils/BundleUtils.cs
+++ b/Engine/Source/Programs/AutomationTool/AutomationUtils/BundleUtils.cs
@@ -33,110 +33,140 @@ namespace AutomationUtils.Automation
 			LoadBundleConfig<BundleSettings>(ProjectDir, Platform, out Bundles, delegate (BundleSettings Settings, ConfigHierarchy BundleConfig, string Section) { });
 		}
 
-		public static void LoadBundleConfig<TPlatformBundleSettings>(DirectoryReference ProjectDir, UnrealTargetPlatform Platform, 
+		public static void LoadBundleConfig_IncludeEditorConfig(DirectoryReference ProjectDir, UnrealTargetPlatform Platform, out IReadOnlyDictionary<string, BundleSettings> Bundles)
+		{
+			LoadBundleConfig_IncludeEditorConfig<BundleSettings>(ProjectDir, Platform, out Bundles, delegate (BundleSettings Settings, ConfigHierarchy BundleConfig, string Section) { });
+		}
+
+		public static void LoadBundleConfig<TPlatformBundleSettings>(DirectoryReference ProjectDir, UnrealTargetPlatform Platform,
+			out IReadOnlyDictionary<string, TPlatformBundleSettings> Bundles,
+			Action<TPlatformBundleSettings, ConfigHierarchy, string> GetPlatformSettings)
+			where TPlatformBundleSettings : BundleSettings, new()
+		{
+			LoadBundleConfigInternal<TPlatformBundleSettings>(false, ProjectDir, Platform, out Bundles, GetPlatformSettings);
+		}
+
+		public static void LoadBundleConfig_IncludeEditorConfig<TPlatformBundleSettings>(DirectoryReference ProjectDir, UnrealTargetPlatform Platform,
+			out IReadOnlyDictionary<string, TPlatformBundleSettings> Bundles,
+			Action<TPlatformBundleSettings, ConfigHierarchy, string> GetPlatformSettings)
+			where TPlatformBundleSettings : BundleSettings, new()
+		{
+			LoadBundleConfigInternal<TPlatformBundleSettings>(true, ProjectDir, Platform, out Bundles, GetPlatformSettings);
+		}
+
+		private static void LoadBundleConfigInternal<TPlatformBundleSettings>(bool bIncludeEditorConfig, DirectoryReference ProjectDir, UnrealTargetPlatform Platform, 
 			out IReadOnlyDictionary<string, TPlatformBundleSettings> Bundles, 
 			Action<TPlatformBundleSettings, ConfigHierarchy, string> GetPlatformSettings) 
 			where TPlatformBundleSettings : BundleSettings, new()
 		{
 			var Results = new List<TPlatformBundleSettings>();
 
-			ConfigHierarchy BundleConfig = ConfigCache.ReadHierarchy(ConfigHierarchyType.InstallBundle, ProjectDir, Platform);
+			List<ConfigHierarchy> BundleConfigs = new List<ConfigHierarchy>();
+			BundleConfigs.Add(ConfigCache.ReadHierarchy(ConfigHierarchyType.InstallBundle, ProjectDir, Platform));
+			if (bIncludeEditorConfig)
+			{
+				BundleConfigs.Add(ConfigCache.ReadHierarchy(ConfigHierarchyType.Editor, ProjectDir, Platform));
+			}
 
 			const string BundleDefinitionPrefix = "InstallBundleDefinition ";
 
-			foreach (string SectionName in BundleConfig.SectionNames)
+			foreach (ConfigHierarchy BundleConfig in BundleConfigs)
 			{
-				if (!SectionName.StartsWith(BundleDefinitionPrefix))
-					continue;
-
-				TPlatformBundleSettings Bundle = new TPlatformBundleSettings();
-				Bundle.Name = SectionName.Substring(BundleDefinitionPrefix.Length);
-				{
-					int Order;
-					if(BundleConfig.GetInt32(SectionName, "Order", out Order))
-					{
-						Bundle.Order = Order;
-					}
-					else
-					{
-						Bundle.Order = int.MaxValue;
-					}
-				}
-				{
-					List<string> Tags;
-					if (BundleConfig.GetArray(SectionName, "Tags", out Tags))
-					{
-						Bundle.Tags = Tags;
-					}
-					else 
-					{
-						Bundle.Tags = new List<string>(); 
-					}
-				}
-				{
-					List<string> Dependencies;
-					if (BundleConfig.GetArray(SectionName, "Dependencies", out Dependencies))
-					{
-						Bundle.Dependencies = Dependencies;
-					}
-					else
-					{
-						Bundle.Dependencies = new List<string>();
-					}
-				}
-				{
-					List<string> FileRegex;
-					if (BundleConfig.GetArray(SectionName, "FileRegex", out FileRegex))
-					{
-						Bundle.FileRegex = FileRegex;
-					}
-					else
-					{
-						Bundle.FileRegex = new List<string>();
-					}
-				}
+				foreach (string SectionName in BundleConfig.SectionNames)
 				{
-					List<string> Files;
-					if (BundleConfig.GetArray(SectionName, "Files", out Files))
+					if (!SectionName.StartsWith(BundleDefinitionPrefix))
+						continue;
+
+					TPlatformBundleSettings Bundle = new TPlatformBundleSettings();
+					Bundle.Name = SectionName.Substring(BundleDefinitionPrefix.Length);
 					{
-						Bundle.Files = Files;
+						int Order;
+						if (BundleConfig.GetInt32(SectionName, "Order", out Order))
+						{
+							Bundle.Order = Order;
+						}
+						else
+						{
+							Bundle.Order = int.MaxValue;
+						}
 					}
-					else
 					{
-						Bundle.Files = new List<string>();
+						List<string> Tags;
+						if (BundleConfig.GetArray(SectionName, "Tags", out Tags))
+						{
+							Bundle.Tags = Tags;
+						}
+						else
+						{
+							Bundle.Tags = new List<string>();
+						}
 					}
-				}
-				{
-					bool bContainsShaderLibrary;
-					if (BundleConfig.GetBool(SectionName, "ContainsShaderLibrary", out bContainsShaderLibrary))
 					{
-						Bundle.bContainsShaderLibrary = bContainsShaderLibrary;
+						List<string> Dependencies;
+						if (BundleConfig.GetArray(SectionName, "Dependencies", out Dependencies))
+						{
+							Bundle.Dependencies = Dependencies;
+						}
+						else
+						{
+							Bundle.Dependencies = new List<string>();
+						}
 					}
-					else 
 					{
-						Bundle.bContainsShaderLibrary = false;
+						List<string> FileRegex;
+						if (BundleConfig.GetArray(SectionName, "FileRegex", out FileRegex))
+						{
+							Bundle.FileRegex = FileRegex;
+						}
+						else
+						{
+							Bundle.FileRegex = new List<string>();
+						}
 					}
-				}
-				{
-					bool bUseChunkDBs;
-					if (!BundleConfig.GetBool(SectionName, "UseChunkDBs", out bUseChunkDBs))
 					{
-						Bundle.UseChunkDBs = false;
+						List<string> Files;
+						if (BundleConfig.GetArray(SectionName, "Files", out Files))
+						{
+							Bundle.Files = Files;
+						}
+						else
+						{
+							Bundle.Files = new List<string>();
+						}
 					}
-					else
 					{
-						Bundle.UseChunkDBs = bUseChunkDBs;
+						bool bContainsShaderLibrary;
+						if (BundleConfig.GetBool(SectionName, "ContainsShaderLibrary", out bContainsShaderLibrary))
+						{
+							Bundle.bContainsShaderLibrary = bContainsShaderLibrary;
+						}
+						else
+						{
+							Bundle.bContainsShaderLibrary = false;
+						}
 					}
-				}
-				{
-					bool bUseDetailedInstallSizes;
-					if (!BundleConfig.GetBool(SectionName, "UseDetailedInstallSizes", out bUseDetailedInstallSizes))
+
 					{
-						Bundle.UseDetailedInstallSizes = true; // default to true
+						bool bUseChunkDBs;
+						if (!BundleConfig.GetBool(SectionName, "UseChunkDBs", out bUseChunkDBs))
+						{
+							Bundle.UseChunkDBs = false;
+						}
+						else
+						{
+							Bundle.UseChunkDBs = bUseChunkDBs;
+						}
 					}
-					else
 					{
-						Bundle.UseDetailedInstallSizes = bUseDetailedInstallSizes;
+						bool bUseDetailedInstallSizes;
+						if (!BundleConfig.GetBool(SectionName, "UseDetailedInstallSizes", out bUseDetailedInstallSizes))
+						{
+							Bundle.UseDetailedInstallSizes = true; // default to true
+						}
+						else
+						{
+							Bundle.UseDetailedInstallSizes = bUseDetailedInstallSizes;
+						}
 					}
 				}
 
