Fix bad merge
- see 27072463

diff --git a/Engine/Source/Programs/AutomationTool/AutomationUtils/BundleUtils.cs b/Engine/Source/Programs/AutomationTool/AutomationUtils/BundleUtils.cs
--- a/Engine/Source/Programs/AutomationTool/AutomationUtils/BundleUtils.cs
+++ b/Engine/Source/Programs/AutomationTool/AutomationUtils/BundleUtils.cs
@@ -168,11 +168,11 @@ namespace AutomationUtils.Automation
 							Bundle.UseDetailedInstallSizes = bUseDetailedInstallSizes;
 						}
 					}
-				}
 
-				GetPlatformSettings(Bundle, BundleConfig, BundleDefinitionPrefix + Bundle.Name);
+					GetPlatformSettings(Bundle, BundleConfig, BundleDefinitionPrefix + Bundle.Name);
 
-				Results.Add(Bundle);
+					Results.Add(Bundle);
+				}
 			}
 
 			// Use OrderBy and not Sort because OrderBy is stable
