Add version of automation Platform.RunClient that takes a DeploymentContext

#rnx

diff --git a/Engine/Source/Programs/AutomationTool/AutomationUtils/Platform.cs b/Engine/Source/Programs/AutomationTool/AutomationUtils/Platform.cs
--- a/Engine/Source/Programs/AutomationTool/AutomationUtils/Platform.cs
+++ b/Engine/Source/Programs/AutomationTool/AutomationUtils/Platform.cs
@@ -388,6 +388,19 @@ namespace AutomationTool
 			Logger.LogWarning("{PlatformType} does not implement Deploy...", PlatformType);
 		}
 
+		/// <summary>
+		/// Run the client application on the platform
+		/// </summary>
+		/// <param name="ClientRunFlags"></param>
+		/// <param name="ClientApp"></param>
+		/// <param name="ClientCmdLine"></param>
+		/// <param name="Params"></param>
+		/// <param name="SC"></param>
+		public virtual IProcessResult RunClient(ERunOptions ClientRunFlags, string ClientApp, string ClientCmdLine, ProjectParams Params, DeploymentContext SC)
+		{
+			return RunClient(ClientRunFlags, ClientApp, ClientCmdLine, Params );
+		}
+
 		/// <summary>
 		/// Run the client application on the platform
 		/// </summary>
diff --git a/Engine/Source/Programs/AutomationTool/Scripts/RunProjectCommand.Automation.cs b/Engine/Source/Programs/AutomationTool/Scripts/RunProjectCommand.Automation.cs
--- a/Engine/Source/Programs/AutomationTool/Scripts/RunProjectCommand.Automation.cs
+++ b/Engine/Source/Programs/AutomationTool/Scripts/RunProjectCommand.Automation.cs
@@ -325,7 +325,7 @@ namespace AutomationScripts
 				int LastAutoFailIndex = -1;
 
 				Logger.LogInformation("Starting Client for unattended test....");
-				ClientProcess = SC.StageTargetPlatform.RunClient(ClientRunFlags, ClientApp, ClientCmdLine, Params);
+				ClientProcess = SC.StageTargetPlatform.RunClient(ClientRunFlags, ClientApp, ClientCmdLine, Params, SC);
 
 				if (DedicatedServerProcess != null)
 				{
@@ -334,7 +334,7 @@ namespace AutomationScripts
 						for (int i = 1; i < NumClients; i++)
 						{
 							Logger.LogInformation("Starting Extra Client {i} for unattended test....", i);
-							ExtraClients.Add(SC.StageTargetPlatform.RunClient(ExtraClientRunFlags, ClientApp, ClientCmdLine, Params));
+							ExtraClients.Add(SC.StageTargetPlatform.RunClient(ExtraClientRunFlags, ClientApp, ClientCmdLine, Params, SC));
 						}
 					}
 				}
@@ -458,7 +458,7 @@ namespace AutomationScripts
 			else
 			{
 				Logger.LogInformation("Starting Client....");
-				ClientProcess = SC.StageTargetPlatform.RunClient(ClientRunFlags, ClientApp, ClientCmdLine, Params);
+				ClientProcess = SC.StageTargetPlatform.RunClient(ClientRunFlags, ClientApp, ClientCmdLine, Params, SC);
 
 				if (DedicatedServerProcess != null)
 				{
@@ -467,7 +467,7 @@ namespace AutomationScripts
 						for (int i = 1; i < NumClients; i++)
 						{
 							Logger.LogInformation("Starting Extra Client {i}....", i);
-							ExtraClients.Add(SC.StageTargetPlatform.RunClient(ExtraClientRunFlags, ClientApp, ClientCmdLine, Params));
+							ExtraClients.Add(SC.StageTargetPlatform.RunClient(ExtraClientRunFlags, ClientApp, ClientCmdLine, Params, SC));
 						}
 					}
 				}
