Gauntlet - Make sure to kill child process when closing Unreal app

#rnx

diff --git a/Engine/Source/Programs/AutomationTool/Gauntlet/Platform/IOS/Gauntlet.TargetDeviceIOS.cs b/Engine/Source/Programs/AutomationTool/Gauntlet/Platform/IOS/Gauntlet.TargetDeviceIOS.cs
--- a/Engine/Source/Programs/AutomationTool/Gauntlet/Platform/IOS/Gauntlet.TargetDeviceIOS.cs
+++ b/Engine/Source/Programs/AutomationTool/Gauntlet/Platform/IOS/Gauntlet.TargetDeviceIOS.cs
@@ -122,7 +122,7 @@ namespace Gauntlet
 			if (!HasExited)
 			{
 				WasKilled = true;
-				ProcessResult.ProcessObject.Kill();
+				ProcessResult.ProcessObject.Kill(true);
 			}
 		}
 
diff --git a/Engine/Source/Programs/AutomationTool/Gauntlet/Platform/Windows/Gauntlet.TargetDeviceWindows.cs b/Engine/Source/Programs/AutomationTool/Gauntlet/Platform/Windows/Gauntlet.TargetDeviceWindows.cs
--- a/Engine/Source/Programs/AutomationTool/Gauntlet/Platform/Windows/Gauntlet.TargetDeviceWindows.cs
+++ b/Engine/Source/Programs/AutomationTool/Gauntlet/Platform/Windows/Gauntlet.TargetDeviceWindows.cs
@@ -56,7 +56,7 @@ namespace Gauntlet
 			if (!HasExited)
 			{
 				WasKilled = true;
-				ProcessResult.ProcessObject.Kill();
+				ProcessResult.ProcessObject.Kill(true);
 			}
 		}
 
