- Removed the quotes around the signing identity for Xcode manual codesigning
- During staging, we now make the UBTGenerated/Info.Template.plist in the project folder instead of Engine folder, for BP projects
- Fixed an AppStoreConnect label in the Xcode project settings that would have confused anyone trying to set it up

diff --git a/Engine/Source/Developer/Mac/MacTargetPlatform/Classes/XcodeProjectSettings.h b/Engine/Source/Developer/Mac/MacTargetPlatform/Classes/XcodeProjectSettings.h
--- a/Engine/Source/Developer/Mac/MacTargetPlatform/Classes/XcodeProjectSettings.h
+++ b/Engine/Source/Developer/Mac/MacTargetPlatform/Classes/XcodeProjectSettings.h
@@ -126,7 +126,7 @@ public:
 	/**
 	 * The Key ID for your App Store Connect generated API key, a 32 hex-character string, including dashes
 	 */
-	UPROPERTY(EditAnywhere, config, Category="Code Signing", meta = (EditCondition="bUseModernXcode && bUseAutomaticCodeSigning && bUseAppStoreConnect", DisplayName = "App Store Connect: Issuer ID"))
+	UPROPERTY(EditAnywhere, config, Category="Code Signing", meta = (EditCondition="bUseModernXcode && bUseAutomaticCodeSigning && bUseAppStoreConnect", DisplayName = "App Store Connect: Key ID"))
 	FString AppStoreConnectKeyID;
 	
 	/**
diff --git a/Engine/Source/Programs/AutomationTool/IOS/IOSPlatform.Automation.cs b/Engine/Source/Programs/AutomationTool/IOS/IOSPlatform.Automation.cs
--- a/Engine/Source/Programs/AutomationTool/IOS/IOSPlatform.Automation.cs
+++ b/Engine/Source/Programs/AutomationTool/IOS/IOSPlatform.Automation.cs
@@ -1373,10 +1373,17 @@ public class IOSPlatform : ApplePlatform
 
 					var TargetConfiguration = SC.StageTargetConfigurations[0];
 
+					DirectoryReference ProjectRoot = SC.ProjectRoot;
+					// keep old logic for BP projects with legacy
+					if (!AppleExports.UseModernXcode(SC.RawProjectPath) && !SC.IsCodeBasedProject)
+					{
+						ProjectRoot = DirectoryReference.Combine(SC.LocalRoot, "Engine");
+					}
+
 					DeployGeneratePList(
 							SC.RawProjectPath,
 							TargetConfiguration,
-							(SC.IsCodeBasedProject ? SC.ProjectRoot : DirectoryReference.Combine(SC.LocalRoot, "Engine")),
+							ProjectRoot,
 							!SC.IsCodeBasedProject,
 							(SC.IsCodeBasedProject ? TargetName : "UnrealGame"),
 							SC.IsCodeBasedProject ? false : Params.Client, // Code based projects will have Client in their executable name already
@@ -1384,9 +1391,12 @@ public class IOSPlatform : ApplePlatform
 							DirectoryReference.Combine((SC.IsCodeBasedProject ? SC.ProjectRoot : DirectoryReference.Combine(SC.LocalRoot, "Engine")), "Binaries", PlatformName, "Payload", (SC.IsCodeBasedProject ? SC.ShortProjectName : "UnrealGame") + ".app"),
 							SC.StageExecutables[0]);
 
-					// copy the plist to the stage dir
-					SC.StageFile(StagedFileType.SystemNonUFS, TargetPListFile, new StagedFileReference("Info.plist"));
-				}
+					if (!AppleExports.UseModernXcode(SC.RawProjectPath))
+					{
+						// copy the plist to the stage dir
+						SC.StageFile(StagedFileType.SystemNonUFS, TargetPListFile, new StagedFileReference("Info.plist"));
+					}
+				}	
 
 				// copy the udebugsymbols if they exist
 				{
diff --git a/Engine/Source/Programs/UnrealBuildTool/ProjectFiles/Xcode/XcodeProject.cs b/Engine/Source/Programs/UnrealBuildTool/ProjectFiles/Xcode/XcodeProject.cs
--- a/Engine/Source/Programs/UnrealBuildTool/ProjectFiles/Xcode/XcodeProject.cs
+++ b/Engine/Source/Programs/UnrealBuildTool/ProjectFiles/Xcode/XcodeProject.cs
@@ -1892,7 +1892,7 @@ namespace UnrealBuildTool.XcodeProjectXcconfig
 					}
 					else if (!bAutomaticSigning && SigningIdentity != null)
 					{
-						Xcconfig.AppendLine($"CODE_SIGN_IDENTITY = \"{SigningIdentity}\"");
+						Xcconfig.AppendLine($"CODE_SIGN_IDENTITY = {SigningIdentity}");
 					}
 				}
 				else
@@ -1907,7 +1907,7 @@ namespace UnrealBuildTool.XcodeProjectXcconfig
 						}
 						if (SigningIdentity != null)
 						{
-							Xcconfig.AppendLine($"CODE_SIGN_IDENTITY = \"{SigningIdentity}\"");
+							Xcconfig.AppendLine($"CODE_SIGN_IDENTITY = {SigningIdentity}");
 						}
 					}
 				}
