[Backout] - CL27070697
Allow loading bundle config from editor INI file.

diff --git a/Engine/Source/Programs/AutomationTool/AutomationUtils/BundleUtils.cs b/Engine/Source/Programs/AutomationTool/AutomationUtils/BundleUtils.cs
--- a/Engine/Source/Programs/AutomationTool/AutomationUtils/BundleUtils.cs
+++ b/Engine/Source/Programs/AutomationTool/AutomationUtils/BundleUtils.cs
@@ -33,129 +33,100 @@ namespace AutomationUtils.Automation
 			LoadBundleConfig<BundleSettings>(ProjectDir, Platform, out Bundles, delegate (BundleSettings Settings, ConfigHierarchy BundleConfig, string Section) { });
 		}
 
-		public static void LoadBundleConfig_IncludeEditorConfig(DirectoryReference ProjectDir, UnrealTargetPlatform Platform, out IReadOnlyDictionary<string, BundleSettings> Bundles)
-		{
-			LoadBundleConfig_IncludeEditorConfig<BundleSettings>(ProjectDir, Platform, out Bundles, delegate (BundleSettings Settings, ConfigHierarchy BundleConfig, string Section) { });
-		}
-
-		public static void LoadBundleConfig<TPlatformBundleSettings>(DirectoryReference ProjectDir, UnrealTargetPlatform Platform,
-			out IReadOnlyDictionary<string, TPlatformBundleSettings> Bundles,
-			Action<TPlatformBundleSettings, ConfigHierarchy, string> GetPlatformSettings)
-			where TPlatformBundleSettings : BundleSettings, new()
-		{
-			LoadBundleConfigInternal<TPlatformBundleSettings>(false, ProjectDir, Platform, out Bundles, GetPlatformSettings);
-		}
-
-		public static void LoadBundleConfig_IncludeEditorConfig<TPlatformBundleSettings>(DirectoryReference ProjectDir, UnrealTargetPlatform Platform,
-			out IReadOnlyDictionary<string, TPlatformBundleSettings> Bundles,
-			Action<TPlatformBundleSettings, ConfigHierarchy, string> GetPlatformSettings)
-			where TPlatformBundleSettings : BundleSettings, new()
-		{
-			LoadBundleConfigInternal<TPlatformBundleSettings>(true, ProjectDir, Platform, out Bundles, GetPlatformSettings);
-		}
-
-		private static void LoadBundleConfigInternal<TPlatformBundleSettings>(bool bIncludeEditorConfig, DirectoryReference ProjectDir, UnrealTargetPlatform Platform, 
+		public static void LoadBundleConfig<TPlatformBundleSettings>(DirectoryReference ProjectDir, UnrealTargetPlatform Platform, 
 			out IReadOnlyDictionary<string, TPlatformBundleSettings> Bundles, 
 			Action<TPlatformBundleSettings, ConfigHierarchy, string> GetPlatformSettings) 
 			where TPlatformBundleSettings : BundleSettings, new()
 		{
 			var Results = new List<TPlatformBundleSettings>();
 
-			List<ConfigHierarchy> BundleConfigs = new List<ConfigHierarchy>();
-			BundleConfigs.Add(ConfigCache.ReadHierarchy(ConfigHierarchyType.InstallBundle, ProjectDir, Platform));
-			if (bIncludeEditorConfig)
-			{
-				BundleConfigs.Add(ConfigCache.ReadHierarchy(ConfigHierarchyType.Editor, ProjectDir, Platform));
-			}
+			ConfigHierarchy BundleConfig = ConfigCache.ReadHierarchy(ConfigHierarchyType.InstallBundle, ProjectDir, Platform);
 
 			const string BundleDefinitionPrefix = "InstallBundleDefinition ";
 
-			foreach (ConfigHierarchy BundleConfig in BundleConfigs)
+			foreach (string SectionName in BundleConfig.SectionNames)
 			{
-				foreach (string SectionName in BundleConfig.SectionNames)
-				{
-					if (!SectionName.StartsWith(BundleDefinitionPrefix))
-						continue;
+				if (!SectionName.StartsWith(BundleDefinitionPrefix))
+					continue;
 
-					TPlatformBundleSettings Bundle = new TPlatformBundleSettings();
-					Bundle.Name = SectionName.Substring(BundleDefinitionPrefix.Length);
+				TPlatformBundleSettings Bundle = new TPlatformBundleSettings();
+				Bundle.Name = SectionName.Substring(BundleDefinitionPrefix.Length);
+				{
+					int Order;
+					if(BundleConfig.GetInt32(SectionName, "Order", out Order))
 					{
-						int Order;
-						if (BundleConfig.GetInt32(SectionName, "Order", out Order))
-						{
-							Bundle.Order = Order;
-						}
-						else
-						{
-							Bundle.Order = int.MaxValue;
-						}
+						Bundle.Order = Order;
 					}
+					else
 					{
-						List<string> Tags;
-						if (BundleConfig.GetArray(SectionName, "Tags", out Tags))
-						{
-							Bundle.Tags = Tags;
-						}
-						else
-						{
-							Bundle.Tags = new List<string>();
-						}
+						Bundle.Order = int.MaxValue;
 					}
+				}
+				{
+					List<string> Tags;
+					if (BundleConfig.GetArray(SectionName, "Tags", out Tags))
 					{
-						List<string> Dependencies;
-						if (BundleConfig.GetArray(SectionName, "Dependencies", out Dependencies))
-						{
-							Bundle.Dependencies = Dependencies;
-						}
-						else
-						{
-							Bundle.Dependencies = new List<string>();
-						}
+						Bundle.Tags = Tags;
 					}
+					else 
 					{
-						List<string> FileRegex;
-						if (BundleConfig.GetArray(SectionName, "FileRegex", out FileRegex))
-						{
-							Bundle.FileRegex = FileRegex;
-						}
-						else
-						{
-							Bundle.FileRegex = new List<string>();
-						}
+						Bundle.Tags = new List<string>(); 
+					}
+				}
+				{
+					List<string> Dependencies;
+					if (BundleConfig.GetArray(SectionName, "Dependencies", out Dependencies))
+					{
+						Bundle.Dependencies = Dependencies;
 					}
+					else
 					{
-						List<string> Files;
-						if (BundleConfig.GetArray(SectionName, "Files", out Files))
-						{
-							Bundle.Files = Files;
-						}
-						else
-						{
-							Bundle.Files = new List<string>();
-						}
+						Bundle.Dependencies = new List<string>();
+					}
+				}
+				{
+					List<string> FileRegex;
+					if (BundleConfig.GetArray(SectionName, "FileRegex", out FileRegex))
+					{
+						Bundle.FileRegex = FileRegex;
 					}
+					else
 					{
-						bool bContainsShaderLibrary;
-						if (BundleConfig.GetBool(SectionName, "ContainsShaderLibrary", out bContainsShaderLibrary))
-						{
-							Bundle.bContainsShaderLibrary = bContainsShaderLibrary;
-						}
-						else
-						{
-							Bundle.bContainsShaderLibrary = false;
-						}
+						Bundle.FileRegex = new List<string>();
+					}
+				}
+				{
+					List<string> Files;
+					if (BundleConfig.GetArray(SectionName, "Files", out Files))
+					{
+						Bundle.Files = Files;
+					}
+					else
+					{
+						Bundle.Files = new List<string>();
+					}
+				}
+				{
+					bool bContainsShaderLibrary;
+					if (BundleConfig.GetBool(SectionName, "ContainsShaderLibrary", out bContainsShaderLibrary))
+					{
+						Bundle.bContainsShaderLibrary = bContainsShaderLibrary;
 					}
+					else 
+					{
+						Bundle.bContainsShaderLibrary = false;
+					}
+				}
 
+				{
+					bool bUseChunkDBs;
+					if (!BundleConfig.GetBool(SectionName, "UseChunkDBs", out bUseChunkDBs))
 					{
-						bool bUseChunkDBs;
-						if (!BundleConfig.GetBool(SectionName, "UseChunkDBs", out bUseChunkDBs))
-						{
-							Bundle.UseChunkDBs = false;
-						}
-						else
-						{
-							Bundle.UseChunkDBs = bUseChunkDBs;
-						}
+						Bundle.UseChunkDBs = false;
+					}
+					else
+					{
+						Bundle.UseChunkDBs = bUseChunkDBs;
 					}
 					{
 						bool bUseDetailedInstallSizes;
@@ -169,10 +140,10 @@ namespace AutomationUtils.Automation
 						}
 					}
 
-					GetPlatformSettings(Bundle, BundleConfig, BundleDefinitionPrefix + Bundle.Name);
-
-					Results.Add(Bundle);
 				}
+				GetPlatformSettings(Bundle, BundleConfig, BundleDefinitionPrefix + Bundle.Name);
+
+				Results.Add(Bundle);
 			}
 
 			// Use OrderBy and not Sort because OrderBy is stable
