Update install scripts and deploy for future AndroidFileServer changes
Remove unused AFS receiver

#android

diff --git a/Engine/Plugins/Runtime/AndroidFileServer/Source/AndroidFileServer/Android/java/com/epicgames/unreal/RemoteFileManagerReceiver.java b/Engine/Plugins/Runtime/AndroidFileServer/Source/AndroidFileServer/Android/java/com/epicgames/unreal/RemoteFileManagerReceiver.java
deleted file mode 100644
diff --git a/Engine/Plugins/Runtime/AndroidFileServer/Source/AndroidFileServer/AndroidFileServer_UPL.xml b/Engine/Plugins/Runtime/AndroidFileServer/Source/AndroidFileServer/AndroidFileServer_UPL.xml
--- a/Engine/Plugins/Runtime/AndroidFileServer/Source/AndroidFileServer/AndroidFileServer_UPL.xml
+++ b/Engine/Plugins/Runtime/AndroidFileServer/Source/AndroidFileServer/AndroidFileServer_UPL.xml
@@ -68,11 +68,6 @@
 											<action android:name="com.epicgames.unreal.RemoteFileManager.intent.COMMAND2" />
 										</intent-filter>
 									</activity>
-									<receiver android:name="com.epicgames.unreal.RemoteFileManagerReceiver" android:exported="true">
-										<intent-filter>
-											<action android:name="com.epicgames.unreal.RemoteFileManager.intent.COMMAND" />
-										</intent-filter>
-									</receiver>
 								</addElements>
 							</true>
 						</if>
@@ -131,16 +126,13 @@
 				<if condition="bExternalStartAFS">
 					<true>
 						<copyFile src="$S(PluginDir)/Android/java/com/epicgames/unreal/RemoteFileManagerActivity.java" dst="$S(BuildDir)/src/com/epicgames/unreal/RemoteFileManagerActivity.java" />
-						<copyFile src="$S(PluginDir)/Android/java/com/epicgames/unreal/RemoteFileManagerReceiver.java" dst="$S(BuildDir)/src/com/epicgames/unreal/RemoteFileManagerReceiver.java" />
 					</true>
 					<false>
 						<isArch arch="arm64-v8a">
 							<deleteFiles filespec="arm64/src/com/epicgames/unreal/RemoteFileManagerActivity.java" />
-							<deleteFiles filespec="arm64/src/com/epicgames/unreal/RemoteFileManagerReceiver.java" />
 						</isArch>
 						<isArch arch="x86_64">
 							<deleteFiles filespec="x64/src/com/epicgames/unreal/RemoteFileManagerActivity.java" />
-							<deleteFiles filespec="x64/src/com/epicgames/unreal/RemoteFileManagerReceiver.java" />
 						</isArch>
 					</false>
 				</if>
@@ -149,12 +141,10 @@
 				<isArch arch="arm64-v8a">
 					<deleteFiles filespec="arm64/src/com/epicgames/unreal/RemoteFileManager.java" />
 					<deleteFiles filespec="arm64/src/com/epicgames/unreal/RemoteFileManagerActivity.java" />
-					<deleteFiles filespec="arm64/src/com/epicgames/unreal/RemoteFileManagerReceiver.java" />
 				</isArch>
 				<isArch arch="x86_64">
 					<deleteFiles filespec="x64/src/com/epicgames/unreal/RemoteFileManager.java" />
 					<deleteFiles filespec="x64/src/com/epicgames/unreal/RemoteFileManagerActivity.java" />
-					<deleteFiles filespec="x64/src/com/epicgames/unreal/RemoteFileManagerReceiver.java" />
 				</isArch>
 			</false>
 		</if>
@@ -287,14 +277,11 @@ tasks.whenTaskAdded { task ->
 	afs_infile = new File(rootProject.projectDir, 'app/src/main/AndroidManifest.xml')
 	afs_outfile = new File(rootProject.projectDir, 'AFSProject/app/src/main/AndroidManifest.xml')
 	def afs_hasActivity = false
-	def afs_hasReceiver = false
 	afs_outfile.withWriter { writer ->
 		afs_infile.withReader { reader ->
 			while ((afs_line = reader.readLine()) != null) {
 				if (afs_line.contains('com.epicgames.unreal.RemoteFileManagerActivity')) {
 					afs_hasActivity = true
-				} else if (afs_line.contains('com.epicgames.unreal.RemoteFileManagerReceiver')) {
-					afs_hasReceiver = true
 				} else if (afs_line.contains('</application>')) {
 					if (afs_hasActivity == false) {
 						writer.writeLine('    <activity android:name="com.epicgames.unreal.RemoteFileManagerActivity" android:exported="true" android:screenOrientation="sensor">')
@@ -303,13 +290,6 @@ tasks.whenTaskAdded { task ->
 						writer.writeLine('    	</intent-filter>')
 						writer.writeLine('    </activity>')
 					}
-					if (afs_hasReceiver == false) {
-						writer.writeLine('    <receiver android:name="com.epicgames.unreal.RemoteFileManagerReceiver" android:exported="true">')
-						writer.writeLine('      <intent-filter>')
-						writer.writeLine('        <action android:name="com.epicgames.unreal.RemoteFileManager.intent.COMMAND" />')
-						writer.writeLine('      </intent-filter>')
-						writer.writeLine('    </receiver>')
-					}
 				}
 				writer.writeLine("${afs_line}")
 			}
diff --git a/Engine/Source/Programs/AutomationTool/Android/AndroidPlatform.Automation.cs b/Engine/Source/Programs/AutomationTool/Android/AndroidPlatform.Automation.cs
--- a/Engine/Source/Programs/AutomationTool/Android/AndroidPlatform.Automation.cs
+++ b/Engine/Source/Programs/AutomationTool/Android/AndroidPlatform.Automation.cs
@@ -1551,6 +1551,7 @@ public class AndroidPlatform : Platform
 		public string OverflowName { get; }
 		public bool bNoOverflowInstall { get; }
 	}
+
 	private List<string> GenerateInstallBatchFile(bool bPackageDataInsideApk, string PackageName, string ApkName, ProjectParams Params, string ObbName, string DeviceObbName, bool bNoObbInstall,
 		string PatchName, string DevicePatchName, bool bNoPatchInstall, List<OverflowBatchInstallInfo> OverflowInfo,
 		bool bIsPC, bool bIsDistribution, bool bRequireRuntimeStoragePermission, bool bDisablePerfHarden, bool bUseAFS, bool bUseAFSProject, string AFSToken, UnrealTargetPlatform Target)
@@ -1558,11 +1559,17 @@ public class AndroidPlatform : Platform
 		List<string> BatchLines = new List<string>();
         string ReadPermissionGrantCommand = "shell pm grant " + PackageName + " android.permission.READ_EXTERNAL_STORAGE";
         string WritePermissionGrantCommand = "shell pm grant " + PackageName + " android.permission.WRITE_EXTERNAL_STORAGE";
+		string ForegroundPermissionGrantCommand = "shell pm grant " + PackageName + " android.permission.FOREGROUND_SERVICE";
+		string ForegroundDataSyncPermissionGrantCommand = "shell pm grant " + PackageName + " android.permission.FOREGROUND_SERVICE_DATA_SYNC";
+		string NotificationPermissionGrantCommand = "shell pm grant " + PackageName + " android.permission.POST_NOTIFICATIONS";
 		string DisablePerfHardenCommand = "shell setprop security.perf_harden 0";
 
+		string NullCmd = bIsPC ? " >nul 2>&1" : " >/dev/null 2>&1";
+
 		// We don't grant runtime permission for distribution build on purpose since we will push the obb file to the folder that doesn't require runtime storage permission.
 		// This way developer can catch permission issue if they try to save/load game file in folder that requires runtime storage permission.
 		bool bNeedGrantStoragePermission = bRequireRuntimeStoragePermission && !bIsDistribution;
+		bool bNeedGrantForegroundPermission = bUseAFS || bUseAFSProject;
 
 		// We can't always push directly to Android/obb so uploads to Download then moves it
 		bool bDontMoveOBB = bUseAFS ? true : bPackageDataInsideApk;
@@ -1643,10 +1650,9 @@ public class AndroidPlatform : Platform
 						"if [ $? -eq 0 ]; then",
 						"\techo",
 						"\t$ADB $DEVICE shell pm list packages " + PackageName,
-						bNeedGrantStoragePermission ? "\techo Grant READ_EXTERNAL_STORAGE and WRITE_EXTERNAL_STORAGE to the apk for reading OBB or game file in external storage." : "",
-						bNeedGrantStoragePermission ? "\t$ADB $DEVICE " + ReadPermissionGrantCommand : "",
-						bNeedGrantStoragePermission ? "\t$ADB $DEVICE " + WritePermissionGrantCommand : "",
-						bDisablePerfHarden ? "\t$ADB $DEVICE " + DisablePerfHardenCommand : "",
+						bNeedGrantForegroundPermission ? "\t$ADB $DEVICE " + ForegroundPermissionGrantCommand + NullCmd : "",
+						bNeedGrantForegroundPermission ? "\t$ADB $DEVICE " + ForegroundDataSyncPermissionGrantCommand + NullCmd : "",
+						bNeedGrantForegroundPermission ? "\t$ADB $DEVICE " + NotificationPermissionGrantCommand + NullCmd : "",
 						"\techo",
 						"\techo Removing old data. Failures here are usually fine - indicating the files were not on the device.",
 						"\t$ADB $DEVICE shell 'rm -r $EXTERNAL_STORAGE/UnrealGame/" + Params.ShortProjectName + "'",
@@ -1668,6 +1674,10 @@ public class AndroidPlatform : Platform
 						bDontMoveOBB ? "" : "\t\t$ADB $DEVICE shell mv " + TargetAndroidTemp + TargetAndroidLocation + PackageName + " $STORAGE/Android/" + TargetAndroidLocation,
 						bDontMoveOBB ? "" : "\t\t$ADB $DEVICE shell rm -r " + TargetAndroidTemp + TargetAndroidLocation,
 						APKReinstallCommand,
+						bNeedGrantStoragePermission ? "\techo Grant READ_EXTERNAL_STORAGE and WRITE_EXTERNAL_STORAGE to the apk for reading OBB or game file in external storage." : "",
+						bNeedGrantStoragePermission ? "\t$ADB $DEVICE " + ReadPermissionGrantCommand + NullCmd : "",
+						bNeedGrantStoragePermission ? "\t$ADB $DEVICE " + WritePermissionGrantCommand + NullCmd : "",
+						bDisablePerfHarden ? "\t$ADB $DEVICE " + DisablePerfHardenCommand : "",
 						"\t\techo",
 						"\t\techo Installation successful",
 						"\t\texit 0",
@@ -1751,6 +1761,9 @@ public class AndroidPlatform : Platform
 						APKInstallCommand,
 						"@if \"%ERRORLEVEL%\" NEQ \"0\" goto Error",
 						"%ADB% %DEVICE% shell pm list packages " + PackageName,
+						bNeedGrantForegroundPermission ? "%ADB% %DEVICE% " + ForegroundPermissionGrantCommand + NullCmd : "",
+						bNeedGrantForegroundPermission ? "%ADB% %DEVICE% " + ForegroundDataSyncPermissionGrantCommand + NullCmd : "",
+						bNeedGrantForegroundPermission ? "%ADB% %DEVICE% " + NotificationPermissionGrantCommand + NullCmd : "",
 						"%ADB% %DEVICE% shell rm -r %STORAGE%/UnrealGame/" + Params.ShortProjectName,
 						"%ADB% %DEVICE% shell rm -r %STORAGE%/UnrealGame/UECommandLine.txt", // we need to delete the commandline in UnrealGame or it will mess up loading
 						"%ADB% %DEVICE% shell rm -r %STORAGE%/" + TargetAndroidLocation + PackageName,
@@ -1774,8 +1787,8 @@ public class AndroidPlatform : Platform
 						bUseAFSProject ? "if \"%ERRORLEVEL%\" NEQ \"0\" goto Error" : "",
 						"@echo.",
 						bNeedGrantStoragePermission ? "@echo Grant READ_EXTERNAL_STORAGE and WRITE_EXTERNAL_STORAGE to the apk for reading OBB file or game file in external storage." : "",
-						bNeedGrantStoragePermission ? "%ADB% %DEVICE% " + ReadPermissionGrantCommand : "",
-						bNeedGrantStoragePermission ? "%ADB% %DEVICE% " + WritePermissionGrantCommand : "",
+						bNeedGrantStoragePermission ? "%ADB% %DEVICE% " + ReadPermissionGrantCommand + NullCmd : "",
+						bNeedGrantStoragePermission ? "%ADB% %DEVICE% " + WritePermissionGrantCommand + NullCmd : "",
 						bDisablePerfHarden ? "%ADB% %DEVICE% " + DisablePerfHardenCommand : "",
                         "@echo.",
                         "@echo Installation successful",
@@ -2622,6 +2635,14 @@ public class AndroidPlatform : Platform
 				adb.Shell(DeviceName, ReadPermissionCommandLine);
 				adb.Shell(DeviceName, WritePermissionCommandLine);
 
+				// grant permission for the foreground service to start (otherwise a security violation on Android 28+)
+				string ForegroundPermissionGrantCommand = "pm grant " + PackageName + " android.permission.FOREGROUND_SERVICE";
+				string ForegroundDataSyncPermissionGrantCommand = "pm grant " + PackageName + " android.permission.FOREGROUND_SERVICE_DATA_SYNC";
+				string NotificationsGrantCommand = "pm grant " + PackageName + " android.permission.POST_NOTIFICATIONS";
+				adb.Shell(DeviceName, ForegroundPermissionGrantCommand);
+				adb.Shell(DeviceName, ForegroundDataSyncPermissionGrantCommand);
+				adb.Shell(DeviceName, NotificationsGrantCommand);
+
 				// time for receivers to be registered by pm after install
 				Thread.Sleep(350);
 			}
