IAS - changed cosmetic bundling to use wildcards

#rnx

diff --git a/Engine/Source/Programs/AutomationTool/Scripts/CopyBuildToStagingDirectory.Automation.cs b/Engine/Source/Programs/AutomationTool/Scripts/CopyBuildToStagingDirectory.Automation.cs
--- a/Engine/Source/Programs/AutomationTool/Scripts/CopyBuildToStagingDirectory.Automation.cs
+++ b/Engine/Source/Programs/AutomationTool/Scripts/CopyBuildToStagingDirectory.Automation.cs
@@ -2216,7 +2216,7 @@ namespace AutomationScripts
 			public string EncryptionKeyGuid;
 
 			// List of allowed chunk names when generating on demand chunk(s) 
-			public HashSet<string> OnDemandAllowedChunks;
+			public List<Regex> OnDemandAllowedChunkFilters;
 
 			public static bool IsMatch(PakFileRules PakRules, KeyValuePair<string, string> StagingFile)
 			{
@@ -2392,13 +2392,20 @@ namespace AutomationScripts
 					continue;
 				}
 
-				IReadOnlyList<string> AllowedChunks;
-				if (PakRulesConfig.TryGetValues(SectionName, "OnDemandAllowedChunks", out AllowedChunks))
+				IReadOnlyList<string> AllowedChunkFilters;
+				if (PakRulesConfig.TryGetValues(SectionName, "OnDemandAllowedChunkFilters", out AllowedChunkFilters))
 				{
-					PakRules.OnDemandAllowedChunks = new HashSet<string>();
-					foreach (var ChunkName in AllowedChunks)
+					PakRules.OnDemandAllowedChunkFilters = new List<Regex>();
+					foreach (string Filter in AllowedChunkFilters)
 					{
-						PakRules.OnDemandAllowedChunks.Add(ChunkName);
+						try
+						{
+							PakRules.OnDemandAllowedChunkFilters.Add(new Regex(Filter));
+						}
+						catch(System.Text.RegularExpressions.RegexParseException Ex)
+						{
+							Logger.LogError("Failed to parse allowed chunk filter, reason '{Arg1}'", Ex.Message);
+						}
 					}
 				}
 
@@ -4227,9 +4234,9 @@ namespace AutomationScripts
 						if (MatchingRule.HasValue && MatchingRule.Value.bOnDemand)
 						{
 							bOnDemand = true;
-							if (MatchingRule.Value.OnDemandAllowedChunks != null)
+							if (MatchingRule.Value.OnDemandAllowedChunkFilters != null)
 							{
-								bOnDemand = MatchingRule.Value.OnDemandAllowedChunks.Contains(TargetChunk.ChunkName);
+								bOnDemand = MatchingRule.Value.OnDemandAllowedChunkFilters.Any(Filter => Filter.IsMatch(TargetChunk.ChunkName));
 							}
 						}
 
