- add webtests on horde
- add server output and server error logs
- link webtests to dashboard

diff --git a/Engine/Source/Programs/AutomationTool/LowLevelTests/Tests/LowLevelTests.WebTestsLowLevelTestsExtension.cs b/Engine/Source/Programs/AutomationTool/LowLevelTests/Tests/LowLevelTests.WebTestsLowLevelTestsExtension.cs
--- a/Engine/Source/Programs/AutomationTool/LowLevelTests/Tests/LowLevelTests.WebTestsLowLevelTestsExtension.cs
+++ b/Engine/Source/Programs/AutomationTool/LowLevelTests/Tests/LowLevelTests.WebTestsLowLevelTestsExtension.cs
@@ -4,18 +4,21 @@ using JetBrains.Annotations;
 using System;
 using System.Diagnostics;
 using System.IO;
+using System.Text;
 using EpicGames.Core;
 using AutomationTool;
 using UnrealBuildTool;
 using Gauntlet;
 using UnrealBuildBase;
 using EpicGames.Horde.Storage;
+using System.Net.Sockets;
 
 namespace LowLevelTests
 {
 	public class WebTestsLowLevelTestsExtension : ILowLevelTestsExtension
 	{
 		private Process ServerProcess;
+		private StringBuilder logBuilder = new StringBuilder();
 
 		public bool IsSupported(UnrealTargetPlatform InPlatform, string InTestApp)
 		{
@@ -27,6 +30,22 @@ namespace LowLevelTests
 			return string.Format("--web_server_ip{0}", UnrealHelpers.GetHostIpAddress());
 		}
 
+		private bool IsServerPortOpen(string ipAddress, int port)
+		{
+			using (TcpClient client = new TcpClient())
+			{
+				try
+				{
+					client.Connect(ipAddress, port);
+					return true;
+				}
+				catch
+				{
+					return false;
+				}
+			}
+		}
+
 		public void PreRunTests()
 		{
 			ProcessStartInfo StartInfo = new ProcessStartInfo();
@@ -41,7 +60,43 @@ namespace LowLevelTests
 
 			ServerProcess = new Process();
 			ServerProcess.StartInfo = StartInfo;
+
+			ServerProcess.OutputDataReceived += (sender, e) =>
+			{
+				if (!string.IsNullOrEmpty(e.Data))
+				{
+					logBuilder.AppendLine(e.Data);
+				}
+			};
+
+			ServerProcess.ErrorDataReceived += (sender, e) =>
+			{
+				if (!string.IsNullOrEmpty(e.Data))
+				{
+					logBuilder.AppendLine(e.Data);
+				}
+			};
+
 			ServerProcess.Start();
+			ServerProcess.BeginOutputReadLine();
+			ServerProcess.BeginErrorReadLine();
+
+			Stopwatch sw = new Stopwatch();
+			sw.Start();
+
+			while (!IsServerPortOpen(UnrealHelpers.GetHostIpAddress(), 8000))
+			{
+				if (sw.ElapsedMilliseconds > 60000)
+				{
+					sw.Stop();
+					throw new TimeoutException("Server port did not open within the specified time.");
+				}
+				System.Threading.Thread.Sleep(1000);
+			}
+
+			sw.Stop();
+
+			Console.WriteLine("server port is now open.");
 		}
 
 		public void PostRunTests()
@@ -50,6 +105,13 @@ namespace LowLevelTests
 			{
 				ServerProcess.Kill();
 				ServerProcess = null;
+
+				string serverOutput = logBuilder.ToString();
+				Console.WriteLine("Server Output:");
+				Console.WriteLine(serverOutput);
+
+				string serverOutputLog = Path.Combine(Path.GetFullPath(Globals.LogDir).Replace("GauntletTemp", "Engine\\Programs\\AutomationTool\\Saved"), "ServerOutput.log");
+				File.WriteAllText(serverOutputLog, serverOutput);
 			}
 		}
 	}
diff --git a/Engine/Source/Programs/WebTests/WebTests.Build.cs b/Engine/Source/Programs/WebTests/WebTests.Build.cs
--- a/Engine/Source/Programs/WebTests/WebTests.Build.cs
+++ b/Engine/Source/Programs/WebTests/WebTests.Build.cs
@@ -15,5 +15,13 @@ public class WebTests : TestModuleRules
 				"Chaos",
 				"ChaosUserDataPT"
 			});
+		
+		UpdateBuildGraphPropertiesFile(new Metadata() { 
+			TestName = "Web", 
+			TestShortName = "Web", 
+			ReportType = "xml", 
+		});
 	}
 }
+
+
