CookedEditor Staging: AdditionalPluginDirectories: Handle paths from the staged RemappedPlugins in DeploymentContext.UnmakeRelativeStagedReference.
#rnx

diff --git a/Engine/Source/Programs/AutomationTool/AutomationUtils/DeploymentContext.cs b/Engine/Source/Programs/AutomationTool/AutomationUtils/DeploymentContext.cs
--- a/Engine/Source/Programs/AutomationTool/AutomationUtils/DeploymentContext.cs
+++ b/Engine/Source/Programs/AutomationTool/AutomationUtils/DeploymentContext.cs
@@ -1172,31 +1172,75 @@ public class DeploymentContext //: ProjectParams
 
 	public static StagedFileReference MakeRelativeStagedReference(DeploymentContext SC, FileSystemReference Ref, out DirectoryReference RootDir)
 	{
+		foreach (DirectoryReference AdditionalPluginDir in SC.AdditionalPluginDirectories)
+		{
+			if (Ref.IsUnderDirectory(AdditionalPluginDir))
+			{
+				// This is a plugin that lives outside of the Engine/Plugins or Game/Plugins directory so needs to be remapped for staging/packaging
+				// We need to remap C:\SomePath\PluginName\RelativePath to RemappedPlugins\PluginName\RelativePath
+				string RemainingPath = Ref.MakeRelativeTo(AdditionalPluginDir).Replace('\\', '/');
+				int PluginEndIndex = RemainingPath.IndexOf("/");
+				if (PluginEndIndex >= 0 && PluginEndIndex < RemainingPath.Length - 1)
+				{
+					string PluginName = RemainingPath.Substring(0, PluginEndIndex);
+					RemainingPath = RemainingPath.Substring(PluginEndIndex + 1);
+					RootDir = DirectoryReference.Combine(AdditionalPluginDir, PluginName);
+					StagedFileReference StagedFile = new StagedFileReference(String.Format("RemappedPlugins/{0}/{1}", PluginName, RemainingPath));
+					return ApplyDirectoryRemap(SC, StagedFile);
+				}
+			}
+		}
+
 		if (Ref.IsUnderDirectory(SC.ProjectRoot))
 		{
 			RootDir = SC.ProjectRoot;
 			return ApplyDirectoryRemap(SC, new StagedFileReference(SC.ShortProjectName + "/" + Ref.MakeRelativeTo(SC.ProjectRoot).Replace('\\', '/')));
 		}
-		else if (Ref.IsUnderDirectory(SC.EngineRoot))
+
+		if (Ref.IsUnderDirectory(SC.EngineRoot))
 		{
 			RootDir = SC.EngineRoot;
 			return ApplyDirectoryRemap(SC, new StagedFileReference("Engine/" + Ref.MakeRelativeTo(SC.EngineRoot).Replace('\\', '/')));
 		}
+
 		throw new Exception();
 	}
 	public static FileReference UnmakeRelativeStagedReference(DeploymentContext SC, StagedFileReference Ref)
 	{
-		// paths will be in the form "Engine/Foo" or "{ProjectName}/Foo" (or something that we don't handle, so assert)
-		// So, replace the Engine/ with {EngineDir} and {ProjectName}/ with {ProjectDir}, and then append Foo
+		// paths will be in the form "Engine/Foo" or "{ProjectName}/Foo" or "RemappedPlugins/{PluginName}/Foo
+		// Anything else we don't handle.
+		// So, replace the Engine/ with {EngineDir} and {ProjectName}/ with {ProjectDir}, or change PluginDir to RemappedPlugins/{PluginName}
+		// with the plugin path from AdditionalPluginDirectories, and then append Foo
+
+		string RemappedPluginsStr = "RemappedPlugins/";
+		if (Ref.Name.StartsWith(RemappedPluginsStr, StringComparison.CurrentCultureIgnoreCase))
+		{
+			int PluginEndIndex = Ref.Name.IndexOf("/", RemappedPluginsStr.Length);
+			if (PluginEndIndex >= 0 && PluginEndIndex < Ref.Name.Length - 1)
+			{
+				string PluginName = Ref.Name.Substring(RemappedPluginsStr.Length, PluginEndIndex - RemappedPluginsStr.Length);
+				foreach (DirectoryReference AdditionalPluginDir in SC.AdditionalPluginDirectories)
+				{
+					DirectoryReference PossiblePluginDir = DirectoryReference.Combine(AdditionalPluginDir, PluginName);
+					if (System.IO.Directory.Exists(PossiblePluginDir.FullName))
+					{
+						return FileReference.Combine(PossiblePluginDir, Ref.Name.Substring(PluginEndIndex+1));
+					}
+				}
+			}
+		}
+
 		if (Ref.Name.StartsWith("Engine/"))
 		{
 			// skip over "Engine/" which is 7 chars long
 			return FileReference.Combine(SC.EngineRoot, Ref.Name.Substring(7));
 		}
-		else if (Ref.Name.StartsWith(SC.ShortProjectName + "/"))
+
+		if (Ref.Name.StartsWith(SC.ShortProjectName + "/"))
 		{
 			return FileReference.Combine(SC.ProjectRoot, Ref.Name.Substring(SC.ShortProjectName.Length + 1));
 		}
+
 		throw new Exception();
 	}
 }
