Adding a new ZenImport task on the buildgraph to be used with Multi-Agent cooking

diff --git a/Engine/Source/Programs/AutomationTool/BuildGraph/Tasks/ZenExportSnapshotTask.cs b/Engine/Source/Programs/AutomationTool/BuildGraph/Tasks/ZenExportSnapshotTask.cs
--- a/Engine/Source/Programs/AutomationTool/BuildGraph/Tasks/ZenExportSnapshotTask.cs
+++ b/Engine/Source/Programs/AutomationTool/BuildGraph/Tasks/ZenExportSnapshotTask.cs
@@ -79,6 +79,18 @@ namespace AutomationTool.Tasks
 		/// </summary>
 		[TaskParameter(Optional = true)]
 		public string DestinationCloudBucket;
+
+		/// <summary>
+		/// The directory to store the exported oplog
+		/// </summary>
+		[TaskParameter(Optional = true)]
+		public string DestinationDir;
+
+		/// <summary>
+		/// The name of the oplog to write on disk
+		/// </summary>
+		[TaskParameter(Optional = true)]
+		public string DestinationOplogName;
 	}
 
 	/// <summary>
@@ -270,6 +282,13 @@ namespace AutomationTool.Tasks
 						}
 					}
 
+					break;
+				case SnapshotStorageType.File:
+					string ProjectId = ProjectUtils.GetProjectPathId(ProjectFile);
+					OplogExportCommandline.AppendFormat(" --file {0} --name {1} {2} {3}", Parameters.DestinationDir, Parameters.DestinationOplogName, ProjectId, Parameters.Platform);
+
+					Logger.LogInformation("Running '{Arg0} {Arg1}'", CommandUtils.MakePathSafeToUseWithCommandLine(ZenExe.FullName), OplogExportCommandline.ToString());
+					CommandUtils.RunAndLog(CommandUtils.CmdEnv, ZenExe.FullName, OplogExportCommandline.ToString(), Options: CommandUtils.ERunOptions.Default);
 					break;
 				default:
 					throw new AutomationException("Unknown/invalid/unimplemented destination storage type - {0}", Parameters.DestinationStorageType);
diff --git a/Engine/Source/Programs/AutomationTool/BuildGraph/Tasks/ZenImportOplogTask.cs b/Engine/Source/Programs/AutomationTool/BuildGraph/Tasks/ZenImportOplogTask.cs
new file mode 100644
--- /dev/null
+++ b/Engine/Source/Programs/AutomationTool/BuildGraph/Tasks/ZenImportOplogTask.cs
@@ -0,0 +1,304 @@
+// Copyright Epic Games, Inc. All Rights Reserved.
+
+using AutomationTool;
+using EpicGames.BuildGraph;
+using System;
+using System.Collections.Generic;
+using System.Text;
+using System.Threading.Tasks;
+using System.Xml;
+using EpicGames.Core;
+using UnrealBuildTool;
+using Microsoft.Extensions.Logging;
+using EpicGames.Serialization;
+
+namespace AutomationTool.Tasks
+{
+
+	/// <summary>
+	/// Parameters for a task that exports an snapshot from ZenServer
+	/// </summary>
+	public class ZenImportOplogTaskParameters
+	{
+		/// <summary>
+		/// The type of destination to import from to (cloud, file...)
+		/// </summary>
+		[TaskParameter]
+		public string ImportType;
+
+		/// <summary>
+		/// comma separated full path to the oplog dir to import into the local zen server
+		/// Files="Path1,Path2"
+		/// </summary>
+		[TaskParameter(Optional = true)]
+		public string Files;
+
+		/// <summary>
+		/// The project from which to import for
+		/// </summary>
+		[TaskParameter(Optional = true)]
+		public FileReference Project;
+
+		/// <summary>
+		/// The name of the newly created Zen Project we will be importing into
+		/// </summary>
+		[TaskParameter(Optional = true)]
+		public string ProjectName;
+
+
+		/// <summary>
+		/// The target platform to import the snapshot for
+		/// </summary>
+		[TaskParameter(Optional = true)]
+		public string Platform;
+
+		/// <summary>
+		/// Root dir for the UE project. Used to derive the Enging folder and the Project folder
+		/// </summary>
+		[TaskParameter(Optional = true)]
+		public string RootDir;
+
+		/// <summary>
+		/// The name of the imported oplog
+		/// </summary>
+		[TaskParameter(Optional = true)]
+		public string OplogName;
+
+		/// <summary>
+		/// The host URL for the zen server we are importing from
+		/// </summary>
+		[TaskParameter(Optional = true)]
+		public string HostName = "localhost";
+
+		/// <summary>
+		/// The host port for the zen server we are importing from
+		/// </summary>
+		[TaskParameter(Optional = true)]
+		public string HostPort = "1337";
+		
+
+		/// <summary>
+		/// The cloud URL to import from
+		/// </summary>
+		[TaskParameter(Optional = true)]
+		public string CloudURL;
+
+		/// <summary>
+		/// what namespace to use when importing from cloud
+		/// </summary>
+		[TaskParameter(Optional = true)]
+		public string Namespace;
+
+		/// <summary>
+		/// what bucket to use when importing from cloud
+		/// </summary>
+		[TaskParameter(Optional = true)]
+		public string Bucket;
+
+		/// <summary>
+		/// What key to use when importing from cloud
+		/// </summary>
+		[TaskParameter(Optional = true)]
+		public string Key;
+	}
+
+	/// <summary>
+	/// Imports an oplog from Zen to a specified destination.
+	/// </summary>
+	[TaskElement("ZenImportOplog", typeof(ZenImportOplogTaskParameters))]
+	public class ZenImportOplogTask : BgTaskImpl
+	{
+		/// <summary>
+		/// Parameters for the task
+		/// </summary>
+		ZenImportOplogTaskParameters Parameters;
+
+		FileReference ProjectFile;
+
+		/// <summary>
+		/// Constructor.
+		/// </summary>
+		/// <param name="InParameters">Parameters for this task</param>
+		public ZenImportOplogTask(ZenImportOplogTaskParameters InParameters)
+		{
+			Parameters = InParameters;
+		}
+
+		/// <summary>
+		/// Execute the task.
+		/// </summary>
+		/// <param name="Job">Information about the current job</param>
+		/// <param name="BuildProducts">Set of build products produced by this node.</param>
+		/// <param name="TagNameToFileSet">Mapping from tag names to the set of files they include</param>
+		public override Task ExecuteAsync(JobContext Job, HashSet<FileReference> BuildProducts, Dictionary<string, HashSet<FileReference>> TagNameToFileSet)
+		{
+			SnapshotStorageType ImportMethod = SnapshotStorageType.Invalid;
+			if (!string.IsNullOrEmpty(Parameters.ImportType))
+			{
+				ImportMethod = (SnapshotStorageType)Enum.Parse(typeof(SnapshotStorageType), Parameters.ImportType);
+			}
+
+			ProjectFile = Parameters.Project;
+			if (!FileReference.Exists(ProjectFile))
+			{
+				throw new AutomationException("Missing project file - {0}", ProjectFile.FullName);
+			}
+
+			// Get the executable path
+			FileReference ZenExe;
+			if(HostPlatform.Current.HostEditorPlatform == UnrealTargetPlatform.Win64)
+			{
+				ZenExe = ResolveFile("Engine/Binaries/Win64/zen.exe");
+			}
+			else
+			{
+				ZenExe = ResolveFile(String.Format("Engine/Binaries/{0}/zen", HostPlatform.Current.HostEditorPlatform.ToString()));
+			}
+
+			{
+				if (String.IsNullOrEmpty(Parameters.RootDir))
+				{
+					throw new AutomationException("RootDir was not specified");
+				}
+				if (String.IsNullOrEmpty(Parameters.ProjectName))
+				{
+					throw new AutomationException("ProjectName was not specified");
+				}
+				
+				// Create a new project to import everything into.
+				string RootDir = Parameters.RootDir;
+				string EngineDir = System.IO.Path.Combine(Parameters.RootDir, "Engine");
+				string ProjectDir = System.IO.Path.Combine(Parameters.RootDir, ProjectFile.GetFileNameWithoutAnyExtensions());
+				string HostURLArg = string.Format("--hosturl http://{0}:{1}", Parameters.HostName, Parameters.HostPort);
+				StringBuilder OplogProjectCreateCommandline = new StringBuilder();
+				OplogProjectCreateCommandline.AppendFormat("project-create -p {0} --rootdir {1} --enginedir {2} --projectdir {3} --projectfile {4} {5}",
+					Parameters.ProjectName,
+					RootDir,
+					EngineDir,
+					ProjectDir,
+					ProjectFile.FullName,
+					HostURLArg);
+
+				Logger.LogInformation("Running '{Arg0} {Arg1}'", CommandUtils.MakePathSafeToUseWithCommandLine(ZenExe.FullName), OplogProjectCreateCommandline.ToString());
+				CommandUtils.RunAndLog(CommandUtils.CmdEnv, ZenExe.FullName, OplogProjectCreateCommandline.ToString(), Options: CommandUtils.ERunOptions.Default);
+			}
+
+			switch (ImportMethod)
+			{
+				case SnapshotStorageType.File:
+					ImportFromFile(ZenExe);
+					break;
+				case SnapshotStorageType.Cloud:
+					ImportFromCloud(ZenExe);
+					break;
+				default:
+					throw new AutomationException("Unknown/invalid/unimplemented import type - {0}", Parameters.ImportType);
+			}
+
+			WriteProjectStoreFile();
+			return Task.CompletedTask;
+		}
+
+		private void ImportFromFile(FileReference ZenExe)
+		{
+			if (String.IsNullOrEmpty(Parameters.OplogName))
+			{
+				throw new AutomationException("OplogName was not specified");
+			}
+
+			foreach (string FileToImport in Parameters.Files.Split(','))
+			{
+				if (DirectoryReference.Exists(new DirectoryReference(FileToImport)))
+				{
+					StringBuilder OplogImportCommandline = new StringBuilder();
+					OplogImportCommandline.AppendFormat("oplog-import --file {0} --oplog {1} -p {2}", FileToImport, Parameters.OplogName, Parameters.ProjectName);
+
+					Logger.LogInformation("Running '{Arg0} {Arg1}'", CommandUtils.MakePathSafeToUseWithCommandLine(ZenExe.FullName), OplogImportCommandline.ToString());
+					CommandUtils.RunAndLog(CommandUtils.CmdEnv, ZenExe.FullName, OplogImportCommandline.ToString(), Options: CommandUtils.ERunOptions.Default);
+				}
+			}
+		}
+
+		private void WriteProjectStoreFile()
+		{
+			DirectoryReference PlatformCookedDirectory = DirectoryReference.Combine(ProjectFile.Directory, "Saved", "Cooked", Parameters.Platform);
+			if (!DirectoryReference.Exists(PlatformCookedDirectory))
+			{
+				DirectoryReference.CreateDirectory(PlatformCookedDirectory);
+			}
+			FileReference ProjectStoreFile = FileReference.Combine(PlatformCookedDirectory, ".projectstore");
+
+			CbWriter Writer = new CbWriter();
+			Writer.BeginObject();
+			Writer.BeginObject("zenserver");
+
+			Writer.WriteBool("islocalhost", true);
+			Writer.WriteString("hostname", "localhost");
+			Writer.WriteInteger("hostport", 1337);
+			Writer.WriteString("projectid", Parameters.ProjectName);
+			Writer.WriteString("oplogid", Parameters.OplogName);
+
+			Writer.EndObject();
+			Writer.EndObject();
+
+			FileReference.WriteAllBytes(ProjectStoreFile, Writer.ToByteArray());
+		}
+
+		private void ImportFromCloud(FileReference ZenExe)
+		{
+			if (string.IsNullOrEmpty(Parameters.CloudURL))
+			{
+				throw new AutomationException("Missing destination cloud host");
+			}
+			if (string.IsNullOrEmpty(Parameters.Namespace))
+			{
+				throw new AutomationException("Missing destination cloud namespace");
+			}
+			if (string.IsNullOrEmpty(Parameters.Key))
+			{
+				throw new AutomationException("Missing destination cloud storage key");
+			}
+
+			string BucketName = Parameters.Bucket;
+			string ProjectNameAsBucketName = ProjectFile.GetFileNameWithoutAnyExtensions().ToLowerInvariant();
+			if (string.IsNullOrEmpty(BucketName))
+			{
+				BucketName = ProjectNameAsBucketName;
+			}
+
+			string HostURLArg = string.Format("--hosturl http://{0}:{1}", Parameters.HostName, Parameters.HostPort);
+			StringBuilder OplogImportCommandline = new StringBuilder();
+			OplogImportCommandline.AppendFormat("oplog-import {0} --cloud {1} --namespace {2} --bucket {3}", HostURLArg, Parameters.CloudURL, Parameters.Namespace, BucketName);
+			OplogImportCommandline.AppendFormat(" {0}", Parameters.Key);
+
+			Logger.LogInformation("Running '{Arg0} {Arg1}'", CommandUtils.MakePathSafeToUseWithCommandLine(ZenExe.FullName), OplogImportCommandline.ToString());
+			CommandUtils.RunAndLog(CommandUtils.CmdEnv, ZenExe.FullName, OplogImportCommandline.ToString(), Options: CommandUtils.ERunOptions.Default);
+		}
+
+		/// <summary>
+		/// Output this task out to an XML writer.
+		/// </summary>
+		public override void Write(XmlWriter Writer)
+		{
+			Write(Writer, Parameters);
+		}
+
+		/// <summary>
+		/// Find all the tags which are used as inputs to this task
+		/// </summary>
+		/// <returns>The tag names which are read by this task</returns>
+		public override IEnumerable<string> FindConsumedTagNames()
+		{
+			yield break;
+		}
+
+		/// <summary>
+		/// Find all the tags which are modified by this task
+		/// </summary>
+		/// <returns>The tag names which are modified by this task</returns>
+		public override IEnumerable<string> FindProducedTagNames()
+		{
+			yield break;
+		}
+	}
+}
