Add post build task to upload any stats of interest to the Gauntlet MySQL DB. Start with Switch executable size.

diff --git a/Engine/Source/Programs/AutomationTool/AutomationUtils/Platform.cs b/Engine/Source/Programs/AutomationTool/AutomationUtils/Platform.cs
--- a/Engine/Source/Programs/AutomationTool/AutomationUtils/Platform.cs
+++ b/Engine/Source/Programs/AutomationTool/AutomationUtils/Platform.cs
@@ -801,6 +801,12 @@ namespace AutomationTool
 			return false;
 		}
 
+		public virtual int GetExecutableSize(DirectoryReference BinariesDirectory, string ClientName, HashSet<FileReference> BuildProducts)
+		{
+			Logger.LogWarning("GetExecutableSize() has not been implemented for {Arg0}", PlatformType.ToString());
+			return -1;
+		}
+
 		/// <summary>
 		/// When overridden, returns the directory structure of the platform's symbol server.
 		/// Each element is a semi-colon separated string of possible directory names.
