AutomationTool Staging: CookedEditor: Add feedback about which file is failing in the failed assertion in DeploymentContext.UnmakeRelativeStagedReference.
#rnx

diff --git a/Engine/Source/Programs/AutomationTool/AutomationUtils/DeploymentContext.cs b/Engine/Source/Programs/AutomationTool/AutomationUtils/DeploymentContext.cs
--- a/Engine/Source/Programs/AutomationTool/AutomationUtils/DeploymentContext.cs
+++ b/Engine/Source/Programs/AutomationTool/AutomationUtils/DeploymentContext.cs
@@ -1241,6 +1241,6 @@ public class DeploymentContext //: ProjectParams
 			return FileReference.Combine(SC.ProjectRoot, Ref.Name.Substring(SC.ShortProjectName.Length + 1));
 		}
 
-		throw new Exception();
+		throw new Exception($"Don't know how to convert staged file {Ref.Name} to its original editor path, because it is not in a recognized root directory.");
 	}
 }
