- Fixed plist generation for TVOS BP projects

diff --git a/Engine/Source/Programs/AutomationTool/TVOS/TVOSPlatform.Automation.cs b/Engine/Source/Programs/AutomationTool/TVOS/TVOSPlatform.Automation.cs
--- a/Engine/Source/Programs/AutomationTool/TVOS/TVOSPlatform.Automation.cs
+++ b/Engine/Source/Programs/AutomationTool/TVOS/TVOSPlatform.Automation.cs
@@ -111,10 +111,17 @@ public class TVOSPlatform : IOSPlatform
 
                     var TargetConfiguration = SC.StageTargetConfigurations[0];
 
-                    DeployGeneratePList(
+					DirectoryReference ProjectRoot = SC.ProjectRoot;
+					// keep old logic for BP projects with legacy
+					if (!AppleExports.UseModernXcode(SC.RawProjectPath) && !SC.IsCodeBasedProject)
+					{
+						ProjectRoot = DirectoryReference.Combine(SC.LocalRoot, "Engine");
+					}
+
+					DeployGeneratePList(
 						SC.RawProjectPath,
 						TargetConfiguration,
-						(SC.IsCodeBasedProject ? SC.ProjectRoot : SC.EngineRoot),
+						ProjectRoot,
 						!SC.IsCodeBasedProject,
 						(SC.IsCodeBasedProject ? TargetName : "UnrealGame"),
 						Params.Client,
@@ -129,26 +136,31 @@ public class TVOSPlatform : IOSPlatform
                 }
 
 
-            // copy the udebugsymbols if they exist
-            {
-                ConfigHierarchy PlatformGameConfig;
-                bool bIncludeSymbols = false;
-                if (Params.EngineConfigs.TryGetValue(SC.StageTargetPlatform.PlatformType, out PlatformGameConfig))
-                {
-                    PlatformGameConfig.GetBool("/Script/IOSRuntimeSettings.IOSRuntimeSettings", "bGenerateCrashReportSymbols", out bIncludeSymbols);
-                }
-                if (bIncludeSymbols)
+                // copy the udebugsymbols if they exist
                 {
-                    FileReference SymbolFileName = FileReference.Combine((SC.IsCodeBasedProject ? SC.ProjectRoot : SC.EngineRoot), "Binaries", "TVOS", SC.StageExecutables[0] + ".udebugsymbols");
-                    if (FileReference.Exists(SymbolFileName))
+                    ConfigHierarchy PlatformGameConfig;
+                    bool bIncludeSymbols = false;
+                    if (Params.EngineConfigs.TryGetValue(SC.StageTargetPlatform.PlatformType, out PlatformGameConfig))
                     {
-                        SC.StageFile(StagedFileType.NonUFS, SymbolFileName, new StagedFileReference((Params.ShortProjectName + ".udebugsymbols").ToLowerInvariant()));
+                        PlatformGameConfig.GetBool("/Script/IOSRuntimeSettings.IOSRuntimeSettings", "bGenerateCrashReportSymbols", out bIncludeSymbols);
+                    }
+                    if (bIncludeSymbols)
+                    {
+                        FileReference SymbolFileName = FileReference.Combine((SC.IsCodeBasedProject ? SC.ProjectRoot : SC.EngineRoot), "Binaries", "TVOS", SC.StageExecutables[0] + ".udebugsymbols");
+                        if (FileReference.Exists(SymbolFileName))
+                        {
+                            SC.StageFile(StagedFileType.NonUFS, SymbolFileName, new StagedFileReference((Params.ShortProjectName + ".udebugsymbols").ToLowerInvariant()));
+                        }
                     }
                 }
-            }
-                SC.StageFile(StagedFileType.SystemNonUFS, TargetPListFile, new StagedFileReference("Info.plist"));
-            }
-        }
+
+				if (!AppleExports.UseModernXcode(SC.RawProjectPath))
+				{
+					// copy the plist to the stage dir
+					SC.StageFile(StagedFileType.SystemNonUFS, TargetPListFile, new StagedFileReference("Info.plist"));
+				}
+			}
+		}
 
         // copy the movies from the project
         {
diff --git a/Engine/Source/Programs/UnrealBuildTool/Platform/IOS/UEDeployIOS.cs b/Engine/Source/Programs/UnrealBuildTool/Platform/IOS/UEDeployIOS.cs
--- a/Engine/Source/Programs/UnrealBuildTool/Platform/IOS/UEDeployIOS.cs
+++ b/Engine/Source/Programs/UnrealBuildTool/Platform/IOS/UEDeployIOS.cs
@@ -426,6 +426,13 @@ namespace UnrealBuildTool
 				Text.AppendLine("\t\t</dict>");
 			}
 
+			// add a TVOS setting since they share this file
+			Text.AppendLine("\t<key>TVTopShelfImage</key>");
+			Text.AppendLine("\t<dict>");
+			Text.AppendLine("\t\t<key>TVTopShelfPrimaryImageWide</key>");
+			Text.AppendLine("\t\t<string>Top Shelf Image Wide</string>");
+			Text.AppendLine("\t</dict>");
+
 			if (!String.IsNullOrEmpty(ExtraData))
 			{
 				ExtraData = ExtraData.Replace("\\n", "\n");
diff --git a/Engine/Source/Programs/UnrealBuildTool/Platform/TVOS/UEDeployTVOS.cs b/Engine/Source/Programs/UnrealBuildTool/Platform/TVOS/UEDeployTVOS.cs
--- a/Engine/Source/Programs/UnrealBuildTool/Platform/TVOS/UEDeployTVOS.cs
+++ b/Engine/Source/Programs/UnrealBuildTool/Platform/TVOS/UEDeployTVOS.cs
@@ -200,6 +200,10 @@ namespace UnrealBuildTool
 
 		public override bool GeneratePList(FileReference? ProjectFile, UnrealTargetConfiguration Config, string ProjectDirectory, bool bIsUnrealGame, string GameName, bool bIsClient, string ProjectName, string InEngineDir, string AppDirectory, List<string> UPLScripts, string? BundleID, bool bBuildAsFramework)
 		{
+			if (AppleExports.UseModernXcode(ProjectFile))
+			{
+				return base.GeneratePList(ProjectFile, Config, ProjectDirectory, bIsUnrealGame, GameName, bIsClient, ProjectName, InEngineDir, AppDirectory, UPLScripts, BundleID, bBuildAsFramework);
+			}
 			return GenerateTVOSPList(ProjectDirectory, bIsUnrealGame, GameName, bIsClient, ProjectName, InEngineDir, AppDirectory, null, BundleID, Logger);
 		}
 	}
