AutomationTool Staging: CookedEditor: Fix DeploymentContext::UnmakeRelativeStagedReference directory comparisons to be case insensitive.
#rnx

diff --git a/Engine/Source/Programs/AutomationTool/AutomationUtils/DeploymentContext.cs b/Engine/Source/Programs/AutomationTool/AutomationUtils/DeploymentContext.cs
--- a/Engine/Source/Programs/AutomationTool/AutomationUtils/DeploymentContext.cs
+++ b/Engine/Source/Programs/AutomationTool/AutomationUtils/DeploymentContext.cs
@@ -1230,13 +1230,13 @@ public class DeploymentContext //: ProjectParams
 			}
 		}
 
-		if (Ref.Name.StartsWith("Engine/"))
+		if (Ref.Name.StartsWith("Engine/", StringComparison.CurrentCultureIgnoreCase))
 		{
 			// skip over "Engine/" which is 7 chars long
 			return FileReference.Combine(SC.EngineRoot, Ref.Name.Substring(7));
 		}
 
-		if (Ref.Name.StartsWith(SC.ShortProjectName + "/"))
+		if (Ref.Name.StartsWith(SC.ShortProjectName + "/", StringComparison.CurrentCultureIgnoreCase))
 		{
 			return FileReference.Combine(SC.ProjectRoot, Ref.Name.Substring(SC.ShortProjectName.Length + 1));
 		}
