Use ZenLaunch executable to ensure that the ZenExportSnapshotTask has ZenServer running if needed.

diff --git a/Engine/Source/Programs/AutomationTool/BuildGraph/Tasks/ZenExportSnapshotTask.cs b/Engine/Source/Programs/AutomationTool/BuildGraph/Tasks/ZenExportSnapshotTask.cs
--- a/Engine/Source/Programs/AutomationTool/BuildGraph/Tasks/ZenExportSnapshotTask.cs
+++ b/Engine/Source/Programs/AutomationTool/BuildGraph/Tasks/ZenExportSnapshotTask.cs
@@ -6,6 +6,7 @@ using System;
 using System.Collections.Generic;
 using System.IO;
 using System.Linq;
+using System.Runtime.InteropServices;
 using System.Text;
 using System.Threading.Tasks;
 using System.Xml;
@@ -123,6 +124,18 @@ namespace AutomationTool.Tasks
 			Parameters = InParameters;
 		}
 
+		private void ZenLaunch(FileReference ProjectFile)
+		{
+			// Get the ZenLaunch executable path
+			FileReference ZenLaunchExe = ResolveFile(String.Format("Engine/Binaries/{0}/ZenLaunch{}", HostPlatform.Current.HostEditorPlatform.ToString(), RuntimeInformation.IsOSPlatform(OSPlatform.Windows) ? ".exe" : ""));
+
+			StringBuilder ZenLaunchCommandline = new StringBuilder();
+			ZenLaunchCommandline.AppendFormat("{0} -SponsorProcessID {1}", CommandUtils.MakePathSafeToUseWithCommandLine(ProjectFile.FullName), Environment.ProcessId);
+
+			Logger.LogInformation("Running '{Arg0} {Arg1}'", CommandUtils.MakePathSafeToUseWithCommandLine(ZenLaunchExe.FullName), ZenLaunchCommandline.ToString());
+			CommandUtils.RunAndLog(CommandUtils.CmdEnv, ZenLaunchExe.FullName, ZenLaunchCommandline.ToString(), Options: CommandUtils.ERunOptions.Default);
+		}
+
 		/// <summary>
 		/// Execute the task.
 		/// </summary>
@@ -143,6 +156,8 @@ namespace AutomationTool.Tasks
 				throw new AutomationException("Missing project file - {0}", ProjectFile.FullName);
 			}
 
+			ZenLaunch(ProjectFile);
+
 			List<ExportSourceData> ExportSources = new List<ExportSourceData>();
 			foreach (string Platform in Parameters.Platform.Split('+'))
 			{
@@ -175,16 +190,8 @@ namespace AutomationTool.Tasks
 				}
 			}
 
-			// Get the executable path
-			FileReference ZenExe;
-			if(HostPlatform.Current.HostEditorPlatform == UnrealTargetPlatform.Win64)
-			{
-				ZenExe = ResolveFile("Engine/Binaries/Win64/zen.exe");
-			}
-			else
-			{
-				ZenExe = ResolveFile(String.Format("Engine/Binaries/{0}/zen", HostPlatform.Current.HostEditorPlatform.ToString()));
-			}
+			// Get the Zen executable path
+			FileReference ZenExe = ResolveFile(String.Format("Engine/Binaries/{0}/zen{}", HostPlatform.Current.HostEditorPlatform.ToString(), RuntimeInformation.IsOSPlatform(OSPlatform.Windows) ? ".exe" : ""));
 
 			// Format the command lines
 			StringBuilder OplogSnapshotCommandline = new StringBuilder();
