Gauntlet: Changing up how crash-dump saving works, just to make sure platforms have a little more consistency before copying them to the Artifact folder for a test run
Base Gauntlet Changes

diff --git a/Engine/Source/Programs/AutomationTool/Gauntlet/Framework/Base/Gauntlet.TargetDevice.cs b/Engine/Source/Programs/AutomationTool/Gauntlet/Framework/Base/Gauntlet.TargetDevice.cs
--- a/Engine/Source/Programs/AutomationTool/Gauntlet/Framework/Base/Gauntlet.TargetDevice.cs
+++ b/Engine/Source/Programs/AutomationTool/Gauntlet/Framework/Base/Gauntlet.TargetDevice.cs
@@ -90,6 +90,24 @@ namespace Gauntlet
 		IAppInstance Run(IAppInstall App);
 
 		string GetPackagedExecutableLocation() { return null; }
+
+		/// <summary>
+		/// Path to the crash dumps on a device
+		/// </summary>
+		string CrashDumpPath
+		{
+			get
+			{
+				return Path.Combine(Globals.TempDir, "CrashDumps", Platform.ToString() + "_" + Name);
+			}
+		}
+
+		/// <summary>
+		/// Ensures the crash dump copy has occurred already - and does the copy if it hasn't happened yet
+		/// Returns true if there were any crash dumps for the run, and false otherwise
+		/// </summary>
+		bool CopyCrashDumps() { return false; }
+		
 	};
 
 	/// <summary>
diff --git a/Engine/Source/Programs/AutomationTool/Gauntlet/Unreal/Base/Gauntlet.UnrealSession.cs b/Engine/Source/Programs/AutomationTool/Gauntlet/Unreal/Base/Gauntlet.UnrealSession.cs
--- a/Engine/Source/Programs/AutomationTool/Gauntlet/Unreal/Base/Gauntlet.UnrealSession.cs
+++ b/Engine/Source/Programs/AutomationTool/Gauntlet/Unreal/Base/Gauntlet.UnrealSession.cs
@@ -1225,6 +1225,23 @@ namespace Gauntlet
 						}
 					}
 
+					// Copy any Crash Dumps as well
+					if (InRunningRole.AppInstance.Device.CopyCrashDumps())
+					{
+						string DeviceCrashDumpPath = InRunningRole.AppInstance.Device.CrashDumpPath;
+						if (!string.IsNullOrEmpty(DeviceCrashDumpPath))
+						{
+							string ArtifactCrashDumpDir = Path.Combine(DestSavedDirInfo.FullName, "CrashDumps");
+							Log.Info("Copying any CrashDumps from {0} to {1}", DeviceCrashDumpPath, ArtifactCrashDumpDir);
+							Utils.SystemHelpers.CopyDirectory(DeviceCrashDumpPath, Utils.SystemHelpers.GetFullyQualifiedPath(ArtifactCrashDumpDir), Utils.SystemHelpers.CopyOptions.Default, TruncateLongPathFilter);
+							Log.Info("Cleaning files in {0}", DeviceCrashDumpPath);
+							foreach (FileInfo CrashDumpFile in Utils.SystemHelpers.GetFiles(new DirectoryInfo(DeviceCrashDumpPath), "*", SearchOption.AllDirectories))
+							{
+								Utils.SystemHelpers.Delete(CrashDumpFile);
+							}
+						}
+					}
+
 					// Copy any lose files from the source
 					foreach (string SourceFile in Directory.EnumerateFiles(SourceSavedDir))
 					{
