Begin creation of new UAT command to run many clients, run a server, and connect them. This initial version only spawns the clients, but if there is an existing server running for them to connect to, they will connect to it.

diff --git a/Engine/Source/Programs/AutomationTool/Scripts/MultiClientLauncher.Automation.cs b/Engine/Source/Programs/AutomationTool/Scripts/MultiClientLauncher.Automation.cs
new file mode 100644
--- /dev/null
+++ b/Engine/Source/Programs/AutomationTool/Scripts/MultiClientLauncher.Automation.cs
@@ -0,0 +1,289 @@
+// Copyright Epic Games, Inc. All Rights Reserved.
+
+using System;
+using System.Collections.Generic;
+using System.Diagnostics;
+using System.IO;
+using System.Threading;
+using AutomationTool;
+using EpicGames.Core;
+using UnrealBuildBase;
+
+namespace MultiClientLauncher.Automation
+{
+	[Help("Run many game clients, a server, and connect them")]
+	[ParamHelp("ClientExe", "Absolute path to the client to run", ParamType = typeof(FileReference))]
+	[ParamHelp("ClientLogFile", "Absolute path to the client log file", ParamType = typeof(FileReference))]
+	[ParamHelp("ClientArgsFile", "Absolute path to a file containing the client arguments", ParamType = typeof(FileReference))]
+	[ParamHelp("FirstClientNumber", "The number of the first LoadBot client", ParamType = typeof(int))]
+	[ParamHelp("ClientCount", "How many bot clients to run, must consecutively follow FirstClientNumber", ParamType = typeof(int))]
+	[ParamHelp("SleepTimeBetweenLaunches", "How long to sleep between running clients in milliseconds (could prevent race conditions)", ParamType = typeof(int))]
+	[ParamHelp("MaxRunAttemptsPerClient", "Maximum number of attempts to run a client which crashes or fails to connect to the server, defaults to 3", ParamType = typeof(int))]
+	[ParamHelp("ClientSessionCompleted", "Log message indicating that a client has completed a game session and may be terminated", ParamType = typeof(string))]
+	[ParamHelp("ClientFailed", "Log message indicating that a client failed to connect to the server", ParamType = typeof(string))]
+	[ParamHelp("ClientConnected", "Log message indicating that a client connected to the server", ParamType = typeof(string))]
+	public class MultiClientLauncher : BuildCommand
+	{
+		private bool CancelClientProcesses = false;
+		
+		private string ClientExe;
+		private string ClientLog;
+		private string ClientArgs;
+
+		private int FirstClientNumber;
+		private int ClientCount;
+
+		private int SleepTimeBetweenLaunches = -1;
+		private int MaxRunAttemptsPerClient = 3;
+
+		private const int SleepTimeBetweenChecksForClientRelaunches = 1000;
+
+		protected ClientLogIndicators ClientIndicators;
+		
+		private void ParseCommandLine()
+		{
+			ClientExe = ParseRequiredFileReferenceParam("ClientExe").ToString();
+			ClientLog = ParseRequiredFileReferenceParam("ClientLogFile").ToString();
+			ClientArgs = FileReference.ReadAllText(ParseRequiredFileReferenceParam("ClientArgsFile"));
+			
+			FirstClientNumber = int.Parse(ParseRequiredStringParam("FirstClientNumber"));
+			ClientCount = int.Parse(ParseRequiredStringParam("ClientCount"));
+
+			SleepTimeBetweenLaunches = ParseParamInt("SleepTimeBetweenLaunches", -1);
+			MaxRunAttemptsPerClient = ParseParamInt("MaxRunAttemptsPerClient", 3);
+			
+			InitializeClientLogIndicators();
+		}
+
+		// Derived commands may hardcode these log indicators
+		protected virtual void InitializeClientLogIndicators()
+		{
+			ClientIndicators.FinishedGameAndDisconnected = ParseRequiredStringParam("ClientSessionCompleted");
+			ClientIndicators.FailedToConnectToServer = ParseRequiredStringParam("ClientFailed");
+			ClientIndicators.ConnectedToServer = ParseRequiredStringParam("ClientConnected");
+		}
+		
+		public override ExitCode Execute()
+		{
+			ParseCommandLine();
+			
+			List<ClientProcess> ClientProcesses = new List<ClientProcess>();
+			
+			// Allow ctrl C to terminate all client processes
+			Console.CancelKeyPress += delegate
+			{
+				CancelClientProcesses = true;
+				KillProcesses(ClientProcesses);
+			};
+			
+			// Run all client processes
+			Console.WriteLine("Spawning clients");
+			for (int i = 0; i < ClientCount; i++)
+			{
+				if (CancelClientProcesses)
+				{
+					break;
+				}
+				
+				Console.WriteLine("Spawning client {0}...", i);
+
+				string ClientNumber = (FirstClientNumber + i).ToString();
+				string CurrentClientArgs = ClientArgs.Replace("#REPLACE_CLIENT_ID#", ClientNumber);
+				
+				string CurrentClientLog = ClientLog;
+				if (i > 0)
+				{
+					CurrentClientLog += "_" + (i + 1);
+				}
+				CurrentClientLog += ".log";
+				
+				ClientProcesses.Add(SpawnClientProcess(ClientExe, CurrentClientLog, CurrentClientArgs));
+
+				if (SleepTimeBetweenLaunches != -1)
+				{
+					Console.WriteLine("Sleeping for {0}ms...", SleepTimeBetweenLaunches);
+					Thread.Sleep(SleepTimeBetweenLaunches);
+				}
+			}
+
+			while (!CancelClientProcesses && ClientProcesses.Count > 0)
+			{
+				// Iterate through clients in reverse order for safe removal
+				for (int i = ClientProcesses.Count - 1; i >= 0; i--)
+				{
+					if (ClientProcesses[i].Stopped())
+					{
+						bool OutOfTries = !ClientProcesses[i].Start();
+						
+						if (OutOfTries)
+						{
+							ClientProcesses.RemoveAt(i);
+						}
+					}
+				}
+				
+				// Todo: Consider replacing sleeps with waiting for LogSkimmer threads to finish reading
+				// Todo:   - That may still happen very fast and result in more busywaiting
+				Thread.Sleep(SleepTimeBetweenChecksForClientRelaunches);
+			}
+
+			return ExitCode.Success;
+		}
+		
+		private ClientProcess SpawnClientProcess(string ExeFilename, string ClientLogFilename, string ExeArguments)
+		{
+			ClientProcess ClientProc = new ClientProcess(ExeFilename, ExeArguments, MaxRunAttemptsPerClient, ClientLogFilename, ClientIndicators);
+			ClientProc.Start();
+			return ClientProc;
+		}
+
+		private static void KillProcesses(List<ClientProcess> Processes)
+		{
+			Console.WriteLine("Killing all client processes");
+			foreach (ClientProcess CurrentProcess in Processes)
+			{
+				CurrentProcess.Kill();
+			}
+		}
+
+		protected struct ClientLogIndicators
+		{
+			public string FinishedGameAndDisconnected;
+			public string FailedToConnectToServer;
+			public string ConnectedToServer;
+		}
+
+		private class ClientProcess
+		{
+			private readonly Process Proc;
+			private int RemainingRunAttempts;
+
+			private Thread LogSkimmer;
+			private const int SleepTimeBetweenLogSkims = 1000;
+			private readonly ClientLogIndicators LogIndicators;
+			private readonly string LogFilepath;
+
+			private readonly Stopwatch ConnectionTimer;
+			private const int MinutesUntilTimeout = 2; // Todo: maybe make this a command line argument
+
+			public ClientProcess(string Exe, string Args, int MaxRunAttempts, string ClientLog, ClientLogIndicators ClientIndicators)
+			{
+				Proc = new Process();
+				Proc.StartInfo.FileName = Exe;
+				Proc.StartInfo.Arguments = Args;
+				RemainingRunAttempts = MaxRunAttempts;
+
+				LogFilepath = ClientLog;
+				LogIndicators = ClientIndicators;
+
+				ConnectionTimer = new Stopwatch();
+			}
+			
+			public bool Stopped()
+			{
+				return Proc.HasExited && !LogSkimmer.IsAlive;
+			}
+
+			// Returns false if the client is out of attempts to start the process
+			public bool Start()
+			{
+				if (RemainingRunAttempts <= 0)
+				{
+					return false;
+				}
+				RemainingRunAttempts--;
+				
+				bool Success = Proc.Start();
+				if (!Success)
+				{
+					Console.WriteLine("Failed to start process for {0}", Proc.StartInfo.FileName);
+					return true;
+				}
+
+				ConnectionTimer.Restart();
+
+				LogSkimmer = new Thread(MonitorClient);
+				LogSkimmer.Start();
+
+				return true;
+			}
+
+			private string GetProcessName()
+			{
+				return "(" + Proc.ProcessName + ": " + Proc.Id + ")";
+			}
+
+			public void Kill()
+			{
+				Proc.Kill();
+			}
+
+			private void MonitorClient()
+			{
+				// Wait long enough for log files to be created
+				Thread.Sleep(2000);
+				
+				string AllServerOutput = "";
+				using FileStream ProcessLog = File.Open(LogFilepath, FileMode.Open, FileAccess.Read, FileShare.ReadWrite);
+				using StreamReader LogReader = new StreamReader(ProcessLog);
+
+				// Read until the process has exited or we found the success text in the log
+				while (!Proc.HasExited)
+				{
+					while (!LogReader.EndOfStream)
+					{
+						string Output = LogReader.ReadToEnd();
+						if (string.IsNullOrEmpty(Output))
+						{
+							continue;
+						}
+
+						AllServerOutput += Output;
+
+						if (ConnectionTimer.IsRunning)
+						{
+							// If a client connected to the server, there's no more need to check for timing out
+							if (AllServerOutput.Contains(LogIndicators.ConnectedToServer))
+							{
+								Console.WriteLine("Client {0} successfully connected to the server.", GetProcessName());
+								ConnectionTimer.Stop();
+								break;
+							}
+
+							// If client failed to connect, kill process and decrement attempts
+							if (AllServerOutput.Contains(LogIndicators.FailedToConnectToServer))
+							{
+								Console.WriteLine("Client {0} failed to connect to server. Relaunching client. Attempts left: {1}",
+									GetProcessName(), RemainingRunAttempts);
+								Kill();
+								return;
+							}
+						}
+						else
+						{
+							// If the client completed a session, kill the process
+							if (AllServerOutput.Contains(LogIndicators.FinishedGameAndDisconnected))
+							{
+								RemainingRunAttempts = 0;
+								Console.WriteLine("Client {0} completed session.", GetProcessName());
+								Kill();
+								return;
+							}
+						}
+					}
+
+					// If timed out while attempting to connect, kill process and decrement attempts
+					if (ConnectionTimer.IsRunning && ConnectionTimer.Elapsed.Minutes >= MinutesUntilTimeout)
+					{
+						Console.WriteLine("Client {0} timed out. Attempts left: {1}", GetProcessName(), RemainingRunAttempts);
+						Kill();
+						return;
+					}
+					
+					// Wait for more logging to occur
+					Thread.Sleep(SleepTimeBetweenLogSkims);
+				}
+			}
+		}
+	}
+}
\ No newline at end of file
