Add PreferenceOrder for IBuild. Give preference to packaged builds.

diff --git a/Engine/Source/Programs/AutomationTool/Gauntlet/Framework/Base/Gauntlet.BuildSource.cs b/Engine/Source/Programs/AutomationTool/Gauntlet/Framework/Base/Gauntlet.BuildSource.cs
--- a/Engine/Source/Programs/AutomationTool/Gauntlet/Framework/Base/Gauntlet.BuildSource.cs
+++ b/Engine/Source/Programs/AutomationTool/Gauntlet/Framework/Base/Gauntlet.BuildSource.cs
@@ -29,6 +29,11 @@ namespace Gauntlet
 	/// </summary>
     public interface IBuild
 	{
+		/// <summary>
+		/// Preference order for this build type vs. other build types. Lower is more prefered.
+		/// </summary>
+		int PreferenceOrder { get; }
+
 		/// <summary>
 		/// Flags that describe the properties of this build
 		/// </summary>
diff --git a/Engine/Source/Programs/AutomationTool/Gauntlet/Platform/Android/Gauntlet.AndroidBuildSource.cs b/Engine/Source/Programs/AutomationTool/Gauntlet/Platform/Android/Gauntlet.AndroidBuildSource.cs
--- a/Engine/Source/Programs/AutomationTool/Gauntlet/Platform/Android/Gauntlet.AndroidBuildSource.cs
+++ b/Engine/Source/Programs/AutomationTool/Gauntlet/Platform/Android/Gauntlet.AndroidBuildSource.cs
@@ -14,6 +14,8 @@ namespace Gauntlet
 {
 	public class AndroidBuild : IBuild
 	{
+		public int PreferenceOrder { get { return 0; } }
+
 		public UnrealTargetConfiguration Configuration { get; protected set; }
 
 		public string SourceApkPath;
diff --git a/Engine/Source/Programs/AutomationTool/Gauntlet/Platform/IOS/Gauntlet.IOSBuildSource.cs b/Engine/Source/Programs/AutomationTool/Gauntlet/Platform/IOS/Gauntlet.IOSBuildSource.cs
--- a/Engine/Source/Programs/AutomationTool/Gauntlet/Platform/IOS/Gauntlet.IOSBuildSource.cs
+++ b/Engine/Source/Programs/AutomationTool/Gauntlet/Platform/IOS/Gauntlet.IOSBuildSource.cs
@@ -14,6 +14,8 @@ namespace Gauntlet
 {
 	public class IOSBuild : IBuild
 	{
+		public int PreferenceOrder { get { return 0; } }
+
 		public UnrealTargetConfiguration Configuration { get; protected set; }
 
 		public string SourceIPAPath;
diff --git a/Engine/Source/Programs/AutomationTool/Gauntlet/Unreal/BuildSource/Gauntlet.StagedBuildSource.cs b/Engine/Source/Programs/AutomationTool/Gauntlet/Unreal/BuildSource/Gauntlet.StagedBuildSource.cs
--- a/Engine/Source/Programs/AutomationTool/Gauntlet/Unreal/BuildSource/Gauntlet.StagedBuildSource.cs
+++ b/Engine/Source/Programs/AutomationTool/Gauntlet/Unreal/BuildSource/Gauntlet.StagedBuildSource.cs
@@ -14,6 +14,8 @@ namespace Gauntlet
 
 	public class EditorBuild : IBuild 
 	{
+		public int PreferenceOrder { get { return 0; } }
+
 		public UnrealTargetPlatform Platform { get { return BuildHostPlatform.Current.Platform; } }
 
 		public UnrealTargetConfiguration Configuration { get; protected set; }
@@ -35,6 +37,8 @@ namespace Gauntlet
 
 	public class PackagedBuild : IBuild
 	{
+		public int PreferenceOrder { get { return 0; } }
+
 		public UnrealTargetPlatform Platform { get; protected set; }
 
 		public UnrealTargetConfiguration Configuration { get; protected set; }
@@ -70,6 +74,8 @@ namespace Gauntlet
 
 	public class StagedBuild : IBuild
 	{
+		public int PreferenceOrder { get { return 1; } }
+
 		public UnrealTargetPlatform Platform { get; protected set; }
 
 		public UnrealTargetConfiguration Configuration { get; protected set; }
diff --git a/Engine/Source/Programs/AutomationTool/Gauntlet/Unreal/BuildSource/Gauntlet.UnrealBuildSource.cs b/Engine/Source/Programs/AutomationTool/Gauntlet/Unreal/BuildSource/Gauntlet.UnrealBuildSource.cs
--- a/Engine/Source/Programs/AutomationTool/Gauntlet/Unreal/BuildSource/Gauntlet.UnrealBuildSource.cs
+++ b/Engine/Source/Programs/AutomationTool/Gauntlet/Unreal/BuildSource/Gauntlet.UnrealBuildSource.cs
@@ -360,6 +360,7 @@ namespace Gauntlet
 
 							foreach (IBuild Build in Builds)
 							{
+								Log.Info("Adding build {0} with flags {1}", BS.BuildName, Build.Flags);
 								AddBuild(Build);
 							}
 						}
@@ -511,7 +512,7 @@ namespace Gauntlet
             Config.FilesToCopy = new List<UnrealFileToCopy>();
 
 			// new system of retrieving and encapsulating the info needed to install/launch. Android & Mac
-			Config.Build = GetMatchingBuilds(Role.RoleType, Role.Platform, Role.Configuration, Role.RequiredBuildFlags, Role.RequiredFlavor).FirstOrDefault();
+			Config.Build = GetMatchingBuilds(Role.RoleType, Role.Platform, Role.Configuration, Role.RequiredBuildFlags, Role.RequiredFlavor).OrderBy(B => B.PreferenceOrder).FirstOrDefault();
 
 			if (Config.Build == null && Role.IsNullRole() == false)
 			{
