- Correctly name BP non-development IOS/Mac .app's

diff --git a/Engine/Source/Programs/AutomationTool/Apple/ApplePlatform.Automation.cs b/Engine/Source/Programs/AutomationTool/Apple/ApplePlatform.Automation.cs
--- a/Engine/Source/Programs/AutomationTool/Apple/ApplePlatform.Automation.cs
+++ b/Engine/Source/Programs/AutomationTool/Apple/ApplePlatform.Automation.cs
@@ -207,7 +207,8 @@ public abstract class ApplePlatform : Platform
 				{
 					// override where the .app will be located and named
 					ExtraOptions += $" SYMROOT=\"{SC.ProjectRoot}/Binaries\"";
-					ExtraOptions += $" PRODUCT_NAME={Params.ShortProjectName}";
+					string ProductName = AppleExports.MakeBinaryFileName(SC.ShortProjectName, Receipt.Platform, Receipt.Configuration, Receipt.Architectures, UnrealTargetConfiguration.Development, null);
+					ExtraOptions += $" PRODUCT_NAME={ProductName}";
 
 					if (Params.RawProjectPath != null)
 					{
