Fix malformed string format in ZenExportSnapshotTask.

diff --git a/Engine/Source/Programs/AutomationTool/BuildGraph/Tasks/ZenExportSnapshotTask.cs b/Engine/Source/Programs/AutomationTool/BuildGraph/Tasks/ZenExportSnapshotTask.cs
--- a/Engine/Source/Programs/AutomationTool/BuildGraph/Tasks/ZenExportSnapshotTask.cs
+++ b/Engine/Source/Programs/AutomationTool/BuildGraph/Tasks/ZenExportSnapshotTask.cs
@@ -127,7 +127,7 @@ namespace AutomationTool.Tasks
 		private void ZenLaunch(FileReference ProjectFile)
 		{
 			// Get the ZenLaunch executable path
-			FileReference ZenLaunchExe = ResolveFile(String.Format("Engine/Binaries/{0}/ZenLaunch{}", HostPlatform.Current.HostEditorPlatform.ToString(), RuntimeInformation.IsOSPlatform(OSPlatform.Windows) ? ".exe" : ""));
+			FileReference ZenLaunchExe = ResolveFile(String.Format("Engine/Binaries/{0}/ZenLaunch{1}", HostPlatform.Current.HostEditorPlatform.ToString(), RuntimeInformation.IsOSPlatform(OSPlatform.Windows) ? ".exe" : ""));
 
 			StringBuilder ZenLaunchCommandline = new StringBuilder();
 			ZenLaunchCommandline.AppendFormat("{0} -SponsorProcessID {1}", CommandUtils.MakePathSafeToUseWithCommandLine(ProjectFile.FullName), Environment.ProcessId);
@@ -191,7 +191,7 @@ namespace AutomationTool.Tasks
 			}
 
 			// Get the Zen executable path
-			FileReference ZenExe = ResolveFile(String.Format("Engine/Binaries/{0}/zen{}", HostPlatform.Current.HostEditorPlatform.ToString(), RuntimeInformation.IsOSPlatform(OSPlatform.Windows) ? ".exe" : ""));
+			FileReference ZenExe = ResolveFile(String.Format("Engine/Binaries/{0}/zen{1}", HostPlatform.Current.HostEditorPlatform.ToString(), RuntimeInformation.IsOSPlatform(OSPlatform.Windows) ? ".exe" : ""));
 
 			// Format the command lines
 			StringBuilder OplogSnapshotCommandline = new StringBuilder();
