Automation platforms can define a custom deployment handler that can be used to alter staging, packaging and deployment behavior.

 - Allows a target to change how it is packaged, staged and deployed - for example, when packaging for a specific game store

Hooked into Windows only at the moment, configured via:

[/Script/WindowsTargetPlatform.WindowsTargetSettings]
CustomDeployment=MyCustomDeploymentHandler

#rnx

diff --git a/Engine/Source/Programs/AutomationTool/AutomationUtils/CustomDeploymentHandler.cs b/Engine/Source/Programs/AutomationTool/AutomationUtils/CustomDeploymentHandler.cs
new file mode 100644
--- /dev/null
+++ b/Engine/Source/Programs/AutomationTool/AutomationUtils/CustomDeploymentHandler.cs
@@ -0,0 +1,100 @@
+// Copyright Epic Games, Inc. All Rights Reserved.
+using System;
+using System.Collections.Generic;
+using System.Reflection;
+using AutomationTool;
+using UnrealBuildTool;
+using System.Linq;
+
+/// <summary>
+/// Allows a target to change how it is packaged, staged and deployed - for example, when packaging for a specific game store
+/// </summary>
+public abstract class CustomDeploymentHandler
+{
+	protected readonly Platform Platform;
+	protected CustomDeploymentHandler( Platform AutomationPlatform )
+	{
+		this.Platform = AutomationPlatform;
+	}
+	public abstract bool SupportsPlatform( UnrealTargetPlatform Platform );
+
+	// will return false if the caller should do the default platform processing
+	public virtual bool GetPlatformPakCommandLine(ProjectParams Params, DeploymentContext SC, ref string CmdLine) => false;
+	public virtual bool GetPlatformIoStoreCommandLine(ProjectParams Params, DeploymentContext SC, ref string CmdLine) => false;
+	public virtual bool PrePackage(ProjectParams Params, DeploymentContext SC, int WorkingCL) => false;
+	public virtual bool Deploy(ProjectParams Params, DeploymentContext SC) => false;
+	public virtual bool PreGetFilesToDeployOrStage(ProjectParams Params, DeploymentContext SC) => false;
+	public virtual bool GetFilesToStageForDLC(ProjectParams Params, DeploymentContext SC) => false;
+	public virtual bool GetFilesToArchive(ProjectParams Params, DeploymentContext SC) => false;
+
+	// always called
+	public virtual void PostPackage(ProjectParams Params, DeploymentContext SC, int WorkingCL) { }
+	public virtual void PostGetFilesToDeployOrStage(ProjectParams Params, DeploymentContext SC) { }
+
+	// returns null if the caller should do the default platform processing
+	public virtual bool? GetPlatformPatchesWithDiffPak(ProjectParams Params, DeploymentContext SC) => null;
+	public virtual IProcessResult RunClient(CommandUtils.ERunOptions ClientRunFlags, string ClientApp, string ClientCmdLine, ProjectParams Params, DeploymentContext SC) => null;
+
+
+	#region Private/boilerplate
+
+	/// <summary>
+	/// Validate & instantiate the given custom deployment handler
+	/// </summary>
+	public static CustomDeploymentHandler Create(string Name, Platform AutomationPlatform)
+	{
+		// one-time static init. NB. Can't use the partial struct convention because automation classes are all in different assemblies
+		if (HandlerRegistry == null)
+		{
+			HandlerRegistry = new();
+
+			foreach (Assembly LoadedAssembly in ScriptManager.AllScriptAssemblies)
+			{
+				IEnumerable<Type> HandlerTypes = LoadedAssembly.GetTypes().Where(X => !X.IsAbstract && X.IsSubclassOf(typeof(CustomDeploymentHandler)) && X.GetCustomAttribute(typeof(CustomDeploymentHandlerAttribute)) != null);
+				foreach (Type HandlerType in HandlerTypes)
+				{
+					string HandlerName = HandlerType.GetCustomAttribute<CustomDeploymentHandlerAttribute>().Name;
+					HandlerRegistry.Add(HandlerName, HandlerType);
+				}
+			}
+		}
+
+		Type SelectedType;
+		if (!HandlerRegistry.TryGetValue(Name, out SelectedType) || SelectedType == null)
+		{
+			throw new BuildException($"Unknown custom deployment handler {Name}");
+		}
+
+		CustomDeploymentHandler Handler = (CustomDeploymentHandler)Activator.CreateInstance(SelectedType, AutomationPlatform);
+		if (Handler == null)
+		{
+			throw new BuildException($"Could not instantiate the custom deployment handler {Name}");
+		}
+
+		if (!Handler.SupportsPlatform(AutomationPlatform.PlatformType))
+		{
+			throw new BuildException($"Custom deployment handler {Name} does not support platform {AutomationPlatform.PlatformType}");
+		}
+
+		return Handler;
+
+	}
+
+	private static Dictionary<string, Type> HandlerRegistry = null;
+
+	#endregion
+}
+
+
+/// <summary>
+/// CustomDeploymentHandler class declarations must be tagged with this, used to find the handler by name
+/// </summary>
+[System.AttributeUsage(System.AttributeTargets.Class)]
+public class CustomDeploymentHandlerAttribute : System.Attribute
+{
+	public string Name;
+	public CustomDeploymentHandlerAttribute( string Name)
+	{
+		this.Name = Name;
+	}
+}
\ No newline at end of file
diff --git a/Engine/Source/Programs/AutomationTool/AutomationUtils/DeploymentContext.cs b/Engine/Source/Programs/AutomationTool/AutomationUtils/DeploymentContext.cs
--- a/Engine/Source/Programs/AutomationTool/AutomationUtils/DeploymentContext.cs
+++ b/Engine/Source/Programs/AutomationTool/AutomationUtils/DeploymentContext.cs
@@ -147,6 +147,11 @@ public class DeploymentContext //: ProjectParams
 	/// </summary>
 	public string CustomConfig;
 
+	/// <summary>
+	/// Allows a platform to change how it is packaged, staged and deployed - for example, when packaging for a specific game store
+	/// </summary>
+	public CustomDeploymentHandler CustomDeployment = null;
+
 	/// <summary>
 	/// This is the root directory that contains the engine: d:\a\UE\
 	/// </summary>
diff --git a/Engine/Source/Programs/AutomationTool/Scripts/ArchiveCommand.Automation.cs b/Engine/Source/Programs/AutomationTool/Scripts/ArchiveCommand.Automation.cs
--- a/Engine/Source/Programs/AutomationTool/Scripts/ArchiveCommand.Automation.cs
+++ b/Engine/Source/Programs/AutomationTool/Scripts/ArchiveCommand.Automation.cs
@@ -21,7 +21,10 @@ namespace AutomationScripts
 			}
 			var ThisPlatform = SC.StageTargetPlatform;
 
-			ThisPlatform.GetFilesToArchive(Params, SC);
+			if (SC.CustomDeployment == null || !SC.CustomDeployment.GetFilesToArchive(Params, SC))
+			{
+				ThisPlatform.GetFilesToArchive(Params, SC);
+			}
 
 			//@todo add any archive meta data files as needed
 
diff --git a/Engine/Source/Programs/AutomationTool/Scripts/CopyBuildToStagingDirectory.Automation.cs b/Engine/Source/Programs/AutomationTool/Scripts/CopyBuildToStagingDirectory.Automation.cs
--- a/Engine/Source/Programs/AutomationTool/Scripts/CopyBuildToStagingDirectory.Automation.cs
+++ b/Engine/Source/Programs/AutomationTool/Scripts/CopyBuildToStagingDirectory.Automation.cs
@@ -1109,7 +1109,10 @@ namespace AutomationScripts
 					}
 				}
 
-				ThisPlatform.GetFilesToStageForDLC(Params, SC);
+				if (SC.CustomDeployment == null || !SC.CustomDeployment.GetFilesToStageForDLC(Params, SC))
+				{
+					ThisPlatform.GetFilesToStageForDLC(Params, SC);
+				}
 
 				bCreatePluginManifest = true;
 			}
@@ -1117,7 +1120,11 @@ namespace AutomationScripts
 			{
 				if (!Params.IterateSharedBuildUsePrecompiledExe)
 				{
-					ThisPlatform.GetFilesToDeployOrStage(Params, SC);
+					if (SC.CustomDeployment == null || !SC.CustomDeployment.PreGetFilesToDeployOrStage(Params, SC))
+					{
+						ThisPlatform.GetFilesToDeployOrStage(Params, SC);
+					}
+					SC.CustomDeployment?.PostGetFilesToDeployOrStage(Params, SC);
 
 					// Stage any extra runtime dependencies from the receipts
 					foreach (StageTarget Target in SC.StageTargets)
@@ -2797,8 +2804,9 @@ namespace AutomationScripts
 		/// <param name="PakName"></param>
 		private static void CreatePaks(ProjectParams Params, DeploymentContext SC, List<CreatePakParams> PakParamsList, EncryptionAndSigning.CryptoSettings CryptoSettings, FileReference CryptoKeysCacheFilename)
 		{
+			bool? bCustomDeploymentNeedsDiffPak = SC.CustomDeployment?.GetPlatformPatchesWithDiffPak(Params, SC);
 			bool bShouldCreateIoStoreContainerFiles = ShouldCreateIoStoreContainerFiles(Params, SC);
-			bool bShouldGeneratePatch = Params.IsGeneratingPatch && SC.StageTargetPlatform.GetPlatformPatchesWithDiffPak(Params, SC);
+			bool bShouldGeneratePatch = Params.IsGeneratingPatch && (bCustomDeploymentNeedsDiffPak ?? SC.StageTargetPlatform.GetPlatformPatchesWithDiffPak(Params, SC));
 
 			if (bShouldGeneratePatch && !Params.HasBasedOnReleaseVersion)
 			{
@@ -3202,7 +3210,12 @@ namespace AutomationScripts
 				SecondaryOrderFiles = OrderFiles.FindAll(x => x.OrderType == OrderFile.OrderFileType.Cooker);
 			}
 
-			string CommonAdditionalArgs = SC.StageTargetPlatform.GetPlatformPakCommandLine(Params, SC) + BulkOption + PatchOptions + CompressionFormats + " " + Params.AdditionalPakOptions;
+			string PlatformPakCmdLine = "";
+			if (SC.CustomDeployment == null || !SC.CustomDeployment.GetPlatformPakCommandLine(Params, SC, ref PlatformPakCmdLine))
+			{
+				PlatformPakCmdLine += SC.StageTargetPlatform.GetPlatformPakCommandLine(Params, SC);
+			}
+			string CommonAdditionalArgs = PlatformPakCmdLine + BulkOption + PatchOptions + CompressionFormats + " " + Params.AdditionalPakOptions;
 
 			bool bPakFallbackOrderForNonUassetFiles = false;
 			PlatformGameConfig.GetBool("/Script/UnrealEd.ProjectPackagingSettings", "bPakFallbackOrderForNonUassetFiles", out bPakFallbackOrderForNonUassetFiles);
@@ -3484,9 +3497,21 @@ namespace AutomationScripts
 					FinalOrderFiles.AddRange(OrderFiles.FindAll(x => x.OrderType == OrderFile.OrderFileType.Cooker));
 				}
 
+				string AdditionalPlatformPakCmdLine = "";
+				if (SC.CustomDeployment == null || !SC.CustomDeployment.GetPlatformPakCommandLine(Params, SC, ref AdditionalPlatformPakCmdLine))
+				{
+					PlatformPakCmdLine += SC.StageTargetPlatform.GetPlatformPakCommandLine(Params, SC);
+				}
+
+				string AdditionalPlatformIoStoreCmdLine = "";
+				if (SC.CustomDeployment == null || !SC.CustomDeployment.GetPlatformIoStoreCommandLine(Params, SC, ref AdditionalPlatformIoStoreCmdLine))
+				{
+					AdditionalPlatformIoStoreCmdLine += SC.StageTargetPlatform.GetPlatformIoStoreCommandLine(Params, SC);
+				}
+
 				string AdditionalArgs =
-					SC.StageTargetPlatform.GetPlatformPakCommandLine(Params, SC)
-					+ SC.StageTargetPlatform.GetPlatformIoStoreCommandLine(Params, SC)
+					AdditionalPlatformPakCmdLine
+					+ AdditionalPlatformIoStoreCmdLine
 					+ BulkOption
 					+ CompressionFormats
 					+ " " + AdditionalCompressionOptionsOnCommandLine
diff --git a/Engine/Source/Programs/AutomationTool/Scripts/DeployCommand.Automation.cs b/Engine/Source/Programs/AutomationTool/Scripts/DeployCommand.Automation.cs
--- a/Engine/Source/Programs/AutomationTool/Scripts/DeployCommand.Automation.cs
+++ b/Engine/Source/Programs/AutomationTool/Scripts/DeployCommand.Automation.cs
@@ -55,7 +55,10 @@ namespace AutomationScripts
 					}
 					else
 					{
-						SC.StageTargetPlatform.Deploy(Params, SC);
+						if (SC.CustomDeployment == null || !SC.CustomDeployment.Deploy(Params, SC))
+						{
+							SC.StageTargetPlatform.Deploy(Params, SC);
+						}
 					}
 				}
 			}
diff --git a/Engine/Source/Programs/AutomationTool/Scripts/PackageCommand.Automation.cs b/Engine/Source/Programs/AutomationTool/Scripts/PackageCommand.Automation.cs
--- a/Engine/Source/Programs/AutomationTool/Scripts/PackageCommand.Automation.cs
+++ b/Engine/Source/Programs/AutomationTool/Scripts/PackageCommand.Automation.cs
@@ -46,7 +46,11 @@ namespace AutomationScripts
 					{
 						if (Params.Package || (SC.StageTargetPlatform.RequiresPackageToDeploy(Params) && Params.Deploy))
 						{
-							SC.StageTargetPlatform.Package(Params, SC, WorkingCL);
+							if (SC.CustomDeployment == null || !SC.CustomDeployment.PrePackage(Params, SC, WorkingCL))
+							{
+								SC.StageTargetPlatform.Package(Params, SC, WorkingCL);
+							}
+							SC.CustomDeployment?.PostPackage(Params, SC, WorkingCL);
 						}
 					}
 
diff --git a/Engine/Source/Programs/AutomationTool/Scripts/RunProjectCommand.Automation.cs b/Engine/Source/Programs/AutomationTool/Scripts/RunProjectCommand.Automation.cs
--- a/Engine/Source/Programs/AutomationTool/Scripts/RunProjectCommand.Automation.cs
+++ b/Engine/Source/Programs/AutomationTool/Scripts/RunProjectCommand.Automation.cs
@@ -275,6 +275,12 @@ namespace AutomationScripts
 			}
 		}
 
+		private static IProcessResult RunClientInternal( ERunOptions ClientRunFlags, string ClientApp, string ClientCmdLine, ProjectParams Params, DeploymentContext SC )
+		{
+			IProcessResult CustomResult = SC.CustomDeployment?.RunClient(ClientRunFlags, ClientApp, ClientCmdLine, Params, SC);
+			return CustomResult ?? SC.StageTargetPlatform.RunClient(ClientRunFlags, ClientApp, ClientCmdLine, Params, SC);
+		}
+
 		private static void RunClient(
 				List<DeploymentContext> DeployContextList,
 				string ClientLogFile,
@@ -325,7 +331,7 @@ namespace AutomationScripts
 				int LastAutoFailIndex = -1;
 
 				Logger.LogInformation("Starting Client for unattended test....");
-				ClientProcess = SC.StageTargetPlatform.RunClient(ClientRunFlags, ClientApp, ClientCmdLine, Params, SC);
+				ClientProcess = RunClientInternal(ClientRunFlags, ClientApp, ClientCmdLine, Params, SC);
 
 				if (DedicatedServerProcess != null)
 				{
@@ -334,7 +340,7 @@ namespace AutomationScripts
 						for (int i = 1; i < NumClients; i++)
 						{
 							Logger.LogInformation("Starting Extra Client {i} for unattended test....", i);
-							ExtraClients.Add(SC.StageTargetPlatform.RunClient(ExtraClientRunFlags, ClientApp, ClientCmdLine, Params, SC));
+							ExtraClients.Add(RunClientInternal(ExtraClientRunFlags, ClientApp, ClientCmdLine, Params, SC));
 						}
 					}
 				}
@@ -458,7 +464,7 @@ namespace AutomationScripts
 			else
 			{
 				Logger.LogInformation("Starting Client....");
-				ClientProcess = SC.StageTargetPlatform.RunClient(ClientRunFlags, ClientApp, ClientCmdLine, Params, SC);
+				ClientProcess = RunClientInternal(ClientRunFlags, ClientApp, ClientCmdLine, Params, SC);
 
 				if (DedicatedServerProcess != null)
 				{
@@ -467,7 +473,7 @@ namespace AutomationScripts
 						for (int i = 1; i < NumClients; i++)
 						{
 							Logger.LogInformation("Starting Extra Client {i}....", i);
-							ExtraClients.Add(SC.StageTargetPlatform.RunClient(ExtraClientRunFlags, ClientApp, ClientCmdLine, Params, SC));
+							ExtraClients.Add(RunClientInternal(ExtraClientRunFlags, ClientApp, ClientCmdLine, Params, SC));
 						}
 					}
 				}
diff --git a/Engine/Source/Programs/AutomationTool/Win/WinPlatform.Automation.cs b/Engine/Source/Programs/AutomationTool/Win/WinPlatform.Automation.cs
--- a/Engine/Source/Programs/AutomationTool/Win/WinPlatform.Automation.cs
+++ b/Engine/Source/Programs/AutomationTool/Win/WinPlatform.Automation.cs
@@ -47,6 +47,24 @@ public class Win64Platform : Platform
 		return Devices.ToArray();
 	}
 
+	public override void PlatformSetupParams(ref ProjectParams Params)
+	{
+		base.PlatformSetupParams(ref Params);
+
+		// use a custom deployment handler if one is requested
+		Params.PreModifyDeploymentContextCallback = new Action<ProjectParams, DeploymentContext>((ProjectParams Params, DeploymentContext SC) =>
+		{
+			if (SC.CustomDeployment == null)
+			{			
+				ConfigHierarchy EngineIni = ConfigCache.ReadHierarchy(ConfigHierarchyType.Engine, Params.RawProjectPath.Directory, PlatformType, SC.CustomConfig);
+				if (EngineIni.GetString("/Script/WindowsTargetPlatform.WindowsTargetSettings", "CustomDeployment", out string CustomDeploymentName))
+				{
+					SC.CustomDeployment = CustomDeploymentHandler.Create(CustomDeploymentName, this);
+				}
+			}			
+		});
+	}
+
 	public override void Deploy(ProjectParams Params, DeploymentContext SC)
 	{
 		// We only care about deploying for SteamDeck
