Add support for self-installing Windows Gauntlet builds, such as packages for specific game stores.

#rnx

diff --git a/Engine/Source/Programs/AutomationTool/Gauntlet/Platform/Windows/Gauntlet.TargetDeviceWindows.cs b/Engine/Source/Programs/AutomationTool/Gauntlet/Platform/Windows/Gauntlet.TargetDeviceWindows.cs
--- a/Engine/Source/Programs/AutomationTool/Gauntlet/Platform/Windows/Gauntlet.TargetDeviceWindows.cs
+++ b/Engine/Source/Programs/AutomationTool/Gauntlet/Platform/Windows/Gauntlet.TargetDeviceWindows.cs
@@ -143,7 +143,7 @@ namespace Gauntlet
 	}
 
 
-	class WindowsAppInstall : IAppInstall, IAppInstall.IDynamicCommandLine
+	public class WindowsAppInstall : IAppInstall, IAppInstall.IDynamicCommandLine
 	{
 		public string Name { get; private set; }
 
@@ -241,6 +241,11 @@ namespace Gauntlet
 		}
 	}
 
+	public interface IWindowsSelfInstallingBuild
+	{
+		WindowsAppInstall Install(TargetDeviceWindows TargetDevice, UnrealAppConfig AppConfig, out string BasePath);
+	}
+
 	public class Win64DeviceFactory : IDeviceFactory
 	{
 		public bool CanSupportPlatform(UnrealTargetPlatform? Platform)
@@ -566,6 +571,37 @@ namespace Gauntlet
 		}
 
 
+		protected IAppInstall InstallSelfInstallingBuild(UnrealAppConfig AppConfig, IWindowsSelfInstallingBuild Build)
+		{
+			WindowsAppInstall WinApp = Build.Install(this, AppConfig, out string BasePath);
+
+			if (Log.IsVeryVerbose)
+			{
+				WinApp.RunOptions |= CommandUtils.ERunOptions.AllowSpew;
+			}
+
+			if (string.IsNullOrEmpty(UserDir) == false)
+			{
+				WinApp.CommandArguments += string.Format(" -userdir=\"{0}\"", UserDir);
+				WinApp.ArtifactPath = Path.Combine(UserDir, @"Saved");
+				Utils.SystemHelpers.MarkDirectoryForCleanup(UserDir);
+			}
+			else
+			{
+				WinApp.ArtifactPath = Path.Combine(BasePath, AppConfig.ProjectName, @"Saved");
+			}
+			WinApp.CleanDeviceArtifacts();
+
+			if (LocalDirectoryMappings.Count == 0)
+			{
+				PopulateDirectoryMappings(Path.Combine(BasePath, AppConfig.ProjectName), UserDir);
+			}
+
+			CopyAdditionalFiles(AppConfig);
+			return WinApp;
+		}
+
+
 		public IAppInstall InstallApplication(UnrealAppConfig AppConfig)
 		{
 			if (AppConfig.Build is NativeStagedBuild)
@@ -576,6 +612,10 @@ namespace Gauntlet
 			{
 				return InstallStagedBuild(AppConfig, AppConfig.Build as StagedBuild);
 			}
+			else if (AppConfig.Build is IWindowsSelfInstallingBuild)
+			{
+				return InstallSelfInstallingBuild(AppConfig, AppConfig.Build as IWindowsSelfInstallingBuild);
+			}
 
 			EditorBuild EditorBuild = AppConfig.Build as EditorBuild;
 
