Fix for use of SponsorProcessID when calling ZenLaunch.

diff --git a/Engine/Source/Programs/AutomationTool/BuildGraph/Tasks/ZenExportSnapshotTask.cs b/Engine/Source/Programs/AutomationTool/BuildGraph/Tasks/ZenExportSnapshotTask.cs
--- a/Engine/Source/Programs/AutomationTool/BuildGraph/Tasks/ZenExportSnapshotTask.cs
+++ b/Engine/Source/Programs/AutomationTool/BuildGraph/Tasks/ZenExportSnapshotTask.cs
@@ -130,7 +130,7 @@ namespace AutomationTool.Tasks
 			FileReference ZenLaunchExe = ResolveFile(String.Format("Engine/Binaries/{0}/ZenLaunch{1}", HostPlatform.Current.HostEditorPlatform.ToString(), RuntimeInformation.IsOSPlatform(OSPlatform.Windows) ? ".exe" : ""));
 
 			StringBuilder ZenLaunchCommandline = new StringBuilder();
-			ZenLaunchCommandline.AppendFormat("{0} -SponsorProcessID {1}", CommandUtils.MakePathSafeToUseWithCommandLine(ProjectFile.FullName), Environment.ProcessId);
+			ZenLaunchCommandline.AppendFormat("{0} -SponsorProcessID={1}", CommandUtils.MakePathSafeToUseWithCommandLine(ProjectFile.FullName), Environment.ProcessId);
 
 			Logger.LogInformation("Running '{Arg0} {Arg1}'", CommandUtils.MakePathSafeToUseWithCommandLine(ZenLaunchExe.FullName), ZenLaunchCommandline.ToString());
 			CommandUtils.RunAndLog(CommandUtils.CmdEnv, ZenLaunchExe.FullName, ZenLaunchCommandline.ToString(), Options: CommandUtils.ERunOptions.Default);
