Ensure Zen is running before running OplogImport commands
Expose Zen launch commands from ZenExport to keep one source of truth

diff --git a/Engine/Source/Programs/AutomationTool/BuildGraph/Tasks/ZenExportSnapshotTask.cs b/Engine/Source/Programs/AutomationTool/BuildGraph/Tasks/ZenExportSnapshotTask.cs
--- a/Engine/Source/Programs/AutomationTool/BuildGraph/Tasks/ZenExportSnapshotTask.cs
+++ b/Engine/Source/Programs/AutomationTool/BuildGraph/Tasks/ZenExportSnapshotTask.cs
@@ -124,10 +124,26 @@ namespace AutomationTool.Tasks
 			Parameters = InParameters;
 		}
 
-		private void ZenLaunch(FileReference ProjectFile)
+		/// <summary>
+		/// Gets the assumed path to where the Zen exe should exist.
+		/// </summary>
+		/// <returns></returns>
+		public static FileReference ZenExeFileReference()
+		{
+			return ResolveFile(String.Format("Engine/Binaries/{0}/ZenLaunch{}", HostPlatform.Current.HostEditorPlatform.ToString(), RuntimeInformation.IsOSPlatform(OSPlatform.Windows) ? ".exe" : ""));
+		}
+
+
+		/// <summary>
+		/// Ensures that ZenServer is running on this current machine. This is needed before running any oplog commands
+		/// This passes the sponsor'd process Id to launch zen.
+		/// This ensures that zen does not live longer than the lifetime of a particular a process that needs Zen to be running
+		/// </summary>
+		/// <param name="ProjectFile"></param>
+		public static void ZenLaunch(FileReference ProjectFile)
 		{
 			// Get the ZenLaunch executable path
-			FileReference ZenLaunchExe = ResolveFile(String.Format("Engine/Binaries/{0}/ZenLaunch{1}", HostPlatform.Current.HostEditorPlatform.ToString(), RuntimeInformation.IsOSPlatform(OSPlatform.Windows) ? ".exe" : ""));
+			FileReference ZenLaunchExe = ZenExeFileReference();
 
 			StringBuilder ZenLaunchCommandline = new StringBuilder();
 			ZenLaunchCommandline.AppendFormat("{0} -SponsorProcessID={1}", CommandUtils.MakePathSafeToUseWithCommandLine(ProjectFile.FullName), Environment.ProcessId);
@@ -191,7 +207,7 @@ namespace AutomationTool.Tasks
 			}
 
 			// Get the Zen executable path
-			FileReference ZenExe = ResolveFile(String.Format("Engine/Binaries/{0}/zen{1}", HostPlatform.Current.HostEditorPlatform.ToString(), RuntimeInformation.IsOSPlatform(OSPlatform.Windows) ? ".exe" : ""));
+			FileReference ZenExe = ZenExeFileReference();
 
 			// Format the command lines
 			StringBuilder OplogSnapshotCommandline = new StringBuilder();
diff --git a/Engine/Source/Programs/AutomationTool/BuildGraph/Tasks/ZenImportOplogTask.cs b/Engine/Source/Programs/AutomationTool/BuildGraph/Tasks/ZenImportOplogTask.cs
--- a/Engine/Source/Programs/AutomationTool/BuildGraph/Tasks/ZenImportOplogTask.cs
+++ b/Engine/Source/Programs/AutomationTool/BuildGraph/Tasks/ZenImportOplogTask.cs
@@ -144,17 +144,10 @@ namespace AutomationTool.Tasks
 				throw new AutomationException("Missing project file - {0}", ProjectFile.FullName);
 			}
 
-			// Get the executable path
-			FileReference ZenExe;
-			if(HostPlatform.Current.HostEditorPlatform == UnrealTargetPlatform.Win64)
-			{
-				ZenExe = ResolveFile("Engine/Binaries/Win64/zen.exe");
-			}
-			else
-			{
-				ZenExe = ResolveFile(String.Format("Engine/Binaries/{0}/zen", HostPlatform.Current.HostEditorPlatform.ToString()));
-			}
+			ZenExportSnapshotTask.ZenLaunch(ProjectFile);
 
+			// Get the Zen executable path
+			FileReference ZenExe = ZenExportSnapshotTask.ZenExeFileReference();
 			{
 				if (String.IsNullOrEmpty(Parameters.RootDir))
 				{
