Fix up incorrect exe specification

diff --git a/Engine/Source/Programs/AutomationTool/BuildGraph/Tasks/ZenExportSnapshotTask.cs b/Engine/Source/Programs/AutomationTool/BuildGraph/Tasks/ZenExportSnapshotTask.cs
--- a/Engine/Source/Programs/AutomationTool/BuildGraph/Tasks/ZenExportSnapshotTask.cs
+++ b/Engine/Source/Programs/AutomationTool/BuildGraph/Tasks/ZenExportSnapshotTask.cs
@@ -125,12 +125,12 @@ namespace AutomationTool.Tasks
 		}
 
 		/// <summary>
-		/// Gets the assumed path to where the Zen exe should exist.
+		/// Gets the assumed path to where Zen should exist
 		/// </summary>
 		/// <returns></returns>
 		public static FileReference ZenExeFileReference()
 		{
-			return ResolveFile(String.Format("Engine/Binaries/{0}/ZenLaunch{}", HostPlatform.Current.HostEditorPlatform.ToString(), RuntimeInformation.IsOSPlatform(OSPlatform.Windows) ? ".exe" : ""));
+			return ResolveFile(String.Format("Engine/Binaries/{0}/zen{1}", HostPlatform.Current.HostEditorPlatform.ToString(), RuntimeInformation.IsOSPlatform(OSPlatform.Windows) ? ".exe" : ""));
 		}
 
 
@@ -143,7 +143,7 @@ namespace AutomationTool.Tasks
 		public static void ZenLaunch(FileReference ProjectFile)
 		{
 			// Get the ZenLaunch executable path
-			FileReference ZenLaunchExe = ZenExeFileReference();
+			FileReference ZenLaunchExe = ResolveFile(String.Format("Engine/Binaries/{0}/ZenLaunch{1}", HostPlatform.Current.HostEditorPlatform.ToString(), RuntimeInformation.IsOSPlatform(OSPlatform.Windows) ? ".exe" : ""));
 
 			StringBuilder ZenLaunchCommandline = new StringBuilder();
 			ZenLaunchCommandline.AppendFormat("{0} -SponsorProcessID={1}", CommandUtils.MakePathSafeToUseWithCommandLine(ProjectFile.FullName), Environment.ProcessId);
