Localization:
Fixed issue with Smartling downloads failing during localization import. Issue is caused by an authentication token being retrieved/refreshed but not set as part of HttpClient Authentication request header.
This resulted in HTTP 401 errors where the authentication fails with the endpoint.
- Updated retrieval/refreshing of authentication token to also set the newly retreived/refreshed authentication data in the HttpClient Authentication request header.
- Introduce HTTP 401 status checks in the downloading/uploading of PO files to detect authentication token expiration. Token refreshing is performed the upload/download operation is retried.
- Refactored UploadLatestPOFile and DownloadLatestPOFile to reduce nesting blocks and minimize code in try-catch blocks.
- This is an urgent fix so no approved reviews yet. Will fix/address problems after.

#rnx
#test Locally ran the Localize UAT command with the Download step. Before these changes, the download step would have HTTP 401 errors failing to authenticate. After these changes, I haven't encountered them.

diff --git a/Engine/Source/Programs/AutomationTool/SmartlingLocalization/SmartlingLocalizationProvider.cs b/Engine/Source/Programs/AutomationTool/SmartlingLocalization/SmartlingLocalizationProvider.cs
--- a/Engine/Source/Programs/AutomationTool/SmartlingLocalization/SmartlingLocalizationProvider.cs
+++ b/Engine/Source/Programs/AutomationTool/SmartlingLocalization/SmartlingLocalizationProvider.cs
@@ -107,7 +107,7 @@ namespace EpicGames.SmartlingLocalization
 		private async Task DownloadLatestPOFile(string EpicLocale, string Platform, ProjectImportExportInfo ProjectImportInfo)
 		{
 			await GetAuthenticationToken();
-			Client.DefaultRequestHeaders.Authorization = new AuthenticationHeaderValue(AuthenticationToken.TokenType, AuthenticationToken.AccessToken);
+			
 
 			var DestinationDirectory = String.IsNullOrEmpty(Platform)
 			? new DirectoryInfo(CommandUtils.CombinePaths(RootWorkingDirectory, ProjectImportInfo.DestinationPath))
@@ -136,30 +136,41 @@ namespace EpicGames.SmartlingLocalization
 			int CurrentTries = 1;
 			// The base for which all the exponential back off will be derived from. By default HttpClient has a default timeout of 100s 
 			int InitialTimeOut = 150;
+			int CurrentTimeOut = CurrentTries * InitialTimeOut;
 			while (true)
 			{
-				int CurrentTimeOut= CurrentTries * InitialTimeOut;
+				CurrentTimeOut = CurrentTries * InitialTimeOut;
 				Console.WriteLine($"Current time out for download response {CurrentTimeOut}s.");
 				TimeSpan Timeout = TimeSpan.FromSeconds(CurrentTimeOut);
-				using (var CancellationToken = new CancellationTokenSource(Timeout))
+				using var CancellationToken = new CancellationTokenSource(Timeout);
+				try
 				{
-					try
-					{
-						DownloadResponse = await Client.GetAsync(DownloadUriBuilder.Uri, CancellationToken.Token);
-						break;
-					}
-					catch (Exception Ex)
+					DownloadResponse = await Client.GetAsync(DownloadUriBuilder.Uri, CancellationToken.Token);
+					if (DownloadResponse.StatusCode == HttpStatusCode.Unauthorized)
 					{
 						++CurrentTries;
 						if (CurrentTries > MaxTries)
 						{
-							Console.WriteLine($"[FAILED] Exporting: '{ExportFile.FullName}' ({EpicLocale}) teimed out. - {Ex}");
+							Console.WriteLine($"[FAILED] Exporting: '{ExportFile.FullName}' ({EpicLocale}) exhausted all retries.");
 							return;
 						}
-						Console.WriteLine($"Failed to get download response. Retrying {CurrentTries}/{MaxTries} times.");
-						// We need to retreive the authentication token again because after each timeout, we may exceed the validity of the authentication token 
+						Console.WriteLine($"Encountered HTTP Status Code 401. Authentication most likely expired. Retrying {CurrentTries}/{MaxTries} times with refreshed authentication token.");
 						await GetAuthenticationToken();
+						continue;
+					}
+					break;
+				}
+				catch (Exception Ex)
+				{
+					++CurrentTries;
+					if (CurrentTries > MaxTries)
+					{
+						Console.WriteLine($"[FAILED] Exporting: '{ExportFile.FullName}' ({EpicLocale}) exhausted all retries. - {Ex}");
+						return;
 					}
+					Console.WriteLine($"Failed to get download response. Retrying {CurrentTries}/{MaxTries} times.");
+					// We need to retreive the authentication token again because after each timeout, we may exceed the validity of the authentication token 
+					await GetAuthenticationToken();
 				}
 			}
 
@@ -254,7 +265,6 @@ namespace EpicGames.SmartlingLocalization
 		private async Task UploadLatestPOFile(string EpicLocale, string Platform, ProjectImportExportInfo ProjectExportInfo)
 		{
 			await GetAuthenticationToken();
-			Client.DefaultRequestHeaders.Authorization = new AuthenticationHeaderValue(AuthenticationToken.TokenType, AuthenticationToken.AccessToken);
 
 			var SourceDirectory = String.IsNullOrEmpty(Platform)
 				? new DirectoryInfo(CommandUtils.CombinePaths(RootWorkingDirectory, ProjectExportInfo.DestinationPath))
@@ -274,44 +284,67 @@ namespace EpicGames.SmartlingLocalization
 			Console.WriteLine($"Uploading: '{FileToUpload.FullName}' as '{SmartlingFileUri}' ({EpicLocale})");
 			// For now we only upload the file in the native locale
 			// It is dissuaded to import translations to Smartling. Perform all translations in the Smartling dashboard instead.
-			try
+
+			string UploadEndpoint = $"https://api.smartling.com/files-api/v2/projects/{Config.ProjectId}/file";
+
+			using var UploadMultipartFormDataContent = new MultipartFormDataContent();
+				var UploadFileStreamContent = new StreamContent(File.OpenRead(FileToUpload.FullName));
+			// generic binary stream we will simply consider the octet-stream mime type 
+			UploadFileStreamContent.Headers.ContentType = new MediaTypeHeaderValue("application/octet-stream");
+			UploadMultipartFormDataContent.Add(UploadFileStreamContent, "file", SmartlingFileUri);
+			UploadMultipartFormDataContent.Add(new StringContent(SmartlingFileUri, Encoding.UTF8, "application/json"), "fileUri");
+			UploadMultipartFormDataContent.Add(new StringContent("gettext", Encoding.UTF8, "application/json"), "fileType");
+			// We introduce a Smartling namespace to leverage Smartling's string sharing feature
+			// Following Smartling best practices, we make the Smartling namespace the same as the full file path of the file.
+			// https://help.smartling.com/hc/en-us/articles/360008143833-String-Sharing-and-Namespaces-via-Smartling-API
+			UploadMultipartFormDataContent.Add(new StringContent(SmartlingFilename, Encoding.UTF8, "application/json"), "smartling.namespace");
+			// all placeholder values in the source files will raise warnings in Smartling which need to be manually reviewed.
+			// This allows the warnings to be resolved. The regex designates anything within {} to be a placeholder value.
+			// This accounts for the common FText formatted string placeholders of {0} or {MyPlaceholderVariable}
+			UploadMultipartFormDataContent.Add(new StringContent("\\{([^}]+)\\}", Encoding.UTF8, "application/json"), "smartling.placeholder_format_custom");
+			// Controls whether base characters ( > < & " ) are "escaped" into entities 
+			// We set this as false so that we don't automatically escape and have Smartling think these are HTML tags or something.
+			// https://help.smartling.com/hc/en-us/articles/360007894594-Gettext-PO-POT
+			UploadMultipartFormDataContent.Add(new StringContent("false", Encoding.UTF8, "application/json"), "smartling.entity_escaping");
+
+			HttpResponseMessage UploadResponse = null;
+			int CurrentTries = 1;
+			int MaxTries = 5;
+			// @TODOLocalization: Doing retries with exponential backoff won't work great with uploads as we'll run into HTTP 402 errors where the resource is locked. Need to find a way around that.
+			// For now, we don't run into issues of timing out or failure to authenticate, so we're fine. But this isn't correct  
+			while (true)
 			{
-				string UploadEndpoint = $"https://api.smartling.com/files-api/v2/projects/{Config.ProjectId}/file";
-				using (var UploadMultipartFormDataContent = new MultipartFormDataContent())
+				try
 				{
-					var UploadFileStreamContent = new StreamContent(File.OpenRead(FileToUpload.FullName));
-					// generic binary stream we will simply consider the octet-stream mime type 
-					UploadFileStreamContent.Headers.ContentType = new MediaTypeHeaderValue("application/octet-stream");
-					UploadMultipartFormDataContent.Add(UploadFileStreamContent, "file", SmartlingFileUri);	
-					UploadMultipartFormDataContent.Add(new StringContent(SmartlingFileUri, Encoding.UTF8, "application/json"), "fileUri");
-					UploadMultipartFormDataContent.Add(new StringContent("gettext", Encoding.UTF8, "application/json"), "fileType");
-					// We introduce a Smartling namespace to leverage Smartling's string sharing feature
-					// Following Smartling best practices, we make the Smartling namespace the same as the full file path of the file.
-					// https://help.smartling.com/hc/en-us/articles/360008143833-String-Sharing-and-Namespaces-via-Smartling-API
-					UploadMultipartFormDataContent.Add(new StringContent(SmartlingFilename, Encoding.UTF8, "application/json"), "smartling.namespace");
-					// all placeholder values in the source files will raise warnings in Smartling which need to be manually reviewed.
-					// This allows the warnings to be resolved. The regex designates anything within {} to be a placeholder value.
-					// This accounts for the common FText formatted string placeholders of {0} or {MyPlaceholderVariable}
-					UploadMultipartFormDataContent.Add(new StringContent("\\{([^}]+)\\}", Encoding.UTF8, "application/json"), "smartling.placeholder_format_custom");
-					// Controls whether base characters ( > < & " ) are "escaped" into entities 
-					// We set this as false so that we don't automatically escape and have Smartling think these are HTML tags or something.
-					// https://help.smartling.com/hc/en-us/articles/360007894594-Gettext-PO-POT
-					UploadMultipartFormDataContent.Add(new StringContent("false", Encoding.UTF8, "application/json"), "smartling.entity_escaping");
-					var UploadResponse = await Client.PostAsync(UploadEndpoint, UploadMultipartFormDataContent);
-					if (UploadResponse.IsSuccessStatusCode)
+					UploadResponse = await Client.PostAsync(UploadEndpoint, UploadMultipartFormDataContent);
+					if (UploadResponse.StatusCode == HttpStatusCode.Unauthorized)
 					{
-						Console.WriteLine($"[SUCCESS] Uploading: '{FileToUpload.FullName}' ({EpicLocale})");
-					}
-					else
-					{
-						Console.WriteLine($"[FAILED] Uploading: '{FileToUpload.FullName}' ({EpicLocale})");
-						await PrintRequestErrors(UploadResponse);
+						++CurrentTries;
+						if (CurrentTries > MaxTries)
+						{
+							Console.WriteLine($"[FAILED] Uploading: '{FileToUpload.FullName}' ({EpicLocale}). Exhausted all retries.");
+							return;
+						}
+						Console.WriteLine("Encountered HTTP Status Code 401. Authentication token most likely expired. Retrying {CurrentTries}/{MaxTries} times with refreshed authentication token.");
+						continue;
 					}
+					break;
+				}
+				catch (Exception Ex)
+				{
+					Console.WriteLine($"[FAILED] Uploading: '{FileToUpload.FullName}' ({EpicLocale}) - {Ex}");
+					return;
 				}
 			}
-			catch (Exception Ex)
+
+			if (UploadResponse.IsSuccessStatusCode)
+			{
+				Console.WriteLine($"[SUCCESS] Uploading: '{FileToUpload.FullName}' ({EpicLocale})");
+			}
+			else
 			{
-				Console.WriteLine($"[FAILED] Uploading: '{FileToUpload.FullName}' ({EpicLocale}) - {Ex}");
+				Console.WriteLine($"[FAILED] Uploading: '{FileToUpload.FullName}' ({EpicLocale})");
+				await PrintRequestErrors(UploadResponse);
 			}
 		}
 
@@ -422,6 +455,7 @@ namespace EpicGames.SmartlingLocalization
 					var AuthenticateResponseEnvelope = JsonSerializer.Deserialize<SmartlingResponseEnvelope<SmartlingDataEnvelope<SmartlingAuthenticationToken>>>(AuthenticateResponseString, JsonOptions);
 					AuthenticationToken = AuthenticateResponseEnvelope.Response.Data;
 					AuthenticationToken.LastSuccessfulUpdateTime = DateTime.UtcNow;
+					Client.DefaultRequestHeaders.Authorization = new AuthenticationHeaderValue(AuthenticationToken.TokenType, AuthenticationToken.AccessToken);
 					Console.WriteLine("Successfully retrieved authentication token!");
 				}
 				else
@@ -455,6 +489,7 @@ namespace EpicGames.SmartlingLocalization
 					var RefreshResponseEnvelope = JsonSerializer.Deserialize<SmartlingResponseEnvelope<SmartlingDataEnvelope<SmartlingAuthenticationToken>>>(RefreshResponseString, JsonOptions);
 					AuthenticationToken = RefreshResponseEnvelope.Response.Data;
 					AuthenticationToken.LastSuccessfulUpdateTime = DateTime.UtcNow;
+					Client.DefaultRequestHeaders.Authorization = new AuthenticationHeaderValue(AuthenticationToken.TokenType, AuthenticationToken.AccessToken);
 					Console.WriteLine("Successfully refreshed authentication token!");
 				}
 				else
