Change the zen export task to use the newly added --embedloosefiles argument instead of calling oplog-snapshot before calling oplog-export.

diff --git a/Engine/Source/Programs/AutomationTool/BuildGraph/Tasks/ZenExportSnapshotTask.cs b/Engine/Source/Programs/AutomationTool/BuildGraph/Tasks/ZenExportSnapshotTask.cs
--- a/Engine/Source/Programs/AutomationTool/BuildGraph/Tasks/ZenExportSnapshotTask.cs
+++ b/Engine/Source/Programs/AutomationTool/BuildGraph/Tasks/ZenExportSnapshotTask.cs
@@ -209,9 +209,7 @@ namespace AutomationTool.Tasks
 			// Get the Zen executable path
 			FileReference ZenExe = ZenExeFileReference();
 
-			// Format the command lines
-			StringBuilder OplogSnapshotCommandline = new StringBuilder();
-			OplogSnapshotCommandline.AppendFormat("oplog-snapshot");
+			// Format the command line
 			StringBuilder OplogExportCommandline = new StringBuilder();
 			OplogExportCommandline.AppendFormat("oplog-export");
 
@@ -246,12 +244,6 @@ namespace AutomationTool.Tasks
 					{
 						String HostUrlArg = string.Format("--hosturl http://{0}:{1}", ExportSource.IsLocalHost ? "localhost" : ExportSource.HostName, ExportSource.HostPort);
 
-						StringBuilder SnapshotSingleSourceCommandline = new StringBuilder(OplogSnapshotCommandline.Length);
-						SnapshotSingleSourceCommandline.Append(OplogSnapshotCommandline);
-						SnapshotSingleSourceCommandline.AppendFormat(" {0} {1} {2}", HostUrlArg, ExportSource.ProjectId, ExportSource.OplogId);
-						Logger.LogInformation("Running '{Arg0} {Arg1}'", CommandUtils.MakePathSafeToUseWithCommandLine(ZenExe.FullName), SnapshotSingleSourceCommandline.ToString());
-						CommandUtils.RunAndLog(CommandUtils.CmdEnv, ZenExe.FullName, SnapshotSingleSourceCommandline.ToString(), Options: CommandUtils.ERunOptions.Default);
-
 						StringBuilder ExportSingleSourceCommandline = new StringBuilder(OplogExportCommandline.Length);
 						ExportSingleSourceCommandline.Append(OplogExportCommandline);
 
@@ -264,7 +256,7 @@ namespace AutomationTool.Tasks
 							{
 								return null;
 							});
-						ExportSingleSourceCommandline.AppendFormat(" {0} --key {1} {2} {3}", HostUrlArg, DestinationKeyHash.ToString().ToLowerInvariant(), ExportSource.ProjectId, ExportSource.OplogId);
+						ExportSingleSourceCommandline.AppendFormat(" {0} --embedloosefiles --key {1} {2} {3}", HostUrlArg, DestinationKeyHash.ToString().ToLowerInvariant(), ExportSource.ProjectId, ExportSource.OplogId);
 						Logger.LogInformation("Running '{Arg0} {Arg1}'", CommandUtils.MakePathSafeToUseWithCommandLine(ZenExe.FullName), ExportSingleSourceCommandline.ToString());
 						CommandUtils.RunAndLog(CommandUtils.CmdEnv, ZenExe.FullName, ExportSingleSourceCommandline.ToString(), MaxSuccessCode: int.MaxValue, Options: CommandUtils.ERunOptions.Default, SpewFilterCallback: SilentOutputFilter);
 
