UAT: BuildCookRun Stage fix for missing patchpaddingalign iostore argument to UnrealPak.exe (fallout from 27370750)

#rnx

diff --git a/Engine/Source/Programs/AutomationTool/Scripts/CopyBuildToStagingDirectory.Automation.cs b/Engine/Source/Programs/AutomationTool/Scripts/CopyBuildToStagingDirectory.Automation.cs
--- a/Engine/Source/Programs/AutomationTool/Scripts/CopyBuildToStagingDirectory.Automation.cs
+++ b/Engine/Source/Programs/AutomationTool/Scripts/CopyBuildToStagingDirectory.Automation.cs
@@ -3500,7 +3500,7 @@ namespace AutomationScripts
 				string AdditionalPlatformPakCmdLine = "";
 				if (SC.CustomDeployment == null || !SC.CustomDeployment.GetPlatformPakCommandLine(Params, SC, ref AdditionalPlatformPakCmdLine))
 				{
-					PlatformPakCmdLine += SC.StageTargetPlatform.GetPlatformPakCommandLine(Params, SC);
+					AdditionalPlatformPakCmdLine += SC.StageTargetPlatform.GetPlatformPakCommandLine(Params, SC);
 				}
 
 				string AdditionalPlatformIoStoreCmdLine = "";
