Some improvements to MultiClientLauncher tool.

- It is now possible to simply launch it as
   RunUAT MultiClientLauncher -ClientCount=X -ClientExe=Foo/Bar.exe (-BuildIdOverride=Y if needed)

diff --git a/Engine/Source/Programs/AutomationTool/Scripts/MultiClientLauncher.Automation.cs b/Engine/Source/Programs/AutomationTool/Scripts/MultiClientLauncher.Automation.cs
--- a/Engine/Source/Programs/AutomationTool/Scripts/MultiClientLauncher.Automation.cs
+++ b/Engine/Source/Programs/AutomationTool/Scripts/MultiClientLauncher.Automation.cs
@@ -5,19 +5,25 @@ using System.Collections.Generic;
 using System.Diagnostics;
 using System.IO;
 using System.Threading;
+using Amazon;
 using AutomationTool;
 using EpicGames.Core;
+using Microsoft.CodeAnalysis.CSharp.Syntax;
+using Microsoft.Extensions.Logging;
 using UnrealBuildBase;
+using UnrealBuildTool;
 
 namespace MultiClientLauncher.Automation
 {
 	[Help("Run many game clients, a server, and connect them")]
 	[ParamHelp("ClientExe", "Absolute path to the client to run", ParamType = typeof(FileReference))]
-	[ParamHelp("ClientLogFile", "Absolute path to the client log file", ParamType = typeof(FileReference))]
-	[ParamHelp("ClientArgsFile", "Absolute path to a file containing the client arguments", ParamType = typeof(FileReference))]
-	[ParamHelp("FirstClientNumber", "The number of the first LoadBot client", ParamType = typeof(int))]
 	[ParamHelp("ClientCount", "How many bot clients to run, must consecutively follow FirstClientNumber", ParamType = typeof(int))]
-	[ParamHelp("SleepTimeBetweenLaunches", "How long to sleep between running clients in milliseconds (could prevent race conditions)", ParamType = typeof(int))]
+	[ParamHelp("BuildIdOverride", "Parameter for -buildidoverride switch, often used to narrow down matchmaking to a particular server (optional)", ParamType = typeof(int))]
+	[ParamHelp("ClientArgsFile", "Absolute path to a file containing the client arguments (Engine/Build/AutomationWorkflows/ManyBotClientsDefault.txt by default)", ParamType = typeof(FileReference))]
+	[ParamHelp("ClientLogDir", "Absolute path to the directory with client log files (relative to the exe by default)", ParamType = typeof(FileReference))]
+	[ParamHelp("FirstClientNumber", "The number of the first LoadBot client (0 by default)", ParamType = typeof(int))]
+	[ParamHelp("GridLayout", "If clients aren't nullrhi, lay them out in 320x240 fashion, defaults to true", ParamType = typeof(bool))]
+	[ParamHelp("SleepTimeBetweenLaunches", "How long to sleep between running clients in milliseconds (could prevent race conditions), 100 by default", ParamType = typeof(int))]
 	[ParamHelp("MaxRunAttemptsPerClient", "Maximum number of attempts to run a client which crashes or fails to connect to the server, defaults to 3", ParamType = typeof(int))]
 	[ParamHelp("ClientSessionCompleted", "Log message indicating that a client has completed a game session and may be terminated", ParamType = typeof(string))]
 	[ParamHelp("ClientFailed", "Log message indicating that a client failed to connect to the server", ParamType = typeof(string))]
@@ -27,13 +33,16 @@ namespace MultiClientLauncher.Automation
 		private bool CancelClientProcesses = false;
 		
 		private string ClientExe;
-		private string ClientLog;
+		private string ClientLogDir;
 		private string ClientArgs;
+		private string ClientLogFilenameGuess;
 
 		private int FirstClientNumber;
 		private int ClientCount;
+		private int BuildIdOverride;
+		private bool GridLayout = true;
 
-		private int SleepTimeBetweenLaunches = -1;
+		private int SleepTimeBetweenLaunches = 100;
 		private int MaxRunAttemptsPerClient = 3;
 
 		private const int SleepTimeBetweenChecksForClientRelaunches = 1000;
@@ -42,19 +51,61 @@ namespace MultiClientLauncher.Automation
 		
 		private void ParseCommandLine()
 		{
-			ClientExe = ParseRequiredFileReferenceParam("ClientExe").ToString();
-			ClientLog = ParseRequiredFileReferenceParam("ClientLogFile").ToString();
-			ClientArgs = FileReference.ReadAllText(ParseRequiredFileReferenceParam("ClientArgsFile"));
-			
-			FirstClientNumber = int.Parse(ParseRequiredStringParam("FirstClientNumber"));
+			FileReference ClientExeFile = ParseRequiredFileReferenceParam("ClientExe");
+			ClientExe = ClientExeFile.ToString();
 			ClientCount = int.Parse(ParseRequiredStringParam("ClientCount"));
+			BuildIdOverride = ParseParamInt("BuildIdOverride", -1);
+
+			// guess the log filename from the binary, e.g. FooClient-Linux-Shipping -> FooGame
+			ClientLogFilenameGuess = ClientExeFile.GetFileNameWithoutAnyExtensions();
+			if (ClientLogFilenameGuess.Contains("-"))
+			{
+				ClientLogFilenameGuess = ClientLogFilenameGuess.Split('-')[0];
+			}
+			if (ClientLogFilenameGuess.Contains("Client"))
+			{
+				ClientLogFilenameGuess = ClientLogFilenameGuess.Replace("Client", "Game");
+			}
+
+			// file comm sucks and is unreliable, but better than nothing
+			ClientLogDir = ParseParamValue("ClientLogDir", "");
+			if (string.IsNullOrEmpty(ClientLogDir))
+			{
+				// figure out from ClientExe path
+				ClientLogDir = Utils.CollapseRelativeDirectories(CommandUtils.CombinePaths(ClientExeFile.Directory.ToString(), "../../Saved/Logs"));
+			}
+
+			FileReference ClientArgsFileRef = new FileReference(ParseParamValue("ClientArgsFile", GetDefaultArgsFile()));
+			if (!FileReference.Exists(ClientArgsFileRef))
+			{
+				throw new BuildException("ClientArgs file {0} does not exist (override with -ClientArgsFile=...)", ClientArgsFileRef);
+			}
+			ClientArgs = FileReference.ReadAllText(ClientArgsFileRef);
+			if (BuildIdOverride != -1)
+			{
+				ClientArgs += string.Format(" -buildidoverride={0} ", BuildIdOverride);
+			}
+
+			FirstClientNumber = int.Parse(ParseParamValue("FirstClientNumber", "1"));
 
 			SleepTimeBetweenLaunches = ParseParamInt("SleepTimeBetweenLaunches", -1);
 			MaxRunAttemptsPerClient = ParseParamInt("MaxRunAttemptsPerClient", 3);
-			
+			GridLayout = ParseParamBool("GridLayout", true);
+
+			// disable grid layout for nullrhi
+			if (ClientArgs.Contains("-nullrhi"))
+			{
+				GridLayout = false;
+			}
+
 			InitializeClientLogIndicators();
 		}
 
+		protected virtual string GetDefaultArgsFile()
+		{
+			return "Engine/Build/AutomationWorkflows/ManyBotClientsDefault.txt";
+		}
+
 		// Derived commands may hardcode these log indicators
 		protected virtual void InitializeClientLogIndicators()
 		{
@@ -75,59 +126,92 @@ namespace MultiClientLauncher.Automation
 				CancelClientProcesses = true;
 				KillProcesses(ClientProcesses);
 			};
-			
-			// Run all client processes
-			Console.WriteLine("Spawning clients");
-			for (int i = 0; i < ClientCount; i++)
+
+			// Delete all previous log files
+			Console.WriteLine("Deleting all existing log files in the log directory {0}", ClientLogDir);			
+			string[] LogFiles = Directory.GetFiles(ClientLogDir);
+			foreach (string Filename in LogFiles)
 			{
-				if (CancelClientProcesses)
+				if (Filename.EndsWith(".log"))
 				{
-					break;
+					File.Delete(Filename);
 				}
-				
-				Console.WriteLine("Spawning client {0}...", i);
+			}
 
-				string ClientNumber = (FirstClientNumber + i).ToString();
-				string CurrentClientArgs = ClientArgs.Replace("#REPLACE_CLIENT_ID#", ClientNumber);
-				
-				string CurrentClientLog = ClientLog;
-				if (i > 0)
+			try
+			{
+				// Run all client processes
+				Console.WriteLine("Spawning clients.");
+				for (int ClientIdx = 0; ClientIdx < ClientCount; ++ClientIdx)
 				{
-					CurrentClientLog += "_" + (i + 1);
-				}
-				CurrentClientLog += ".log";
-				
-				ClientProcesses.Add(SpawnClientProcess(ClientExe, CurrentClientLog, CurrentClientArgs));
+					if (CancelClientProcesses)
+					{
+						break;
+					}
 
-				if (SleepTimeBetweenLaunches != -1)
-				{
-					Console.WriteLine("Sleeping for {0}ms...", SleepTimeBetweenLaunches);
-					Thread.Sleep(SleepTimeBetweenLaunches);
+					Console.WriteLine("Spawning client {0}...", ClientIdx);
+
+					string ClientNumber = (FirstClientNumber + ClientIdx).ToString("00.##");
+					string CurrentClientArgs = ClientArgs.Replace("#REPLACED_WITH_TWO_DIGIT_CLIENT_ID#", ClientNumber);
+
+					string CurrentClientLog = CommandUtils.CombinePaths(ClientLogDir, ClientLogFilenameGuess);
+					if (ClientIdx > 0)
+					{
+						CurrentClientLog += "_" + (ClientIdx + 1);
+					}
+					CurrentClientLog += ".log";
+
+					CurrentClientArgs += string.Format(" -abslog={0} ", CurrentClientLog);
+
+					if (GridLayout)
+					{
+						// assume 4k screen, which can fit 90 (10x9) clients running in 384x240. Note - not trying 320x240 as these days client will not use 4:3 aspect ratio
+						const int ResX = 384;
+						const int ResY = 240;
+						const int ClientsPerRow = 3840 / ResX;
+						int WinY = ResY * (ClientIdx / ClientsPerRow);
+						int WinX = ResX * (ClientIdx % ClientsPerRow);
+						CurrentClientArgs += string.Format(" -WinX={0} -WinY={1} -ResX={2} -ResY={3} ", WinX, WinY, ResX, ResY);
+					}
+
+					System.Console.WriteLine("Args: {0}", CurrentClientArgs);
+					ClientProcesses.Add(SpawnClientProcess(ClientExe, CurrentClientLog, CurrentClientArgs));
+
+					if (SleepTimeBetweenLaunches != -1)
+					{
+						Console.WriteLine("Sleeping for {0}ms...", SleepTimeBetweenLaunches);
+						Thread.Sleep(SleepTimeBetweenLaunches);
+					}
 				}
-			}
 
-			while (!CancelClientProcesses && ClientProcesses.Count > 0)
-			{
-				// Iterate through clients in reverse order for safe removal
-				for (int i = ClientProcesses.Count - 1; i >= 0; i--)
+				while (!CancelClientProcesses && ClientProcesses.Count > 0)
 				{
-					if (ClientProcesses[i].Stopped())
+					// Iterate through clients in reverse order for safe removal
+					for (int i = ClientProcesses.Count - 1; i >= 0; i--)
 					{
-						bool OutOfTries = !ClientProcesses[i].Start();
-						
-						if (OutOfTries)
+						if (ClientProcesses[i].Stopped())
 						{
-							ClientProcesses.RemoveAt(i);
+							bool OutOfTries = !ClientProcesses[i].Start();
+
+							if (OutOfTries)
+							{
+								ClientProcesses.RemoveAt(i);
+							}
 						}
 					}
+
+					// Todo: Consider replacing sleeps with waiting for LogSkimmer threads to finish reading
+					// Todo:   - That may still happen very fast and result in more busywaiting
+					Thread.Sleep(SleepTimeBetweenChecksForClientRelaunches);
 				}
-				
-				// Todo: Consider replacing sleeps with waiting for LogSkimmer threads to finish reading
-				// Todo:   - That may still happen very fast and result in more busywaiting
-				Thread.Sleep(SleepTimeBetweenChecksForClientRelaunches);
-			}
 
-			return ExitCode.Success;
+				return ExitCode.Success;
+			}
+			catch(Exception)
+			{
+				KillProcesses(ClientProcesses);
+				return ExitCode.Error_Unknown;
+			}
 		}
 		
 		private ClientProcess SpawnClientProcess(string ExeFilename, string ClientLogFilename, string ExeArguments)
@@ -221,8 +305,25 @@ namespace MultiClientLauncher.Automation
 			private void MonitorClient()
 			{
 				// Wait long enough for log files to be created
-				Thread.Sleep(2000);
-				
+				double MaxSecondsToWait = 10;
+				double WaitedSoFar = 0;
+				do
+				{
+					if (File.Exists(LogFilepath))
+					{
+						break;
+					}
+
+					Thread.Sleep(2000);
+					WaitedSoFar += 2;
+
+					if (WaitedSoFar >= MaxSecondsToWait)
+					{
+						throw new BuildException("Log file {0} was not created after {1} seconds (did process crash on start?)", LogFilepath, WaitedSoFar);
+					}
+				}
+				while (true);
+
 				string AllServerOutput = "";
 				using FileStream ProcessLog = File.Open(LogFilepath, FileMode.Open, FileAccess.Read, FileShare.ReadWrite);
 				using StreamReader LogReader = new StreamReader(ProcessLog);
