Use Project Target to find editor even when -build point to a staged build
When -editordir is specified trim the editor path to be relative
Add support for -target command line argument to Gauntlet

#rnx

diff --git a/Engine/Build/Graph/Tasks/BuildAndTestProject.xml b/Engine/Build/Graph/Tasks/BuildAndTestProject.xml
--- a/Engine/Build/Graph/Tasks/BuildAndTestProject.xml
+++ b/Engine/Build/Graph/Tasks/BuildAndTestProject.xml
@@ -591,7 +591,7 @@
 						<Property Name="GauntletReportArgs" Value="$(GauntletReportArgs) -ReportURL=&quot;$(AutomationReportUrl)/$(OutputFragmentPath)/&quot;" If="'$(AutomationReportUrl)' != ''" />
 					</Do>
 					<!-- Run the boot test -->
-					<Command Name="RunUnreal" Arguments="-project=$(TargetProject) -test=&quot;UE.BootTest&quot; -platform=$(TargetPlatform) -configuration=$(TestConfiguration) $(GauntletBuildArg) $(GauntletTestArgs) $(GauntletReportArgs)" />
+					<Command Name="RunUnreal" Arguments="-project=$(TargetProject) -test=&quot;UE.BootTest&quot; -platform=$(TargetPlatform) -configuration=$(TestConfiguration) -target=$(TargetName) $(GauntletBuildArg) $(GauntletTestArgs) $(GauntletReportArgs)" />
 				</Node>
 				<!-- Link the boot test in to the prerequisites and generated nodes if its enabled-->
 				<Property Name="TestPrerequisites" Value="$(TestPrerequisites);$(PlatformEditorNodeName)" />
@@ -615,7 +615,7 @@
 					<Property Name="NodeAnnotations" Value="Workflow=$(AnnotationsTarget)" If="ContainsItem('$(AnnotationsPlatforms)', '$(TargetPlatform)', ';') and ContainsItem('$(AnnotationsTests)', '$(TestName)', ';')"/>
 					<!-- Run the test. No 'Requires' as we run on the same machine locally and use the publish folder otherwise -->
 					<Node Name="$(TestNodeName)" Requires="$(TestPrerequisites)" Annotations="$(NodeAnnotations)">
-						<Command Name="RunUnreal" Arguments="-project=$(TargetProject) -test=&quot;$(TestName)&quot; -platform=$(TargetPlatform) -configuration=$(TestConfiguration) $(GauntletBuildArg) $(GauntletTestArgs) $(GauntletReportArgs)"/>
+						<Command Name="RunUnreal" Arguments="-project=$(TargetProject) -test=&quot;$(TestName)&quot; -platform=$(TargetPlatform) -configuration=$(TestConfiguration) -target=$(TargetName) $(GauntletBuildArg) $(GauntletTestArgs) $(GauntletReportArgs)"/>
 					</Node>
 					<!-- Add it to our list -->
 					<Do If="!ContainsItem('$(TargetPlatformsToSkipTests)', '$(TargetPlatform)', '+')">
diff --git a/Engine/Source/Programs/AutomationTool/Gauntlet/Unreal/Base/Gauntlet.UnrealTestContext.cs b/Engine/Source/Programs/AutomationTool/Gauntlet/Unreal/Base/Gauntlet.UnrealTestContext.cs
--- a/Engine/Source/Programs/AutomationTool/Gauntlet/Unreal/Base/Gauntlet.UnrealTestContext.cs
+++ b/Engine/Source/Programs/AutomationTool/Gauntlet/Unreal/Base/Gauntlet.UnrealTestContext.cs
@@ -202,6 +202,12 @@ namespace Gauntlet
 		[AutoParam("")]
 		public string CustomModuleRoles { get; set; }
 
+		/// <summary>
+		/// Set the target name of the module name to identify a custom project target executable
+		/// </summary>
+		[AutoParam("")]
+		public string Target { get; set; }
+
 		/// <summary>
 		/// Less logging
 		/// </summary>
@@ -382,6 +388,12 @@ namespace Gauntlet
 					}
 				}
 			}
+			if(!string.IsNullOrEmpty(Target) && Target != Project)
+			{
+				bool IsEditor = Build.Equals("Editor", StringComparison.InvariantCultureIgnoreCase) || Globals.Params.ParseParam("editor");
+				UnrealTargetRole Role = IsEditor? UnrealTargetRole.Editor : UnrealTargetRole.Client;
+				UnrealHelpers.AddCustomModuleName(IsEditor? string.Format("{0}Editor", Target): Target, Role);
+			}
 		}
 	}
 	
diff --git a/Engine/Source/Programs/AutomationTool/Gauntlet/Unreal/BuildSource/Gauntlet.UnrealBuildSource.cs b/Engine/Source/Programs/AutomationTool/Gauntlet/Unreal/BuildSource/Gauntlet.UnrealBuildSource.cs
--- a/Engine/Source/Programs/AutomationTool/Gauntlet/Unreal/BuildSource/Gauntlet.UnrealBuildSource.cs
+++ b/Engine/Source/Programs/AutomationTool/Gauntlet/Unreal/BuildSource/Gauntlet.UnrealBuildSource.cs
@@ -648,27 +648,46 @@ namespace Gauntlet
 
 			if (TargetRole.UsesEditor() || TargetRole.IsEditor())
 			{
+				bool HasCustomTarget = UnrealHelpers.CustomModuleToRoles.ContainsValue(UnrealTargetRole.Editor);
+				string EditorTarget = HasCustomTarget ? UnrealHelpers.CustomModuleToRoles.FirstOrDefault(M => M.Value == UnrealTargetRole.Editor).Key : string.Empty;
 				FileSystemReference EditorExe = null;
 
-				if (BuildName == "Editor" && string.IsNullOrEmpty(Globals.Params.ParseValue("EditorDir", null)))
+				if (!HasCustomTarget)
 				{
 					try
 					{
 						EditorExe = ProjectUtils.GetProjectTarget(ProjectPath, UnrealBuildTool.TargetType.Editor, TargetPlatform, TargetConfiguration);
 					}
-					catch(Exception Ex)
+					catch (Exception Ex)
 					{
-						throw new AutomationException("No suitable editor build for {0} configuration.\n{1}", TargetConfiguration, Ex.ToString());
+						string Message = string.Format("The project config is overriding build targets.\n"
+								+ "But no suitable editor build for {0} configuration found from target file.\n"
+								+ "{1}", TargetConfiguration, Ex.Message); 
+
+						if (BuildName.Equals("Editor", StringComparison.OrdinalIgnoreCase))
+						{
+							// An editor is being explicitly requested but no executable could be found from target file
+							// Hightlight the issue in the log and return an empty string to avoid taking an inappropriate editor
+							Log.Warning(Message);
+							return string.Empty;
+						}
+						Log.Info(Message);
 					}
 				}
 
 				if (EditorExe != null)
 				{
 					ExePath = EditorExe.FullName;
+					if (!string.IsNullOrEmpty(Globals.Params.ParseValue("EditorDir", null)))
+					{
+						/// Trim the Editor absolute path from what the target file provided as the editor dir is being overriden
+						/// https://regex101.com/r/7BttxH/1
+						ExePath = Regex.Replace(ExePath, @"(.+?)[/\\]((Engine[/\\])?Binaries[/\\].+)", "$2");
+					}
 				}
 				else
 				{
-					string ExeFileName = "UnrealEditor";
+					string ExeFileName = HasCustomTarget ? EditorTarget : "UnrealEditor";
 					if (TargetConfiguration != UnrealTargetConfiguration.Development)
 					{
 						ExeFileName += string.Format("-{0}-{1}", TargetPlatform.ToString(), TargetConfiguration.ToString());
@@ -676,7 +695,8 @@ namespace Gauntlet
 
 					ExeFileName += Platform.GetExeExtension(TargetPlatform);
 
-					ExePath = string.Format("Engine/Binaries/{0}/{1}", BuildHostPlatform.Current.Platform, ExeFileName);
+					string BasePath = HasCustomTarget ? Globals.Params.ParseValue("EditorDir", ProjectPath.Directory.FullName) : "Engine";
+					ExePath = string.Format("{0}/Binaries/{1}/{2}", BasePath, BuildHostPlatform.Current.Platform, ExeFileName);
 				}
 			}
 			else
