- Disallow BaseSourceControlSettings.ini from being staged

diff --git a/Engine/Source/Programs/AutomationTool/AutomationUtils/DeploymentContext.cs b/Engine/Source/Programs/AutomationTool/AutomationUtils/DeploymentContext.cs
--- a/Engine/Source/Programs/AutomationTool/AutomationUtils/DeploymentContext.cs
+++ b/Engine/Source/Programs/AutomationTool/AutomationUtils/DeploymentContext.cs
@@ -643,7 +643,8 @@ public class DeploymentContext //: ProjectParams
 			"editorsettings.ini",
 			"editorusersettings.ini",
 			"lightmass.ini",
-			"pakfilerules.ini"
+			"pakfilerules.ini",
+			"sourcecontrolsettings.ini"
 		};
 
 		// If we were configured to use manifests across the whole project, then this platform should use manifests.
