Using new Gauntlet command line parameter RequiredSettings to check device settings (for example, -RequiredSettings=LargeMem). The device is skipped if it doesn't have a required configuration. This is a new Parameter that is not being used by any gauntlet test.

diff --git a/Engine/Source/Programs/AutomationTool/Gauntlet/Framework/Base/Gauntlet.TargetDevice.cs b/Engine/Source/Programs/AutomationTool/Gauntlet/Framework/Base/Gauntlet.TargetDevice.cs
--- a/Engine/Source/Programs/AutomationTool/Gauntlet/Framework/Base/Gauntlet.TargetDevice.cs
+++ b/Engine/Source/Programs/AutomationTool/Gauntlet/Framework/Base/Gauntlet.TargetDevice.cs
@@ -89,6 +89,8 @@ namespace Gauntlet
 
 		IAppInstance Run(IAppInstall App);
 
+		bool CheckRequiredSettings(List<string> RequiredSettingsList) { return true; }
+			
 		string GetPackagedExecutableLocation() { return null; }
 
 		/// <summary>
diff --git a/Engine/Source/Programs/AutomationTool/Gauntlet/Framework/Devices/Gauntlet.DevicePool.cs b/Engine/Source/Programs/AutomationTool/Gauntlet/Framework/Devices/Gauntlet.DevicePool.cs
--- a/Engine/Source/Programs/AutomationTool/Gauntlet/Framework/Devices/Gauntlet.DevicePool.cs
+++ b/Engine/Source/Programs/AutomationTool/Gauntlet/Framework/Devices/Gauntlet.DevicePool.cs
@@ -1211,6 +1211,17 @@ namespace Gauntlet
 					NewDevice = Factory.CreateDevice(Def.Address, ClientTempDir, Def.DeviceData);
 				}
 
+				List<string> RequiredSettings = Globals.Params.ParseValues("RequiredSettings");
+				if (RequiredSettings.Count > 0)
+				{
+					Log.Info("Checking required settings for device {0}...", NewDevice.Name);
+					if (!NewDevice.CheckRequiredSettings(RequiredSettings))
+					{
+						Log.Warning("Device settings doesn't meet the requirements, skipping device.");
+						return null;
+					}
+					Log.Info("Device has all required settings, keeping device.");
+				}
 				if (NewDevice.IsAvailable == false)
 				{
 					Log.Info(KnownLogEvents.Gauntlet_DeviceEvent, "Assigned device {Name} reports unavailable. Requesting a forced disconnect", NewDevice);
