UAT - removed warning about missing encryption key for on demand PAK rules

#rnx

diff --git a/Engine/Source/Programs/AutomationTool/Scripts/CopyBuildToStagingDirectory.Automation.cs b/Engine/Source/Programs/AutomationTool/Scripts/CopyBuildToStagingDirectory.Automation.cs
--- a/Engine/Source/Programs/AutomationTool/Scripts/CopyBuildToStagingDirectory.Automation.cs
+++ b/Engine/Source/Programs/AutomationTool/Scripts/CopyBuildToStagingDirectory.Automation.cs
@@ -4151,7 +4151,6 @@ namespace AutomationScripts
 							if (string.IsNullOrEmpty(Rule.EncryptionKeyGuid))
 							{
 								bIsEncryptionKeyValid = false;
-								Logger.LogWarning("No encryption key specified for content-on-demand PAK rule '{Arg0}'", Rule.Name);
 							}
 							else
 							{
