Add more precise client-side test start tracing

#rnx

diff --git a/Engine/Source/Programs/AutomationTool/Gauntlet/Unreal/Automation/UE.Automation.cs b/Engine/Source/Programs/AutomationTool/Gauntlet/Unreal/Automation/UE.Automation.cs
--- a/Engine/Source/Programs/AutomationTool/Gauntlet/Unreal/Automation/UE.Automation.cs
+++ b/Engine/Source/Programs/AutomationTool/Gauntlet/Unreal/Automation/UE.Automation.cs
@@ -6,8 +6,6 @@ using System;
 using System.Collections.Generic;
 using System.IO;
 using System.Linq;
-using System.Text;
-using System.Threading.Tasks;
 using UnrealBuildTool;
 using EpicGames.Core;
 using Log = Gauntlet.Log;
@@ -659,7 +657,7 @@ namespace UE
 			bool HasTimeout = RoleResults != null && RoleResults.Where(R => R.ProcessResult == UnrealProcessResult.TimeOut).Any();
 			string HordeArtifactPath = GetConfiguration().HordeArtifactPath;
 			var MainRole = GetConfiguration().GetMainRequiredRole();
-			string LastTestWithCriticalFailure = string.Empty;
+			UnrealAutomatedTestResult LastTestInProgress = null;
 			if (JsonTestPassResults.InProcess > 0)
 			{
 				// The test pass did not run completely
@@ -668,19 +666,20 @@ namespace UE
 				if (RoleResults != null)
 				{
 					UnrealLog.CallstackMessage FatalError = null;
+					UnrealRoleResult FatalErrorRoleResult = null;
 					foreach (var Result in RoleResults)
 					{
 						if (Result.LogSummary.FatalError != null)
 						{
 							FatalError = Result.LogSummary.FatalError;
+							FatalErrorRoleResult = Result;
 							break;
 						}
 					}
-					var Test = JsonTestPassResults.Tests.FirstOrDefault((T => T.State == TestStateType.InProcess));
-					if (!String.IsNullOrEmpty(Test.TestDisplayName))
+					LastTestInProgress = JsonTestPassResults.Tests.FirstOrDefault((T => T.State == TestStateType.InProcess));
+					if (!String.IsNullOrEmpty(LastTestInProgress.TestDisplayName))
 					{
-						LastTestWithCriticalFailure = Test.FullTestPath;
-						string ErrorMessage;
+						string ErrorMessage = null;
 						if (HasTimeout)
 						{
 							ErrorMessage = String.Format("Session reached timeout after {0} seconds.", MaxDuration);
@@ -689,22 +688,98 @@ namespace UE
 						{
 							if (FatalError != null)
 							{
-								ErrorMessage = FatalError.FormatForLog();
+								var LogEntry = FatalErrorRoleResult.LogSummary.LogEntries.LastOrDefault(Entry => Entry.Category == "AutomationTestStateTrace", null);
+
+								if (LogEntry != null)
+								{
+									string LogEntryPrefix = "Test is about to start. Name={";
+									string LogEntrySuffix = "}";
+									if (LogEntry.Message.StartsWith(LogEntryPrefix) && LogEntry.Message.EndsWith(LogEntrySuffix))
+									{
+										string StartedTestFullName = LogEntry.Message.Substring(
+											LogEntryPrefix.Length,
+											LogEntry.Message.Length - LogEntryPrefix.Length - LogEntrySuffix.Length);
+
+										if (LastTestInProgress.FullTestPath == StartedTestFullName)
+										{
+											LastTestInProgress.State = TestStateType.Fail;
+											JsonTestPassResults.Failed++;
+											JsonTestPassResults.InProcess--;
+											ErrorMessage = FatalError.FormatForLog();
+										}
+										else
+										{
+											LastTestInProgress.State = TestStateType.NotRun;
+											JsonTestPassResults.NotRun++;
+											JsonTestPassResults.InProcess--;
+
+											// We have to find another test that has been crashed and set it as the last test with critical failure
+											LastTestInProgress = JsonTestPassResults.Tests.FirstOrDefault((T => T.FullTestPath == StartedTestFullName));
+											if (LastTestInProgress == null)
+											{
+												Log.Warning("Failed to find test for the test state log entry '{0}'", LogEntry.Message);
+											}
+											else
+											{
+												if (LastTestInProgress.State != TestStateType.Fail)
+												{
+													switch (LastTestInProgress.State)
+													{
+														case TestStateType.InProcess:
+															JsonTestPassResults.InProcess--;
+															break;
+
+														case TestStateType.NotRun:
+															Log.Warning("The current state for the test '{0}' is NotRun (it will be changed to Fail). Log entry is '{1}'. That state is inconsistent with the state that is expected by design", StartedTestFullName, LogEntry.Message);
+															JsonTestPassResults.NotRun--;
+															break;
+
+														case TestStateType.Success:
+															JsonTestPassResults.Succeeded--;
+															break;
+
+														default:
+															break;
+													};
+													LastTestInProgress.State = TestStateType.Fail;
+													JsonTestPassResults.Failed++;
+												}
+												ErrorMessage = FatalError.FormatForLog();
+											}
+										}
+									}
+									else
+									{
+										// Handling the test state without using AutomationTestStateTrace log channel
+										string WarningMessage = string.Format(
+											"Please have a look at this case. It is possible that the crash is not caused by the test execution. The last test's state trace log entry is '{0}'",
+											LogEntry.Message);
+										Log.Warning(WarningMessage + "\nCallstack:\n" + FatalError.FormatForLog());
+										ErrorMessage = FatalError.FormatForLog() + "\n" + WarningMessage;
+									}
+								}
+								else
+								{
+									ErrorMessage = FatalError.FormatForLog();
+								}
 							}
 							else
 							{
 								ErrorMessage = "No callstack found in the log.";
 							}
 						}
-						Test.AddError(ErrorMessage, !HasTimeout);
+						if ((LastTestInProgress != null) && !String.IsNullOrEmpty(ErrorMessage))
+						{
+							LastTestInProgress.AddError(ErrorMessage, !HasTimeout);
+						}
 						if (!CanRetry() || JsonTestPassResults.NotRun == 0)
 						{
 							// Setting the test as fail because no retry will be done anymore.
 							// The InProcess state won't be used for pass resume
-							Test.State = TestStateType.Fail;
+							LastTestInProgress.State = TestStateType.Fail;
 							if (!CanRetry())
 							{
-								Test.AddWarning(string.Format("Session reached maximum of retries({0}) to resume on critical failure!", Retries));
+								LastTestInProgress.AddWarning(string.Format("Session reached maximum of retries({0}) to resume on critical failure!", Retries));
 							}
 						}
 						JsonTestPassResults.WriteToJson();
@@ -744,9 +819,9 @@ namespace UE
 						Log.Warning(KnownLogEvents.Gauntlet_TestEvent, "Reach maximum of retries({Count}) to resume on critical failure!", Retries);
 						// Adding a note to the report about why the not-run are not going to be run
 						string Message = string.Format("Session reached maximum of retries({0}) to resume on critical failure!", Retries);
-						if (!string.IsNullOrEmpty(LastTestWithCriticalFailure))
+						if (LastTestInProgress != null && !string.IsNullOrEmpty(LastTestInProgress.FullTestPath))
 						{
-							Message += string.Format(" \nLast critical failure was caught on {0}", LastTestWithCriticalFailure);
+							Message += string.Format(" \nLast critical failure was caught on {0}", LastTestInProgress.FullTestPath);
 						}
 						var NotRunTests = JsonTestPassResults.Tests.Where((T => T.State == TestStateType.NotRun));
 						foreach (var Test in NotRunTests)
diff --git a/Engine/Source/Runtime/Core/Private/Misc/AutomationTest.cpp b/Engine/Source/Runtime/Core/Private/Misc/AutomationTest.cpp
--- a/Engine/Source/Runtime/Core/Private/Misc/AutomationTest.cpp
+++ b/Engine/Source/Runtime/Core/Private/Misc/AutomationTest.cpp
@@ -21,6 +21,7 @@
 #include "Modules/ModuleManager.h"
 
 DEFINE_LOG_CATEGORY(LogLatentCommands)
+DEFINE_LOG_CATEGORY_STATIC(LogAutomationTestStateTrace, Log, All);
 DEFINE_LOG_CATEGORY_STATIC(LogAutomationTest, Warning, All);
 
 namespace AutomationTest
@@ -43,6 +44,12 @@ namespace AutomationTest
 		bLogBPTestMetadata,
 		TEXT("Whether to output blueprint functional test metadata to the log when test is running"));
 
+	static bool bLogTestStateTrace = false;
+	static FAutoConsoleVariableRef CVarAutomationLogTestStateTrace(
+		TEXT("Automation.LogTestStateTrace"),
+		bLogTestStateTrace,
+		TEXT("Whether to enable or disable logging of test state trace"));
+
 	// The method prepares the filename and LineNumber to be placed in the form that could be extracted by SAutomationWindow widget if it is additionally eclosed into []
 	// The result format is filename(line)
 	static FString CreateFileLineDescription(const FString& Filename, const int32 LineNumber)
@@ -876,6 +883,10 @@ void FAutomationTestFramework::InternalStartTest( const FString& InTestToRun )
 		uint32 NonSmokeTestFlags = (EAutomationTestFlags::FilterMask & (~EAutomationTestFlags::SmokeFilter));
 		if (RequestedTestFilter & NonSmokeTestFlags)
 		{
+			if (AutomationTest::bLogTestStateTrace)
+			{
+				UE_LOG(LogAutomationTestStateTrace, Log, TEXT("Test is about to start. Name={%s}"), *CurrentTest->GetTestFullName());
+			}
 			UE_LOG(LogAutomationTest, Log, TEXT("%s %s is starting at %f"), *CurrentTest->GetBeautifiedTestName(), *Parameters, StartTime);
 		}
 
@@ -918,6 +929,10 @@ bool FAutomationTestFramework::InternalStopTest(FAutomationTestExecutionInfo& Ou
 	if (RequestedTestFilter & NonSmokeTestFlags)
 	{
 		UE_LOG(LogAutomationTest, Log, TEXT("%s %s ran in %f"), *CurrentTest->GetBeautifiedTestName(), *Parameters, TimeForTest);
+		if (AutomationTest::bLogTestStateTrace)
+		{
+			UE_LOG(LogAutomationTestStateTrace, Log, TEXT("Test has stopped execution. Name={%s}"), *CurrentTest->GetTestFullName());
+		}
 	}
 
 	// Fill out the provided execution info with the info from the test
