Make sure the web server process can be closed properly when run web tests on automation tool, to fix the pipeline not closed warning;
Fix the low level report parser error, when there are sections in test cases.

diff --git a/Engine/Source/Programs/AutomationTool/LowLevelTests/Tests/Gauntlet.LowLevelTestNode.cs b/Engine/Source/Programs/AutomationTool/LowLevelTests/Tests/Gauntlet.LowLevelTestNode.cs
--- a/Engine/Source/Programs/AutomationTool/LowLevelTests/Tests/Gauntlet.LowLevelTestNode.cs
+++ b/Engine/Source/Programs/AutomationTool/LowLevelTests/Tests/Gauntlet.LowLevelTestNode.cs
@@ -279,7 +279,7 @@ namespace LowLevelTests
 						else
 						{
 							LowLevelTestResult = TestResult.Failed;
-							ExitReason = "Tests failed";
+							ExitReason = "Tests failed according to console report";
 						}
 					}
 					else if (ReportType == "xml")
@@ -293,7 +293,7 @@ namespace LowLevelTests
 						else
 						{
 							LowLevelTestResult = TestResult.Failed;
-							ExitReason = "Tests failed";
+							ExitReason = "Tests failed according to xml report";
 						}
 					}
 				}
diff --git a/Engine/Source/Programs/AutomationTool/LowLevelTests/Tests/LowLevelTests.WebTestsLowLevelTestsExtension.cs b/Engine/Source/Programs/AutomationTool/LowLevelTests/Tests/LowLevelTests.WebTestsLowLevelTestsExtension.cs
--- a/Engine/Source/Programs/AutomationTool/LowLevelTests/Tests/LowLevelTests.WebTestsLowLevelTestsExtension.cs
+++ b/Engine/Source/Programs/AutomationTool/LowLevelTests/Tests/LowLevelTests.WebTestsLowLevelTestsExtension.cs
@@ -1,24 +1,18 @@
 // Copyright Epic Games, Inc. All Rights Reserved.
 
-using JetBrains.Annotations;
 using System;
 using System.Diagnostics;
 using System.IO;
-using System.Text;
 using EpicGames.Core;
-using AutomationTool;
 using UnrealBuildTool;
 using Gauntlet;
 using UnrealBuildBase;
-using EpicGames.Horde.Storage;
-using System.Net.Sockets;
 
 namespace LowLevelTests
 {
 	public class WebTestsLowLevelTestsExtension : ILowLevelTestsExtension
 	{
 		private Process ServerProcess;
-		private StringBuilder logBuilder = new StringBuilder();
 
 		public bool IsSupported(UnrealTargetPlatform InPlatform, string InTestApp)
 		{
@@ -30,88 +24,62 @@ namespace LowLevelTests
 			return string.Format("--web_server_ip{0}", UnrealHelpers.GetHostIpAddress());
 		}
 
-		private bool IsServerPortOpen(string ipAddress, int port)
+		public void PreRunTests()
 		{
-			using (TcpClient client = new TcpClient())
-			{
-				try
-				{
-					client.Connect(ipAddress, port);
-					return true;
-				}
-				catch
-				{
-					return false;
-				}
-			}
+			InstallWebServer();
+			RunWebServer();
 		}
 
-		public void PreRunTests()
+		private string WebTestsServerDir()
+		{
+			return Path.Combine(Unreal.EngineDirectory.FullName, "Source", "Programs", "WebTestsServer");
+		}
+
+		private void InstallWebServer()
 		{
 			ProcessStartInfo StartInfo = new ProcessStartInfo();
-			StartInfo.WorkingDirectory = Path.Combine(Unreal.EngineDirectory.FullName, "Source", "Programs", "WebTestsServer");
+			StartInfo.WorkingDirectory = WebTestsServerDir();
 			StartInfo.FileName = RuntimePlatform.IsWindows ? "cmd.exe" : "/bind/sh";
-			StartInfo.Arguments = RuntimePlatform.IsWindows ? "/c runserver.bat" : "-c 'runserver.sh'";
+			StartInfo.Arguments = RuntimePlatform.IsWindows ? "/c createenv.bat" : "-c 'createenv.sh'";
 			StartInfo.UseShellExecute = false;
 			StartInfo.CreateNoWindow = true;
-			StartInfo.RedirectStandardInput = true;
-			StartInfo.RedirectStandardOutput = true;
-			StartInfo.RedirectStandardError = true;
-
-			ServerProcess = new Process();
-			ServerProcess.StartInfo = StartInfo;
 
-			ServerProcess.OutputDataReceived += (sender, e) =>
-			{
-				if (!string.IsNullOrEmpty(e.Data))
-				{
-					logBuilder.AppendLine(e.Data);
-				}
-			};
+			Process InstallProcess = new Process();
+			InstallProcess.StartInfo = StartInfo;
+			InstallProcess.Start();
+			InstallProcess.WaitForExit();
 
-			ServerProcess.ErrorDataReceived += (sender, e) =>
-			{
-				if (!string.IsNullOrEmpty(e.Data))
-				{
-					logBuilder.AppendLine(e.Data);
-				}
-			};
-
-			ServerProcess.Start();
-			ServerProcess.BeginOutputReadLine();
-			ServerProcess.BeginErrorReadLine();
+			Console.WriteLine("Requirements installed.");
+		}
 
-			Stopwatch sw = new Stopwatch();
-			sw.Start();
+		private void RunWebServer()
+		{
+			string WorkingDir = WebTestsServerDir();
+			string PythonFile = Path.Combine(WorkingDir, "env", "Scripts", RuntimePlatform.IsWindows ? "python.exe" : "python");
 
-			while (!IsServerPortOpen(UnrealHelpers.GetHostIpAddress(), 8000))
-			{
-				if (sw.ElapsedMilliseconds > 60000)
-				{
-					sw.Stop();
-					throw new TimeoutException("Server port did not open within the specified time.");
-				}
-				System.Threading.Thread.Sleep(1000);
-			}
+			ProcessStartInfo StartInfo = new ProcessStartInfo();
+			StartInfo.WorkingDirectory = WorkingDir;
+			StartInfo.FileName = PythonFile;
+			StartInfo.WindowStyle = ProcessWindowStyle.Normal;
+			StartInfo.Arguments = "manage.py runserver 0.0.0.0:8000";
+			StartInfo.UseShellExecute = true;
+			StartInfo.CreateNoWindow = false;
 
-			sw.Stop();
+			ServerProcess = new Process();
+			ServerProcess.StartInfo = StartInfo;
+			ServerProcess.Start();
 
-			Console.WriteLine("server port is now open.");
+			Console.WriteLine("Web server is now running.");
 		}
 
 		public void PostRunTests()
 		{
 			if (ServerProcess != null)
 			{
-				ServerProcess.Kill();
+				ServerProcess.CloseMainWindow();
 				ServerProcess = null;
 
-				string serverOutput = logBuilder.ToString();
-				Console.WriteLine("Server Output:");
-				Console.WriteLine(serverOutput);
-
-				string serverOutputLog = Path.Combine(Path.GetFullPath(Globals.LogDir).Replace("GauntletTemp", "Engine\\Programs\\AutomationTool\\Saved"), "ServerOutput.log");
-				File.WriteAllText(serverOutputLog, serverOutput);
+				Console.WriteLine("Web server killed.");
 			}
 		}
 	}
diff --git a/Engine/Source/Programs/AutomationTool/LowLevelTests/Utility/LowLevelTests.ReportParser.cs b/Engine/Source/Programs/AutomationTool/LowLevelTests/Utility/LowLevelTests.ReportParser.cs
--- a/Engine/Source/Programs/AutomationTool/LowLevelTests/Utility/LowLevelTests.ReportParser.cs
+++ b/Engine/Source/Programs/AutomationTool/LowLevelTests/Utility/LowLevelTests.ReportParser.cs
@@ -37,13 +37,14 @@ namespace LowLevelTests
 			int NrOverallResultsSuccesses = -1;
 			int NrOverallResultsCasesFailures = -1;
 			int NrOverallResultsCasesSuccesses = -1;
-			XElement OverallResults = ReportDoc.Descendants("OverallResults").FirstOrDefault();
+			// OverallResults also exists in each SECTION of every test case, we need to use the last one in the end of the xml report
+			XElement OverallResults = ReportDoc.Descendants("OverallResults").LastOrDefault();
 			if (OverallResults != null)
 			{
 				NrOverallResultsFailures = int.Parse(OverallResults.Attribute("failures").Value);
 				NrOverallResultsSuccesses = int.Parse(OverallResults.Attribute("successes").Value);
 			}
-			XElement OverallResultsCases = ReportDoc.Descendants("OverallResultsCases").FirstOrDefault();
+			XElement OverallResultsCases = ReportDoc.Descendants("OverallResultsCases").LastOrDefault();
 			if (OverallResultsCases != null)
 			{
 				NrOverallResultsCasesFailures = int.Parse(OverallResultsCases.Attribute("failures").Value);
diff --git a/Engine/Source/Programs/WebTestsServer/createenv.bat b/Engine/Source/Programs/WebTestsServer/createenv.bat
new file mode 100755
--- /dev/null
+++ b/Engine/Source/Programs/WebTestsServer/createenv.bat
@@ -0,0 +1,14 @@
+@echo off
+
+if exist env (
+	call .\env\Scripts\activate.bat
+) else (
+	..\..\..\Binaries\ThirdParty\Python3\Win64\python.exe -m venv env
+	call .\env\Scripts\activate.bat
+	python -m pip install --upgrade pip
+	pip install -r requirements.txt
+	if %ERRORLEVEL% NEQ 0 (
+		rmdir /s /q env
+		exit %ERRORLEVEL%
+	)
+)
diff --git a/Engine/Source/Programs/WebTestsServer/createenv.sh b/Engine/Source/Programs/WebTestsServer/createenv.sh
new file mode 100644
--- /dev/null
+++ b/Engine/Source/Programs/WebTestsServer/createenv.sh
@@ -0,0 +1,12 @@
+if [ -d "env" ]; then
+	source ./env/Scripts/activate.sh
+else
+	run ../../../Binaries/ThirdParty/Python3/Mac/python3 -m venv env
+	source ./env/Scripts/activate.sh
+	run python -m  pip install --upgrade pip
+	run pip install -r requirements.txt
+	if [ $? -ne 0 ]; then
+		rm -rf env
+		exit $?
+	fi
+fi
diff --git a/Engine/Source/Programs/WebTestsServer/runserver.bat b/Engine/Source/Programs/WebTestsServer/runserver.bat
--- a/Engine/Source/Programs/WebTestsServer/runserver.bat
+++ b/Engine/Source/Programs/WebTestsServer/runserver.bat
@@ -1,16 +1,5 @@
 @echo off
 
-if exist env (
-	call .\env\Scripts\activate.bat
-) else (
-	..\..\..\Binaries\ThirdParty\Python3\Win64\python.exe -m venv env
-	call .\env\Scripts\activate.bat
-	python -m pip install --upgrade pip
-	pip install -r requirements.txt
-	if %ERRORLEVEL% NEQ 0 (
-		rmdir /s /q env
-		exit %ERRORLEVEL%
-	)
-)
+call createenv.bat
 
 python manage.py runserver 0.0.0.0:8000
diff --git a/Engine/Source/Programs/WebTestsServer/runserver.sh b/Engine/Source/Programs/WebTestsServer/runserver.sh
--- a/Engine/Source/Programs/WebTestsServer/runserver.sh
+++ b/Engine/Source/Programs/WebTestsServer/runserver.sh
@@ -1,14 +1,2 @@
-if [ -d "env" ]; then
-	source ./env/Scripts/activate.sh
-else
-	run ../../../Binaries/ThirdParty/Python3/Mac/python3 -m venv env
-	source ./env/Scripts/activate.sh
-	run python -m  pip install --upgrade pip
-	run pip install -r requirements.txt
-	if [ $? -ne 0 ]; then
-		rm -rf env
-		exit $?
-	fi
-fi
-
+sh createenv.sh
 cmd python manage.py runserver 0.0.0.0:8000
