MultiClientLauncher: added ability to pass -nullrhi and disable timeouts.

- Also a better command line for clients.

#rnx

diff --git a/Engine/Source/Programs/AutomationTool/Scripts/MultiClientLauncher.Automation.cs b/Engine/Source/Programs/AutomationTool/Scripts/MultiClientLauncher.Automation.cs
--- a/Engine/Source/Programs/AutomationTool/Scripts/MultiClientLauncher.Automation.cs
+++ b/Engine/Source/Programs/AutomationTool/Scripts/MultiClientLauncher.Automation.cs
@@ -22,7 +22,9 @@ namespace MultiClientLauncher.Automation
 	[ParamHelp("ClientArgsFile", "Absolute path to a file containing the client arguments (Engine/Build/AutomationWorkflows/ManyBotClientsDefault.txt by default)", ParamType = typeof(FileReference))]
 	[ParamHelp("ClientLogDir", "Absolute path to the directory with client log files (relative to the exe by default)", ParamType = typeof(FileReference))]
 	[ParamHelp("FirstClientNumber", "The number of the first LoadBot client (0 by default)", ParamType = typeof(int))]
+	[ParamHelp("NullRHI", "Pass -nullrhi to the clients, defaults to false", ParamType = typeof(bool))]
 	[ParamHelp("GridLayout", "If clients aren't nullrhi, lay them out in 320x240 fashion, defaults to true", ParamType = typeof(bool))]
+	[ParamHelp("NoTimeouts", "Disable timeouts in this script (defaults to false)", ParamType = typeof(bool))]
 	[ParamHelp("SleepTimeBetweenLaunches", "How long to sleep between running clients in milliseconds (could prevent race conditions), 100 by default", ParamType = typeof(int))]
 	[ParamHelp("MaxRunAttemptsPerClient", "Maximum number of attempts to run a client which crashes or fails to connect to the server, defaults to 3", ParamType = typeof(int))]
 	[ParamHelp("ClientSessionCompleted", "Log message indicating that a client has completed a game session and may be terminated", ParamType = typeof(string))]
@@ -40,7 +42,9 @@ namespace MultiClientLauncher.Automation
 		private int FirstClientNumber;
 		private int ClientCount;
 		private int BuildIdOverride;
+		private bool NullRHI = false;
 		private bool GridLayout = true;
+		private bool NoTimeouts = false;
 
 		private int SleepTimeBetweenLaunches = 100;
 		private int MaxRunAttemptsPerClient = 3;
@@ -90,14 +94,21 @@ namespace MultiClientLauncher.Automation
 
 			SleepTimeBetweenLaunches = ParseParamInt("SleepTimeBetweenLaunches", -1);
 			MaxRunAttemptsPerClient = ParseParamInt("MaxRunAttemptsPerClient", 3);
-			GridLayout = ParseParamBool("GridLayout", true);
+			NullRHI = ParseParamBool("NullRHI", NullRHI);
+			GridLayout = ParseParamBool("GridLayout", GridLayout);
+			NoTimeouts = ParseParamBool("NoTimeouts", NoTimeouts);
 
 			// disable grid layout for nullrhi
-			if (ClientArgs.Contains("-nullrhi"))
+			if (NullRHI || ClientArgs.Contains("-nullrhi"))
 			{
 				GridLayout = false;
 			}
 
+			if (NullRHI)
+			{
+				ClientArgs += " -nullrhi ";
+			}
+
 			InitializeClientLogIndicators();
 		}
 
@@ -216,7 +227,7 @@ namespace MultiClientLauncher.Automation
 		
 		private ClientProcess SpawnClientProcess(string ExeFilename, string ClientLogFilename, string ExeArguments)
 		{
-			ClientProcess ClientProc = new ClientProcess(ExeFilename, ExeArguments, MaxRunAttemptsPerClient, ClientLogFilename, ClientIndicators);
+			ClientProcess ClientProc = new ClientProcess(ExeFilename, ExeArguments, MaxRunAttemptsPerClient, NoTimeouts, ClientLogFilename, ClientIndicators);
 			ClientProc.Start();
 			return ClientProc;
 		}
@@ -248,14 +259,16 @@ namespace MultiClientLauncher.Automation
 			private readonly string LogFilepath;
 
 			private readonly Stopwatch ConnectionTimer;
-			private const int MinutesUntilTimeout = 2; // Todo: maybe make this a command line argument
+			private const int MinutesUntilTimeout = 15; // Todo: maybe make this a command line argument
+			private bool NoTimeouts = false;
 
-			public ClientProcess(string Exe, string Args, int MaxRunAttempts, string ClientLog, ClientLogIndicators ClientIndicators)
+			public ClientProcess(string Exe, string Args, int MaxRunAttempts, bool IgnoreTimeouts, string ClientLog, ClientLogIndicators ClientIndicators)
 			{
 				Proc = new Process();
 				Proc.StartInfo.FileName = Exe;
 				Proc.StartInfo.Arguments = Args;
 				RemainingRunAttempts = MaxRunAttempts;
+				NoTimeouts = IgnoreTimeouts;
 
 				LogFilepath = ClientLog;
 				LogIndicators = ClientIndicators;
@@ -374,7 +387,7 @@ namespace MultiClientLauncher.Automation
 					}
 
 					// If timed out while attempting to connect, kill process and decrement attempts
-					if (ConnectionTimer.IsRunning && ConnectionTimer.Elapsed.Minutes >= MinutesUntilTimeout)
+					if (!NoTimeouts && ConnectionTimer.IsRunning && ConnectionTimer.Elapsed.Minutes >= MinutesUntilTimeout)
 					{
 						Console.WriteLine("Client {0} timed out. Attempts left: {1}", GetProcessName(), RemainingRunAttempts);
 						Kill();
