[Mac] Invert Preference order for packaged and staged build so automation works.

diff --git a/Engine/Source/Programs/AutomationTool/Gauntlet/Platform/Mac/Gauntlet.MacBuildSource.cs b/Engine/Source/Programs/AutomationTool/Gauntlet/Platform/Mac/Gauntlet.MacBuildSource.cs
--- a/Engine/Source/Programs/AutomationTool/Gauntlet/Platform/Mac/Gauntlet.MacBuildSource.cs
+++ b/Engine/Source/Programs/AutomationTool/Gauntlet/Platform/Mac/Gauntlet.MacBuildSource.cs
@@ -45,12 +45,21 @@ namespace Gauntlet
 		}
 	}
 
+	public class MacStagedBuild : StagedBuild
+	{
+		public MacStagedBuild(UnrealTargetPlatform InPlatform, UnrealTargetConfiguration InConfig, UnrealTargetRole InRole, string InBuildPath, string InExecutablePath, string InFlavor = "")
+		: base(InPlatform, InConfig, InRole, InBuildPath, InExecutablePath, InFlavor)
+		{
+		}
+
+		public override int PreferenceOrder { get { return 0; } }
+	}
 
 	/// <summary>
 	/// Represents a loose collection of files. For Mac the executable will point to the binary.app bundle and
 	/// not the inner executable file
 	/// </summary>
-	public class MacStagedBuildSource : StagedBuildSource<StagedBuild>
+	public class MacStagedBuildSource : StagedBuildSource<MacStagedBuild>
 	{
 		public override string BuildName { get { return "MacStagedBuildSource"; } }
 
@@ -74,6 +83,8 @@ namespace Gauntlet
 			: base(InPlatform, InConfig, InRole, InBuildPath, InFlags)
 		{
 		}
+
+		public override int PreferenceOrder { get { return 1; } }
 	}
 
 	/// <summary>
diff --git a/Engine/Source/Programs/AutomationTool/Gauntlet/Unreal/BuildSource/Gauntlet.StagedBuildSource.cs b/Engine/Source/Programs/AutomationTool/Gauntlet/Unreal/BuildSource/Gauntlet.StagedBuildSource.cs
--- a/Engine/Source/Programs/AutomationTool/Gauntlet/Unreal/BuildSource/Gauntlet.StagedBuildSource.cs
+++ b/Engine/Source/Programs/AutomationTool/Gauntlet/Unreal/BuildSource/Gauntlet.StagedBuildSource.cs
@@ -37,7 +37,7 @@ namespace Gauntlet
 
 	public class PackagedBuild : IBuild
 	{
-		public int PreferenceOrder { get { return 0; } }
+		public virtual int PreferenceOrder { get { return 0; } }
 
 		public UnrealTargetPlatform Platform { get; protected set; }
 
@@ -74,7 +74,7 @@ namespace Gauntlet
 
 	public class StagedBuild : IBuild
 	{
-		public int PreferenceOrder { get { return 1; } }
+		public virtual int PreferenceOrder { get { return 1; } }
 
 		public UnrealTargetPlatform Platform { get; protected set; }
 
diff --git a/Engine/Source/Programs/AutomationTool/Gauntlet/Unreal/BuildSource/Gauntlet.UnrealBuildSource.cs b/Engine/Source/Programs/AutomationTool/Gauntlet/Unreal/BuildSource/Gauntlet.UnrealBuildSource.cs
--- a/Engine/Source/Programs/AutomationTool/Gauntlet/Unreal/BuildSource/Gauntlet.UnrealBuildSource.cs
+++ b/Engine/Source/Programs/AutomationTool/Gauntlet/Unreal/BuildSource/Gauntlet.UnrealBuildSource.cs
@@ -360,7 +360,7 @@ namespace Gauntlet
 
 							foreach (IBuild Build in Builds)
 							{
-								Log.Info("Adding build {0} with flags {1}", BS.BuildName, Build.Flags);
+								Log.Info("Adding build {0} with flags {1} priority {2}", BS.BuildName, Build.Flags, Build.PreferenceOrder);
 								AddBuild(Build);
 							}
 						}
