Fix Failed to save completion automated test report

#rnx

diff --git a/Engine/Source/Programs/AutomationTool/Gauntlet/Unreal/Automation/UE.Automation.cs b/Engine/Source/Programs/AutomationTool/Gauntlet/Unreal/Automation/UE.Automation.cs
--- a/Engine/Source/Programs/AutomationTool/Gauntlet/Unreal/Automation/UE.Automation.cs
+++ b/Engine/Source/Programs/AutomationTool/Gauntlet/Unreal/Automation/UE.Automation.cs
@@ -700,27 +700,25 @@ namespace UE
 											LogEntryPrefix.Length,
 											LogEntry.Message.Length - LogEntryPrefix.Length - LogEntrySuffix.Length);
 
-										if (LastTestInProgress.FullTestPath == StartedTestFullName)
+										if (LastTestInProgress.FullTestPath != StartedTestFullName)
 										{
-											LastTestInProgress.State = TestStateType.Fail;
-											JsonTestPassResults.Failed++;
-											JsonTestPassResults.InProcess--;
-											ErrorMessage = FatalError.FormatForLog();
-										}
-										else
-										{
-											LastTestInProgress.State = TestStateType.NotRun;
-											JsonTestPassResults.NotRun++;
-											JsonTestPassResults.InProcess--;
-
-											// We have to find another test that has been crashed and set it as the last test with critical failure
-											LastTestInProgress = JsonTestPassResults.Tests.FirstOrDefault((T => T.FullTestPath == StartedTestFullName));
-											if (LastTestInProgress == null)
+											// We have to find the test that has last started. Find that test result in the json report.
+											UnrealAutomatedTestResult LastStartedTestFromLog = JsonTestPassResults.Tests.FirstOrDefault((T => T.FullTestPath == StartedTestFullName));
+											if (LastStartedTestFromLog == null)
 											{
-												Log.Warning("Failed to find test for the test state log entry '{0}'", LogEntry.Message);
+												Log.Warning("Failed to find the test from the test state log entry '{0}'", LogEntry.Message);
+												// We revert to the standard approached as the json report seems inconsistent with the log.
+												// The test in-progress from the json report will get the callstack attached.
 											}
 											else
 											{
+												// The last known running test needs to be rescheduled
+												LastTestInProgress.State = TestStateType.NotRun;
+												JsonTestPassResults.NotRun++;
+												JsonTestPassResults.InProcess--;
+												// Then the last started test according to the log gets flagged as failed
+												LastTestInProgress = LastStartedTestFromLog;
+
 												if (LastTestInProgress.State != TestStateType.Fail)
 												{
 													switch (LastTestInProgress.State)
@@ -744,9 +742,9 @@ namespace UE
 													LastTestInProgress.State = TestStateType.Fail;
 													JsonTestPassResults.Failed++;
 												}
-												ErrorMessage = FatalError.FormatForLog();
 											}
 										}
+										ErrorMessage = FatalError.FormatForLog();
 									}
 									else
 									{
@@ -768,18 +766,21 @@ namespace UE
 								ErrorMessage = "No callstack found in the log.";
 							}
 						}
-						if ((LastTestInProgress != null) && !String.IsNullOrEmpty(ErrorMessage))
-						{
-							LastTestInProgress.AddError(ErrorMessage, !HasTimeout);
-						}
-						if (!CanRetry() || JsonTestPassResults.NotRun == 0)
+						if (LastTestInProgress != null)
 						{
-							// Setting the test as fail because no retry will be done anymore.
-							// The InProcess state won't be used for pass resume
-							LastTestInProgress.State = TestStateType.Fail;
-							if (!CanRetry())
+							if (!String.IsNullOrEmpty(ErrorMessage))
+							{
+								LastTestInProgress.AddError(ErrorMessage, !HasTimeout);
+							}
+							if (!CanRetry() || JsonTestPassResults.NotRun == 0)
 							{
-								LastTestInProgress.AddWarning(string.Format("Session reached maximum of retries({0}) to resume on critical failure!", Retries));
+								// Setting the test as fail because no retry will be done anymore.
+								// The InProcess state won't be used for pass resume
+								LastTestInProgress.State = TestStateType.Fail;
+								if (!CanRetry())
+								{
+									LastTestInProgress.AddWarning(string.Format("Session reached maximum of retries({0}) to resume on critical failure!", Retries));
+								}
 							}
 						}
 						JsonTestPassResults.WriteToJson();
diff --git a/Engine/Source/Runtime/AutomationWorker/Private/AutomationWorkerModule.cpp b/Engine/Source/Runtime/AutomationWorker/Private/AutomationWorkerModule.cpp
--- a/Engine/Source/Runtime/AutomationWorker/Private/AutomationWorkerModule.cpp
+++ b/Engine/Source/Runtime/AutomationWorker/Private/AutomationWorkerModule.cpp
@@ -187,8 +187,6 @@ void FAutomationWorkerModule::ReportTestComplete()
 		//see if there are any more network commands left to execute
 		bool bAllLatentCommandsComplete = FAutomationTestFramework::Get().ExecuteLatentCommands();
 
-		FString TestFullName = FAutomationTestFramework::Get().GetCurrentTest()->GetTestFullName();
-
 		//structure to track error/warning/log messages
 		FAutomationTestExecutionInfo ExecutionInfo;
 
@@ -224,7 +222,7 @@ void FAutomationWorkerModule::ReportTestComplete()
 
 			if (ExecutionInfo.TelemetryItems.Num() > 0)
 			{
-				HandleTelemetryData(ExecutionInfo.TelemetryStorage, TestFullName, ExecutionInfo.TelemetryItems);
+				HandleTelemetryData(ExecutionInfo.TelemetryStorage, FullTestPath, ExecutionInfo.TelemetryItems);
 			}
 
 			MessageEndpoint->Send(Message, TestRequesterAddress);
@@ -235,6 +233,8 @@ void FAutomationWorkerModule::ReportTestComplete()
 		TestRequesterAddress.Invalidate();
 		ExecutionCount = INDEX_NONE;
 		TestName.Empty();
+		FullTestPath.Empty();
+		BeautifiedTestName.Empty();
 		StopTestEvent.Unbind();
 	}
 }
@@ -623,6 +623,7 @@ void FAutomationWorkerModule::HandleRunTestsMessage( const FAutomationWorkerRunT
 	ExecutionCount = Message.ExecutionCount;
 	TestName = Message.TestName;
 	BeautifiedTestName = Message.BeautifiedTestName;
+	FullTestPath = Message.FullTestPath;
 	bSendAnalytics = Message.bSendAnalytics;
 	TestRequesterAddress = Context->GetSender();
 
@@ -632,7 +633,7 @@ void FAutomationWorkerModule::HandleRunTestsMessage( const FAutomationWorkerRunT
 	// We are not executing network command sub-commands right now
 	bExecutingNetworkCommandResults = false;
 
-	FAutomationTestFramework::Get().StartTestByName(Message.TestName, Message.RoleIndex);
+	FAutomationTestFramework::Get().StartTestByName(Message.TestName, Message.RoleIndex, Message.FullTestPath);
 }
 
 
diff --git a/Engine/Source/Runtime/AutomationWorker/Private/AutomationWorkerModule.h b/Engine/Source/Runtime/AutomationWorker/Private/AutomationWorkerModule.h
--- a/Engine/Source/Runtime/AutomationWorker/Private/AutomationWorkerModule.h
+++ b/Engine/Source/Runtime/AutomationWorker/Private/AutomationWorkerModule.h
@@ -158,6 +158,9 @@ private:
 	/** Beautified name of the test */
 	FString BeautifiedTestName;
 
+	/** Full path of the test */
+	FString FullTestPath;
+
 	/** Whether to send analytics events to the backend - sent from controller */
 	bool bSendAnalytics;
 
diff --git a/Engine/Source/Runtime/Core/Private/Misc/AutomationTest.cpp b/Engine/Source/Runtime/Core/Private/Misc/AutomationTest.cpp
--- a/Engine/Source/Runtime/Core/Private/Misc/AutomationTest.cpp
+++ b/Engine/Source/Runtime/Core/Private/Misc/AutomationTest.cpp
@@ -468,7 +468,7 @@ void FAutomationTestFramework::ResetTests()
 	IFileManager::Get().DeleteDirectory(*FPaths::AutomationTransientDir(), bEnsureExists, bDeleteEntireTree);
 }
 
-void FAutomationTestFramework::StartTestByName( const FString& InTestToRun, const int32 InRoleIndex )
+void FAutomationTestFramework::StartTestByName( const FString& InTestToRun, const int32 InRoleIndex, const FString& InFullTestPath )
 {
 	if (GIsAutomationTesting)
 	{
@@ -492,6 +492,7 @@ void FAutomationTestFramework::StartTestByName( const FString& InTestToRun, cons
 	{
 		TestName = InTestToRun;
 	}
+	FString TestPath = InFullTestPath.IsEmpty() ? InTestToRun : InFullTestPath;
 
 	NetworkRoleIndex = InRoleIndex;
 
@@ -504,16 +505,16 @@ void FAutomationTestFramework::StartTestByName( const FString& InTestToRun, cons
 			// Make any setting changes that have to occur to support unit testing
 			PrepForAutomationTests();
 
-			InternalStartTest( InTestToRun );
+			InternalStartTest( InTestToRun, TestPath);
 		}
 		else
 		{
-			UE_LOG(LogAutomationTest, Error, TEXT("Test %s does not exist and could not be run."), *InTestToRun);
+			UE_LOG(LogAutomationTest, Error, TEXT("Test %s does not exist and could not be run."), *TestPath);
 		}
 	}
 	else
 	{
-		UE_LOG(LogAutomationTest, Error, TEXT("Test %s is too slow and could not be run."), *InTestToRun);
+		UE_LOG(LogAutomationTest, Error, TEXT("Test %s is too slow and could not be run."), *TestPath);
 	}
 }
 
@@ -855,9 +856,10 @@ void FAutomationTestFramework::DumpAutomationTestExecutionInfo( const TMap<FStri
 	}
 }
 
-void FAutomationTestFramework::InternalStartTest( const FString& InTestToRun )
+void FAutomationTestFramework::InternalStartTest( const FString& InTestToRun, const FString& InFullTestPath)
 {
 	Parameters.Empty();
+	FullTestPath.Empty();
 
 	FString TestName;
 	if (!InTestToRun.Split(TEXT(" "), &TestName, &Parameters, ESearchCase::CaseSensitive))
@@ -879,19 +881,20 @@ void FAutomationTestFramework::InternalStartTest( const FString& InTestToRun )
 
 		StartTime = FPlatformTime::Seconds();
 
+		CurrentTest->SetTestContext(Parameters);
+		FullTestPath = InFullTestPath;
+
 		// If not a smoke test, log the test has started.
 		uint32 NonSmokeTestFlags = (EAutomationTestFlags::FilterMask & (~EAutomationTestFlags::SmokeFilter));
 		if (RequestedTestFilter & NonSmokeTestFlags)
 		{
 			if (AutomationTest::bLogTestStateTrace)
 			{
-				UE_LOG(LogAutomationTestStateTrace, Log, TEXT("Test is about to start. Name={%s}"), *CurrentTest->GetTestFullName());
+				UE_LOG(LogAutomationTestStateTrace, Log, TEXT("Test is about to start. Name={%s}"), *FullTestPath);
 			}
 			UE_LOG(LogAutomationTest, Log, TEXT("%s %s is starting at %f"), *CurrentTest->GetBeautifiedTestName(), *Parameters, StartTime);
 		}
 
-		CurrentTest->SetTestContext(Parameters);
-
 		OnTestStartEvent.Broadcast(CurrentTest);
 
 		{
@@ -931,7 +934,7 @@ bool FAutomationTestFramework::InternalStopTest(FAutomationTestExecutionInfo& Ou
 		UE_LOG(LogAutomationTest, Log, TEXT("%s %s ran in %f"), *CurrentTest->GetBeautifiedTestName(), *Parameters, TimeForTest);
 		if (AutomationTest::bLogTestStateTrace)
 		{
-			UE_LOG(LogAutomationTestStateTrace, Log, TEXT("Test has stopped execution. Name={%s}"), *CurrentTest->GetTestFullName());
+			UE_LOG(LogAutomationTestStateTrace, Log, TEXT("Test has stopped execution. Name={%s}"), *FullTestPath);
 		}
 	}
 
diff --git a/Engine/Source/Runtime/Core/Public/Misc/AutomationTest.h b/Engine/Source/Runtime/Core/Public/Misc/AutomationTest.h
--- a/Engine/Source/Runtime/Core/Public/Misc/AutomationTest.h
+++ b/Engine/Source/Runtime/Core/Public/Misc/AutomationTest.h
@@ -976,8 +976,9 @@ public:
 	 *
 	 * @param	InTestToRun			Name of the test that should be run
 	 * @param	InRoleIndex			Identifier for which worker in this group that should execute a command
+	 * @param	InFullTestPath		Full test path
 	 */
-	CORE_API void StartTestByName( const FString& InTestToRun, const int32 InRoleIndex );
+	CORE_API void StartTestByName( const FString& InTestToRun, const int32 InRoleIndex, const FString& InFullTestPath = FString() );
 
 	/**
 	 * Stop the current test and return the results of execution
@@ -1255,9 +1256,9 @@ private:
 	 * Internal helper method designed to simply start the provided test name.
 	 *
 	 * @param	InTestToRun			Name of the test that should be run
-	 * @param	OutExecutionInfo	Results of executing the test
+	 * @param	InFullTestPath		Full test path
 	 */
-	CORE_API void InternalStartTest( const FString& InTestToRun );
+	CORE_API void InternalStartTest( const FString& InTestToRun, const FString& InFullTestPath );
 
 	/**
 	 * Internal helper method designed to stop current executing test and return the results of execution.
@@ -1309,6 +1310,9 @@ private:
 	/** Copy of the parameters for the active test */
 	FString Parameters;
 
+	/** Full test path as given by the automation controller of the active test */
+	FString FullTestPath;
+
 	/** Whether we want to run automation tests on content within the Developer Directories */
 	bool bDeveloperDirectoryIncluded;
 
