CustomDeployment handler can be specified on the command line.

- still only used on Windows at the moment.

#rnx

diff --git a/Engine/Source/Programs/AutomationTool/AutomationUtils/ProjectParams.cs b/Engine/Source/Programs/AutomationTool/AutomationUtils/ProjectParams.cs
--- a/Engine/Source/Programs/AutomationTool/AutomationUtils/ProjectParams.cs
+++ b/Engine/Source/Programs/AutomationTool/AutomationUtils/ProjectParams.cs
@@ -430,6 +430,7 @@ namespace AutomationTool
 			this.PackageEncryptionKeyFile = InParams.PackageEncryptionKeyFile;
 			this.Prereqs = InParams.Prereqs;
 			this.AppLocalDirectory = InParams.AppLocalDirectory;
+			this.CustomDeploymentHandler = InParams.CustomDeploymentHandler;
 			this.NoBootstrapExe = InParams.NoBootstrapExe;
             this.Prebuilt = InParams.Prebuilt;
             this.RunTimeoutSeconds = InParams.RunTimeoutSeconds;
@@ -557,6 +558,7 @@ namespace AutomationTool
 			bool? GenerateOptimizationData = null,
 			bool? Prereqs = null,
 			string AppLocalDirectory = null,
+			string CustomDeploymentHandler = null,
 			bool? NoBootstrapExe = null,
             bool? SignedPak = null,
 			bool? PakAlignForMemoryMapping = null,
@@ -873,6 +875,7 @@ namespace AutomationTool
 			this.PackageEncryptionKeyFile = ParseParamValueIfNotSpecified(Command, PackageEncryptionKeyFile, "packageencryptionkeyfile", null);
 			this.Prereqs = GetParamValueIfNotSpecified(Command, Prereqs, this.Prereqs, "prereqs");
 			this.AppLocalDirectory = ParseParamValueIfNotSpecified(Command, AppLocalDirectory, "applocaldirectory", String.Empty, true);
+			this.CustomDeploymentHandler = ParseParamValueIfNotSpecified(Command, CustomDeploymentHandler, "customdeployment", String.Empty, true );
 			this.NoBootstrapExe = GetParamValueIfNotSpecified(Command, NoBootstrapExe, this.NoBootstrapExe, "nobootstrapexe");
             this.Prebuilt = GetParamValueIfNotSpecified(Command, Prebuilt, this.Prebuilt, "prebuilt");
             if (this.Prebuilt)
@@ -2030,6 +2033,11 @@ namespace AutomationTool
 		/// </summary>
 		public Action<ProjectParams, DeploymentContext> FinalizeDeploymentContextCallback = null;
 
+		/// <summary>
+		/// Name of the custom deployment handler to change how the build packaged, staged and deployed - for example, when packaging for a specific game store
+		/// </summary>
+		public string CustomDeploymentHandler { get; set; }
+
 		/// <summary>
 		/// On Windows, adds an executable to the root of the staging directory which checks for prerequisites being 
 		/// installed and launches the game with a path to the .uproject file.
diff --git a/Engine/Source/Programs/AutomationTool/Win/WinPlatform.Automation.cs b/Engine/Source/Programs/AutomationTool/Win/WinPlatform.Automation.cs
--- a/Engine/Source/Programs/AutomationTool/Win/WinPlatform.Automation.cs
+++ b/Engine/Source/Programs/AutomationTool/Win/WinPlatform.Automation.cs
@@ -56,8 +56,17 @@ public class Win64Platform : Platform
 		{
 			if (SC.CustomDeployment == null)
 			{			
+				string CustomDeploymentName = null;
+
 				ConfigHierarchy EngineIni = ConfigCache.ReadHierarchy(ConfigHierarchyType.Engine, Params.RawProjectPath.Directory, PlatformType, SC.CustomConfig);
-				if (EngineIni.GetString("/Script/WindowsTargetPlatform.WindowsTargetSettings", "CustomDeployment", out string CustomDeploymentName))
+				EngineIni.GetString("/Script/WindowsTargetPlatform.WindowsTargetSettings", "CustomDeployment", out CustomDeploymentName);
+
+				if (string.IsNullOrEmpty(CustomDeploymentName))
+				{
+					CustomDeploymentName = Params.CustomDeploymentHandler;
+				}
+
+				if (!string.IsNullOrEmpty(CustomDeploymentName))
 				{
 					SC.CustomDeployment = CustomDeploymentHandler.Create(CustomDeploymentName, this);
 				}
