Fixing vulnerable NuGet deps.
System.Drawing.Common 4.7.0 => 4.7.2  ([REDACTED_PROFILE_URL]/GHSA-rxg9-xrhp-64gj)
MongoDB.Driver 2.11.5 => 2.19.0  ([REDACTED_PROFILE_URL]/GHSA-7j9m-j397-g4wx)

diff --git a/Engine/Source/Programs/AutomationTool/AutomationUtils/AutomationUtils.Automation.csproj b/Engine/Source/Programs/AutomationTool/AutomationUtils/AutomationUtils.Automation.csproj
--- a/Engine/Source/Programs/AutomationTool/AutomationUtils/AutomationUtils.Automation.csproj
+++ b/Engine/Source/Programs/AutomationTool/AutomationUtils/AutomationUtils.Automation.csproj
@@ -59,7 +59,7 @@
   <ItemGroup>
     <PackageReference Include="System.Data.DataSetExtensions" Version="4.5.0" />
     <PackageReference Include="System.Net.Http" Version="4.3.4" />
-    <PackageReference Include="System.Drawing.Common" Version="4.7.0" />
+    <PackageReference Include="System.Drawing.Common" Version="4.7.2" />
 
     <!--Required by Ionic -->
     <PackageReference Include="System.Security.Permissions" Version="4.7.0" />
diff --git a/Engine/Source/Programs/AutomationTool/BuildGraph/BuildGraph.Automation.csproj b/Engine/Source/Programs/AutomationTool/BuildGraph/BuildGraph.Automation.csproj
--- a/Engine/Source/Programs/AutomationTool/BuildGraph/BuildGraph.Automation.csproj
+++ b/Engine/Source/Programs/AutomationTool/BuildGraph/BuildGraph.Automation.csproj
@@ -61,6 +61,6 @@
     <PackageReference Include="Microsoft.CSharp" Version="4.7.0" />
     <PackageReference Include="System.Data.DataSetExtensions" Version="4.5.0" />
     <PackageReference Include="System.Net.Http" Version="4.3.4" />
-    <PackageReference Include="System.Drawing.Common" Version="4.7.0" />
+    <PackageReference Include="System.Drawing.Common" Version="4.7.2" />
   </ItemGroup>
 </Project>
\ No newline at end of file
diff --git a/Engine/Source/Programs/AutomationTool/IOS/IOS.Automation.csproj b/Engine/Source/Programs/AutomationTool/IOS/IOS.Automation.csproj
--- a/Engine/Source/Programs/AutomationTool/IOS/IOS.Automation.csproj
+++ b/Engine/Source/Programs/AutomationTool/IOS/IOS.Automation.csproj
@@ -56,7 +56,7 @@
     <PackageReference Include="Microsoft.CSharp" Version="4.7.0" PrivateAssets="all" />
     <PackageReference Include="System.Data.DataSetExtensions" Version="4.5.0" PrivateAssets="all" />
     <PackageReference Include="System.Net.Http" Version="4.3.4" PrivateAssets="all" />
-    <PackageReference Include="System.Drawing.Common" Version="4.7.0" PrivateAssets="all" />
+    <PackageReference Include="System.Drawing.Common" Version="4.7.2" PrivateAssets="all" />
     <PackageReference Include="System.Resources.Extensions" Version="4.7.1" PrivateAssets="all" />
     <PackageReference Include="Portable.BouncyCastle" Version="1.8.6.7" PrivateAssets="all" />
   </ItemGroup>
diff --git a/Engine/Source/Programs/AutomationTool/Linux/Linux.Automation.csproj b/Engine/Source/Programs/AutomationTool/Linux/Linux.Automation.csproj
--- a/Engine/Source/Programs/AutomationTool/Linux/Linux.Automation.csproj
+++ b/Engine/Source/Programs/AutomationTool/Linux/Linux.Automation.csproj
@@ -45,6 +45,6 @@
     <PackageReference Include="Microsoft.CSharp" Version="4.7.0" />
     <PackageReference Include="System.Data.DataSetExtensions" Version="4.5.0" />
     <PackageReference Include="System.Net.Http" Version="4.3.4" />
-    <PackageReference Include="System.Drawing.Common" Version="4.7.0" />
+    <PackageReference Include="System.Drawing.Common" Version="4.7.2" />
   </ItemGroup>
 </Project>
\ No newline at end of file
diff --git a/Engine/Source/Programs/AutomationTool/Mac/Mac.Automation.csproj b/Engine/Source/Programs/AutomationTool/Mac/Mac.Automation.csproj
--- a/Engine/Source/Programs/AutomationTool/Mac/Mac.Automation.csproj
+++ b/Engine/Source/Programs/AutomationTool/Mac/Mac.Automation.csproj
@@ -44,6 +44,6 @@
     <PackageReference Include="Microsoft.CSharp" Version="4.7.0" />
     <PackageReference Include="System.Data.DataSetExtensions" Version="4.5.0" />
     <PackageReference Include="System.Net.Http" Version="4.3.4" />
-    <PackageReference Include="System.Drawing.Common" Version="4.7.0" />
+    <PackageReference Include="System.Drawing.Common" Version="4.7.2" />
   </ItemGroup>
 </Project>
\ No newline at end of file
diff --git a/Engine/Source/Programs/AutomationTool/Scripts/AutomationScripts.Automation.csproj b/Engine/Source/Programs/AutomationTool/Scripts/AutomationScripts.Automation.csproj
--- a/Engine/Source/Programs/AutomationTool/Scripts/AutomationScripts.Automation.csproj
+++ b/Engine/Source/Programs/AutomationTool/Scripts/AutomationScripts.Automation.csproj
@@ -57,7 +57,7 @@
     <PackageReference Include="Microsoft.CSharp" Version="4.7.0" />
     <PackageReference Include="System.Data.DataSetExtensions" Version="4.5.0" />
     <PackageReference Include="System.Net.Http" Version="4.3.4" />
-    <PackageReference Include="System.Drawing.Common" Version="4.7.0" />
+    <PackageReference Include="System.Drawing.Common" Version="4.7.2" />
     <PackageReference Include="AWSSDK.S3" Version="3.3.113.2" PrivateAssets="all" />
   </ItemGroup>
 </Project>
\ No newline at end of file
diff --git a/Engine/Source/Programs/AutomationTool/SteamDeck/SteamDeck.Automation.csproj b/Engine/Source/Programs/AutomationTool/SteamDeck/SteamDeck.Automation.csproj
--- a/Engine/Source/Programs/AutomationTool/SteamDeck/SteamDeck.Automation.csproj
+++ b/Engine/Source/Programs/AutomationTool/SteamDeck/SteamDeck.Automation.csproj
@@ -44,6 +44,6 @@
     <PackageReference Include="Microsoft.CSharp" Version="4.7.0" />
     <PackageReference Include="System.Data.DataSetExtensions" Version="4.5.0" />
     <PackageReference Include="System.Net.Http" Version="4.3.4" />
-    <PackageReference Include="System.Drawing.Common" Version="4.7.0" />
+    <PackageReference Include="System.Drawing.Common" Version="4.7.2" />
   </ItemGroup>
 </Project>
\ No newline at end of file
diff --git a/Engine/Source/Programs/AutomationTool/TVOS/TVOS.Automation.csproj b/Engine/Source/Programs/AutomationTool/TVOS/TVOS.Automation.csproj
--- a/Engine/Source/Programs/AutomationTool/TVOS/TVOS.Automation.csproj
+++ b/Engine/Source/Programs/AutomationTool/TVOS/TVOS.Automation.csproj
@@ -51,6 +51,6 @@
     <PackageReference Include="Microsoft.CSharp" Version="4.7.0" />
     <PackageReference Include="System.Data.DataSetExtensions" Version="4.5.0" />
     <PackageReference Include="System.Net.Http" Version="4.3.4" />
-    <PackageReference Include="System.Drawing.Common" Version="4.7.0" />
+    <PackageReference Include="System.Drawing.Common" Version="4.7.2" />
   </ItemGroup>
 </Project>
\ No newline at end of file
diff --git a/Engine/Source/Programs/AutomationTool/Turnkey/Turnkey.Automation.csproj b/Engine/Source/Programs/AutomationTool/Turnkey/Turnkey.Automation.csproj
--- a/Engine/Source/Programs/AutomationTool/Turnkey/Turnkey.Automation.csproj
+++ b/Engine/Source/Programs/AutomationTool/Turnkey/Turnkey.Automation.csproj
@@ -58,7 +58,7 @@
     <PackageReference Include="Microsoft.CSharp" Version="4.7.0" />
     <PackageReference Include="System.Data.DataSetExtensions" Version="4.5.0" />
     <PackageReference Include="System.Net.Http" Version="4.3.4" />
-    <PackageReference Include="System.Drawing.Common" Version="4.7.0" />
+    <PackageReference Include="System.Drawing.Common" Version="4.7.2" />
     <PackageReference Include="System.IO.FileSystem.AccessControl" Version="4.7.0" />
     <PackageReference Include="Google.Apis.Drive.v3" Version="1.33.1.1225" PrivateAssets="all" />
   </ItemGroup>
diff --git a/Engine/Source/Programs/AutomationTool/Win/Win.Automation.csproj b/Engine/Source/Programs/AutomationTool/Win/Win.Automation.csproj
--- a/Engine/Source/Programs/AutomationTool/Win/Win.Automation.csproj
+++ b/Engine/Source/Programs/AutomationTool/Win/Win.Automation.csproj
@@ -52,6 +52,6 @@
     <PackageReference Include="Microsoft.CSharp" Version="4.7.0" />
     <PackageReference Include="System.Data.DataSetExtensions" Version="4.5.0" />
     <PackageReference Include="System.Net.Http" Version="4.3.4" />
-    <PackageReference Include="System.Drawing.Common" Version="4.7.0" />
+    <PackageReference Include="System.Drawing.Common" Version="4.7.2" />
   </ItemGroup>
 </Project>
\ No newline at end of file
diff --git a/Engine/Source/Programs/Shared/EpicGames.MongoDB/EpicGames.MongoDB.csproj b/Engine/Source/Programs/Shared/EpicGames.MongoDB/EpicGames.MongoDB.csproj
--- a/Engine/Source/Programs/Shared/EpicGames.MongoDB/EpicGames.MongoDB.csproj
+++ b/Engine/Source/Programs/Shared/EpicGames.MongoDB/EpicGames.MongoDB.csproj
@@ -26,6 +26,6 @@
     </Compile>
   </ItemGroup>
   <ItemGroup>
-    <PackageReference Include="MongoDB.Driver" Version="2.11.5" />
+    <PackageReference Include="MongoDB.Driver" Version="2.19.0" />
   </ItemGroup>
 </Project>
