- Allow Gauntlet to run a .app, but running the exe inside of it

diff --git a/Engine/Source/Programs/AutomationTool/Gauntlet/Platform/Mac/Gauntlet.TargetDeviceMac.cs b/Engine/Source/Programs/AutomationTool/Gauntlet/Platform/Mac/Gauntlet.TargetDeviceMac.cs
--- a/Engine/Source/Programs/AutomationTool/Gauntlet/Platform/Mac/Gauntlet.TargetDeviceMac.cs
+++ b/Engine/Source/Programs/AutomationTool/Gauntlet/Platform/Mac/Gauntlet.TargetDeviceMac.cs
@@ -436,7 +436,7 @@ namespace Gauntlet
 				Log.Verbose("\t{0}", MacInstall.CommandArguments);
 
 				bool bAllowSpew = MacInstall.RunOptions.HasFlag(CommandUtils.ERunOptions.AllowSpew);
-				Result = CommandUtils.Run(MacInstall.ExecutablePath, MacInstall.CommandArguments, Options: MacInstall.RunOptions, SpewFilterCallback: new SpewFilterCallbackType(delegate (string M) { return bAllowSpew ? M : null; }) /* make sure stderr does not spew in the stdout */);
+				Result = CommandUtils.Run(GetExecutableIfBundle(MacInstall.ExecutablePath), MacInstall.CommandArguments, Options: MacInstall.RunOptions, SpewFilterCallback: new SpewFilterCallbackType(delegate (string M) { return bAllowSpew ? M : null; }) /* make sure stderr does not spew in the stdout */);
 
 				if (Result.HasExited && Result.ExitCode != 0)
 				{
