Check for fatal error when test pass is incomplete and attach its callstack to the last test run (even if that test is not in progress)

#rnx
#prelfight 64fba60ddc7c83509c50bafc

diff --git a/Engine/Source/Programs/AutomationTool/Gauntlet/Unreal/Automation/UE.Automation.cs b/Engine/Source/Programs/AutomationTool/Gauntlet/Unreal/Automation/UE.Automation.cs
--- a/Engine/Source/Programs/AutomationTool/Gauntlet/Unreal/Automation/UE.Automation.cs
+++ b/Engine/Source/Programs/AutomationTool/Gauntlet/Unreal/Automation/UE.Automation.cs
@@ -657,132 +657,138 @@ namespace UE
 			bool HasTimeout = RoleResults != null && RoleResults.Where(R => R.ProcessResult == UnrealProcessResult.TimeOut).Any();
 			string HordeArtifactPath = GetConfiguration().HordeArtifactPath;
 			var MainRole = GetConfiguration().GetMainRequiredRole();
+			// Get any critical error and push it to json report and resave it.
+			UnrealLog.CallstackMessage FatalError = null;
+			UnrealRoleResult FatalErrorRoleResult = null;
+			if (RoleResults != null)
+			{
+				foreach (var Result in RoleResults)
+				{
+					if (Result.LogSummary.FatalError != null)
+					{
+						FatalError = Result.LogSummary.FatalError;
+						FatalErrorRoleResult = Result;
+						break;
+					}
+				}
+			}
 			UnrealAutomatedTestResult LastTestInProgress = null;
 			if (JsonTestPassResults.InProcess > 0)
 			{
 				// The test pass did not run completely
 				Log.Verbose("Found in-process tests: {Count}", JsonTestPassResults.InProcess);
-				// Get any critical error and push it to json report and resave it.
-				if (RoleResults != null)
+				LastTestInProgress = JsonTestPassResults.Tests.FirstOrDefault((T => T.State == TestStateType.InProcess));
+				if (!String.IsNullOrEmpty(LastTestInProgress.TestDisplayName))
 				{
-					UnrealLog.CallstackMessage FatalError = null;
-					UnrealRoleResult FatalErrorRoleResult = null;
-					foreach (var Result in RoleResults)
+					string ErrorMessage = null;
+					if (HasTimeout)
 					{
-						if (Result.LogSummary.FatalError != null)
-						{
-							FatalError = Result.LogSummary.FatalError;
-							FatalErrorRoleResult = Result;
-							break;
-						}
+						ErrorMessage = String.Format("Session reached timeout after {0} seconds.", MaxDuration);
 					}
-					LastTestInProgress = JsonTestPassResults.Tests.FirstOrDefault((T => T.State == TestStateType.InProcess));
-					if (!String.IsNullOrEmpty(LastTestInProgress.TestDisplayName))
+					else
 					{
-						string ErrorMessage = null;
-						if (HasTimeout)
-						{
-							ErrorMessage = String.Format("Session reached timeout after {0} seconds.", MaxDuration);
-						}
-						else
+						if (FatalError != null)
 						{
-							if (FatalError != null)
+							string StartedTestFullName = GetLastStartedTestFullNameFromRoleResult(FatalErrorRoleResult);
+							if (!string.IsNullOrEmpty(StartedTestFullName))
 							{
-								var LogEntry = FatalErrorRoleResult.LogSummary.LogEntries.LastOrDefault(Entry => Entry.Category == "AutomationTestStateTrace", null);
-
-								if (LogEntry != null)
+								if (LastTestInProgress.FullTestPath != StartedTestFullName)
 								{
-									string LogEntryPrefix = "Test is about to start. Name={";
-									string LogEntrySuffix = "}";
-									if (LogEntry.Message.StartsWith(LogEntryPrefix) && LogEntry.Message.EndsWith(LogEntrySuffix))
+									// We have to find the test that has last started. Find that test result in the json report.
+									UnrealAutomatedTestResult LastStartedTestFromLog = JsonTestPassResults.Tests.FirstOrDefault((T => T.FullTestPath == StartedTestFullName));
+									if (LastStartedTestFromLog == null)
 									{
-										string StartedTestFullName = LogEntry.Message.Substring(
-											LogEntryPrefix.Length,
-											LogEntry.Message.Length - LogEntryPrefix.Length - LogEntrySuffix.Length);
-
-										if (LastTestInProgress.FullTestPath != StartedTestFullName)
-										{
-											// We have to find the test that has last started. Find that test result in the json report.
-											UnrealAutomatedTestResult LastStartedTestFromLog = JsonTestPassResults.Tests.FirstOrDefault((T => T.FullTestPath == StartedTestFullName));
-											if (LastStartedTestFromLog == null)
-											{
-												Log.Warning("Failed to find the test from the test state log entry '{0}'", LogEntry.Message);
-												// We revert to the standard approached as the json report seems inconsistent with the log.
-												// The test in-progress from the json report will get the callstack attached.
-											}
-											else
-											{
-												// The last known running test needs to be rescheduled
-												LastTestInProgress.State = TestStateType.NotRun;
-												JsonTestPassResults.NotRun++;
-												JsonTestPassResults.InProcess--;
-												// Then the last started test according to the log gets flagged as failed
-												LastTestInProgress = LastStartedTestFromLog;
-
-												if (LastTestInProgress.State != TestStateType.Fail)
-												{
-													switch (LastTestInProgress.State)
-													{
-														case TestStateType.InProcess:
-															JsonTestPassResults.InProcess--;
-															break;
-
-														case TestStateType.NotRun:
-															Log.Warning("The current state for the test '{0}' is NotRun (it will be changed to Fail). Log entry is '{1}'. That state is inconsistent with the state that is expected by design", StartedTestFullName, LogEntry.Message);
-															JsonTestPassResults.NotRun--;
-															break;
-
-														case TestStateType.Success:
-															JsonTestPassResults.Succeeded--;
-															break;
-
-														default:
-															break;
-													};
-													LastTestInProgress.State = TestStateType.Fail;
-													JsonTestPassResults.Failed++;
-												}
-											}
-										}
-										ErrorMessage = FatalError.FormatForLog();
+										Log.Warning("Failed to find the test from the test state log entry '{0}'", StartedTestFullName);
+										// We revert to the standard approached as the json report seems inconsistent with the log.
+										// The test in-progress from the json report will get the callstack attached.
 									}
 									else
 									{
-										// Handling the test state without using AutomationTestStateTrace log channel
-										string WarningMessage = string.Format(
-											"Please have a look at this case. It is possible that the crash is not caused by the test execution. The last test's state trace log entry is '{0}'",
-											LogEntry.Message);
-										Log.Warning(WarningMessage + "\nCallstack:\n" + FatalError.FormatForLog());
-										ErrorMessage = FatalError.FormatForLog() + "\n" + WarningMessage;
+										// The last known running test needs to be rescheduled
+										LastTestInProgress.State = TestStateType.NotRun;
+										JsonTestPassResults.NotRun++;
+										JsonTestPassResults.InProcess--;
+										// Then the last started test according to the log gets flagged as failed
+										LastTestInProgress = LastStartedTestFromLog;
+
+										if (LastTestInProgress.State != TestStateType.Fail)
+										{
+											switch (LastTestInProgress.State)
+											{
+												case TestStateType.InProcess:
+													JsonTestPassResults.InProcess--;
+													break;
+
+												case TestStateType.NotRun:
+													Log.Warning("The current state from the json report for the test '{0}' is NotRun (it will be changed to Fail). Log entry says it was started. That state is inconsistent.", StartedTestFullName);
+													JsonTestPassResults.NotRun--;
+													break;
+
+												case TestStateType.Success:
+													JsonTestPassResults.Succeeded--;
+													break;
+
+												default:
+													break;
+											};
+											LastTestInProgress.State = TestStateType.Fail;
+											JsonTestPassResults.Failed++;
+										}
 									}
 								}
-								else
-								{
-									ErrorMessage = FatalError.FormatForLog();
-								}
-							}
-							else
-							{
-								ErrorMessage = "No callstack found in the log.";
 							}
+							ErrorMessage = FatalError.FormatForLog();
 						}
-						if (LastTestInProgress != null)
+						else
 						{
-							if (!String.IsNullOrEmpty(ErrorMessage))
-							{
-								LastTestInProgress.AddError(ErrorMessage, !HasTimeout);
-							}
-							if (!CanRetry() || JsonTestPassResults.NotRun == 0)
+							ErrorMessage = "No callstack found in the log.";
+						}
+					}
+					if (LastTestInProgress != null)
+					{
+						if (!String.IsNullOrEmpty(ErrorMessage))
+						{
+							LastTestInProgress.AddError(ErrorMessage, !HasTimeout);
+						}
+						if (!CanRetry() || JsonTestPassResults.NotRun == 0)
+						{
+							// Setting the test as fail because no retry will be done anymore.
+							// The InProcess state won't be used for pass resume
+							LastTestInProgress.State = TestStateType.Fail;
+							if (!CanRetry())
 							{
-								// Setting the test as fail because no retry will be done anymore.
-								// The InProcess state won't be used for pass resume
-								LastTestInProgress.State = TestStateType.Fail;
-								if (!CanRetry())
-								{
-									LastTestInProgress.AddWarning(string.Format("Session reached maximum of retries({0}) to resume on critical failure!", Retries));
-								}
+								LastTestInProgress.AddWarning(string.Format("Session reached maximum of retries({0}) to resume on critical failure!", Retries));
 							}
 						}
+					}
+					JsonTestPassResults.WriteToJson();
+				}
+			}
+			else if (FatalError != null)
+			{
+				string ErrorMessage = FatalError.FormatForLog();
+				string StartedTestFullName = GetLastStartedTestFullNameFromRoleResult(FatalErrorRoleResult);
+				if (!string.IsNullOrEmpty(StartedTestFullName))
+				{
+					// We have to find the test that has last started. Find that test result in the json report.
+					UnrealAutomatedTestResult LastStartedTestFromLog = JsonTestPassResults.Tests.FirstOrDefault((T => T.FullTestPath == StartedTestFullName));
+					if (LastStartedTestFromLog == null)
+					{
+						Log.Warning("Failed to find the test from the test state log entry '{0}'", StartedTestFullName);
+					}
+					else
+					{
+						LastStartedTestFromLog.AddError(ErrorMessage, true);
+						if(LastStartedTestFromLog.State == TestStateType.Success)
+						{
+							LastStartedTestFromLog.State = TestStateType.Fail;
+							JsonTestPassResults.Succeeded--;
+							JsonTestPassResults.Failed++;
+						}
+						if (!CanRetry())
+						{
+							LastStartedTestFromLog.AddWarning(string.Format("Session reached maximum of retries({0}) to resume on critical failure!", Retries));
+						}
 						JsonTestPassResults.WriteToJson();
 					}
 				}
@@ -846,6 +852,32 @@ namespace UE
 			}
 		}
 
+		/// <summary>
+		/// Get last started test full name from Unreal Role result
+		/// </summary>
+		/// <param name="RoleResult"></param>
+		/// <returns></returns>
+		private string GetLastStartedTestFullNameFromRoleResult(UnrealRoleResult RoleResult)
+		{
+			var LogEntry = RoleResult == null ? null : RoleResult.LogSummary.LogEntries.LastOrDefault(Entry => Entry.Category == "AutomationTestStateTrace", null);
+
+			if (LogEntry != null)
+			{
+				string LogEntryPrefix = "Test is about to start. Name={";
+				string LogEntrySuffix = "}";
+				if (LogEntry.Message.StartsWith(LogEntryPrefix) && LogEntry.Message.EndsWith(LogEntrySuffix))
+				{
+					string StartedTestFullName = LogEntry.Message.Substring(
+						LogEntryPrefix.Length,
+						LogEntry.Message.Length - LogEntryPrefix.Length - LogEntrySuffix.Length);
+
+					return StartedTestFullName;
+				}
+			}
+
+			return null;
+		}
+
 		/// <summary>
 		/// Override GetExitCodeAndReason to provide additional checking of success / failure based on what occurred
 		/// </summary>
