Enable Linux, Android for WebTests.

diff --git a/Engine/Build/LowLevelTests/Web.xml b/Engine/Build/LowLevelTests/Web.xml
--- a/Engine/Build/LowLevelTests/Web.xml
+++ b/Engine/Build/LowLevelTests/Web.xml
@@ -13,7 +13,7 @@
 	<Property Name="WebHasAfterSteps" Value="True" />
 	<Property Name="WebUsesCatch2" Value="True" />
 	<Option Name="RunWebTests" DefaultValue="False" Description="Run Web Tests" />
-	<Property Name="WebSupportedPlatforms" Value="Win64" />
+	<Property Name="WebSupportedPlatforms" Value="Win64;Linux;Android" />
 	<Property Name="Win64WebTags" Value="" />
 	<Property Name="MacWebTags" Value="" />
 	<Property Name="IOSWebTags" Value="" />
diff --git a/Engine/Source/Programs/AutomationTool/LowLevelTests/Tests/LowLevelTests.WebTestsLowLevelTestsExtension.cs b/Engine/Source/Programs/AutomationTool/LowLevelTests/Tests/LowLevelTests.WebTestsLowLevelTestsExtension.cs
--- a/Engine/Source/Programs/AutomationTool/LowLevelTests/Tests/LowLevelTests.WebTestsLowLevelTestsExtension.cs
+++ b/Engine/Source/Programs/AutomationTool/LowLevelTests/Tests/LowLevelTests.WebTestsLowLevelTestsExtension.cs
@@ -7,6 +7,7 @@ using EpicGames.Core;
 using UnrealBuildTool;
 using Gauntlet;
 using UnrealBuildBase;
+using System.Net.Sockets;
 
 namespace LowLevelTests
 {
@@ -27,7 +28,9 @@ namespace LowLevelTests
 		public void PreRunTests()
 		{
 			InstallWebServer();
-			RunWebServer();
+			AsyncLaunchWebServerProcess();
+			// If start WebTests right after launching web server process without waiting, it could get refused to connect, especially on Linux
+			WaitUntilWebServerPortOpen();
 		}
 
 		private string WebTestsServerDir()
@@ -39,8 +42,8 @@ namespace LowLevelTests
 		{
 			ProcessStartInfo StartInfo = new ProcessStartInfo();
 			StartInfo.WorkingDirectory = WebTestsServerDir();
-			StartInfo.FileName = RuntimePlatform.IsWindows ? "cmd.exe" : "/bind/sh";
-			StartInfo.Arguments = RuntimePlatform.IsWindows ? "/c createenv.bat" : "-c 'createenv.sh'";
+			StartInfo.FileName = RuntimePlatform.IsWindows ? "cmd.exe" : "/bin/sh";
+			StartInfo.Arguments = RuntimePlatform.IsWindows ? "/c createenv.bat" : "-c './createenv.sh'";
 			StartInfo.UseShellExecute = false;
 			StartInfo.CreateNoWindow = true;
 
@@ -52,10 +55,10 @@ namespace LowLevelTests
 			Console.WriteLine("Requirements installed.");
 		}
 
-		private void RunWebServer()
+		private void AsyncLaunchWebServerProcess()
 		{
 			string WorkingDir = WebTestsServerDir();
-			string PythonFile = Path.Combine(WorkingDir, "env", "Scripts", RuntimePlatform.IsWindows ? "python.exe" : "python");
+			string PythonFile = Path.Combine(WorkingDir, "env", RuntimePlatform.IsWindows ? "Scripts" : "bin", RuntimePlatform.IsWindows ? "python.exe" : "python");
 
 			ProcessStartInfo StartInfo = new ProcessStartInfo();
 			StartInfo.WorkingDirectory = WorkingDir;
@@ -69,10 +72,51 @@ namespace LowLevelTests
 			ServerProcess.StartInfo = StartInfo;
 			ServerProcess.Start();
 
-			Console.WriteLine("Web server is now running.");
+			Console.WriteLine("Web server process is now running.");
+		}
+
+		private void WaitUntilWebServerPortOpen()
+		{
+			Stopwatch sw = new Stopwatch();
+			sw.Start();
+
+			while (!IsServerPortOpen(UnrealHelpers.GetHostIpAddress(), 8000))
+			{
+				if (sw.ElapsedMilliseconds > 60000)
+				{
+					sw.Stop();
+					throw new TimeoutException("Server port did not open within the specified time.");
+				}
+				System.Threading.Thread.Sleep(1000);
+			}
+
+			sw.Stop();
+
+			Console.WriteLine("Web server port is now open.");
+		}
+
+		private bool IsServerPortOpen(string ipAddress, int port)
+		{
+			using (TcpClient client = new TcpClient())
+			{
+				try
+				{
+					client.Connect(ipAddress, port);
+					return true;
+				}
+				catch
+				{
+					return false;
+				}
+			}
 		}
 
 		public void PostRunTests()
+		{
+			CloseWebServer();
+		}
+
+		private void CloseWebServer()
 		{
 			if (ServerProcess != null)
 			{
diff --git a/Engine/Source/Programs/WebTestsServer/createenv.sh b/Engine/Source/Programs/WebTestsServer/createenv.sh
old mode 100644
new mode 100755
--- a/Engine/Source/Programs/WebTestsServer/createenv.sh
+++ b/Engine/Source/Programs/WebTestsServer/createenv.sh
@@ -1,12 +1,17 @@
 if [ -d "env" ]; then
-	source ./env/Scripts/activate.sh
+	. ./env/bin/activate
 else
-	run ../../../Binaries/ThirdParty/Python3/Mac/python3 -m venv env
-	source ./env/Scripts/activate.sh
-	run python -m  pip install --upgrade pip
-	run pip install -r requirements.txt
-	if [ $? -ne 0 ]; then
+	if [ "$OSTYPE" = "darwin" ]; then
+		../../../Binaries/ThirdParty/Python3/Mac/bin/python3 -m venv env
+	else
+		../../../Binaries/ThirdParty/Python3/Linux/bin/python3 -m venv env
+	fi
+	. ./env/bin/activate
+	python -m pip install --upgrade pip
+	pip install -r requirements.txt
+	local install_exit_code = $?
+	if [ $install_exit_code -ne 0 ]; then
 		rm -rf env
-		exit $?
+		exit $install_exit_code
 	fi
 fi
diff --git a/Engine/Source/Programs/WebTestsServer/webtests/settings.py b/Engine/Source/Programs/WebTestsServer/webtests/settings.py
--- a/Engine/Source/Programs/WebTestsServer/webtests/settings.py
+++ b/Engine/Source/Programs/WebTestsServer/webtests/settings.py
@@ -77,10 +77,11 @@ WSGI_APPLICATION = 'webtests.wsgi.application'
 # https://docs.djangoproject.com/en/4.1/ref/settings/#databases
 
 DATABASES = {
-    'default': {
-        'ENGINE': 'django.db.backends.sqlite3',
-        'NAME': BASE_DIR / 'db.sqlite3',
-    }
+    # No need of database for web tests, at least for now
+    # 'default': {
+    #     'ENGINE': 'django.db.backends.sqlite3',
+    #     'NAME': BASE_DIR / 'db.sqlite3',
+    # }
 }
 
 
