Properly generate staged build for OnlineTests and try to use metadata file to deploy low level tests content
Redirected binaries so they end up on the proper path

diff --git a/Engine/Source/Programs/AutomationTool/LowLevelTests/RunLowLevelTests.cs b/Engine/Source/Programs/AutomationTool/LowLevelTests/RunLowLevelTests.cs
--- a/Engine/Source/Programs/AutomationTool/LowLevelTests/RunLowLevelTests.cs
+++ b/Engine/Source/Programs/AutomationTool/LowLevelTests/RunLowLevelTests.cs
@@ -527,7 +527,7 @@ namespace LowLevelTests
 	{
 		bool CanSupportPlatform(UnrealTargetPlatform InPlatform);
 
-		LowLevelTestsBuild CreateBuild(UnrealTargetPlatform InPlatform, UnrealTargetConfiguration InConfiguration, string InTestApp, string InBuildPath, bool bSkipStage);
+		StagedBuild CreateBuild(UnrealTargetPlatform InPlatform, UnrealTargetConfiguration InConfiguration, string InTestApp, string InBuildPath, bool bSkipStage);
 		protected static string GetExecutable(UnrealTargetPlatform InPlatform, UnrealTargetConfiguration InConfiguration, string InTestApp, string InBuildPath, string FileRegEx)
 		{
 			IEnumerable<string> Executables = DirectoryUtils.FindMatchingFiles(InBuildPath, FileRegEx, -1).Select(FileInfo => FileInfo.FullName);
@@ -615,7 +615,7 @@ namespace LowLevelTests
 			return InPlatform.IsInGroup(UnrealPlatformGroup.Desktop);
 		}
 
-		public LowLevelTestsBuild CreateBuild(UnrealTargetPlatform InPlatform, UnrealTargetConfiguration InConfiguration, string InTestApp, string InBuildPath, bool bSkipStage)
+		public StagedBuild CreateBuild(UnrealTargetPlatform InPlatform, UnrealTargetConfiguration InConfiguration, string InTestApp, string InBuildPath, bool bSkipStage)
 		{
 			string ExecutablePath = ILowLevelTestsBuildFactory.GetExecutable(InPlatform, InConfiguration, InTestApp, InBuildPath, GetExecutableRegex(InPlatform));
 			return new LowLevelTestsBuild(InPlatform, InConfiguration, InBuildPath, ExecutablePath);
@@ -686,7 +686,7 @@ namespace LowLevelTests
 
 		public UnrealTargetPlatform Platform { get; protected set; }
 		public UnrealTargetConfiguration Configuration { get; protected set; }
-		public LowLevelTestsBuild DiscoveredBuild { get; protected set; }
+		public StagedBuild DiscoveredBuild { get; protected set; }
 
 		public LowLevelTestsBuildSource(string InTestApp, string InBuildPath, UnrealTargetPlatform InTargetPlatform, UnrealTargetConfiguration InConfiguration, bool InSkipStage)
 		{
