Added function to UAT BuildCommand for checking whether there are any pending P4 changes

diff --git a/Engine/Source/Programs/AutomationTool/AutomationUtils/P4Utils.cs b/Engine/Source/Programs/AutomationTool/AutomationUtils/P4Utils.cs
--- a/Engine/Source/Programs/AutomationTool/AutomationUtils/P4Utils.cs
+++ b/Engine/Source/Programs/AutomationTool/AutomationUtils/P4Utils.cs
@@ -4334,5 +4334,27 @@ namespace AutomationTool
 			}
 			return Text;
 		}
+
+		public bool CheckClientHasPendingChanges(string P4Client = null)
+		{
+			if (P4Client == null)
+			{
+				P4Client = P4Env.Client;
+			}
+
+			if (String.IsNullOrEmpty(P4Client))
+			{
+				Logger.LogWarning("No Perforce client found.");
+				return false;
+			}
+
+			List<ChangeRecord> PendingChanges;
+			string P4GetPendingChangesArgs = $"-s pending -c {P4Client}";
+
+			Logger.LogInformation("Checking for pending changes...");
+			Changes(out PendingChanges, P4GetPendingChangesArgs);
+
+			return PendingChanges.Count > 0;
+		}
 	}
 }
