Horde: Fix warnings building Horde due to implicit usings.

diff --git a/Engine/Source/Programs/AutomationTool/AutomationUtils/Matchers/.editorconfig b/Engine/Source/Programs/AutomationTool/AutomationUtils/Matchers/.editorconfig
--- a/Engine/Source/Programs/AutomationTool/AutomationUtils/Matchers/.editorconfig
+++ b/Engine/Source/Programs/AutomationTool/AutomationUtils/Matchers/.editorconfig
@@ -1,2 +1,3 @@
 [*.cs]
+dotnet_diagnostic.IDE0005.severity = silent
 dotnet_diagnostic.IDE0090.severity = silent
diff --git a/Engine/Source/Programs/AutomationTool/AutomationUtils/Matchers/LogChannelEventMatcher.cs b/Engine/Source/Programs/AutomationTool/AutomationUtils/Matchers/LogChannelEventMatcher.cs
--- a/Engine/Source/Programs/AutomationTool/AutomationUtils/Matchers/LogChannelEventMatcher.cs
+++ b/Engine/Source/Programs/AutomationTool/AutomationUtils/Matchers/LogChannelEventMatcher.cs
@@ -13,7 +13,7 @@ namespace AutomationUtils.Matchers
 	/// </summary>
 	class LogChannelEventMatcher : ILogEventMatcher
 	{
-		readonly static Regex s_pattern = new Regex(
+		static readonly Regex s_pattern = new Regex(
 			@"^(\s*)" +
 			@"(?:\[[\d\.\-: ]+\])*" +
 			@"(?<channel>[a-zA-Z_][a-zA-Z0-9_]*):\s*" +
