Add a new UAT script to report on binary sizes.

diff --git a/Engine/Source/Programs/AutomationTool/Scripts/ReportItemizedExecutableCode.cs b/Engine/Source/Programs/AutomationTool/Scripts/ReportItemizedExecutableCode.cs
new file mode 100644
--- /dev/null
+++ b/Engine/Source/Programs/AutomationTool/Scripts/ReportItemizedExecutableCode.cs
@@ -0,0 +1,328 @@
+// Copyright Epic Games, Inc. All Rights Reserved.
+
+using AutomationTool;
+using EpicGames.Core;
+using System.Collections.Generic;
+using System.IO;
+using System;
+using UnrealBuildBase;
+using IdentityModel.OidcClient;
+using System.Linq;
+using System.Drawing;
+using Microsoft.CodeAnalysis.CSharp.Syntax;
+
+namespace ReportItemizedExecutableCode.Automation
+{
+	struct ModuleInfos
+	{
+		/** Holds the actual itemized data (module to its bytes) */
+		public Dictionary<string, uint> Modules;
+
+		/** Holds the mappings - which sources we categorized as belonging to this module (for debugging) */
+		public Dictionary<string, List<string>> ModuleNameToSources;
+
+		/** Which sources we couldn't classify */
+		public List<string> SourcesWithoutModule;
+
+		/** Size of sources that we couldn't classify */
+		public uint TotalSizeForSourcesWithoutModule;
+
+		public ModuleInfos()
+		{
+			Modules = new Dictionary<string, uint>();
+			ModuleNameToSources = new Dictionary<string, List<string>>();
+			SourcesWithoutModule = new List<string>();
+			TotalSizeForSourcesWithoutModule = 0;
+		}
+	};
+
+	[Help("Reports itemized binary size in a format consumed by PerfReportServer")]
+	[ParamHelp("ExeToItemize", "Absolute path to the binary to itemize", ParamType = typeof(FileReference))]
+	class ReportItemizedExecutableCode : BuildCommand
+	{
+		/** Attempts to assign a reasonable module name from a source file name */
+		string FindModuleName(string SourcePath)
+		{
+			// For context, here are the examples we might be dealing with:
+
+			// ../../FooGame/Intermediate/Build/Platform/Arch/FooClient/Shipping/MovieSceneTracks\Module.MovieSceneTracks.2.cpp  ->   we want it to be name "MovieSceneTracks"
+			// (general rule for that: last directory name)
+
+			// but this does not work for non-unity modules (some are still encountered in unity builds), e.g.
+			// ../Plugins/Runtime/Database/SQLiteCore/Source/SQLiteCore/Private\SQLiteEmbedded.c
+			// So, the rule is amends - if we can find Source in the path, take the next path element after it (if any)
+
+			// But that does not work for a path like
+			// D:/foo/DevAudio/Engine/Source/ThirdParty/Vorbis/libvorbis-1.3.2/Platform/../lib\vorbisenc.c -> we want it named "Vorbis" and not "lib"
+			// So, the rule added: skip "ThirdParty" when moving away from Source
+
+			// Again a bad path where taking the directory after "source" is a bad choice:
+			// D:/P4Libs/depot/3rdParty/libwebsocket/source/2.2/libwebsockets_src/lib\client-parser.c -> we want it named "libwebsocket" and not "2.2"
+			// Solution: look for "3rdParty" (a well known path locally) before looking for Source
+
+			// But this agains fails for foreign strings that don't contain any
+			// D:\home\teamcity\work\sdk\Externals\curl\lib\if2ip.c
+			// For this, we apply the following heurstics: if last directory name is a name like "lib" or "src", take the pre-last directory
+
+			// so the algo becomes:
+			// look for "3rdParty" in the path. Found?  Take the next path element, if any, otherwise take "3rdParty". If not found, continue
+			// look for "Source" in the path. Found?  Take the next path element, if any, otherwise take "3rdParty". If not found, continue
+			// Take the last directory. Is it "lib", "src" or a few other names that we don't think are a good module name?  If yes, keep walking up
+
+			string[] PathElements = SourcePath.Split(new char[] { '\\', '/' }, StringSplitOptions.RemoveEmptyEntries);
+
+			// look for Source for known Unreal names (both non-unity files and ThirdParty libs will be found)
+			for (int IdxElem = 0; IdxElem < PathElements.Length; ++IdxElem)
+			{
+				string CandidateName = PathElements[IdxElem];
+				if (CandidateName.Equals("3rdParty", StringComparison.InvariantCultureIgnoreCase))
+				{
+					if (IdxElem < PathElements.Length - 2)  // last element is the file name, don't want to return that
+					{
+						CandidateName = PathElements[IdxElem + 1];
+
+						// sometimes people build 3rdParty in a 3rdParty folder, e.g. D:/perforce/3rdParty/3rdParty/libwebsocket/source/2.2/libwebsockets_src/lib\libwebsockets.c
+						if (CandidateName.Equals("3rdParty", StringComparison.InvariantCultureIgnoreCase))
+						{
+							if (IdxElem < PathElements.Length - 3)  // last element is the file name, don't want to return that
+							{
+								CandidateName = PathElements[IdxElem + 2];
+							}
+						}
+					}
+
+					return CandidateName;
+				}
+				else if (CandidateName.Equals("Source", StringComparison.InvariantCultureIgnoreCase))
+				{
+					if (IdxElem < PathElements.Length - 2)  // last element is the file name, don't want to return that
+					{
+						CandidateName = PathElements[IdxElem + 1];
+						if (CandidateName.Equals("ThirdParty", StringComparison.InvariantCultureIgnoreCase) && IdxElem < PathElements.Length - 3)
+						{
+							CandidateName = PathElements[IdxElem + 2];
+						}
+					}
+					return CandidateName;
+				}
+			}
+
+			// didn't find neither "3rdParty", nor "Source" string, try to walk up to the module name from behind
+			if (PathElements.Length < 2)
+			{
+				// if source path is does not have any directories, well, then we cannot group it anyhow, let it be the module name
+				return SourcePath;
+			}
+
+			string[] NamesThatCannotBeModuleNames = { "..", "src", "lib", "float" };
+
+			for (int IdxElem = PathElements.Length - 2; IdxElem >= 0; --IdxElem)
+			{
+				string ModuleCandidate = PathElements[IdxElem];
+				
+				bool SkipThisDirectory = false;
+				foreach (string UndesirableModuleName in NamesThatCannotBeModuleNames)
+				{
+					if (ModuleCandidate.Equals(UndesirableModuleName, StringComparison.InvariantCultureIgnoreCase))
+					{
+						SkipThisDirectory = true;
+						break;
+					}
+				}
+
+				if (SkipThisDirectory)
+				{
+					continue;
+				}
+
+				return ModuleCandidate;
+			}
+
+			// some known badly compiled modules
+			if (SourcePath.StartsWith("../..\\png"))
+			{
+				return "png";
+			}
+
+			// if we arrived here we couldn't find anything that looks like a module name. Return the whole path
+			return SourcePath;
+		}
+
+		/** Module names can be used as perf metrics, so we need to remove spaces and non-alphanumeric characters except underscores and hyphens */
+		string SanitizeModuleName(string ModuleName)
+		{
+			return ModuleName.Replace(' ', '_').Replace('.', '-').Replace('#', '-').Replace("[", "").Replace("]", "");
+		}
+
+		bool ValidateModuleName(string ModuleName, string SourcePath)
+		{
+			if (string.IsNullOrEmpty(ModuleName) || ModuleName == ".." || char.IsDigit(ModuleName[0]))
+			{
+				//Logger.LogWarning("Module name '{0}' is poorly chosen from source '{1}' - change FindModuleName heuristics!", ModuleName, SourcePath);
+				return false;
+			}
+			return true;
+		}
+
+		/** Attempts to run bloaty on the executable and returns the results */
+		protected ModuleInfos ItemizeExecutableWithBloaty(FileReference Binary)
+		{
+			ModuleInfos ModInfos = new ModuleInfos();
+			string BloatyExePath = Unreal.RootDirectory.ToString() + @"\Engine\Extras\ThirdPartyNotUE\Bloaty\bloaty.exe";
+
+			if (File.Exists(BloatyExePath))
+			{
+				if (FileReference.Exists(Binary))
+				{
+					string Args = string.Format("-d compileunits -n 0 --csv {0}", Binary.FullName);
+					// Hack!
+					//string Results = HardCodedResults;
+
+					int ExitCode = 0;
+					string Results = UnrealBuildTool.Utils.RunLocalProcessAndReturnStdOut(BloatyExePath, Args, out ExitCode);
+
+					if (ExitCode == 0)
+					{
+						// Get the results from Bloaty as an array of lines
+						string[] LineArray = Results.Split(new char[] { '\r', '\n' }, StringSplitOptions.RemoveEmptyEntries);
+
+						// since we told bloaty to output csv, the result will be already machine readable. Just make sure that the first line is as we expect
+						const string ExpectedCSVOutput = "compileunits,vmsize,filesize";
+						if (LineArray[0] != ExpectedCSVOutput)
+						{
+							throw new AutomationException("CSV output ({0}) does not match expectation ({1})", LineArray[0], ExpectedCSVOutput);
+						}
+
+						uint Total = 0;
+						
+						for (int IdxLine = 1; IdxLine < LineArray.Length; ++IdxLine)
+						{
+							string[] LineElems = LineArray[IdxLine].Split(',');
+							if (LineElems.Length != 3)
+							{
+								throw new AutomationException("CSV line ({0}) does not match expectation ({1}) - comma in source file?", LineArray[IdxLine], ExpectedCSVOutput);
+							}
+
+							string ModuleName = SanitizeModuleName(FindModuleName(LineElems[0]));
+							uint Size = uint.Parse(LineElems[1]);
+							Total += Size;
+
+							if (ValidateModuleName(ModuleName, LineArray[IdxLine]))
+							{
+								if (ModInfos.Modules.ContainsKey(ModuleName))
+								{
+									ModInfos.Modules[ModuleName] += Size;
+								}
+								else
+								{
+									ModInfos.Modules.Add(ModuleName, Size);
+								}
+
+								if (ModInfos.ModuleNameToSources.ContainsKey(ModuleName))
+								{
+									ModInfos.ModuleNameToSources[ModuleName].Add(LineElems[0]);
+								}
+								else
+								{
+									List<string> Sources = new List<string>();
+									Sources.Add(LineElems[0]);
+									ModInfos.ModuleNameToSources.Add(ModuleName, Sources);
+								}
+							}
+							else
+							{
+								ModInfos.SourcesWithoutModule.Add(LineArray[IdxLine]);
+								ModInfos.TotalSizeForSourcesWithoutModule += Size;
+							}
+						}
+
+						// add a _Total metric
+						ModInfos.Modules.Add("_Total", Total);
+
+						return ModInfos;
+					}
+				}
+			}
+
+			return ModInfos;
+		}
+
+		/** Whether the binary on this platform can be itemized with bloaty - expected to be overriden in platform-specific classes */
+		virtual public ModuleInfos PlatformItemizeExecutable(FileReference Binary)
+		{
+			return ItemizeExecutableWithBloaty(Binary);
+		}
+
+		public override void ExecuteBuild()
+		{
+			FileReference ExecutableFileToItemize = ParseRequiredFileReferenceParam("ExeToItemize");
+
+			// Hack for quick iteration on formatting (run bloaty manually and just point this at the csv file)
+			//HardCodedResults = FileReference.ReadAllText(ExecutableFileToItemize);
+			
+			System.Console.WriteLine("\n\n\n\nExecuting itemization of executable: {0}", ExecutableFileToItemize);
+
+			ModuleInfos Info = PlatformItemizeExecutable(ExecutableFileToItemize);
+			if (Info.Modules.Count == 0)
+			{
+				throw new AutomationException("Unable to itemize binary '{0}'", ExecutableFileToItemize.ToString());
+			}
+
+			System.Console.WriteLine("\n\n\n\nExec itemization, here's how we arrived at module names:");
+			foreach (KeyValuePair<string, List<string>> ModuleMap in Info.ModuleNameToSources)
+			{
+				uint SizeForThisModule = Info.Modules[ModuleMap.Key];
+
+				System.Console.WriteLine("{0} (takes {1:F2} MB):", ModuleMap.Key, SizeForThisModule / (1024.0 * 1024.0));
+				foreach (string Source in ModuleMap.Value)
+				{
+					System.Console.WriteLine("\t{0}", Source);
+				}
+			}
+
+			System.Console.WriteLine("\n\n\n\nAdditionally, {0} source files of {1} bytes ({2:F2} MB in total) failed to be categorized:", 
+				Info.SourcesWithoutModule.Count, Info.TotalSizeForSourcesWithoutModule, Info.TotalSizeForSourcesWithoutModule / (1024.0 * 1024.0));
+			if (Info.SourcesWithoutModule.Count > 0)
+			{
+				foreach (string Source in Info.SourcesWithoutModule)
+				{
+					System.Console.WriteLine("\t{0}", Source);
+				}
+			}
+			else
+			{
+				System.Console.WriteLine("None!");
+			}
+
+
+			// sort modules alphabetically and print out
+			List<KeyValuePair<string, uint>> ModInfos = Info.Modules.ToList();
+
+			System.Console.WriteLine("\n\n\n\n-------------------------------------------------------------------------------------------------------------------------------------------------------------------");
+			System.Console.WriteLine("Exec itemization, modules sorted alphabetically by name");
+			System.Console.WriteLine("-------------------------------------------------------------------------------------------------------------------------------------------------------------------");
+			System.Console.WriteLine("Module, Megabytes (float), Bytes (uint)");
+			ModInfos.SortBy(ModInfo => ModInfo.Key);
+			foreach (KeyValuePair<string, uint> ModInfo in ModInfos)
+			{
+				System.Console.WriteLine("{0}, {1:F2}, {2}", ModInfo.Key, ModInfo.Value / (1024.0 * 1024.0), ModInfo.Value);
+			}
+
+			System.Console.WriteLine("\n\n\n\n-------------------------------------------------------------------------------------------------------------------------------------------------------------------");
+			System.Console.WriteLine("Exec itemization, modules sorted by their size");
+			System.Console.WriteLine("-------------------------------------------------------------------------------------------------------------------------------------------------------------------");
+			System.Console.WriteLine("Module, Megabytes (float), Bytes (uint)");
+			ModInfos.SortBy(ModInfo => ModInfo.Value);
+			ModInfos.Reverse();
+			foreach (KeyValuePair<string, uint> ModInfo in ModInfos)
+			{
+				System.Console.WriteLine("{0}, {1:F2}, {2}", ModInfo.Key, ModInfo.Value / (1024.0 * 1024.0), ModInfo.Value);
+			}
+
+			System.Console.WriteLine("\n\n");
+		}
+
+		/** Hack for a quick iteration on formatting - stores the results from bloaty. */
+		//string HardCodedResults;
+	}
+}
