[Gauntlet] Stubs out installation flow changes to enable gradual platform changes
- Adds ITargetDevice.ClearSavedDirectory
- Adds ITargetDevice.CopyAppConfigFiles

Platforms that have these functions implemented change per pass implementation:

Adjusts per pass from FROM
Pass 1
 - InstallApplication (including clearing device cache and copying additional files)
 - Run the app

Pass 2 and beyond
- Run the app

TO
Pass 1
- ClearSavedDirectory
- InstallApplication
- CopyAppConfigFiles
- Run the app

Pass 2 and beyond
- ClearSavedDirectory
- CopyAppConfigFiles
- Run the app

diff --git a/Engine/Source/Programs/AutomationTool/Gauntlet/Framework/Base/Gauntlet.TargetDevice.cs b/Engine/Source/Programs/AutomationTool/Gauntlet/Framework/Base/Gauntlet.TargetDevice.cs
--- a/Engine/Source/Programs/AutomationTool/Gauntlet/Framework/Base/Gauntlet.TargetDevice.cs
+++ b/Engine/Source/Programs/AutomationTool/Gauntlet/Framework/Base/Gauntlet.TargetDevice.cs
@@ -85,8 +85,12 @@ namespace Gauntlet
 		/// <returns></returns>
 		Dictionary<EIntendedBaseCopyDirectory, string> GetPlatformDirectoryMappings();
 
+		void ClearSavedDirectory(UnrealAppConfig AppConfiguration);
+
 		IAppInstall InstallApplication(UnrealAppConfig AppConfiguration);
 
+		void CopyAppConfigurationFiles(UnrealAppConfig AppConfiguration);
+
 		IAppInstance Run(IAppInstall App);
 
 		bool CheckRequiredSettings(List<string> RequiredSettingsList) { return true; }
diff --git a/Engine/Source/Programs/AutomationTool/Gauntlet/Platform/Android/Gauntlet.TargetDeviceAndroid.cs b/Engine/Source/Programs/AutomationTool/Gauntlet/Platform/Android/Gauntlet.TargetDeviceAndroid.cs
--- a/Engine/Source/Programs/AutomationTool/Gauntlet/Platform/Android/Gauntlet.TargetDeviceAndroid.cs
+++ b/Engine/Source/Programs/AutomationTool/Gauntlet/Platform/Android/Gauntlet.TargetDeviceAndroid.cs
@@ -873,6 +873,11 @@ namespace Gauntlet
 			return true;
 		}
 
+		public void ClearSavedDirectory(UnrealAppConfig AppConfiguration)
+		{
+
+		}
+
 		public IAppInstall InstallApplication(UnrealAppConfig AppConfig)
 		{
 			// todo - pass this through
@@ -886,7 +891,7 @@ namespace Gauntlet
 
 			// kill any currently running instance:
 			KillRunningProcess(Build.AndroidPackageName);
-	
+
 			string DeviceStorageQueryCommand = AndroidPlatform.GetStorageQueryCommand();
 			IProcessResult StorageQueryResult = RunAdbDeviceCommand(DeviceStorageQueryCommand);
 			string StorageLocation = StorageQueryResult.Output.Trim(); // "/mnt/sdcard";
@@ -1098,6 +1103,11 @@ namespace Gauntlet
 			return AppInstall;
 		}
 
+		public void CopyAppConfigurationFiles(UnrealAppConfig AppConfiguration)
+		{
+
+		}
+
 		public IAppInstance Run(IAppInstall App)
 		{
 			AndroidAppInstall DroidAppInstall = App as AndroidAppInstall;
diff --git a/Engine/Source/Programs/AutomationTool/Gauntlet/Platform/IOS/Gauntlet.TargetDeviceIOS.cs b/Engine/Source/Programs/AutomationTool/Gauntlet/Platform/IOS/Gauntlet.TargetDeviceIOS.cs
--- a/Engine/Source/Programs/AutomationTool/Gauntlet/Platform/IOS/Gauntlet.TargetDeviceIOS.cs
+++ b/Engine/Source/Programs/AutomationTool/Gauntlet/Platform/IOS/Gauntlet.TargetDeviceIOS.cs
@@ -462,21 +462,24 @@ namespace Gauntlet
 			{
 				throw new AutomationException("Failed to resign application");
 			}
+		}
+
+		public void ClearSavedDirectory(UnrealAppConfig AppConfiguration)
+		{
 
 		}
 
 		// We need to lock around setting up the IPA
 		static object IPALock = new object();
-
 		public IAppInstall InstallApplication(UnrealAppConfig AppConfig)
-		{            
+		{
             IOSBuild Build = AppConfig.Build as IOSBuild;
 
 			// Ensure Build exists
 			if (Build == null)
 			{
 				throw new AutomationException("Invalid build for IOS!");
-			}	
+			}
 
 			bool CacheResigned = false;
 			bool UseLocalExecutable = Globals.Params.ParseParam("dev");
@@ -488,7 +491,7 @@ namespace Gauntlet
 				// device artifact path
 				DeviceArtifactPath = string.Format("/Documents/{0}/Saved", AppConfig.ProjectName);
 
-				CacheResigned = File.Exists(CacheResignedFilename);				
+				CacheResigned = File.Exists(CacheResignedFilename);
 
 				if (CacheResigned && !UseLocalExecutable)
 				{
@@ -496,15 +499,15 @@ namespace Gauntlet
 					{
 						Log.Verbose("App was resigned, invalidating app cache");
 						File.Delete(IPAHashFilename);
-					}								
+					}
 				}
 
 				PrepareIPA(Build);
 
-				// local executable support			
+				// local executable support
 				if (UseLocalExecutable)
-				{					
-					ResignApplication(AppConfig);				
+				{
+					ResignApplication(AppConfig);
 				}
 			}
 
@@ -518,10 +521,10 @@ namespace Gauntlet
 				// remove device artifacts
 				CleanDeviceArtifacts(Build);
 			}
-			
+
 			// parallel iOS tests use same app install folder, so lock it as setup is quick
 			lock (Globals.MainLock)
-			{				
+			{
 				// local app install with additional files, this directory will be mirrored to device in a single operation
 				string AppInstallPath;
 
@@ -578,13 +581,18 @@ namespace Gauntlet
 
 				// store the IPA hash to avoid redundant deployments
 				CopyCommand = String.Format("--bundle_id {0} --upload={1} --to {2}", Build.PackageName, IPAHashFilename, "/Documents/IPAHash.txt");
-				ExecuteIOSDeployCommand(CopyCommand, 120);				
+				ExecuteIOSDeployCommand(CopyCommand, 120);
 			}
 
-			IOSAppInstall IOSApp = new IOSAppInstall(AppConfig.Name, this, Build.PackageName, AppConfig.CommandLine);	
+			IOSAppInstall IOSApp = new IOSAppInstall(AppConfig.Name, this, Build.PackageName, AppConfig.CommandLine);
 			return IOSApp;
 		}
 
+		public void CopyAppConfigurationFiles(UnrealAppConfig AppConfiguration)
+		{
+
+		}
+
 		public void PopulateDirectoryMappings(string ProjectDir)
 		{
 			LocalDirectoryMappings.Add(EIntendedBaseCopyDirectory.Build, Path.Combine(ProjectDir, "Build"));
@@ -607,11 +615,11 @@ namespace Gauntlet
 
 
 		/// <summary>
-		/// Artifact (e.g. Saved) path on the device		
+		/// Artifact (e.g. Saved) path on the device
 		/// </summary>
 		public string DeviceArtifactPath { get; protected set;  }
 
-		
+
 		#region Device State Management
 
 		// NOTE: We check that a default device UUID or the one specifed is connected with 'ios-deploy --detect' at device creation time
diff --git a/Engine/Source/Programs/AutomationTool/Gauntlet/Platform/Linux/Gauntlet.TargetDeviceLinux.cs b/Engine/Source/Programs/AutomationTool/Gauntlet/Platform/Linux/Gauntlet.TargetDeviceLinux.cs
--- a/Engine/Source/Programs/AutomationTool/Gauntlet/Platform/Linux/Gauntlet.TargetDeviceLinux.cs
+++ b/Engine/Source/Programs/AutomationTool/Gauntlet/Platform/Linux/Gauntlet.TargetDeviceLinux.cs
@@ -433,6 +433,10 @@ namespace Gauntlet
 			return LinuxApp;
 		}
 
+		public void ClearSavedDirectory(UnrealAppConfig AppConfiguration)
+		{
+
+		}
 
 		public IAppInstall InstallApplication(UnrealAppConfig AppConfig)
 		{
@@ -456,7 +460,7 @@ namespace Gauntlet
 
 			LinuxApp.WorkingDirectory = Path.GetDirectoryName(EditorBuild.ExecutablePath);
 			LinuxApp.RunOptions = RunOptions;
-	
+
 			// Force this to stop logs and other artifacts going to different places
 			LinuxApp.CommandArguments = AppConfig.CommandLine + string.Format(" -userdir=\"{0}\"", UserDir);
 			LinuxApp.ArtifactPath = Path.Combine(UserDir, @"Saved");
@@ -471,7 +475,12 @@ namespace Gauntlet
 
 			return LinuxApp;
 		}
-		
+
+		public void CopyAppConfigurationFiles(UnrealAppConfig AppConfiguration)
+		{
+
+		}
+
 		public bool CanRunFromPath(string InPath)
 		{
 			return !Utils.SystemHelpers.IsNetworkPath(InPath);
diff --git a/Engine/Source/Programs/AutomationTool/Gauntlet/Platform/Mac/Gauntlet.TargetDeviceMac.cs b/Engine/Source/Programs/AutomationTool/Gauntlet/Platform/Mac/Gauntlet.TargetDeviceMac.cs
--- a/Engine/Source/Programs/AutomationTool/Gauntlet/Platform/Mac/Gauntlet.TargetDeviceMac.cs
+++ b/Engine/Source/Programs/AutomationTool/Gauntlet/Platform/Mac/Gauntlet.TargetDeviceMac.cs
@@ -12,7 +12,7 @@ using static AutomationTool.ProcessResult;
 
 namespace Gauntlet
 {
-	
+
 	public class MacDeviceFactory : IDeviceFactory
 	{
 		public bool CanSupportPlatform(UnrealTargetPlatform? Platform)
@@ -364,7 +364,7 @@ namespace Gauntlet
 				Log.Verbose("LocalFile exists: {0}. Newer: {1}", LocalFileExists, LocalFileNewer);
 
 				if (LocalFileExists && LocalFileNewer)
-				{					
+				{
 					// need to -basedir to have our exe load content from the path that the bundle sits in
 					MacApp.CommandArguments += string.Format(" -basedir={0}", Path.GetDirectoryName(BundlePath));
 					MacApp.ExecutablePath = LocalProjectBinary;
@@ -374,6 +374,10 @@ namespace Gauntlet
 			return MacApp;
 		}
 
+		public void ClearSavedDirectory(UnrealAppConfig AppConfiguration)
+		{
+
+		}
 
 		public IAppInstall InstallApplication(UnrealAppConfig AppConfig)
 		{
@@ -415,6 +419,11 @@ namespace Gauntlet
 			return MacApp;
 		}
 
+		public void CopyAppConfigurationFiles(UnrealAppConfig AppConfiguration)
+		{
+
+		}
+
 		public IAppInstance Run(IAppInstall App)
 		{
 			MacAppInstall MacInstall = App as MacAppInstall;
diff --git a/Engine/Source/Programs/AutomationTool/Gauntlet/Platform/Null/Gauntlet.TargetDeviceNull.cs b/Engine/Source/Programs/AutomationTool/Gauntlet/Platform/Null/Gauntlet.TargetDeviceNull.cs
--- a/Engine/Source/Programs/AutomationTool/Gauntlet/Platform/Null/Gauntlet.TargetDeviceNull.cs
+++ b/Engine/Source/Programs/AutomationTool/Gauntlet/Platform/Null/Gauntlet.TargetDeviceNull.cs
@@ -61,7 +61,7 @@ namespace Gauntlet
 		public ITargetDevice Device { get; protected set; }
 
 		public string Name { get; protected set; }
-			
+
 		public string CommandLine { get; protected set; }
 
 		public NullAppInstall(string InName, TargetDeviceNull InDevice, string InCommandLine)
@@ -151,13 +151,19 @@ namespace Gauntlet
 			return new NullAppInstance(this);
 		}
 
+		public void ClearSavedDirectory(UnrealAppConfig AppConfiguration)
+		{
+
+		}
+
 		public IAppInstall InstallApplication(UnrealAppConfig AppConfig)
 		{
+			return new NullAppInstall(AppConfig.Name, this, AppConfig.CommandLine);
+		}
 
-			NullAppInstall NullApp = new NullAppInstall(AppConfig.Name, this, AppConfig.CommandLine);
+		public void CopyAppConfigurationFiles(UnrealAppConfig AppConfiguration)
+		{
 
-	
-			return NullApp;
 		}
 
 		public UnrealTargetPlatform? Platform { get { return null; } }
diff --git a/Engine/Source/Programs/AutomationTool/Gauntlet/Platform/Windows/Gauntlet.TargetDeviceWindows.cs b/Engine/Source/Programs/AutomationTool/Gauntlet/Platform/Windows/Gauntlet.TargetDeviceWindows.cs
--- a/Engine/Source/Programs/AutomationTool/Gauntlet/Platform/Windows/Gauntlet.TargetDeviceWindows.cs
+++ b/Engine/Source/Programs/AutomationTool/Gauntlet/Platform/Windows/Gauntlet.TargetDeviceWindows.cs
@@ -570,6 +570,10 @@ namespace Gauntlet
 			return WinApp;
 		}
 
+		public void ClearSavedDirectory(UnrealAppConfig AppConfiguration)
+		{
+
+		}
 
 		protected IAppInstall InstallSelfInstallingBuild(UnrealAppConfig AppConfig, IWindowsSelfInstallingBuild Build)
 		{
@@ -628,7 +632,7 @@ namespace Gauntlet
 
 			WinApp.WorkingDirectory = Path.GetDirectoryName(EditorBuild.ExecutablePath);
 			WinApp.RunOptions = RunOptions;
-	
+
 			// Force this to stop logs and other artifacts going to different places
 			WinApp.CommandArguments = AppConfig.CommandLine + string.Format(" -userdir=\"{0}\"", UserDir);
 			WinApp.ArtifactPath = Path.Combine(UserDir, @"Saved");
@@ -643,7 +647,12 @@ namespace Gauntlet
 
 			return WinApp;
 		}
-		
+
+		public void CopyAppConfigurationFiles(UnrealAppConfig AppConfiguration)
+		{
+
+		}
+
 		public bool CanRunFromPath(string InPath)
 		{
 			return !Utils.SystemHelpers.IsNetworkPath(InPath);
diff --git a/Engine/Source/Programs/AutomationTool/Gauntlet/Unreal/Base/Gauntlet.UnrealSession.cs b/Engine/Source/Programs/AutomationTool/Gauntlet/Unreal/Base/Gauntlet.UnrealSession.cs
--- a/Engine/Source/Programs/AutomationTool/Gauntlet/Unreal/Base/Gauntlet.UnrealSession.cs
+++ b/Engine/Source/Programs/AutomationTool/Gauntlet/Unreal/Base/Gauntlet.UnrealSession.cs
@@ -853,6 +853,8 @@ namespace Gauntlet
 					// todo - should this be elsewhere?
 					AppConfig.Sandbox = Sandbox;
 
+					Device.ClearSavedDirectory(AppConfig);
+
 					IAppInstall Install = null;
 					bool bReinstallPerPass = Globals.Params.ParseParam("ReinstallPerPass");
 					if (RolesToInstalls == null || !RolesToInstalls.ContainsKey(Role) || bReinstallPerPass)
@@ -930,6 +932,8 @@ namespace Gauntlet
 						InstallsToRoles[Install] = Role;
 						Log.Info("Using previous install of {0} on {1}", Install.Name, Install.Device.Name);
 					}
+
+					Device.CopyAppConfigurationFiles(AppConfig);
 				}
 
 				if (InstallSuccess == false)
