Horde: Functionality to support multi-agent cooks.

diff --git a/Engine/Source/Programs/AutomationTool/BuildGraph/BuildGraph.cs b/Engine/Source/Programs/AutomationTool/BuildGraph/BuildGraph.cs
--- a/Engine/Source/Programs/AutomationTool/BuildGraph/BuildGraph.cs
+++ b/Engine/Source/Programs/AutomationTool/BuildGraph/BuildGraph.cs
@@ -691,14 +691,14 @@ namespace AutomationTool
 			{
 				if(SingleNode != null)
 				{
-					if(!await BuildNodeAsync(new JobContext(this), Graph, SingleNode, NodeToExecutor, Storage, bWithBanner: true))
+					if(!await BuildNodeAsync(Graph, SingleNode, NodeToExecutor, Storage, bWithBanner: true))
 					{
 						return ExitCode.Error_Unknown;
 					}
 				}
 				else
 				{
-					if(!await BuildAllNodesAsync(new JobContext(this), Graph, NodeToExecutor, Storage))
+					if(!await BuildAllNodesAsync(Graph, NodeToExecutor, Storage))
 					{
 						return ExitCode.Error_Unknown;
 					}
@@ -922,12 +922,11 @@ namespace AutomationTool
 		/// <summary>
 		/// Builds all the nodes in the graph
 		/// </summary>
-		/// <param name="Job">Information about the current job</param>
 		/// <param name="Graph">The graph instance</param>
 		/// <param name="NodeToExecutor">Map from node to executor</param>
 		/// <param name="Storage">The temp storage backend which stores the shared state</param>
 		/// <returns>True if everything built successfully</returns>
-		async Task<bool> BuildAllNodesAsync(JobContext Job, BgGraphDef Graph, Dictionary<BgNodeDef, BgNodeExecutor> NodeToExecutor, TempStorage Storage)
+		async Task<bool> BuildAllNodesAsync(BgGraphDef Graph, Dictionary<BgNodeDef, BgNodeExecutor> NodeToExecutor, TempStorage Storage)
 		{
 			// Build a flat list of nodes to execute, in order
 			BgNodeDef[] NodesToExecute = Graph.Agents.SelectMany(x => x.Nodes).ToArray();
@@ -952,7 +951,7 @@ namespace AutomationTool
 				if(!Storage.IsComplete(NodeToExecute.Name))
 				{
 					Logger.LogInformation("");
-					if(!await BuildNodeAsync(Job, Graph, NodeToExecute, NodeToExecutor, Storage, bWithBanner: false))
+					if(!await BuildNodeAsync(Graph, NodeToExecute, NodeToExecutor, Storage, bWithBanner: false))
 					{
 						return false;
 					} 
@@ -1008,14 +1007,13 @@ namespace AutomationTool
 		/// <summary>
 		/// Build a node
 		/// </summary>
-		/// <param name="Job">Information about the current job</param>
 		/// <param name="Graph">The graph to which the node belongs. Used to determine which outputs need to be transferred to temp storage.</param>
 		/// <param name="Node">The node to build</param>
 		/// <param name="NodeToExecutor">Map from node to executor</param>
 		/// <param name="Storage">The temp storage backend which stores the shared state</param>
 		/// <param name="bWithBanner">Whether to write a banner before and after this node's log output</param>
 		/// <returns>True if the node built successfully, false otherwise.</returns>
-		async Task<bool> BuildNodeAsync(JobContext Job, BgGraphDef Graph, BgNodeDef Node, Dictionary<BgNodeDef, BgNodeExecutor> NodeToExecutor, TempStorage Storage, bool bWithBanner)
+		async Task<bool> BuildNodeAsync(BgGraphDef Graph, BgNodeDef Node, Dictionary<BgNodeDef, BgNodeExecutor> NodeToExecutor, TempStorage Storage, bool bWithBanner)
 		{
 			DirectoryReference RootDir = new DirectoryReference(CommandUtils.CmdEnv.LocalRoot);
 
@@ -1075,7 +1073,7 @@ namespace AutomationTool
 				Console.WriteLine();
 				Logger.LogInformation("========== Starting: {Arg0} ==========", Node.Name);
 			}
-			if(!await NodeToExecutor[Node].ExecuteAsync(Job, TagNameToFileSet))
+			if(!await NodeToExecutor[Node].ExecuteAsync(new JobContext(Node.Name, this), TagNameToFileSet))
 			{
 				return false;
 			}
diff --git a/Engine/Source/Programs/AutomationTool/BuildGraph/JobContext.cs b/Engine/Source/Programs/AutomationTool/BuildGraph/JobContext.cs
--- a/Engine/Source/Programs/AutomationTool/BuildGraph/JobContext.cs
+++ b/Engine/Source/Programs/AutomationTool/BuildGraph/JobContext.cs
@@ -1,4 +1,4 @@
-﻿// Copyright Epic Games, Inc. All Rights Reserved.
+// Copyright Epic Games, Inc. All Rights Reserved.
 
 using System;
 using System.Collections.Generic;
@@ -13,17 +13,32 @@ namespace AutomationTool
 	/// </summary>
 	public class JobContext
 	{
+		/// <summary>
+		/// The current node name
+		/// </summary>
+		public string CurrentNode { get; }
+
 		/// <summary>
 		/// The command that is running the current job.
 		/// </summary>
-		public readonly BuildCommand OwnerCommand;
+		public BuildCommand OwnerCommand { get; }
+
+		/// <summary>
+		/// Constructor
+		/// </summary>
+		/// <param name="InOwnerCommand">The command running the current job</param>
+		public JobContext(BuildCommand InOwnerCommand) : this("Unknown", InOwnerCommand)
+		{
+		}
 
 		/// <summary>
 		/// Constructor
 		/// </summary>
+		/// <param name="InCurrentNode">The current node being executed</param>
 		/// <param name="InOwnerCommand">The command running the current job</param>
-		public JobContext(BuildCommand InOwnerCommand)
+		public JobContext(string InCurrentNode, BuildCommand InOwnerCommand)
 		{
+			CurrentNode = InCurrentNode;
 			OwnerCommand = InOwnerCommand;
 		}
 	}
diff --git a/Engine/Source/Programs/AutomationTool/BuildGraph/Task.cs b/Engine/Source/Programs/AutomationTool/BuildGraph/Task.cs
--- a/Engine/Source/Programs/AutomationTool/BuildGraph/Task.cs
+++ b/Engine/Source/Programs/AutomationTool/BuildGraph/Task.cs
@@ -458,6 +458,34 @@ namespace AutomationTool
 				await FileReference.AppendAllLinesAsync(CleanupScript, NewLines);
 			}
 		}
+
+		/// <summary>
+		/// Name of the environment variable containing a file to write Horde graph updates to
+		/// </summary>
+		public const string GraphUpdateEnvVarName = "UE_HORDE_GRAPH_UPDATE";
+
+		/// <summary>
+		/// Updates the graph currently used by Horde
+		/// </summary>
+		/// <param name="Job">Context for the current job that is being executed</param>
+		public static void UpdateGraphForHorde(JobContext Job)
+		{
+			string exportGraphFile = Environment.GetEnvironmentVariable(GraphUpdateEnvVarName);
+			if (String.IsNullOrEmpty(exportGraphFile))
+			{
+				throw new Exception($"Missing environment variable {GraphUpdateEnvVarName}. This is required to update graphs on Horde.");
+			}
+
+			List<string> newParams = new List<string>();
+			newParams.Add("BuildGraph");
+			newParams.AddRange(Job.OwnerCommand.Params.Select(x => $"-{x}"));
+			newParams.RemoveAll(x => x.StartsWith("-SingleNode=", StringComparison.OrdinalIgnoreCase));
+			newParams.Add($"-HordeExport={exportGraphFile}");
+			newParams.Add($"-ListOnly");
+
+			string newCommandLine = CommandLineArguments.Join(newParams);
+			CommandUtils.RunUAT(CommandUtils.CmdEnv, newCommandLine, "bg");
+		}
 	}
 
 	/// <summary>
diff --git a/Engine/Source/Programs/AutomationTool/BuildGraph/Tasks/CsCompileTask.cs b/Engine/Source/Programs/AutomationTool/BuildGraph/Tasks/CsCompileTask.cs
--- a/Engine/Source/Programs/AutomationTool/BuildGraph/Tasks/CsCompileTask.cs
+++ b/Engine/Source/Programs/AutomationTool/BuildGraph/Tasks/CsCompileTask.cs
@@ -417,7 +417,7 @@ namespace AutomationTool.Tasks
 
 			HashSet<FileReference> BuildProducts = new HashSet<FileReference>();
 			Dictionary<string, HashSet<FileReference>> TagNameToFileSet = new Dictionary<string, HashSet<FileReference>>();
-			await new CsCompileTask(Parameters).ExecuteAsync(new JobContext(null!), BuildProducts, TagNameToFileSet);
+			await new CsCompileTask(Parameters).ExecuteAsync(new JobContext(null!, null!), BuildProducts, TagNameToFileSet);
 
 			FileSet Binaries = FileSet.Empty;
 			FileSet References = FileSet.Empty;
diff --git a/Engine/Source/Programs/AutomationTool/BuildGraph/Tasks/SetVersionTask.cs b/Engine/Source/Programs/AutomationTool/BuildGraph/Tasks/SetVersionTask.cs
--- a/Engine/Source/Programs/AutomationTool/BuildGraph/Tasks/SetVersionTask.cs
+++ b/Engine/Source/Programs/AutomationTool/BuildGraph/Tasks/SetVersionTask.cs
@@ -155,7 +155,7 @@ namespace AutomationTool.Tasks
 		public static async Task<FileSet> ExecuteAsync(BgTaskImpl Task)
 		{
 			HashSet<FileReference> BuildProducts = new HashSet<FileReference>();
-			await Task.ExecuteAsync(new JobContext(null!), BuildProducts, new Dictionary<string, HashSet<FileReference>>());
+			await Task.ExecuteAsync(new JobContext(null!, null!), BuildProducts, new Dictionary<string, HashSet<FileReference>>());
 			return FileSet.FromFiles(Unreal.RootDirectory, BuildProducts);
 		}
 
