Move RetainStagedDirectory setting from launcher profile to project packaging settings.

diff --git a/Engine/Source/Developer/DeveloperToolSettings/Classes/Settings/ProjectPackagingSettings.h b/Engine/Source/Developer/DeveloperToolSettings/Classes/Settings/ProjectPackagingSettings.h
--- a/Engine/Source/Developer/DeveloperToolSettings/Classes/Settings/ProjectPackagingSettings.h
+++ b/Engine/Source/Developer/DeveloperToolSettings/Classes/Settings/ProjectPackagingSettings.h
@@ -630,6 +630,10 @@ public:
 	UPROPERTY(config, EditAnywhere, Category=Packaging, meta=(DisplayName = "Additional builds for this project."))
 	TArray<FProjectBuildSettings> ProjectCustomBuilds;
 
+	/** If set, platforms that destructively edit the iostore containers during packaging will save a copy prior to doing so. */
+	UPROPERTY(config, EditAnywhere, Category = Packaging)
+	bool bRetainStagedDirectory;
+
 	/**
 	 * A list of custom builds, specified in engine ini files, and not editable in editor, that will show up in the Platforms menu to allow customized builds for all projects
 	 */
diff --git a/Engine/Source/Developer/LauncherServices/Private/Launcher/LauncherWorker.cpp b/Engine/Source/Developer/LauncherServices/Private/Launcher/LauncherWorker.cpp
--- a/Engine/Source/Developer/LauncherServices/Private/Launcher/LauncherWorker.cpp
+++ b/Engine/Source/Developer/LauncherServices/Private/Launcher/LauncherWorker.cpp
@@ -884,12 +884,6 @@ FString FLauncherWorker::CreateUATCommand( const ILauncherProfileRef& InProfile,
 		}
 	}
 
-	if (InProfile->IsUsingIoStore() && 
-		InProfile->IsRetainStagedDirectory())
-	{
-		StageAdditionalCommandLine += TEXT(" -RetainStagedDirectory");
-	}
-
 	// stage/package/deploy
 	if (InProfile->GetDeploymentMode() != ELauncherProfileDeploymentModes::DoNotDeploy)
 	{
diff --git a/Engine/Source/Developer/LauncherServices/Private/Profiles/LauncherProfile.h b/Engine/Source/Developer/LauncherServices/Private/Profiles/LauncherProfile.h
--- a/Engine/Source/Developer/LauncherServices/Private/Profiles/LauncherProfile.h
+++ b/Engine/Source/Developer/LauncherServices/Private/Profiles/LauncherProfile.h
@@ -55,6 +55,7 @@ enum ELauncherVersion
 	LAUNCHERSERVICES_ADDEDORIGINALRELEASEVERSION = 33,
 	LAUNCHERSERVICES_ADDBUILDTARGETNAME = 34,
 	LAUNCHERSERVICED_ADDEDRETAINSTAGEDDIRECTORY = 35,
+	LAUNCHERSERVICED_REMOVEDRETAINSTAGEDDIRECTORY = 36,
 	//ADD NEW STUFF HERE
 
 
@@ -578,14 +579,6 @@ public:
 	{
 		ReferenceContainerCryptoKeysFileName = InReferenceContainerCryptoKeysFileName;
 	}
-	virtual bool IsRetainStagedDirectory() const override
-	{
-		return bRetainStagedDirectory;
-	}
-	virtual void SetRetainStagedDirectory(bool bInRetainStagedDirectory) override
-	{
-		bRetainStagedDirectory = bInRetainStagedDirectory;
-	}
 
 	virtual ELauncherProfileDeploymentModes::Type GetDeploymentMode( ) const override
 	{
@@ -1106,10 +1099,11 @@ public:
 			Archive << BuildTargetName;
 		}
 
-		if (Version >= LAUNCHERSERVICED_ADDEDRETAINSTAGEDDIRECTORY)
+		if (Version >= LAUNCHERSERVICED_ADDEDRETAINSTAGEDDIRECTORY && Version < LAUNCHERSERVICED_REMOVEDRETAINSTAGEDDIRECTORY)
 		{
+			bool bRetainStagedDirectory = false;
 			Archive << bRetainStagedDirectory;
-		}		
+		}
 
 		DefaultLaunchRole->Serialize(Archive);
 
@@ -1241,7 +1235,6 @@ public:
 		Writer.WriteValue("BasedOnReleaseVersionName", BasedOnReleaseVersionName);
 		Writer.WriteValue("ReferenceContainerGlobalFileName", ReferenceContainerGlobalFileName);
 		Writer.WriteValue("ReferenceContainerCryptoKeysFileName", ReferenceContainerCryptoKeysFileName);
-		Writer.WriteValue("RetainStagedDirectory", bRetainStagedDirectory);
 		Writer.WriteValue("OriginalReleaseVersionName", OriginalReleaseVersionName);
 		Writer.WriteValue("CreateDLC", CreateDLC);
 		Writer.WriteValue("DLCName", DLCName);
@@ -1582,12 +1575,6 @@ public:
 			}
 		}
 
-		if (bIsStaging)
-		{
-			// (only iostore uses this)
-			Writer.WriteValue("RetainStagedDirectory", IsRetainStagedDirectory());
-		}
-
 		/*
 		"script", ""
 		"project", ""
@@ -1918,11 +1905,6 @@ public:
 			OriginalReleaseVersionName.Empty();
 		}
 
-		if (Version >= LAUNCHERSERVICED_ADDEDRETAINSTAGEDDIRECTORY)
-		{
-			bRetainStagedDirectory = Object.GetBoolField("RetainStagedDirectory");
-		}
-
 		CreateDLC = Object.GetBoolField("CreateDLC");
 		DLCName = Object.GetStringField("DLCName");
 		bGenerateChunks = Object.GetBoolField("GenerateChunks");
@@ -3152,10 +3134,6 @@ private:
 	// the json file containing the keys.
 	FString ReferenceContainerCryptoKeysFileName;
 
-	// Some platforms modify iostore containers such that they can't be read as a reference container - if this is true,
-	// the deployment stage will save the staged directory beforehand..
-	bool bRetainStagedDirectory;
-
 	// create a release version of the content (this can be used to base dlc / patches from)
 	bool CreateReleaseVersion;
 
diff --git a/Engine/Source/Developer/LauncherServices/Public/ILauncherProfile.h b/Engine/Source/Developer/LauncherServices/Public/ILauncherProfile.h
--- a/Engine/Source/Developer/LauncherServices/Public/ILauncherProfile.h
+++ b/Engine/Source/Developer/LauncherServices/Public/ILauncherProfile.h
@@ -1150,14 +1150,6 @@ public:
 	virtual FString GetReferenceContainerCryptoKeysFileName() const = 0;
 	virtual void SetReferenceContainerCryptoKeysFileName(const FString& InReferenceContainerCryptoKeysFileName) = 0;
 
-	/**
-	* Whether or not to save the unaltered staged directory for platforms that modify
-	* i/o store containers for deployment. This is so that the containers can be saved and used as a reference 
-	* container later.
-	*/
-	virtual bool IsRetainStagedDirectory() const = 0;
-	virtual void SetRetainStagedDirectory(bool bInRetainStagedDirectory) = 0;
-
 	/**
 	 * Sets if we are going to generate a patch 
 	 * 
diff --git a/Engine/Source/Developer/ProjectLauncher/Private/Widgets/Package/SProjectLauncherPackagingSettings.cpp b/Engine/Source/Developer/ProjectLauncher/Private/Widgets/Package/SProjectLauncherPackagingSettings.cpp
--- a/Engine/Source/Developer/ProjectLauncher/Private/Widgets/Package/SProjectLauncherPackagingSettings.cpp
+++ b/Engine/Source/Developer/ProjectLauncher/Private/Widgets/Package/SProjectLauncherPackagingSettings.cpp
@@ -145,25 +145,6 @@ void SProjectLauncherPackagingSettings::Construct(const FArguments& InArgs, cons
 								]
 							]
 
-						+ SVerticalBox::Slot()
-							.AutoHeight()
-							.Padding(0.0, 4.0, 0.0, 0.0)
-							[
-								SNew(SCheckBox)
-								.IsEnabled(this, &SProjectLauncherPackagingSettings::IsEditable)
-								.IsChecked(this, &SProjectLauncherPackagingSettings::HandleRetainStagedDirectoryCheckBoxIsChecked)
-								.OnCheckStateChanged(this, &SProjectLauncherPackagingSettings::HandleRetainStagedDirectoryCheckStateChanged)
-								.Padding(FMargin(4.0f, 0.0f))
-								.ToolTipText(LOCTEXT("RetainStagedDirectoryCheckBoxTooltip", "Some platforms modify I/O store containers before deployment which prevents using them as reference chunk database. Set this to true to save the staged directory beforehand."))
-								.Content()
-								[
-									SNew(STextBlock)
-									.Text(LOCTEXT("RetainStagedDirectoryCheckBoxText", "Retain the staged directory if the platform will modify the I/O store containers for deployment"))
-								]
-							]
-
-
-
 						+ SVerticalBox::Slot()
 							.AutoHeight()
 							.Padding(0.0, 4.0, 0.0, 0.0)
@@ -326,29 +307,6 @@ ECheckBoxState SProjectLauncherPackagingSettings::HandleIncludePrerequisitesChec
 	return ECheckBoxState::Unchecked;
 }
 
-
-void SProjectLauncherPackagingSettings::HandleRetainStagedDirectoryCheckStateChanged(ECheckBoxState NewState)
-{
-	ILauncherProfilePtr SelectedProfile = Model->GetSelectedProfile();
-	if (SelectedProfile.IsValid())
-	{
-		SelectedProfile->SetRetainStagedDirectory(NewState == ECheckBoxState::Checked);
-	}
-}
-
-ECheckBoxState SProjectLauncherPackagingSettings::HandleRetainStagedDirectoryCheckBoxIsChecked() const
-{
-	ILauncherProfilePtr SelectedProfile = Model->GetSelectedProfile();
-	if (SelectedProfile.IsValid())
-	{
-		if (SelectedProfile->IsRetainStagedDirectory())
-		{
-			return ECheckBoxState::Checked;
-		}
-	}
-	return ECheckBoxState::Unchecked;
-}
-
 void SProjectLauncherPackagingSettings::HandleUseIoStoreCheckStateChanged(ECheckBoxState NewState)
 {
 	ILauncherProfilePtr SelectedProfile = Model->GetSelectedProfile();
diff --git a/Engine/Source/Developer/ProjectLauncher/Private/Widgets/Package/SProjectLauncherPackagingSettings.h b/Engine/Source/Developer/ProjectLauncher/Private/Widgets/Package/SProjectLauncherPackagingSettings.h
--- a/Engine/Source/Developer/ProjectLauncher/Private/Widgets/Package/SProjectLauncherPackagingSettings.h
+++ b/Engine/Source/Developer/ProjectLauncher/Private/Widgets/Package/SProjectLauncherPackagingSettings.h
@@ -53,9 +53,6 @@ private:
 	void HandleUseIoStoreCheckStateChanged(ECheckBoxState NewState);
 	ECheckBoxState HandleUseIoStoreCheckBoxIsChecked() const;
 
-	void HandleRetainStagedDirectoryCheckStateChanged(ECheckBoxState NewState);
-	ECheckBoxState HandleRetainStagedDirectoryCheckBoxIsChecked() const;
-
 	void HandleMakeBinaryConfigCheckStateChanged(ECheckBoxState NewState);
 	ECheckBoxState HandleMakeBinaryConfigCheckBoxIsChecked() const;
 
diff --git a/Engine/Source/Programs/AutomationTool/AutomationUtils/ProjectParams.cs b/Engine/Source/Programs/AutomationTool/AutomationUtils/ProjectParams.cs
--- a/Engine/Source/Programs/AutomationTool/AutomationUtils/ProjectParams.cs
+++ b/Engine/Source/Programs/AutomationTool/AutomationUtils/ProjectParams.cs
@@ -401,7 +401,6 @@ namespace AutomationTool
 			this.ServerDevice = InParams.ServerDevice;
             this.NullRHI = InParams.NullRHI;
 			this.WriteBackMetadataToAssetRegistry = InParams.WriteBackMetadataToAssetRegistry;
-			this.RetainStagedDirectory = InParams.RetainStagedDirectory;
             this.FakeClient = InParams.FakeClient;
             this.EditorTest = InParams.EditorTest;
             this.RunAutomationTests = InParams.RunAutomationTests;
@@ -614,8 +613,7 @@ namespace AutomationTool
 			string SessionLabel = null,
 			ParamList<string> InMapsToRebuildLightMaps = null,
             ParamList<string> InMapsToRebuildHLOD = null,
-            ParamList<string> TitleID = null,
-            bool? RetainStagedDirectory = null
+            ParamList<string> TitleID = null
 			)
 		{
 			//
@@ -983,7 +981,6 @@ namespace AutomationTool
 			this.UbtArgs = ParseParamValueIfNotSpecified(Command, UbtArgs, "ubtargs", String.Empty);
 			this.AdditionalPackageOptions = ParseParamValueIfNotSpecified(Command, AdditionalPackageOptions, "AdditionalPackageOptions", String.Empty);
 			this.WriteBackMetadataToAssetRegistry = ParseParamValueIfNotSpecified(Command, WriteBackMetadataToAssetRegistry, "WriteBackMetadataToAssetRegistry", String.Empty);
-			this.RetainStagedDirectory = GetParamValueIfNotSpecified(Command, RetainStagedDirectory, this.RetainStagedDirectory, "RetainStagedDirectory");
 
 			string SpecifiedArchString, ServerArchString, EditorArchString, ClientArchString, ProgramArchString;
 			SpecifiedArchString = ParseParamValueIfNotSpecified(Command, SpecifiedArchitecture, "specifiedarchitecture", null);
@@ -2203,9 +2200,6 @@ namespace AutomationTool
         [Help("WriteBackMetadataToAssetRegistry", "Passthru to iostore staging, see IoStoreUtilities.cpp")]
         public string WriteBackMetadataToAssetRegistry;
 
-        [Help("RetainStagedDirectory", "If set, retain the staged directory for platforms that modify the I/O store containers for deployment. This is necessary for using the reference container for patch preventing on such platforms.")]
-        public bool RetainStagedDirectory;
-
         /// <summary>
         /// Run:adds ?fake to the server URL
         /// </summary>
@@ -3246,7 +3240,6 @@ namespace AutomationTool
 				Logger.LogDebug("ForcePackageData={ForcePackageData}", ForcePackageData);
 				Logger.LogDebug("NullRHI={NullRHI}", NullRHI);
 				Logger.LogDebug("WriteBackMetadataToAssetRegistry={WriteBackMetadataToAssetRegistry}", WriteBackMetadataToAssetRegistry);
-				Logger.LogDebug("RetainStagedDirectory={RetainStagedDirectory}", RetainStagedDirectory);
 				Logger.LogDebug("FakeClient={FakeClient}", FakeClient);
                 Logger.LogDebug("EditorTest={EditorTest}", EditorTest);
                 Logger.LogDebug("RunAutomationTests={RunAutomationTests}", RunAutomationTests);
