Horde: Fix UAT build error.

diff --git a/Engine/Source/Programs/AutomationTool/BuildGraph/Tasks/DeployToolTask.cs b/Engine/Source/Programs/AutomationTool/BuildGraph/Tasks/DeployToolTask.cs
--- a/Engine/Source/Programs/AutomationTool/BuildGraph/Tasks/DeployToolTask.cs
+++ b/Engine/Source/Programs/AutomationTool/BuildGraph/Tasks/DeployToolTask.cs
@@ -19,6 +19,7 @@ using EpicGames.Horde.Storage;
 using EpicGames.Horde.Storage.Nodes;
 using System.Threading;
 using System.Data;
+using EpicGames.Horde.Storage.Backends;
 
 #nullable enable
 
@@ -141,7 +142,7 @@ namespace AutomationTool.Tasks
 			HttpClient CreateHttpClient()
 			{
 				HttpClient httpClient = new HttpClient();
-				httpClient.BaseAddress = new Uri(serverUri, $"api/v1/tools/{Parameters.Id}/");
+				httpClient.BaseAddress = serverUri;
 				if (settings?.Token != null)
 				{
 					httpClient.DefaultRequestHeaders.Authorization = new AuthenticationHeaderValue("Bearer", settings.Token);
@@ -151,7 +152,9 @@ namespace AutomationTool.Tasks
 
 			BundleNodeHandle handle;
 
-			HttpStorageClient storageClient = new HttpStorageClient(CreateHttpClient, () => new HttpClient(), BundleReaderCache.None, Logger);
+			string basePath = $"api/v1/tools/{Parameters.Id}";
+			HttpStorageBackend storageBackend = new HttpStorageBackend(basePath, CreateHttpClient, Logger);
+			HttpStorageClient storageClient = new HttpStorageClient(basePath, CreateHttpClient, storageBackend, BundleReaderCache.None, Logger);
 			await using (BundleWriter treeWriter = storageClient.CreateWriter())
 			{
 				DirectoryNode sandbox = new DirectoryNode();
