Fix build error.

diff --git a/Engine/Source/Programs/AutomationTool/BuildGraph/Tasks/DeployToolTask.cs b/Engine/Source/Programs/AutomationTool/BuildGraph/Tasks/DeployToolTask.cs
--- a/Engine/Source/Programs/AutomationTool/BuildGraph/Tasks/DeployToolTask.cs
+++ b/Engine/Source/Programs/AutomationTool/BuildGraph/Tasks/DeployToolTask.cs
@@ -142,7 +142,7 @@ namespace AutomationTool.Tasks
 			HttpClient CreateHttpClient()
 			{
 				HttpClient httpClient = new HttpClient();
-				httpClient.BaseAddress = serverUri;
+				httpClient.BaseAddress = new Uri(serverUri, $"api/v1/tools/{Parameters.Id}/");
 				if (settings?.Token != null)
 				{
 					httpClient.DefaultRequestHeaders.Authorization = new AuthenticationHeaderValue("Bearer", settings.Token);
@@ -152,9 +152,7 @@ namespace AutomationTool.Tasks
 
 			BundleNodeHandle handle;
 
-			string basePath = $"api/v1/tools/{Parameters.Id}";
-			HttpStorageBackend storageBackend = new HttpStorageBackend(basePath, CreateHttpClient, Logger);
-			HttpStorageClient storageClient = new HttpStorageClient(basePath, CreateHttpClient, storageBackend, BundleReaderCache.None, Logger);
+			HttpStorageClient storageClient = new HttpStorageClient(CreateHttpClient, BundleReaderCache.None, Logger);
 			await using (BundleWriter treeWriter = storageClient.CreateWriter())
 			{
 				DirectoryNode sandbox = new DirectoryNode();
