[Gauntlet] Rename TestInstallAndRunBase::TestCopyFilesToCopy to TestInstallAndRunBase::TestCopyAppConfigFiles
#rnx

diff --git a/Engine/Source/Programs/AutomationTool/Gauntlet/Platform/Mac/Gauntlet.SelfTest.TestUnrealInstallAndRunMac.cs b/Engine/Source/Programs/AutomationTool/Gauntlet/Platform/Mac/Gauntlet.SelfTest.TestUnrealInstallAndRunMac.cs
--- a/Engine/Source/Programs/AutomationTool/Gauntlet/Platform/Mac/Gauntlet.SelfTest.TestUnrealInstallAndRunMac.cs
+++ b/Engine/Source/Programs/AutomationTool/Gauntlet/Platform/Mac/Gauntlet.SelfTest.TestUnrealInstallAndRunMac.cs
@@ -27,7 +27,7 @@ namespace Gauntlet.SelfTest
 			return true;
 		}
 
-		protected override bool TestCopyFilesToCopyToDevice()
+		protected override bool TestCopyAppConfigurationFiles()
 		{
 			return true;
 		}
diff --git a/Engine/Source/Programs/AutomationTool/Gauntlet/SelfTest/Unreal/Gauntlet.SelfTest.TestUnrealInstallAndRunBase.cs b/Engine/Source/Programs/AutomationTool/Gauntlet/SelfTest/Unreal/Gauntlet.SelfTest.TestUnrealInstallAndRunBase.cs
--- a/Engine/Source/Programs/AutomationTool/Gauntlet/SelfTest/Unreal/Gauntlet.SelfTest.TestUnrealInstallAndRunBase.cs
+++ b/Engine/Source/Programs/AutomationTool/Gauntlet/SelfTest/Unreal/Gauntlet.SelfTest.TestUnrealInstallAndRunBase.cs
@@ -87,7 +87,7 @@ namespace Gauntlet.SelfTest
 
 			// InstallApplication populates the directory mappings, the ensure copy additional files has proper handling for empty mappings, clear them out now
 			TargetDevice.GetPlatformDirectoryMappings().Clear();
-			if (!CheckResult(TestCopyFilesToCopyToDevice(), "Could not copy {0} additional files to device {1}", AppConfig.FilesToCopy.Count, TargetDevice))
+			if (!CheckResult(TestCopyAppConfigurationFiles(), "Could not copy {0} additional files to device {1}", AppConfig.FilesToCopy.Count, TargetDevice))
 			{
 				return false;
 			}
@@ -127,7 +127,7 @@ namespace Gauntlet.SelfTest
 		protected abstract bool TestInstallApplication(out IAppInstall Install);
 
 		// Should add a file to the AppConfig, copy it over, and then verify it exists in the target diretory
-		protected abstract bool TestCopyFilesToCopyToDevice();
+		protected abstract bool TestCopyAppConfigurationFiles();
 
 		protected FileInfo CreateDummyFile()
 		{
