Horde: Add a matcher for Docker disk space errors.

diff --git a/Engine/Source/Programs/AutomationTool/AutomationUtils/Matchers/DockerEventMatcher.cs b/Engine/Source/Programs/AutomationTool/AutomationUtils/Matchers/DockerEventMatcher.cs
new file mode 100644
--- /dev/null
+++ b/Engine/Source/Programs/AutomationTool/AutomationUtils/Matchers/DockerEventMatcher.cs
@@ -0,0 +1,28 @@
+// Copyright Epic Games, Inc. All Rights Reserved.
+
+using System.Text.RegularExpressions;
+using EpicGames.Core;
+using Microsoft.Extensions.Logging;
+
+#nullable enable
+
+namespace AutomationUtils.Matchers
+{
+	/// <summary>
+	/// Matcher for docker errors
+	/// </summary>
+	class DockerEventMatcher : ILogEventMatcher
+	{
+		static readonly Regex s_noSpacePattern = new Regex(": no space left on device$");
+
+		public LogEventMatch? Match(ILogCursor cursor)
+		{
+			if (cursor.IsMatch(s_noSpacePattern))
+			{
+				LogEventBuilder builder = new LogEventBuilder(cursor);
+				return builder.ToMatch(LogEventPriority.Normal, LogLevel.Error, KnownLogEvents.Systemic_OutOfDiskSpace);
+			}
+			return null;
+		}
+	}
+}
diff --git a/Engine/Source/Programs/Horde/Horde.Agent.Tests/LogParserTests.cs b/Engine/Source/Programs/Horde/Horde.Agent.Tests/LogParserTests.cs
--- a/Engine/Source/Programs/Horde/Horde.Agent.Tests/LogParserTests.cs
+++ b/Engine/Source/Programs/Horde/Horde.Agent.Tests/LogParserTests.cs
@@ -592,6 +592,20 @@ namespace Horde.Agent.Tests
 			CheckEventGroup(logEvents, 2, 2, LogLevel.Warning, KnownLogEvents.Linker);
 		}
 
+		[TestMethod]
+		public void DockerDiskSpaceMatcher()
+		{
+			string[] lines =
+			{
+				@"  Horde.Server -> /app/out/",
+				@"Error processing tar file(exit status 1): write /app/Source/Programs/Horde/Horde.Server/obj/Release/net6.0/Horde.Server.dll: no space left on device",
+				@"Took 32.59s to run docker, ExitCode=1",
+			};
+
+			List<LogEvent> logEvents = Parse(lines);
+			CheckEventGroup(logEvents, 1, 1, LogLevel.Error, KnownLogEvents.Systemic_OutOfDiskSpace);
+		}
+
 		[TestMethod]
 		public void SystemicErrorMatcher()
 		{
diff --git a/Engine/Source/Programs/Shared/EpicGames.Core/KnownLogEvents.cs b/Engine/Source/Programs/Shared/EpicGames.Core/KnownLogEvents.cs
--- a/Engine/Source/Programs/Shared/EpicGames.Core/KnownLogEvents.cs
+++ b/Engine/Source/Programs/Shared/EpicGames.Core/KnownLogEvents.cs
@@ -311,6 +311,11 @@ namespace EpicGames.Core
 		/// </summary>
 		public static EventId Systemic_MissingFileList { get; } = new EventId(767);
 
+		/// <summary>
+		/// Out of disk space
+		/// </summary>
+		public static EventId Systemic_OutOfDiskSpace { get; } = new EventId(768);
+
 		/// <summary>
 		/// Maximum code for systemic events. Add new events in the 700-799 range.
 		/// </summary>
