[Gauntlet] Updates UnrealSession.SaveRoleArtifacts
- Artifacts are now moved from the device to the destination directory instead of being copied
- User can opt-in to do a copy instead with -RetainDeviceArtifacts
- Adds a Self-Test that verifies the integrity of the SaveRoleArtifacts unit - this iteration and future ones
- Some basic code cleanup/re-ordering

#tests Self-Test

diff --git a/Engine/Source/Programs/AutomationTool/Gauntlet/SelfTest/Unreal/Gauntlet.SelfTest.TestUnrealSession.cs b/Engine/Source/Programs/AutomationTool/Gauntlet/SelfTest/Unreal/Gauntlet.SelfTest.TestUnrealSession.cs
--- a/Engine/Source/Programs/AutomationTool/Gauntlet/SelfTest/Unreal/Gauntlet.SelfTest.TestUnrealSession.cs
+++ b/Engine/Source/Programs/AutomationTool/Gauntlet/SelfTest/Unreal/Gauntlet.SelfTest.TestUnrealSession.cs
@@ -1,16 +1,16 @@
 ﻿// Copyright Epic Games, Inc. All Rights Reserved.
 
-using Gauntlet;
+using AutomationTool;
 using System;
 using System.Collections.Generic;
-using System.Diagnostics;
+using System.Data;
 using System.IO;
 using System.Linq;
 using System.Net;
-using System.Text;
 using System.Threading;
-using System.Threading.Tasks;
+using System.Drawing;
 using UnrealBuildTool;
+using System.Drawing.Imaging;
 
 namespace Gauntlet.SelfTest
 {
@@ -31,7 +31,7 @@ namespace Gauntlet.SelfTest
 				new TargetDeviceWindows("Local PC1", Globals.TempDir)
 				, new TargetDeviceWindows("Local PC2", Globals.TempDir)
 			});
-			
+
 			// Create a new build (params come from our base class will be similar to "OrionGame" and "p:\builds\orion\branch-cl")
 			UnrealBuildSource Build = new UnrealBuildSource(this.ProjectName, this.ProjectFile, this.UnrealPath, this.UsesSharedBuildType, this.BuildPath, new string[] { "" });
 
@@ -87,4 +87,322 @@ namespace Gauntlet.SelfTest
 			MarkComplete(TestResult.Passed);
 		}
 	}
+
+	class TestUnrealSessionSaveRoleArtifacts : TestUnrealBase
+	{
+		// Cached references to mocked dependencies
+		private UnrealSession Session;
+		private UnrealSessionRole Client;
+		private UnrealSessionRole Server;
+		private UnrealTestContext Context;
+		private UnrealSessionInstance Instance;
+
+		// Maps a role type to the SessionRole's artifact (saved) directory
+		private Dictionary<UnrealTargetRole, DirectoryInfo> SourceArtifactDirectories;
+
+		// Cached references to files on the
+		private string ClientCrashReporterFileName;
+		private string ClientDummyFileName;
+		private string ClientBuildFileName;
+		private DirectoryInfo ClientBuildDirectory;
+
+		private string ServerCrashReporterFileName;
+		private string ServerDummyFileName;
+		private string ServerBuildFileName;
+		private DirectoryInfo ServerBuildDirectory;
+
+		public TestUnrealSessionSaveRoleArtifacts()
+		{
+			if(Utils.SystemHelpers.IsNetworkPath(BuildPath))
+			{
+				throw new AutomationException("{0} expects to write temp files to the build directory. You to provide a local build path!", this);
+			}
+			SourceArtifactDirectories = new();
+			if (Directory.Exists(Gauntlet.Globals.LogDir))
+			{
+				Directory.Delete(Gauntlet.Globals.LogDir, true);
+			}
+		}
+
+		public override void TickTest()
+		{
+			MarkComplete();
+		}
+
+		public override bool StartTest(int Pass, int NumPasses)
+		{
+			MockDependencies();
+
+			// To test this, we place a few dummy files in the source artifact directory
+			// We then call SaveRoleArtifacts and check the destination artifact directory
+			// to verify files appeared as we would expect.
+			// We also verify the files in the source directory are no longer there
+			// because we want them to move intead of being copied.
+
+			try
+			{
+				// Generate some files in the source artifact directory
+				List<FileInfo> SourceArtifacts = CreateSourceArtifacts();
+
+				// Call SaveRoleArtifacts
+				List<UnrealRoleArtifacts> Artifacts = Session.SaveRoleArtifacts(Context, Instance, Gauntlet.Globals.LogDir).ToList();
+
+				// Ensure the source artifacts are deleted
+				if (!VerifySourceFilesAreDeleted(SourceArtifacts))
+				{
+					Log.Error("Failed to verify SaveRoleArtifacts deleted source artifacts. See above for details.");
+					return false;
+				}
+
+				// Verify files exist in the destination location
+				if (!VerifyArtifactsAreInDestinationDirectory(Artifacts))
+				{
+					Log.Error("Failed to verify SaveRoleArtifacts moved over all expected artifact files. See above for details.");
+					return false;
+				}
+
+			}
+			finally
+			{
+				// Avoid dirtying the build folder
+				FileInfo ClientBuildFile = new(Path.Combine(ClientBuildDirectory.FullName, ClientBuildFileName));
+				if (ClientBuildFile.Exists)
+				{
+					ClientBuildFile.Delete();
+				}
+
+				FileInfo ServerBuildFile = new(Path.Combine(ServerBuildDirectory.FullName, ServerBuildFileName));
+				if (ServerBuildFile.Exists)
+				{
+					ServerBuildFile.Delete();
+				}
+			}
+
+			InnerStatus = TestStatus.Complete;
+			return true;
+		}
+
+		private void MockDependencies()
+		{
+			UnrealTargetPlatform Platform = UnrealTargetPlatform.Win64;
+			UnrealBuildSource Build = new(ProjectName, ProjectFile, UnrealPath, UsesSharedBuildType, BuildPath);
+			DevicePool.Instance.SetLocalOptions(Gauntlet.Globals.TempDir, false, null);
+
+			// Set up roles
+			UnrealDeviceTargetConstraint DefaultConstraint = new UnrealDeviceTargetConstraint(UnrealTargetPlatform.Win64);
+			List<EIntendedBaseCopyDirectory> DefaultAdditionalDirectories = new();
+			Client = new UnrealSessionRole(UnrealTargetRole.Client, Platform, UnrealTargetConfiguration.Test, null)
+			{
+				Constraint = DefaultConstraint,
+				AdditionalArtifactDirectories = DefaultAdditionalDirectories
+			};
+			Server = new UnrealSessionRole(UnrealTargetRole.Server, Platform, UnrealTargetConfiguration.Test, null)
+			{
+				Constraint = DefaultConstraint,
+				AdditionalArtifactDirectories = DefaultAdditionalDirectories
+			};
+			List<UnrealSessionRole> Roles = new() { Client, Server };
+
+			// TestContext
+			UnrealTestRoleContext ClientContext = new();
+			ClientContext.Type = Client.RoleType;
+			ClientContext.Platform = Platform;
+			ClientContext.Configuration = Client.Configuration;
+
+			UnrealTestRoleContext ServerContext = new();
+			ServerContext.Type = Server.RoleType;
+			ServerContext.Platform = Platform;
+			ServerContext.Configuration = Server.Configuration;
+
+			Dictionary<UnrealTargetRole, UnrealTestRoleContext> RoleContexts = new()
+			{
+				{ Client.RoleType, ClientContext },
+				{ Server.RoleType, ServerContext }
+			};
+
+			Session = new UnrealSession(Build, Roles);
+			Session.TryReserveDevices();
+
+			TargetDeviceWindows ClientDevice = (TargetDeviceWindows)Session.UnrealDeviceReservation.ReservedDevices[0];
+			TargetDeviceWindows ServerDevice = (TargetDeviceWindows)Session.UnrealDeviceReservation.ReservedDevices[1];
+			ClientDevice.PopulateDirectoryMappings(Path.Combine(BuildPath, "WindowsClient", ProjectName), Path.Combine(ClientDevice.LocalCachePath, "UserDir"));
+			ServerDevice.PopulateDirectoryMappings(Path.Combine(BuildPath, "WindowsServer", ProjectName), Path.Combine(ServerDevice.LocalCachePath, "UserDir"));
+
+			// SessionInstance
+			UnrealSessionInstance.RoleInstance ClientRoleInstance = new(Client, GetAppInstanceForDevice(ClientDevice, Build.BuildName));
+			UnrealSessionInstance.RoleInstance ServerRoleInstance = new(Server, GetAppInstanceForDevice(ServerDevice, Build.BuildName));
+			UnrealSessionInstance.RoleInstance[] RoleInstances = { ClientRoleInstance, ServerRoleInstance };
+
+			ClientBuildDirectory = new (ClientRoleInstance.AppInstance.Device.GetPlatformDirectoryMappings()[EIntendedBaseCopyDirectory.Build]);
+			ServerBuildDirectory = new(ServerRoleInstance.AppInstance.Device.GetPlatformDirectoryMappings()[EIntendedBaseCopyDirectory.Build]);
+
+			SourceArtifactDirectories.Add(Client.RoleType, new DirectoryInfo(ClientRoleInstance.AppInstance.ArtifactPath));
+			SourceArtifactDirectories.Add(Server.RoleType, new DirectoryInfo(ServerRoleInstance.AppInstance.ArtifactPath));
+
+			Context = new UnrealTestContext(Build, RoleContexts, new UnrealTestOptions());
+			Instance = new UnrealSessionInstance(RoleInstances);
+		}
+
+		private List<FileInfo> CreateSourceArtifacts()
+		{
+			DirectoryInfo ClientArtifactSource = SourceArtifactDirectories[UnrealTargetRole.Client];
+			DirectoryInfo ServerArtifactSource = SourceArtifactDirectories[UnrealTargetRole.Server];
+			string CrashReporterPath = Path.Combine("Config", "CrashReportClient", "UECC-Windows-ABCabc123");
+
+			// Client
+			DirectoryInfo ClientCrashReporterDirectory = new(Path.Combine(ClientArtifactSource.FullName, CrashReporterPath));
+			FileInfo ClientFileWithinLongCrashReporterDirectory = CreateDummyFileInDirectory(ClientCrashReporterDirectory);
+			FileInfo ClientBasicFile = CreateDummyFileInDirectory(ClientArtifactSource);
+			FileInfo ClientBuildFile = CreateDummyFileInDirectory(ClientBuildDirectory);
+			ClientCrashReporterFileName = ClientFileWithinLongCrashReporterDirectory.Name;
+			ClientDummyFileName = ClientBasicFile.Name;
+			ClientBuildFileName = ClientBuildFile.Name;
+
+			Client.AdditionalArtifactDirectories.Add(EIntendedBaseCopyDirectory.Build);
+
+			// Create a couple images for the client - these should get converted to a .gif
+			DirectoryInfo ScreenshotDirectory = new(Path.Combine(ClientArtifactSource.FullName, "Screenshots", "Win64"));
+			FileInfo RedImage = CreateDummyPNGInDirectory("Red", ScreenshotDirectory);
+			FileInfo BlueImage = CreateDummyPNGInDirectory("Blue", ScreenshotDirectory);
+
+			// Server
+			DirectoryInfo ServerCrashReporterDirectory = new(Path.Combine(ServerArtifactSource.FullName, CrashReporterPath));
+			FileInfo ServerFileWithinLongCrashReporterDirectory = CreateDummyFileInDirectory(ServerCrashReporterDirectory);
+			FileInfo ServerBasicFile = CreateDummyFileInDirectory(ServerArtifactSource);
+			FileInfo ServerBuildFile = CreateDummyFileInDirectory(ServerBuildDirectory);
+			ServerCrashReporterFileName = ServerFileWithinLongCrashReporterDirectory.Name;
+			ServerDummyFileName = ServerBasicFile.Name;
+			ServerBuildFileName = ServerBuildFile.Name;
+
+			Server.AdditionalArtifactDirectories.Add(EIntendedBaseCopyDirectory.Build);
+
+			return new List<FileInfo>()
+			{
+				// Build files created to test AdditionalArtifactDirectories should be copied and not moved, so don't include them here
+				ClientFileWithinLongCrashReporterDirectory, ClientBasicFile,
+				ServerFileWithinLongCrashReporterDirectory, ServerBasicFile,
+				RedImage, BlueImage
+			};
+		}
+
+		private bool VerifyArtifactsAreInDestinationDirectory(List<UnrealRoleArtifacts> Artifacts)
+		{
+			UnrealRoleArtifacts ClientArtifacts = Artifacts.Where(Role => Role.SessionRole.RoleType == UnrealTargetRole.Client).First();
+			UnrealRoleArtifacts ServerArtifacts = Artifacts.Where(Role => Role.SessionRole.RoleType == UnrealTargetRole.Server).First();
+
+			DirectoryInfo ClientDirectory = new(ClientArtifacts.ArtifactPath);
+			DirectoryInfo ServerDirectory = new(ServerArtifacts.ArtifactPath);
+			string TruncatedCrashReporterPath = Path.Combine("Config", "CrashReportClient", "UECC-Windows-00");
+
+			FileInfo ClientCrashReporterFile = new(Path.Combine(ClientDirectory.FullName, TruncatedCrashReporterPath, ClientCrashReporterFileName));
+			FileInfo ClientBasicFile = new(Path.Combine(ClientDirectory.FullName, ClientDummyFileName));
+			FileInfo ClientBuildFile = new(Path.Combine(ClientDirectory.FullName, "Build", ClientBuildFileName));
+
+			// We should see the .pngs were converted to .jpgs and a .gif file in the root directory
+			DirectoryInfo ScreenshotDirectory = new(Path.Combine(ClientDirectory.FullName, "Screenshots", "Win64"));
+			FileInfo RedImage = new(Path.Combine(ScreenshotDirectory.FullName, "Red.jpg"));
+			FileInfo BlueImage = new(Path.Combine(ScreenshotDirectory.FullName, "Blue.jpg"));
+			FileInfo TestGif = new(Path.Combine(ClientDirectory.FullName, "ClientTest.gif"));
+
+			FileInfo ServerCrashReporterFile = new(Path.Combine(ServerDirectory.FullName, TruncatedCrashReporterPath, ServerCrashReporterFileName));
+			FileInfo ServerBasicFile = new(Path.Combine(ServerDirectory.FullName, ServerDummyFileName));
+			FileInfo ServerBuildFile = new(Path.Combine(ServerDirectory.FullName, "Build", ServerBuildFileName));
+
+			// A log for each process should also be generated
+			FileInfo ClientLog = new(Path.Combine(ClientDirectory.FullName, "ClientOutput.log"));
+			FileInfo ServerLog = new(Path.Combine(ServerDirectory.FullName, "ServerOutput.log"));
+
+			List<FileInfo> FilesToVerifyExist = new()
+			{
+				ClientCrashReporterFile, ClientBasicFile, ClientBuildFile, ClientLog,
+				ServerCrashReporterFile, ServerBasicFile, ServerBuildFile, ServerLog,
+				RedImage, BlueImage, TestGif,
+			};
+
+			bool bSucceeded = true;
+			foreach (FileInfo File in FilesToVerifyExist)
+			{
+				if (!File.Exists)
+				{
+					bSucceeded = false;
+					Log.Error("Unable locate file {File}", File);
+				}
+			}
+
+			return bSucceeded;
+		}
+
+		private WindowsAppInstance GetAppInstanceForDevice(TargetDeviceWindows Device, string BuildName)
+		{
+			string ArtifactPath = Path.Combine(Device.LocalCachePath, "UserDir", "Saved");
+
+			// This lets us mock the call to StdOut
+			FileInfo Log = new(Path.Combine(ArtifactPath, "Logs", "Output.log"));
+			Log.Directory.Create();
+			File.WriteAllText(Log.FullName, "Foo Log");
+
+			WindowsAppInstall Install = new WindowsAppInstall(BuildName, ProjectName, Device)
+			{
+				ArtifactPath = ArtifactPath
+			};
+
+			ProcessResult DummyProcess = new("Dummy", null, false);
+			WindowsAppInstance AppInstance = new(Install, DummyProcess, Log.FullName);
+
+			// This is a pretty dumb workaround, but the log file reader runs on a separate thread.
+			// We need to wait a few seconds to avoid accidently moving the file in SaveRoleArtifacts
+			// before the thread has had time to map the contents to the stdout
+			Thread.Sleep(5000);
+			return AppInstance;
+		}
+
+		private FileInfo CreateDummyFileInDirectory(DirectoryInfo Directory)
+		{
+			if(!Directory.Exists)
+			{
+				Directory.Create();
+			}
+
+			FileInfo TempFile = new FileInfo(Path.GetTempFileName());
+			FileInfo DummyFile = new FileInfo(Path.Combine(Directory.FullName, TempFile.Name));
+			File.WriteAllText(DummyFile.FullName, "Foo");
+
+			return DummyFile;
+		}
+
+		private FileInfo CreateDummyPNGInDirectory(string BrushColor, DirectoryInfo Directory)
+		{
+			if (!Directory.Exists)
+			{
+				Directory.Create();
+			}
+
+			Bitmap Bitmap = new(16, 16);
+			using (Graphics Graphics = Graphics.FromImage(Bitmap))
+			{
+				Graphics.FillRectangle(new SolidBrush(Color.FromName(BrushColor)), new Rectangle(0, 0, 16, 16));
+			}
+
+			string ImagePath = Path.Combine(Directory.FullName, BrushColor + ".png");
+			Bitmap.Save(ImagePath, ImageFormat.Png);
+
+			return new FileInfo(ImagePath);
+		}
+
+		bool VerifySourceFilesAreDeleted(IEnumerable<FileInfo> SourceFiles)
+		{
+			bool bSucceeded = true;
+			foreach(FileInfo File in SourceFiles)
+			{
+				File.Refresh();
+				if(File.Exists)
+				{
+					Log.Error("Source artifact file {File} should no longer exist, but was found!", File);
+					bSucceeded = false;
+				}
+			}
+
+			return bSucceeded;
+		}
+	}
 }
diff --git a/Engine/Source/Programs/AutomationTool/Gauntlet/Unreal/Base/Gauntlet.UnrealSession.cs b/Engine/Source/Programs/AutomationTool/Gauntlet/Unreal/Base/Gauntlet.UnrealSession.cs
--- a/Engine/Source/Programs/AutomationTool/Gauntlet/Unreal/Base/Gauntlet.UnrealSession.cs
+++ b/Engine/Source/Programs/AutomationTool/Gauntlet/Unreal/Base/Gauntlet.UnrealSession.cs
@@ -6,10 +6,10 @@ using System;
 using System.Collections.Generic;
 using System.IO;
 using System.Linq;
-using System.Text;
 using System.Text.RegularExpressions;
 using System.Threading;
 using UnrealBuildTool;
+using Gauntlet.Utils;
 
 namespace Gauntlet
 {
@@ -1082,45 +1082,102 @@ namespace Gauntlet
 		}
 
 		/// <summary>
-		/// Retrieves and saves all artifacts from the provided session role. Artifacts are saved to the destination path 
+		/// Retrieves and saves all artifacts from the provided session role. Artifacts are saved to the destination path
 		/// </summary>
 		/// <param name="InContext"></param>
 		/// <param name="InRunningRole"></param>
-		/// <param name="InDestArtifactPath"></param>
+		/// <param name="DestinationArtifactPath"></param>
 		/// <returns></returns>
-		public UnrealRoleArtifacts SaveRoleArtifacts(UnrealTestContext InContext, UnrealSessionInstance.RoleInstance InRunningRole, string InDestArtifactPath)
+		public UnrealRoleArtifacts SaveRoleArtifacts(UnrealTestContext InContext, UnrealSessionInstance.RoleInstance InRunningRole, string DestinationArtifactPath)
 		{
-			bool IsServer = InRunningRole.Role.RoleType.IsServer();
-			string RoleName = (InRunningRole.Role.IsDummy() ? "Dummy" : "") + InRunningRole.Role.RoleType.ToString();
-			UnrealTargetPlatform? Platform = InRunningRole.Role.Platform;
-			string RoleConfig = InRunningRole.Role.Configuration.ToString();
+			DirectoryInfo SourceDirectory = new DirectoryInfo(InRunningRole.AppInstance.ArtifactPath);
+			DirectoryInfo DestinationDirectory = new DirectoryInfo(DestinationArtifactPath);
+			bool bRetainArtifacts = InContext.TestParams.ParseParam("RetainDeviceArtifacts");
 
-			if (!Directory.Exists(InDestArtifactPath))
-			{
-				Directory.CreateDirectory(InDestArtifactPath);
-			}
+			// Whether this is a Dummy, Client, Server, Editor, etc
+			string RoleName = (InRunningRole.Role.IsDummy() ? "Dummy" : "") + InRunningRole.Role.RoleType.ToString();
 
+			// We only want to move artifacts for editor data if there was a crash on a buildmachine.
+			// Also, don't move artifacts in dev mode, because peoples saved data could be huuuuuuuge!
 			bool IsDevBuild = InContext.TestParams.ParseParam("dev");
 			bool IsEditorBuild = InRunningRole.Role.RoleType.UsesEditor();
 			bool IsBuildMachine = CommandUtils.IsBuildMachine;
+			bool SkipArchivingAssets = IsDevBuild || (IsEditorBuild && (IsBuildMachine == false || InRunningRole.AppInstance.ExitCode == 0));
+
+			// Check if we should copy artifacts
+			if (!SkipArchivingAssets)
+			{
+				if (SourceDirectory.Exists)
+				{
+					DirectoryInfo PersistentDownloadDirectory = null;
+
+					// Locate any long crash reporter directories and truncate them. This prevents path too long errors.
+					foreach (DirectoryInfo SubDirectory in SourceDirectory.EnumerateDirectories("*", SearchOption.AllDirectories))
+					{
+						string TruncatedName = TruncateLongPathFilter(SubDirectory.Name);
+						if (!SubDirectory.Name.Equals(TruncatedName, StringComparison.OrdinalIgnoreCase))
+						{
+							// This path was successfully truncated, perform a move to rename the path
+							string NewDestination = Path.Combine(SubDirectory.Parent.FullName, TruncatedName);
+							SubDirectory.MoveTo(NewDestination);
+						}
+
+						// If we find a PDD we'll also mark it for deletion before the move.
+						if (SubDirectory.Name.Equals(EIntendedBaseCopyDirectory.PersistentDownloadDir.ToString(), StringComparison.OrdinalIgnoreCase))
+						{
+							PersistentDownloadDirectory = SubDirectory;
+						}
+					}
+
+					if(PersistentDownloadDirectory != null)
+					{
+						PersistentDownloadDirectory.Delete(true);
+					}
 
-			// Unless there was a crash on a builder don't archive editor data (there can be a *lot* of stuff in there).
-			bool SkipArchivingAssets = IsDevBuild ||
-										(IsEditorBuild && 
-											(IsBuildMachine == false || InRunningRole.AppInstance.ExitCode == 0)
-										);
+					// Perform the move/copy
+					SystemHelpers.CopyDirectory(SourceDirectory.FullName, DestinationDirectory.FullName);
+					if (!bRetainArtifacts)
+					{
+						SourceDirectory.Delete(true);
+					}
+				}
+				else
+				{
+					Log.Info("Archive path '{0}' was not found!", SourceDirectory.FullName);
+				}
+			}
+			else
+			{
+				if (IsEditorBuild)
+				{
+					Log.Info("Skipping archival of assets for editor {0}", RoleName);
+				}
+				else if (IsDevBuild)
+				{
+					Log.Info("Skipping archival of assets for dev build");
+				}
+			}
 
-			DirectoryInfo DestSavedDirInfo = new DirectoryInfo(InDestArtifactPath);
-			// save the contents of the saved directory
-			string SourceSavedDir = InRunningRole.AppInstance.ArtifactPath;
-			
-			string ArtifactLogFilePath = String.Empty;
+			// Next, move over any additional artifacts that a role requested
+			foreach (EIntendedBaseCopyDirectory AdditionalDirectory in InRunningRole.Role.AdditionalArtifactDirectories)
+			{
+				Dictionary<EIntendedBaseCopyDirectory, string> PlatformMappings = InRunningRole.AppInstance.Device.GetPlatformDirectoryMappings();
+				if (PlatformMappings.ContainsKey(AdditionalDirectory))
+				{
+					DirectoryInfo AdditionalSourceDirectory = new DirectoryInfo(PlatformMappings[AdditionalDirectory]);
+					if (AdditionalSourceDirectory.Exists)
+					{
+						string TargetDirectory = Path.Combine(DestinationDirectory.FullName, AdditionalSourceDirectory.Name);
+						SystemHelpers.CopyDirectory(AdditionalSourceDirectory.FullName, TargetDirectory);
+					}
+				}
+			}
 
-			// Get the size of the log to determine if it should be saved under a certain size
+			// Now write the role's log file
+			string ArtifactLogFilePath = string.Empty;
 			int MaxLogSize = 1024 * 1024 * 1024;
 			int LogSize = InRunningRole.AppInstance.StdOut.Length * sizeof(char);
 			bool bIgnoreMaxSize = Globals.Params.ParseParam("NoMaxLogSize");
-
 			if (!bIgnoreMaxSize && LogSize > MaxLogSize)
 			{
 				Log.Warning("The process log for Role {0} was over 1 GB in size. A log artifact will not be generated for this process.", InRunningRole.ToString());
@@ -1129,9 +1186,9 @@ namespace Gauntlet
 			{
 				try
 				{
-					ArtifactLogFilePath = Path.Combine(DestSavedDirInfo.FullName, RoleName + "Output.log");
+					ArtifactLogFilePath = Path.Combine(DestinationDirectory.FullName, RoleName + "Output.log");
 
-					// save the output from TTY
+					// Write a short gauntlet blurb before the entire process log
 					using (StreamWriter Writer = new(ArtifactLogFilePath, false))
 					{
 						Writer.WriteLine("------ Gauntlet Test ------");
@@ -1141,6 +1198,7 @@ namespace Gauntlet
 						Writer.Write(InRunningRole.AppInstance.StdOut);
 					}
 					Log.Info($"Wrote {RoleName} Log to {ArtifactLogFilePath}");
+
 					// On build machines, copy all role logs to Horde.
 					if (IsBuildMachine)
 					{
@@ -1155,154 +1213,57 @@ namespace Gauntlet
 				}
 			}
 
-			if (IsServer == false)
+			// Move crash dumps
+			if (InRunningRole.AppInstance.Device.CopyCrashDumps())
 			{
-				// gif-ify and jpeg-ify any screenshots
-				try
-				{
-					string ScreenshotPath = Path.Combine(SourceSavedDir, "Screenshots", Platform.ToString()).ToLower();
-
-					// Check as early as possible for screenshots before creating a temp folder to copy
-					if (Directory.Exists(ScreenshotPath) && Directory.GetFiles(ScreenshotPath).Any())
-					{
-						string TempScreenshotPath = Path.GetTempPath();
-						if (!Directory.Exists(TempScreenshotPath))
-						{
-							Log.Info("Creating temp directory {0}", TempScreenshotPath);
-							Directory.CreateDirectory(TempScreenshotPath);
-						}
-
-						Log.Info("Downsizing and gifying session images at {0}", ScreenshotPath);
-
-						// downsize first so gif-step is quicker and takes less resoruces.
-						Utils.Image.ConvertImages(ScreenshotPath, TempScreenshotPath, "jpg", true);
-
-						string GifPath = Path.Combine(DestSavedDirInfo.FullName, RoleName + "Test.gif");
-						if (Utils.Image.SaveImagesAsGif(TempScreenshotPath, GifPath))
-						{
-							Log.Info("Saved gif to {0}", GifPath);
-						}
-					}
-				}
-				catch (Exception Ex)
+				DirectoryInfo CrashDumpDirectory = new DirectoryInfo(InRunningRole.AppInstance.Device.CrashDumpPath);
+				if (CrashDumpDirectory.Exists)
 				{
-					Log.Info("Failed to downsize and gif-ify images! {0}", Ex.Message);
+					string DesinationCrashDumpDirectory = Path.Combine(DestinationDirectory.FullName, "CrashDumps");
+					Log.Info("Moving any CrashDumps from {0} to {1}", CrashDumpDirectory.FullName, DesinationCrashDumpDirectory);
+					CrashDumpDirectory.MoveTo(DesinationCrashDumpDirectory);
 				}
 			}
 
-			// don't archive data in dev mode, because peoples saved data could be huuuuuuuge!
-			if (SkipArchivingAssets)
-			{
-				if (IsEditorBuild)
-				{
-					Log.Info("Skipping archival of assets for editor {0}", RoleName);
-				}
-				else if (IsDevBuild)
-				{
-					Log.Info("Skipping archival of assets for dev build");
-				}
-			}
-			else
+			// Convert any screenshots to jpegs and create a gif when not running a server
+			if (!InRunningRole.Role.RoleType.IsServer())
 			{
-				LogLevel OldLevel = Log.Level;
-				Log.Level = LogLevel.Normal;
-
-				if (Directory.Exists(SourceSavedDir))
+				try
 				{
-					// Only Copy the artifacts we want
-					foreach (string SubDirectory in Directory.EnumerateDirectories(SourceSavedDir))
-					{
-						DirectoryInfo DirInfo = new DirectoryInfo(SubDirectory);
-
-						if (DirInfo.Name == EIntendedBaseCopyDirectory.PersistentDownloadDir.ToString())
-						{
-							// Do not store the PersistentDownloadDir when possible
-							continue;
-						}
-						
-						// Don't copy an empty directory when possible
-						if (Directory.EnumerateDirectories(SubDirectory).Any() || Directory.EnumerateFiles(SubDirectory).Any())
-						{
-							string FullyQualifiedDestPath = Utils.SystemHelpers.GetFullyQualifiedPath(Path.Combine(DestSavedDirInfo.FullName, DirInfo.Name));
-							Utils.SystemHelpers.CopyDirectory(DirInfo.FullName, FullyQualifiedDestPath, Utils.SystemHelpers.CopyOptions.Default, TruncateLongPathFilter);
-							Log.Info($"Archived artifact \\{DirInfo.Parent.Name}\\{DirInfo.Name}\\ to {FullyQualifiedDestPath}");
-						}
-					}
-
-					// Copy any Crash Dumps as well
-					if (InRunningRole.AppInstance.Device.CopyCrashDumps())
-					{
-						string DeviceCrashDumpPath = InRunningRole.AppInstance.Device.CrashDumpPath;
-						if (!string.IsNullOrEmpty(DeviceCrashDumpPath))
-						{
-							string ArtifactCrashDumpDir = Path.Combine(DestSavedDirInfo.FullName, "CrashDumps");
-							Log.Info("Copying any CrashDumps from {0} to {1}", DeviceCrashDumpPath, ArtifactCrashDumpDir);
-							Utils.SystemHelpers.CopyDirectory(DeviceCrashDumpPath, Utils.SystemHelpers.GetFullyQualifiedPath(ArtifactCrashDumpDir), Utils.SystemHelpers.CopyOptions.Default, TruncateLongPathFilter);
-							Log.Info("Cleaning files in {0}", DeviceCrashDumpPath);
-							foreach (FileInfo CrashDumpFile in Utils.SystemHelpers.GetFiles(new DirectoryInfo(DeviceCrashDumpPath), "*", SearchOption.AllDirectories))
-							{
-								Utils.SystemHelpers.Delete(CrashDumpFile);
-							}
-						}
-					}
+					DirectoryInfo ScreenshotDirectory = new(Path.Combine(DestinationDirectory.FullName, "Screenshots", InRunningRole.Role.Platform.ToString()));
 
-					// Copy any lose files from the source
-					foreach (string SourceFile in Directory.EnumerateFiles(SourceSavedDir))
+					// Check as early as possible for screenshots before creating a temp folder to copy
+					if (ScreenshotDirectory.Exists && ScreenshotDirectory.GetFiles().Any())
 					{
-						FileInfo SourceFileInfo = new FileInfo(SourceFile);
-						string FullyQualifiedDestPath = Utils.SystemHelpers.GetFullyQualifiedPath(Path.Combine(DestSavedDirInfo.FullName, SourceFileInfo.Name));
-						FileInfo DestInfo = new FileInfo(FullyQualifiedDestPath);
+						Log.Info("Downsizing and gifying session images at {0}", ScreenshotDirectory.FullName);
 
-						try
-						{
-							DestInfo = SourceFileInfo.CopyTo(DestInfo.FullName, overwrite:true);
+						// Downsize first so gif-step is quicker and takes less resoruces.
+						Utils.Image.ConvertImages(ScreenshotDirectory.FullName, ScreenshotDirectory.FullName, "jpg", true);
 
-							// Clear attributes and set last write time
-							DestInfo.Attributes = FileAttributes.Normal;
-							DestInfo.LastWriteTime = SourceFileInfo.LastWriteTime;
-							Log.Info($"Archived {SourceFileInfo.Name} to {DestInfo.FullName}");
-						}
-						catch (IOException ex)
+						string GifPath = Path.Combine(DestinationDirectory.FullName, RoleName + "Test.gif");
+						if (Utils.Image.SaveImagesAsGif(ScreenshotDirectory.FullName, GifPath))
 						{
-							Log.Warning($"Archive of {SourceFileInfo.Name} to {DestInfo.FullName} FAILED: Skipping\n{ex.Message}");
+							Log.Info("Saved gif to {0}", GifPath);
 						}
 					}
 				}
-				else
-				{
-					Log.Info("Archive path '{0}' was not found!", SourceSavedDir);
-				}
-
-				Log.Level = OldLevel;
-			}
-
-			foreach (EIntendedBaseCopyDirectory ArtifactDir in InRunningRole.Role.AdditionalArtifactDirectories)
-			{
-				if (InRunningRole.AppInstance.Device.GetPlatformDirectoryMappings().ContainsKey(ArtifactDir))
+				catch (Exception Ex)
 				{
-					string SourcePath = InRunningRole.AppInstance.Device.GetPlatformDirectoryMappings()[ArtifactDir];
-					var DirToCopy = new DirectoryInfo(SourcePath);
-					if (DirToCopy.Exists)
-					{
-						// Grab the final dir name to copy everything into, so everything does not go into the root artifact dir.
-						string IntendedCopyLocation = Path.Combine(DestSavedDirInfo.FullName, DirToCopy.Name);
-						Utils.SystemHelpers.CopyDirectory(SourcePath, Utils.SystemHelpers.GetFullyQualifiedPath(IntendedCopyLocation), Utils.SystemHelpers.CopyOptions.Default, TruncateLongPathFilter);
-					}
+					Log.Info("Failed to downsize and gif-ify images! {0}", Ex.Message);
 				}
 			}
 
 			// TODO REMOVEME- this should go elsewhere, likely a util that can be called or inserted by relevant test nodes.
-			SavePSOs(InContext, InRunningRole, DestSavedDirInfo.FullName);
+			SavePSOs(InContext, InRunningRole, DestinationDirectory.FullName);
 			// END REMOVEME
 
-			UnrealLogParser LogParser = new UnrealLogParser(InRunningRole.AppInstance.StdOut);
-
 			// Save the Artifact filepath
-			return new UnrealRoleArtifacts(InRunningRole.Role, InRunningRole.AppInstance, DestSavedDirInfo.FullName, ArtifactLogFilePath, LogParser);
+			UnrealLogParser LogParser = new UnrealLogParser(InRunningRole.AppInstance.StdOut);
+			return new UnrealRoleArtifacts(InRunningRole.Role, InRunningRole.AppInstance, DestinationDirectory.FullName, ArtifactLogFilePath, LogParser);
 		}
 
 		/// <summary>
-		/// Saves all artifacts from the provided session to the specified output path. 
+		/// Saves all artifacts from the provided session to the specified output path.
 		/// </summary>
 		/// <param name="Context"></param>
 		/// <param name="TestInstance"></param>
@@ -1310,9 +1271,8 @@ namespace Gauntlet
 		/// <returns></returns>
 		public IEnumerable<UnrealRoleArtifacts> SaveRoleArtifacts(UnrealTestContext Context, UnrealSessionInstance TestInstance, string OutputPath)
 		{
-			Dictionary<UnrealTargetRole, int> RoleCounts = new Dictionary<UnrealTargetRole, int>();
 			int DummyClientCount = 0;
-
+			Dictionary<UnrealTargetRole, int> RoleCounts = new Dictionary<UnrealTargetRole, int>();
 			List<UnrealRoleArtifacts> AllArtifacts = new List<UnrealRoleArtifacts>();
 
 			foreach (UnrealSessionInstance.RoleInstance App in TestInstance.RunningRoles)
@@ -1341,7 +1301,6 @@ namespace Gauntlet
 					RoleCount = RoleCounts[App.Role.RoleType];
 				}
 
-				
 				if (RoleCount > 1)
 				{
 					FolderName += string.Format("_{0:00}", RoleCount);
@@ -1385,7 +1344,7 @@ namespace Gauntlet
 						AllArtifacts.Add(Artifacts);
 					}
 				}
-				else 
+				else
 				{
 					Log.Verbose("Skipping SaveRoleArtifacts for Null Role: {0}", App.ToString());
 				}
@@ -1495,11 +1454,11 @@ namespace Gauntlet
 		{
 			Dictionary<string, string> LongCrashReporterStringToIndex = new Dictionary<string, string>();
 
-			Match RegexMatch = Regex.Match(LongFilePath, @"UECC-.+-([\dA-Fa-f]+)");
+			Match RegexMatch = Regex.Match(LongFilePath, @"((?i)UECC)-.+-([\dA-Fa-f]+)");
 
 			if (RegexMatch.Success)
 			{
-				string LongString = RegexMatch.Groups[1].ToString();
+				string LongString = RegexMatch.Groups[2].ToString();
 
 				if (!LongCrashReporterStringToIndex.ContainsKey(LongString))
 				{
