Updated P4Utils to also support adds and deletes on a list of files to avoid blowing the command line.
#rnx
#test: Have some shelved C# scripts that use the functions and behavior seems to be the same using both the Add() function that takes a list of files vs taking a concatenated string of all the files.

diff --git a/Engine/Source/Programs/AutomationTool/AutomationUtils/P4Utils.cs b/Engine/Source/Programs/AutomationTool/AutomationUtils/P4Utils.cs
--- a/Engine/Source/Programs/AutomationTool/AutomationUtils/P4Utils.cs
+++ b/Engine/Source/Programs/AutomationTool/AutomationUtils/P4Utils.cs
@@ -2230,6 +2230,16 @@ namespace AutomationTool
 			LogP4("", "add " + String.Format("-c {0} ", CL) + CommandLine);
 		}
 
+		/// <summary>
+		/// Invokes p4 add command with a list of files.
+		/// </summary>
+		/// <param name="CL">Changelist where the checked out files should be added.</param>
+		/// <param name="Files">The list of files to add.</param>
+		public void Add(int CL, List<string> Files, bool AllowSpew = true)
+		{
+			BatchedCommand($"add -c {CL}", Files, AllowSpew: AllowSpew);
+		}
+
 		/// <summary>
 		/// Invokes p4 delete command.
 		/// </summary>
@@ -2240,6 +2250,16 @@ namespace AutomationTool
 			LogP4("", "delete " + String.Format("-c {0} ", CL) + CommandLine);
 		}
 
+		/// <summary>
+		/// Invokes p4 delete command with a list of files.
+		/// </summary>
+		/// <param name="CL">Changelist where the checked out files should be added.</param>
+		/// <param name="Files">List of files to be deleted.</param>
+		public void Delete(int CL, List<string> Files, bool AllowSpew = true)
+		{
+			BatchedCommand($"delete -c {CL}", Files, AllowSpew: AllowSpew);
+		}
+
 		/// <summary>
 		/// Invokes p4 reconcile command.
 		/// </summary>
