[Automation] Makes UnrealSessionRole.SaveRoleArtifacts copy/delete crashdumps instead of move.
This prevents an exception caused by MoveTo not being allowed across drives.

diff --git a/Engine/Source/Programs/AutomationTool/Gauntlet/Unreal/Base/Gauntlet.UnrealSession.cs b/Engine/Source/Programs/AutomationTool/Gauntlet/Unreal/Base/Gauntlet.UnrealSession.cs
--- a/Engine/Source/Programs/AutomationTool/Gauntlet/Unreal/Base/Gauntlet.UnrealSession.cs
+++ b/Engine/Source/Programs/AutomationTool/Gauntlet/Unreal/Base/Gauntlet.UnrealSession.cs
@@ -1221,7 +1221,9 @@ namespace Gauntlet
 				{
 					string DesinationCrashDumpDirectory = Path.Combine(DestinationDirectory.FullName, "CrashDumps");
 					Log.Info("Moving any CrashDumps from {0} to {1}", CrashDumpDirectory.FullName, DesinationCrashDumpDirectory);
-					CrashDumpDirectory.MoveTo(DesinationCrashDumpDirectory);
+
+					SystemHelpers.CopyDirectory(CrashDumpDirectory.FullName, DesinationCrashDumpDirectory);
+					CrashDumpDirectory.Delete(true);
 				}
 			}
 
