- Fixed an issue with moving remapped files from UFS to NonUFS when making cooked editor, where it was using UnmakeStagedFileReference unnecessarily, and getting it wrong when a file was remapped
- Removed UnmakeStagedFileReference

diff --git a/Engine/Source/Programs/AutomationTool/AutomationUtils/DeploymentContext.cs b/Engine/Source/Programs/AutomationTool/AutomationUtils/DeploymentContext.cs
--- a/Engine/Source/Programs/AutomationTool/AutomationUtils/DeploymentContext.cs
+++ b/Engine/Source/Programs/AutomationTool/AutomationUtils/DeploymentContext.cs
@@ -1211,42 +1211,4 @@ public class DeploymentContext //: ProjectParams
 
 		throw new Exception();
 	}
-	public static FileReference UnmakeRelativeStagedReference(DeploymentContext SC, StagedFileReference Ref)
-	{
-		// paths will be in the form "Engine/Foo" or "{ProjectName}/Foo" or "RemappedPlugins/{PluginName}/Foo
-		// Anything else we don't handle.
-		// So, replace the Engine/ with {EngineDir} and {ProjectName}/ with {ProjectDir}, or change PluginDir to RemappedPlugins/{PluginName}
-		// with the plugin path from AdditionalPluginDirectories, and then append Foo
-
-		string RemappedPluginsStr = "RemappedPlugins/";
-		if (Ref.Name.StartsWith(RemappedPluginsStr, StringComparison.CurrentCultureIgnoreCase))
-		{
-			int PluginEndIndex = Ref.Name.IndexOf("/", RemappedPluginsStr.Length);
-			if (PluginEndIndex >= 0 && PluginEndIndex < Ref.Name.Length - 1)
-			{
-				string PluginName = Ref.Name.Substring(RemappedPluginsStr.Length, PluginEndIndex - RemappedPluginsStr.Length);
-				foreach (DirectoryReference AdditionalPluginDir in SC.AdditionalPluginDirectories)
-				{
-					DirectoryReference PossiblePluginDir = DirectoryReference.Combine(AdditionalPluginDir, PluginName);
-					if (System.IO.Directory.Exists(PossiblePluginDir.FullName))
-					{
-						return FileReference.Combine(PossiblePluginDir, Ref.Name.Substring(PluginEndIndex+1));
-					}
-				}
-			}
-		}
-
-		if (Ref.Name.StartsWith("Engine/", StringComparison.CurrentCultureIgnoreCase))
-		{
-			// skip over "Engine/" which is 7 chars long
-			return FileReference.Combine(SC.EngineRoot, Ref.Name.Substring(7));
-		}
-
-		if (Ref.Name.StartsWith(SC.ShortProjectName + "/", StringComparison.CurrentCultureIgnoreCase))
-		{
-			return FileReference.Combine(SC.ProjectRoot, Ref.Name.Substring(SC.ShortProjectName.Length + 1));
-		}
-
-		throw new Exception($"Don't know how to convert staged file {Ref.Name} to its original editor path, because it is not in a recognized root directory.");
-	}
 }
diff --git a/Engine/Source/Programs/AutomationTool/CookedEditor/CookedEditor.Automation.cs b/Engine/Source/Programs/AutomationTool/CookedEditor/CookedEditor.Automation.cs
--- a/Engine/Source/Programs/AutomationTool/CookedEditor/CookedEditor.Automation.cs
+++ b/Engine/Source/Programs/AutomationTool/CookedEditor/CookedEditor.Automation.cs
@@ -291,7 +291,13 @@ public class ModifyStageContext
 	private void AddUFSFilesToList(List<FileReference> FileList, string Extension, DeploymentContext SC)
 	{
 		// look in SC and UFSFiles
-		FileList.AddRange(SC.FilesToStage.UFSFiles.Keys.Where(x => x.Name.EndsWith(Extension, StringComparison.OrdinalIgnoreCase)).Select(y => DeploymentContext.UnmakeRelativeStagedReference(SC, y)));
+		foreach (var Pair in SC.FilesToStage.UFSFiles)
+		{
+			if (Pair.Key.Name.EndsWith(Extension))
+			{
+				FileList.Add(Pair.Value);
+			}
+		}
 		FileList.AddRange(UFSFilesToStage.Where(x => x.FullName.EndsWith(Extension, StringComparison.InvariantCultureIgnoreCase)));
 	}
 
