Fix: Turnkey uses the AutoSDKPlatformName when setting up AutoSDK & a perforce provider isn't being used

diff --git a/Engine/Source/Programs/AutomationTool/Turnkey/SdkInfo.cs b/Engine/Source/Programs/AutomationTool/Turnkey/SdkInfo.cs
--- a/Engine/Source/Programs/AutomationTool/Turnkey/SdkInfo.cs
+++ b/Engine/Source/Programs/AutomationTool/Turnkey/SdkInfo.cs
@@ -220,7 +220,13 @@ namespace Turnkey
 				}
 
 				// download the AutoSDK, to the UE_SDKS_ROOT dir
-				string AutoSDKDownloadedRoot = Path.Combine(Environment.GetEnvironmentVariable("UE_SDKS_ROOT"), $"Host{HostPlatform.Platform}", Platform.ToString(), Source.Version);
+				string AutoSdkPlatformName = Platform.ToString();
+				UEBuildPlatformSDK PlatformSDK = UEBuildPlatformSDK.GetSDKForPlatform(Platform.ToString());
+				if (PlatformSDK != null)
+				{
+					AutoSdkPlatformName = PlatformSDK.GetAutoSDKPlatformName();
+				}
+				string AutoSDKDownloadedRoot = Path.Combine(Environment.GetEnvironmentVariable("UE_SDKS_ROOT"), $"Host{HostPlatform.Platform}", AutoSdkPlatformName, Source.Version);
 				// it should return what we pass into ig
 				DownloadedRoot = CopyProvider.ExecuteCopy(CopyOperation, CopyExecuteSpecialMode.UsePermanentStorage, AutoSDKDownloadedRoot);
 				if (DownloadedRoot != AutoSDKDownloadedRoot)
