BuildGraph: Add target and output parameters to DockerBuildTask

diff --git a/Engine/Source/Programs/AutomationTool/BuildGraph/Tasks/DockerBuildTask.cs b/Engine/Source/Programs/AutomationTool/BuildGraph/Tasks/DockerBuildTask.cs
--- a/Engine/Source/Programs/AutomationTool/BuildGraph/Tasks/DockerBuildTask.cs
+++ b/Engine/Source/Programs/AutomationTool/BuildGraph/Tasks/DockerBuildTask.cs
@@ -62,6 +62,18 @@ namespace AutomationTool.Tasks
 		[TaskParameter(Optional = true)]
 		public string Tag;
 
+		/// <summary>
+		/// Set the target build stage to build (--target)
+		/// </summary>
+		[TaskParameter(Optional = true)]
+		public string Target;
+
+		/// <summary>
+		/// Custom output exporter. Requires BuildKit (--output)
+		/// </summary>
+		[TaskParameter(Optional = true)]
+		public string Output;
+
 		/// <summary>
 		/// Optional arguments
 		/// </summary>
@@ -144,6 +156,18 @@ namespace AutomationTool.Tasks
 				{
 					Arguments.Append($" -t {Parameters.Tag}");
 				}
+				if (Parameters.Target != null)
+				{
+					Arguments.Append($" --target {Parameters.Target}");
+				}
+				if (Parameters.Output != null)
+				{
+					if (!Parameters.UseBuildKit)
+					{
+						throw new AutomationException($"{nameof(Parameters.UseBuildKit)} must be enabled to use '{nameof(Parameters.Output)}' parameter");
+					}
+					Arguments.Append($" --output {Parameters.Output}");
+				}
 				if (Parameters.DockerFile != null)
 				{
 					FileReference DockerFile = ResolveFile(Parameters.DockerFile);
