Add support for a baseline when making cloud snapshots.  Also unify File and Cloud snapshot mechanism handling of multiple platforms, output descriptors and input baseline descriptors.

diff --git a/Engine/Build/Graph/Tasks/BuildAndTestProject.xml b/Engine/Build/Graph/Tasks/BuildAndTestProject.xml
--- a/Engine/Build/Graph/Tasks/BuildAndTestProject.xml
+++ b/Engine/Build/Graph/Tasks/BuildAndTestProject.xml
@@ -529,7 +529,7 @@
 			<Do If="$(PublishBuilds) and $(PublishCloudSnapshot) and '$(UE-CloudPublishHost)' != '' and '$(UE-CloudPublishNamespace)' != ''">
 				<Property Name="PlatformCloudPublishNodeName" Value="$(PreNodeName)Cloud Publish $(TargetPlatform)" />
 				<Node Name="$(PlatformCloudPublishNodeName)" Requires="$(PlatformHashNodeName)">
-					<ZenExportSnapshot Project="$(TargetProject)" Platform="$(CookPlatform)" DestinationStorageType="Cloud" DestinationCloudHost="$(UE-CloudPublishHost)" DestinationCloudNamespace="$(UE-CloudPublishNamespace)" DestinationCloudIdentifier="$(EscapedBranch)/$(BuildNamePath)" SnapshotDescriptorFile="$(NetworkOutputDirectory)/$(StagedPlatformFolder)/SnapshotDescriptor.json"/>
+					<ZenExportSnapshot Project="$(TargetProject)" Platform="$(CookPlatform)" DestinationStorageType="Cloud" DestinationCloudHost="$(UE-CloudPublishHost)" DestinationCloudNamespace="$(UE-CloudPublishNamespace)" DestinationIdentifier="$(EscapedBranch)/$(BuildNamePath)" SnapshotDescriptorFile="$(NetworkOutputDirectory)/$(StagedPlatformFolder)/SnapshotDescriptor.json"/>
 				</Node>
 				<Property Name="PlatformPublishStageDependencyName" Value="$(PlatformCloudPublishNodeName)" />
 			</Do>
diff --git a/Engine/Source/Programs/AutomationTool/BuildGraph/Tasks/ZenExportSnapshotTask.cs b/Engine/Source/Programs/AutomationTool/BuildGraph/Tasks/ZenExportSnapshotTask.cs
--- a/Engine/Source/Programs/AutomationTool/BuildGraph/Tasks/ZenExportSnapshotTask.cs
+++ b/Engine/Source/Programs/AutomationTool/BuildGraph/Tasks/ZenExportSnapshotTask.cs
@@ -1,22 +1,19 @@
 // Copyright Epic Games, Inc. All Rights Reserved.
 
-using AutomationTool;
-using EpicGames.BuildGraph;
+using Microsoft.Extensions.Logging;
 using System;
 using System.Collections.Generic;
+using System.Diagnostics.CodeAnalysis;
 using System.IO;
 using System.Linq;
 using System.Runtime.InteropServices;
 using System.Text;
+using System.Text.Json;
+using System.Text.Json.Serialization;
 using System.Threading.Tasks;
 using System.Xml;
 using EpicGames.Core;
 using EpicGames.Serialization;
-using UnrealBuildTool;
-using UnrealBuildBase;
-using Microsoft.Extensions.Logging;
-
-using static AutomationTool.CommandUtils;
 
 namespace AutomationTool.Tasks
 {
@@ -45,6 +42,12 @@ namespace AutomationTool.Tasks
 		[TaskParameter(Optional = true)]
 		public string Platform;
 
+		/// <summary>
+		/// A file to read with information about the snapshot that should be used as a base when exporting this new snapshot
+		/// </summary>
+		[TaskParameter(Optional = true)]
+		public FileReference SnapshotBaseDescriptorFile;
+
 		/// <summary>
 		/// A file to create with information about the snapshot that was exported
 		/// </summary>
@@ -58,22 +61,22 @@ namespace AutomationTool.Tasks
 		public string DestinationStorageType;
 
 		/// <summary>
-		/// The host name to use when exporting to a cloud destination
+		/// The identifier to use when exporting to a destination
 		/// </summary>
 		[TaskParameter(Optional = true)]
-		public string DestinationCloudHost;
+		public string DestinationIdentifier;
 
 		/// <summary>
-		/// The namespace to use when exporting to a cloud destination
+		/// The host name to use when exporting to a cloud destination
 		/// </summary>
 		[TaskParameter(Optional = true)]
-		public string DestinationCloudNamespace;
+		public string DestinationCloudHost;
 
 		/// <summary>
-		/// The identifier to use when exporting to a cloud destination
+		/// The namespace to use when exporting to a cloud destination
 		/// </summary>
 		[TaskParameter(Optional = true)]
-		public string DestinationCloudIdentifier;
+		public string DestinationCloudNamespace;
 
 		/// <summary>
 		/// A custom bucket name to use when exporting to a cloud destination
@@ -82,16 +85,16 @@ namespace AutomationTool.Tasks
 		public string DestinationCloudBucket;
 
 		/// <summary>
-		/// The directory to store the exported oplog
+		/// The directory to use when exporting to a file destination
 		/// </summary>
 		[TaskParameter(Optional = true)]
-		public string DestinationDir;
+		public DirectoryReference DestinationFileDir;
 
 		/// <summary>
-		/// The name of the oplog to write on disk
+		/// The filename to use when exporting to a file destination
 		/// </summary>
 		[TaskParameter(Optional = true)]
-		public string DestinationOplogName;
+		public string DestinationFileName;
 	}
 
 	/// <summary>
@@ -100,6 +103,20 @@ namespace AutomationTool.Tasks
 	[TaskElement("ZenExportSnapshot", typeof(ZenExportSnapshotTaskParameters))]
 	public class ZenExportSnapshotTask : BgTaskImpl
 	{
+		private class SnapshotDescriptor
+		{
+			public string Name = null;
+			public SnapshotStorageType Type = SnapshotStorageType.Invalid;
+			public string TargetPlatform = null;
+
+			public string Host = null;
+			public string Namespace = null;
+			public string Bucket = null;
+			public string Key = null;
+
+			public string Directory = null;
+			public string Filename = null;
+		}
 		private class ExportSourceData
 		{
 			public bool IsLocalHost;
@@ -108,6 +125,7 @@ namespace AutomationTool.Tasks
 			public string ProjectId;
 			public string OplogId;
 			public string TargetPlatform;
+			public SnapshotDescriptor SnapshotBaseDescriptor;
 		}
 
 		/// <summary>
@@ -149,7 +167,88 @@ namespace AutomationTool.Tasks
 			ZenLaunchCommandline.AppendFormat("{0} -SponsorProcessID={1}", CommandUtils.MakePathSafeToUseWithCommandLine(ProjectFile.FullName), Environment.ProcessId);
 
 			Logger.LogInformation("Running '{Arg0} {Arg1}'", CommandUtils.MakePathSafeToUseWithCommandLine(ZenLaunchExe.FullName), ZenLaunchCommandline.ToString());
-			CommandUtils.RunAndLog(CommandUtils.CmdEnv, ZenLaunchExe.FullName, ZenLaunchCommandline.ToString(), Options: CommandUtils.ERunOptions.Default);
+			CommandUtils.RunAndLog(CommandUtils.CmdEnv, CommandUtils.MakePathSafeToUseWithCommandLine(ZenLaunchExe.FullName), ZenLaunchCommandline.ToString(), Options: CommandUtils.ERunOptions.Default);
+		}
+		
+		static JsonSerializerOptions GetDefaultJsonSerializerOptions()
+		{
+			JsonSerializerOptions options = new JsonSerializerOptions();
+			options.AllowTrailingCommas = true;
+			options.ReadCommentHandling = JsonCommentHandling.Skip;
+			options.PropertyNameCaseInsensitive = true;
+			options.Converters.Add(new JsonStringEnumConverter());
+			return options;
+		}
+
+#nullable enable
+		static bool TryLoadJson<T>(FileReference file, [NotNullWhen(true)] out T? obj) where T : class
+		{
+			if (!FileReference.Exists(file))
+			{
+				obj = null;
+				return false;
+			}
+
+			try
+			{
+				obj = LoadJson<T>(file);
+				return true;
+			}
+			catch(Exception)
+			{
+				obj = null;
+				return false;
+			}
+		}
+
+		static T? TryDeserializeJson<T>(byte[] data) where T : class
+		{
+			try
+			{
+				return JsonSerializer.Deserialize<T>(data, GetDefaultJsonSerializerOptions())!;
+			}
+			catch (Exception)
+			{
+				return null;
+			}
+		}
+
+		static T LoadJson<T>(FileReference file)
+		{
+			byte[] data = FileReference.ReadAllBytes(file);
+			return JsonSerializer.Deserialize<T>(data, GetDefaultJsonSerializerOptions())!;
+		}
+		private void WriteExportSource(JsonWriter Writer, SnapshotStorageType DestinationStorageType, ExportSourceData ExportSource, string Name)
+		{
+			Writer.WriteObjectStart();
+			switch (DestinationStorageType)
+			{
+				case SnapshotStorageType.Cloud:
+					string BucketName = Parameters.DestinationCloudBucket;
+					string ProjectNameAsBucketName = Parameters.Project.GetFileNameWithoutAnyExtensions().ToLowerInvariant();
+					if (string.IsNullOrEmpty(BucketName))
+					{
+						BucketName = ProjectNameAsBucketName;
+					}
+
+					IoHash DestinationKeyHash = IoHash.Compute(Encoding.UTF8.GetBytes(Name));
+					Writer.WriteValue("name", Name);
+					Writer.WriteValue("type", "cloud");
+					Writer.WriteValue("targetplatform", ExportSource.TargetPlatform);
+					Writer.WriteValue("host", Parameters.DestinationCloudHost);
+					Writer.WriteValue("namespace", Parameters.DestinationCloudNamespace);
+					Writer.WriteValue("bucket", BucketName);
+					Writer.WriteValue("key", DestinationKeyHash.ToString().ToLowerInvariant());
+					break;
+				case SnapshotStorageType.File:
+					Writer.WriteValue("name", Name);
+					Writer.WriteValue("type", "file");
+					Writer.WriteValue("targetplatform", ExportSource.TargetPlatform);
+					Writer.WriteValue("directory", Parameters.DestinationFileDir.FullName);
+					Writer.WriteValue("filename", Parameters.DestinationFileName);
+					break;
+			}
+			Writer.WriteObjectEnd();
 		}
 
 		/// <summary>
@@ -188,7 +287,7 @@ namespace AutomationTool.Tasks
 				{
 					continue;
 				}
-
+				
 				byte[] ProjectStoreData = File.ReadAllBytes(ProjectStoreFile.FullName);
 				CbObject ProjectStoreObject = new CbField(ProjectStoreData).AsObject();
 				CbObject ZenServerObject = ProjectStoreObject["zenserver"].AsObject();
@@ -202,9 +301,23 @@ namespace AutomationTool.Tasks
 					NewExportSource.OplogId = ZenServerObject["oplogid"].AsString();
 					NewExportSource.TargetPlatform = Platform;
 
+					FileReference PlatformSnapshotBase = new FileReference(Parameters.SnapshotBaseDescriptorFile.FullName.Replace("{Platform}", Platform, StringComparison.InvariantCultureIgnoreCase));
+
+					SnapshotDescriptor? ParsedDescriptor = null;
+					if (TryLoadJson(PlatformSnapshotBase, out ParsedDescriptor))
+					{
+						NewExportSource.SnapshotBaseDescriptor = ParsedDescriptor;
+					}
+					else
+					{
+						NewExportSource.SnapshotBaseDescriptor = null;
+					}
+
 					ExportSources.Add(NewExportSource);
 				}
 			}
+			int ExportIndex = 0;
+			string[] ExportNames = new string[ExportSources.Count];
 
 			// Get the Zen executable path
 			FileReference ZenExe = ZenExeFileReference();
@@ -224,9 +337,9 @@ namespace AutomationTool.Tasks
 					{
 						throw new AutomationException("Missing destination cloud namespace");
 					}
-					if (string.IsNullOrEmpty(Parameters.DestinationCloudIdentifier))
+					if (string.IsNullOrEmpty(Parameters.DestinationIdentifier))
 					{
-						throw new AutomationException("Missing destination cloud identifier");
+						throw new AutomationException("Missing destination identifier when exporting to cloud");
 					}
 
 					string BucketName = Parameters.DestinationCloudBucket;
@@ -238,76 +351,138 @@ namespace AutomationTool.Tasks
 
 					OplogExportCommandline.AppendFormat(" --cloud {0} --namespace {1} --bucket {2}", Parameters.DestinationCloudHost, Parameters.DestinationCloudNamespace, BucketName);
 
-					string[] ExportKeyIds = new string[ExportSources.Count];
-					int ExportIndex = 0;
+					ExportIndex = 0;
 					foreach (ExportSourceData ExportSource in ExportSources)
 					{
-						String HostUrlArg = string.Format("--hosturl http://{0}:{1}", ExportSource.IsLocalHost ? "localhost" : ExportSource.HostName, ExportSource.HostPort);
+						string HostUrlArg = string.Format("--hosturl http://{0}:{1}", ExportSource.IsLocalHost ? "localhost" : ExportSource.HostName, ExportSource.HostPort);
+						
+						string BaseKeyArg = string.Empty;
+						if (ExportSource.SnapshotBaseDescriptor != null)
+						{
+							if (ExportSource.SnapshotBaseDescriptor.Type == SnapshotStorageType.Cloud)
+							{
+								BaseKeyArg = " --basekey " + ExportSource.SnapshotBaseDescriptor.Key;
+							}
+							else
+							{
+								Logger.LogWarning("Base snapshot descriptor was for a snapshot storage type {0}, but we're producing a snapshot of type cloud.  Skipping use of base snapshot.", ExportSource.SnapshotBaseDescriptor.Type);
+							}
+						}
 
 						StringBuilder ExportSingleSourceCommandline = new StringBuilder(OplogExportCommandline.Length);
 						ExportSingleSourceCommandline.Append(OplogExportCommandline);
 
 						StringBuilder DestinationKeyBuilder = new StringBuilder();
-						DestinationKeyBuilder.AppendFormat("{0}.{1}.{2}", ProjectNameAsBucketName, Parameters.DestinationCloudIdentifier, ExportSource.OplogId);
-						ExportKeyIds[ExportIndex] = DestinationKeyBuilder.ToString().ToLowerInvariant();
-						IoHash DestinationKeyHash = IoHash.Compute(Encoding.UTF8.GetBytes(ExportKeyIds[ExportIndex]));
+						DestinationKeyBuilder.AppendFormat("{0}.{1}.{2}", ProjectNameAsBucketName, Parameters.DestinationIdentifier, ExportSource.OplogId);
+						ExportNames[ExportIndex] = DestinationKeyBuilder.ToString().ToLowerInvariant();
+						IoHash DestinationKeyHash = IoHash.Compute(Encoding.UTF8.GetBytes(ExportNames[ExportIndex]));
 
 						ProcessResult.SpewFilterCallbackType SilentOutputFilter = new ProcessResult.SpewFilterCallbackType(Line =>
 							{
 								return null;
 							});
-						ExportSingleSourceCommandline.AppendFormat(" {0} --embedloosefiles --key {1} {2} {3}", HostUrlArg, DestinationKeyHash.ToString().ToLowerInvariant(), ExportSource.ProjectId, ExportSource.OplogId);
+						ExportSingleSourceCommandline.AppendFormat(" {0} --embedloosefiles --key {1} --basekey {2} {3} {4}", HostUrlArg, DestinationKeyHash.ToString().ToLowerInvariant(), BaseKeyArg, ExportSource.ProjectId, ExportSource.OplogId);
 						Logger.LogInformation("Running '{Arg0} {Arg1}'", CommandUtils.MakePathSafeToUseWithCommandLine(ZenExe.FullName), ExportSingleSourceCommandline.ToString());
-						CommandUtils.RunAndLog(CommandUtils.CmdEnv, ZenExe.FullName, ExportSingleSourceCommandline.ToString(), MaxSuccessCode: int.MaxValue, Options: CommandUtils.ERunOptions.Default, SpewFilterCallback: SilentOutputFilter);
+						CommandUtils.RunAndLog(CommandUtils.CmdEnv, CommandUtils.MakePathSafeToUseWithCommandLine(ZenExe.FullName), ExportSingleSourceCommandline.ToString(), MaxSuccessCode: int.MaxValue, Options: CommandUtils.ERunOptions.Default, SpewFilterCallback: SilentOutputFilter);
 
 						ExportIndex = ExportIndex + 1;
 					}
-					
-					if ((Parameters.SnapshotDescriptorFile != null) && ExportSources.Any())
+
+					break;
+				case SnapshotStorageType.File:
+					string ProjectId = ProjectUtils.GetProjectPathId(ProjectFile);
+					ExportIndex = 0;
+					foreach (ExportSourceData ExportSource in ExportSources)
 					{
-						DirectoryReference.CreateDirectory(Parameters.SnapshotDescriptorFile.Directory);
-						// Write out a snapshot descriptor
-						using (JsonWriter Writer = new JsonWriter(Parameters.SnapshotDescriptorFile))
+						StringBuilder ExportNameBuilder = new StringBuilder();
+						ExportNameBuilder.AppendFormat("{0}.{1}.{2}", ProjectFile.GetFileNameWithoutAnyExtensions().ToLowerInvariant(), Parameters.DestinationIdentifier, ExportSource.OplogId);
+						ExportNames[ExportIndex] = ExportNameBuilder.ToString().ToLowerInvariant();
+
+						StringBuilder ExportSingleSourceCommandline = new StringBuilder(OplogExportCommandline.Length);
+						ExportSingleSourceCommandline.Append(OplogExportCommandline);
+
+						string DestinationFileName = ExportSource.OplogId;
+						if (!string.IsNullOrEmpty(Parameters.DestinationFileName))
 						{
-							Writer.WriteObjectStart();
+							DestinationFileName = Parameters.DestinationFileName.Replace("{Platform}", ExportSource.TargetPlatform, StringComparison.InvariantCultureIgnoreCase);
+						}
 
-							Writer.WriteArrayStart("snapshots");
-							
-							ExportIndex = 0;
-							foreach (ExportSourceData ExportSource in ExportSources)
+						string BaseNameArg = string.Empty;
+						DirectoryReference PlatformDestinationFileDir = new DirectoryReference(Parameters.DestinationFileDir.FullName.Replace("{Platform}", ExportSource.TargetPlatform, StringComparison.InvariantCultureIgnoreCase));
+						if (ExportSource.SnapshotBaseDescriptor != null)
+						{
+							if (ExportSource.SnapshotBaseDescriptor.Type == SnapshotStorageType.File)
 							{
-								Writer.WriteObjectStart();
-								
-								IoHash DestinationKeyHash = IoHash.Compute(Encoding.UTF8.GetBytes(ExportKeyIds[ExportIndex]));
-								Writer.WriteValue("name", ExportKeyIds[ExportIndex]);
-								Writer.WriteValue("type", "cloud");
-								Writer.WriteValue("targetplatform", ExportSource.TargetPlatform);
-								Writer.WriteValue("host", Parameters.DestinationCloudHost);
-								Writer.WriteValue("namespace", Parameters.DestinationCloudNamespace);
-								Writer.WriteValue("bucket", BucketName);
-								Writer.WriteValue("key", DestinationKeyHash.ToString().ToLowerInvariant());
-
-								Writer.WriteObjectEnd();
-
-								ExportIndex = ExportIndex + 1;
+								FileReference BaseSnapshotFile = new FileReference(Path.Combine(ExportSource.SnapshotBaseDescriptor.Directory, ExportSource.SnapshotBaseDescriptor.Filename));
+								if (FileReference.Exists(BaseSnapshotFile))
+								{
+									BaseNameArg = " --basename " + CommandUtils.MakePathSafeToUseWithCommandLine(BaseSnapshotFile.FullName);
+								}
+								else
+								{
+									Logger.LogWarning("Base snapshot descriptor missing.  Skipping use of base snapshot.");
+								}
+							}
+							else
+							{
+								Logger.LogWarning("Base snapshot descriptor was for a snapshot storage type {0}, but we're producing a snapshot of type file.  Skipping use of base snapshot.", ExportSource.SnapshotBaseDescriptor.Type);
 							}
-
-							Writer.WriteArrayEnd();
-							Writer.WriteObjectEnd();
 						}
-					}
+						ExportSingleSourceCommandline.AppendFormat(" --file {0} --name {1} {2} {3} {4}", CommandUtils.MakePathSafeToUseWithCommandLine(PlatformDestinationFileDir.FullName), DestinationFileName, BaseNameArg, ProjectId, ExportSource.OplogId);
 
-					break;
-				case SnapshotStorageType.File:
-					string ProjectId = ProjectUtils.GetProjectPathId(ProjectFile);
-					OplogExportCommandline.AppendFormat(" --file {0} --name {1} {2} {3}", Parameters.DestinationDir, Parameters.DestinationOplogName, ProjectId, Parameters.Platform);
+						Logger.LogInformation("Running '{Arg0} {Arg1}'", CommandUtils.MakePathSafeToUseWithCommandLine(ZenExe.FullName), ExportSingleSourceCommandline.ToString());
+						CommandUtils.RunAndLog(CommandUtils.CmdEnv, CommandUtils.MakePathSafeToUseWithCommandLine(ZenExe.FullName), ExportSingleSourceCommandline.ToString(), Options: CommandUtils.ERunOptions.Default);
 
-					Logger.LogInformation("Running '{Arg0} {Arg1}'", CommandUtils.MakePathSafeToUseWithCommandLine(ZenExe.FullName), OplogExportCommandline.ToString());
-					CommandUtils.RunAndLog(CommandUtils.CmdEnv, ZenExe.FullName, OplogExportCommandline.ToString(), Options: CommandUtils.ERunOptions.Default);
+						ExportIndex = ExportIndex + 1;
+					}
 					break;
 				default:
 					throw new AutomationException("Unknown/invalid/unimplemented destination storage type - {0}", Parameters.DestinationStorageType);
 			}
+			
+
+			if ((Parameters.SnapshotDescriptorFile != null) && ExportSources.Any())
+			{
+				if (Parameters.SnapshotDescriptorFile.FullName.Contains("{Platform}"))
+				{
+					// Separate descriptor file per platform
+					ExportIndex = 0;
+					foreach (ExportSourceData ExportSource in ExportSources)
+					{
+						FileReference PlatformSnapshotDescriptorFile = new FileReference(Parameters.SnapshotDescriptorFile.FullName.Replace("{Platform}", ExportSource.TargetPlatform, StringComparison.InvariantCultureIgnoreCase));
+						DirectoryReference.CreateDirectory(PlatformSnapshotDescriptorFile.Directory);
+						using (JsonWriter Writer = new JsonWriter(PlatformSnapshotDescriptorFile))
+						{
+							Writer.WriteObjectStart();
+							Writer.WriteArrayStart("snapshots");
+							WriteExportSource(Writer, DestinationStorageType, ExportSource, ExportNames[ExportIndex]);
+							Writer.WriteArrayEnd();
+							Writer.WriteObjectEnd();
+						}
+						ExportIndex = ExportIndex + 1;
+					}
+				}
+				else
+				{
+					// Write out a single snapshot descriptor with info about all snapshots
+					DirectoryReference.CreateDirectory(Parameters.SnapshotDescriptorFile.Directory);
+					using (JsonWriter Writer = new JsonWriter(Parameters.SnapshotDescriptorFile))
+					{
+						Writer.WriteObjectStart();
+						Writer.WriteArrayStart("snapshots");
+								
+						ExportIndex = 0;
+						foreach (ExportSourceData ExportSource in ExportSources)
+						{
+							WriteExportSource(Writer, DestinationStorageType, ExportSource, ExportNames[ExportIndex]);
+							ExportIndex = ExportIndex + 1;
+						}
+
+						Writer.WriteArrayEnd();
+						Writer.WriteObjectEnd();
+					}
+				}
+			}
 
 
 			return Task.CompletedTask;
