Localization:
Refactored the Localize UAT command to be more readable.
#rnx

diff --git a/Engine/Source/Programs/AutomationTool/Scripts/Localisation.Automation.cs b/Engine/Source/Programs/AutomationTool/Scripts/Localisation.Automation.cs
--- a/Engine/Source/Programs/AutomationTool/Scripts/Localisation.Automation.cs
+++ b/Engine/Source/Programs/AutomationTool/Scripts/Localisation.Automation.cs
@@ -75,27 +75,146 @@ class Localize : BuildCommand
 		public List<IProcessResult> GatherProcessResults = new List<IProcessResult>();
 	};
 
+	private string UEProjectRoot = CmdEnv.LocalRoot;
+	private string UEProjectDirectory = "";
+	private string UEProjectName = "";
+	private List<string> LocalizationProjectNames = new();
+	private string LocalizationProviderName = "";
+	private List<string> LocalizationStepNames = new();
+	private bool bShouldGatherPlugins = false;
+	private bool bShouldEnableIncludedPlugins = false;
+	private List<string> IncludePlugins = new();
+	private List<string> ExcludePlugins = new();
+	private bool bShouldGatherPlatforms = false;
+	private string AdditionalCommandletArguments = "";
+	private bool bEnableParallelGather = false;
+	private bool bIsRunningInPreview = false;
+	private int PendingChangeList = -1;
+
 	public override void ExecuteBuild()
 	{
-		var UEProjectRoot = ParseParamValue("UEProjectRoot");
-		if (UEProjectRoot == null)
+		ParseCommandLine();
+
+		var StartTime = DateTime.UtcNow;
+
+		var LocalizationBatches = new List<LocalizationBatch>();
+		// Add the static set of localization projects as a batch
+		if (LocalizationProjectNames.Count > 0)
 		{
-			UEProjectRoot = CmdEnv.LocalRoot;
+			AddStaticLocalizationBatches(LocalizationBatches);
 		}
 
-		var UEProjectDirectory = ParseParamValue("UEProjectDirectory");
-		if (UEProjectDirectory == null)
+		// Build up any additional batches needed for platforms
+		if (bShouldGatherPlatforms)
 		{
-			throw new AutomationException("Missing required command line argument: 'UEProjectDirectory'");
+			AddPlatformLocalizationBatches(LocalizationBatches);
+		}
+
+		// Build up any additional batches needed for plugins
+		if (bShouldGatherPlugins)
+		{
+			AddPluginLocalizationBatches(LocalizationBatches);
+		}
+
+		// Create a single changelist to use for all changes
+		if (P4Enabled && !bIsRunningInPreview)
+		{
+			var ChangeListCommitMessage = String.Format("Localization Automation using CL {0}", P4Env.Changelist);
+			if (File.Exists(CombinePaths(CmdEnv.LocalRoot, @"Engine/Restricted/NotForLicensees/Build/EpicInternal.txt")))
+			{
+				ChangeListCommitMessage += "\n#okforgithub ignore";
+			}
+
+			PendingChangeList = P4.CreateChange(P4Env.Client, ChangeListCommitMessage);
+		}
+
+		// Prepare to process each localization batch
+		var LocalizationTasks = new List<LocalizationTask>();
+		foreach (var LocalizationBatch in LocalizationBatches)
+		{
+			var LocalizationTask = new LocalizationTask(LocalizationBatch, UEProjectRoot, LocalizationProviderName, PendingChangeList, this);
+			LocalizationTasks.Add(LocalizationTask);
+
+			// Make sure the Localization configs and content is up-to-date to ensure we don't get errors later on
+			// we still run this in preview mode to bring all the gather configs up to date
+			if (P4Enabled)
+			{
+				Logger.LogInformation("Sync necessary content to head revision");
+				P4.Sync(P4Env.Branch + "/" + LocalizationTask.Batch.LocalizationTargetDirectory + "/Config/Localization/...");
+				P4.Sync(P4Env.Branch + "/" + LocalizationTask.Batch.LocalizationTargetDirectory + "/Content/Localization/...");
+			}
+
+			// Generate the info we need to gather for each project
+			foreach (var ProjectName in LocalizationTask.Batch.LocalizationProjectNames)
+			{
+				LocalizationTask.ProjectInfos.Add(GenerateProjectInfo(LocalizationTask.RootLocalizationTargetDirectory, ProjectName, LocalizationStepNames));
+			}
+		}
+
+		// Hash the current PO files on disk so we can work out whether they actually change
+		Dictionary<string, byte[]> InitalPOFileHashes = null;
+		if (P4Enabled && !bIsRunningInPreview)
+		{
+			InitalPOFileHashes = GetPOFileHashes(LocalizationBatches, UEProjectRoot);
+		}
+
+		InitializeLocalizationProvider(LocalizationTasks);
+
+		// Download the latest translations from our localization provider
+		if (LocalizationStepNames.Contains("Download"))
+		{
+			DownloadFilesFromLocalizationProvider(LocalizationTasks);
+		}
+
+		// Begin the gather command for each task
+		// These can run in parallel when ParallelGather is enabled
+		StartGatherCommands(LocalizationTasks);
+
+		// Wait for each commandlet process to finish and report the result.
+		// This runs even for non-parallel execution to log the exit state of the process.
+		WaitForCommandletResults(LocalizationTasks);
+
+		// If we are running in preview, we can go ahead and delete all generated preview files after the gather step is complete 
+		if (bIsRunningInPreview)
+        {
+			CleanUpGeneratedPreviewFiles(LocalizationBatches);
+		}
+
+		// Upload the latest sources to our localization provider
+		if (LocalizationStepNames.Contains("Upload"))
+		{
+			UploadFilesToLocalizationProvider(LocalizationTasks);
 		}
 
-		var UEProjectName = ParseParamValue("UEProjectName");
-		if (UEProjectName == null)
+		// Clean-up the changelist so it only contains the changed files, and then submit it (if we were asked to)
+		if (P4Enabled && !bIsRunningInPreview)
 		{
-			UEProjectName = "";
+			// Revert any PO files that haven't changed aside from their header
+			RevertUnchangedFiles(LocalizationBatches, InitalPOFileHashes);
+
+			// Submit that single changelist now
+			if (AllowSubmit)
+			{
+				int SubmittedChangeList;
+				P4.Submit(PendingChangeList, out SubmittedChangeList);
+			}
 		}
 
-		var LocalizationProjectNames = new List<string>();
+		var RunDuration = (DateTime.UtcNow - StartTime).TotalMilliseconds;
+		Logger.LogInformation("Localize command finished in {Arg0} seconds", RunDuration / 1000);
+	}
+
+	private void ParseCommandLine()
+	{
+		UEProjectRoot = ParseParamValue("UEProjectRoot", Default: CmdEnv.LocalRoot);
+		UEProjectDirectory = ParseParamValue("UEProjectDirectory");
+		if (UEProjectDirectory == null)
+		{
+			throw new AutomationException("Missing required command line argument: 'UEProjectDirectory'");
+		}
+
+		UEProjectName = ParseParamValue("UEProjectName", Default: "");
+
 		{
 			var LocalizationProjectNamesStr = ParseParamValue("LocalizationProjectNames");
 			if (LocalizationProjectNamesStr != null)
@@ -107,13 +226,8 @@ class Localize : BuildCommand
 			}
 		}
 
-		var LocalizationProviderName = ParseParamValue("LocalizationProvider");
-		if (LocalizationProviderName == null)
-		{
-			LocalizationProviderName = "";
-		}
+		LocalizationProviderName = ParseParamValue("LocalizationProvider", Default: "");
 
-		var LocalizationStepNames = new List<string>();
 		{
 			var LocalizationStepNamesStr = ParseParamValue("LocalizationSteps");
 			if (LocalizationStepNamesStr == null)
@@ -129,12 +243,9 @@ class Localize : BuildCommand
 			}
 			LocalizationStepNames.Add("Monolithic"); // Always allow the monolithic scripts to run as we don't know which steps they do
 		}
-		
-		bool bShouldGatherPlugins = ParseParam("IncludePlugins");
-		bool bShouldEnableIncludedPlugins = ParseParam("EnableIncludedPlugins");
 
-			var IncludePlugins = new List<string>();
-		var ExcludePlugins = new List<string>();
+		bShouldGatherPlugins = ParseParam("IncludePlugins");
+		bShouldEnableIncludedPlugins = ParseParam("EnableIncludedPlugins");
 
 		string PluginsRootPath = CombinePaths(UEProjectRoot, UEProjectDirectory);
 		string IncludePluginsUnderDirectoryStr = ParseParamValue("IncludePluginsDirectory");
@@ -172,15 +283,11 @@ class Localize : BuildCommand
 			}
 		}
 
-		var ShouldGatherPlatforms = ParseParam("IncludePlatforms");
+		bShouldGatherPlatforms = ParseParam("IncludePlatforms");
 
-		var AdditionalCommandletArguments = ParseParamValue("AdditionalCommandletArguments");
-		if (AdditionalCommandletArguments == null)
-		{
-			AdditionalCommandletArguments = "";
-		}
+		AdditionalCommandletArguments = ParseParamValue("AdditionalCommandletArguments", Default: "");
 		// We remove any leading or trailing quotes from AdditionalCommandletArguments
-		else if (!String.IsNullOrEmpty(AdditionalCommandletArguments))
+		if (!String.IsNullOrEmpty(AdditionalCommandletArguments))
 		{
 			AdditionalCommandletArguments = AdditionalCommandletArguments.Trim();
 			if (AdditionalCommandletArguments.StartsWith("\"") && AdditionalCommandletArguments.EndsWith("\""))
@@ -190,130 +297,163 @@ class Localize : BuildCommand
 			}
 		}
 
-		var EnableParallelGather = ParseParam("ParallelGather");
+		bEnableParallelGather = ParseParam("ParallelGather");
 
-		var IsRunningInPreview = ParseParam("Preview");
+		bIsRunningInPreview = ParseParam("Preview");
 		// We pass the preview switch along to have the gather text commandlets exhibit different behaviors. See UGatherTextCommandlet
-		if (IsRunningInPreview)
+		if (bIsRunningInPreview)
 		{
 			Logger.LogInformation("Running in preview mode. Preview switch will be passed along to all localization commandlets to be run.");
 			AdditionalCommandletArguments += " -Preview";
 		}
-			
-
-			var StartTime = DateTime.UtcNow;
-
-		var LocalizationBatches = new List<LocalizationBatch>();
+	}
 
-		// Add the static set of localization projects as a batch
-		if (LocalizationProjectNames.Count > 0)
-		{
-			LocalizationBatches.Add(new LocalizationBatch(UEProjectDirectory, UEProjectDirectory, "", LocalizationProjectNames));
-		}
+	private void AddStaticLocalizationBatches(List<LocalizationBatch> LocalizationBatches)
+	{
+		LocalizationBatches.Add(new LocalizationBatch(UEProjectDirectory, UEProjectDirectory, "", LocalizationProjectNames));
+	}
 
-		// Build up any additional batches needed for platforms
-		if (ShouldGatherPlatforms)
+	private void AddPlatformLocalizationBatches(List<LocalizationBatch> LocalizationBatches)
+	{
+		var PlatformsRootDirectory = new DirectoryReference(CombinePaths(UEProjectRoot, UEProjectDirectory, "Platforms"));
+		if (DirectoryReference.Exists(PlatformsRootDirectory))
 		{
-			var PlatformsRootDirectory = new DirectoryReference(CombinePaths(UEProjectRoot, UEProjectDirectory, "Platforms"));
-			if (DirectoryReference.Exists(PlatformsRootDirectory))
+			foreach (DirectoryReference PlatformDirectory in DirectoryReference.EnumerateDirectories(PlatformsRootDirectory))
 			{
-				foreach (DirectoryReference PlatformDirectory in DirectoryReference.EnumerateDirectories(PlatformsRootDirectory))
+				// Find the localization targets defined for this platform
+				var PlatformTargetNames = GetLocalizationTargetsFromDirectory(new DirectoryReference(CombinePaths(PlatformDirectory.FullName, "Config", "Localization")));
+				if (PlatformTargetNames.Count > 0)
 				{
-					// Find the localization targets defined for this platform
-					var PlatformTargetNames = GetLocalizationTargetsFromDirectory(new DirectoryReference(CombinePaths(PlatformDirectory.FullName, "Config", "Localization")));
-					if (PlatformTargetNames.Count > 0)
-					{
-						var RootRelativePluginPath = PlatformDirectory.MakeRelativeTo(new DirectoryReference(UEProjectRoot));
-						RootRelativePluginPath = RootRelativePluginPath.Replace('\\', '/'); // Make sure we use / as these paths are used with P4
+					var RootRelativePluginPath = PlatformDirectory.MakeRelativeTo(new DirectoryReference(UEProjectRoot));
+					RootRelativePluginPath = RootRelativePluginPath.Replace('\\', '/'); // Make sure we use / as these paths are used with P4
 
-						LocalizationBatches.Add(new LocalizationBatch(UEProjectDirectory, RootRelativePluginPath, "", PlatformTargetNames));
-					}
+					LocalizationBatches.Add(new LocalizationBatch(UEProjectDirectory, RootRelativePluginPath, "", PlatformTargetNames));
 				}
 			}
 		}
+	}
 
-		// Build up any additional batches needed for plugins
-		if (bShouldGatherPlugins)
+	private void AddPluginLocalizationBatches(List<LocalizationBatch> LocalizationBatches)
+	{
+		var PluginsRootDirectory = new DirectoryReference(CombinePaths(UEProjectRoot, UEProjectDirectory));
+		IReadOnlyList<PluginInfo> AllPlugins = Plugins.ReadPluginsFromDirectory(PluginsRootDirectory, "Plugins", UEProjectName.Length == 0 ? PluginType.Engine : PluginType.Project);
+
+		// Add a batch for each plugin that meets our criteria
+		var AvailablePluginNames = new HashSet<string>();
+		foreach (var PluginInfo in AllPlugins)
 		{
-			var PluginsRootDirectory = new DirectoryReference(CombinePaths(UEProjectRoot, UEProjectDirectory));
-			IReadOnlyList<PluginInfo> AllPlugins = Plugins.ReadPluginsFromDirectory(PluginsRootDirectory, "Plugins", UEProjectName.Length == 0 ? PluginType.Engine : PluginType.Project);
+			AvailablePluginNames.Add(PluginInfo.Name);
 
-			// Add a batch for each plugin that meets our criteria
-			var AvailablePluginNames = new HashSet<string>();
-			foreach (var PluginInfo in AllPlugins)
+			bool bShouldIncludePlugin = (IncludePlugins.Count == 0 || IncludePlugins.Contains(PluginInfo.Name)) && !ExcludePlugins.Contains(PluginInfo.Name);
+			bool bPluginHasLocalizationTarget = PluginInfo.Descriptor.LocalizationTargets != null && PluginInfo.Descriptor.LocalizationTargets.Length > 0;
+			if (bShouldIncludePlugin && bPluginHasLocalizationTarget)
 			{
-				AvailablePluginNames.Add(PluginInfo.Name);
+				var RootRelativePluginPath = PluginInfo.Directory.MakeRelativeTo(new DirectoryReference(UEProjectRoot));
+				RootRelativePluginPath = RootRelativePluginPath.Replace('\\', '/'); // Make sure we use / as these paths are used with P4
 
-				bool ShouldIncludePlugin = (IncludePlugins.Count == 0 || IncludePlugins.Contains(PluginInfo.Name)) && !ExcludePlugins.Contains(PluginInfo.Name);
-				bool PluginHasLocalizationTarget = PluginInfo.Descriptor.LocalizationTargets != null && PluginInfo.Descriptor.LocalizationTargets.Length > 0;
-				if (ShouldIncludePlugin && PluginHasLocalizationTarget)
+				var PluginTargetNames = new List<string>();
+				foreach (var LocalizationTarget in PluginInfo.Descriptor.LocalizationTargets)
 				{
-					var RootRelativePluginPath = PluginInfo.Directory.MakeRelativeTo(new DirectoryReference(UEProjectRoot));
-					RootRelativePluginPath = RootRelativePluginPath.Replace('\\', '/'); // Make sure we use / as these paths are used with P4
-
-					var PluginTargetNames = new List<string>();
-					foreach (var LocalizationTarget in PluginInfo.Descriptor.LocalizationTargets)
-					{
-						PluginTargetNames.Add(LocalizationTarget.Name);
-					}
-
-					LocalizationBatches.Add(new LocalizationBatch(UEProjectDirectory, RootRelativePluginPath, PluginInfo.Name, PluginTargetNames));
+					PluginTargetNames.Add(LocalizationTarget.Name);
 				}
-			}
 
-			// If we had an explicit list of plugins to include, warn if any were missing
-			foreach (string PluginName in IncludePlugins)
-			{
-				if (!AvailablePluginNames.Contains(PluginName))
-				{
-					Logger.LogWarning("The plugin '{PluginName}' specified by -IncludePlugins wasn't found and will be skipped.", PluginName);
-				}
+				LocalizationBatches.Add(new LocalizationBatch(UEProjectDirectory, RootRelativePluginPath, PluginInfo.Name, PluginTargetNames));
 			}
 		}
 
-		// Create a single changelist to use for all changes
-		int PendingChangeList = 0;
-		if (P4Enabled && !IsRunningInPreview)
+		// If we had an explicit list of plugins to include, warn if any were missing
+		foreach (string PluginName in IncludePlugins)
 		{
-			var ChangeListCommitMessage = String.Format("Localization Automation using CL {0}", P4Env.Changelist);
-			if (File.Exists(CombinePaths(CmdEnv.LocalRoot, @"Engine/Restricted/NotForLicensees/Build/EpicInternal.txt")))
+			if (!AvailablePluginNames.Contains(PluginName))
 			{
-				ChangeListCommitMessage += "\n#okforgithub ignore";
+				Logger.LogWarning("The plugin '{PluginName}' specified by -IncludePlugins wasn't found and will be skipped.", PluginName);
 			}
+		}
+	}
 
-			PendingChangeList = P4.CreateChange(P4Env.Client, ChangeListCommitMessage);
+	private string BuildEditorArguments()
+	{
+		var EditorArguments = P4Enabled
+				? String.Format("-SCCProvider=Perforce -P4Port={0} -P4User={1} -P4Client={2} -P4Passwd={3} -P4Changelist={4} -EnableSCC -DisableSCCSubmit", P4Env.ServerAndPort, P4Env.User, P4Env.Client, P4.GetAuthenticationToken(), PendingChangeList)
+				: "-SCCProvider=None";
+		if (IsBuildMachine)
+		{
+			EditorArguments += " -BuildMachine";
+		}
+		EditorArguments += " -Unattended";
+		EditorArguments += " -NoShaderCompile";
+		//EditorArguments += " -LogLocalizationConflicts";
+		if (bEnableParallelGather)
+		{
+			EditorArguments += " -multiprocess";
 		}
 
-		// Prepare to process each localization batch
-		var LocalizationTasks = new List<LocalizationTask>();
-		foreach (var LocalizationBatch in LocalizationBatches)
+		// We append all the included plugins to -EnablePlugins if -EnableIncludedPlugins is enabled. This wil ensure that the plugin content and metadata will be loaded.
+		// @TODOLocalization: Ideally the enabling of plugins should be per batch, otherwise each instance of the editor is enabling a bunch of plugins it doesn't need 
+		if (!string.IsNullOrEmpty(AdditionalCommandletArguments) && bShouldEnableIncludedPlugins && IncludePlugins.Count > 0)
 		{
-			var LocalizationTask = new LocalizationTask(LocalizationBatch, UEProjectRoot, LocalizationProviderName, PendingChangeList, this);
-			LocalizationTasks.Add(LocalizationTask);
+			HashSet<string> PluginsToEnableSet = IncludePlugins.Except(ExcludePlugins).ToHashSet();
 
-			// Make sure the Localization configs and content is up-to-date to ensure we don't get errors later on
-			// we still run this in preview mode to bring all the gather configs up to date
-			if (P4Enabled)
+			// It's possible that there are already specified values for -EnabledPlugins, we willneed to try and parse them first.
+			string EnablePluginsToken = "-EnablePlugins=";
+			string EnablePluginsNewValue = "";
+
+			int EnablePluginsTokenIndex = AdditionalCommandletArguments.IndexOf(EnablePluginsToken);
+			string EnablePluginsOldValue = "";
+			if (EnablePluginsTokenIndex > -1)
 			{
-				Logger.LogInformation("Sync necessary content to head revision");
-				P4.Sync(P4Env.Branch + "/" + LocalizationTask.Batch.LocalizationTargetDirectory + "/Config/Localization/...");
-				P4.Sync(P4Env.Branch + "/" + LocalizationTask.Batch.LocalizationTargetDirectory + "/Content/Localization/...");
+				// -EnablePlugins token exists in the additional commandlet args. We need to process it 
+				int EnablePluginsValueStartIndex = EnablePluginsTokenIndex + EnablePluginsToken.Length;
+				// We try and find the end of the string where it's separated by a space  between the next token 
+				int EnablePluginsValueEndIndex = AdditionalCommandletArguments.IndexOf(" ", EnablePluginsValueStartIndex);
+				// We can't find a next space. THis means we're the last parameter in AdditionalCommandletArguments. The end index will be the length of the string 
+				if (EnablePluginsValueEndIndex == -1)
+				{
+					EnablePluginsValueEndIndex = AdditionalCommandletArguments.Length;
+				}
+				// Isolate the value of -EnablePlugins and add them to our list of plugins to enable 
+				EnablePluginsOldValue = AdditionalCommandletArguments.Substring(EnablePluginsValueStartIndex, EnablePluginsValueEndIndex - EnablePluginsValueStartIndex);
+				foreach (string Plugin in EnablePluginsOldValue.Split(','))
+				{
+					PluginsToEnableSet.Add(Plugin);
+				}
 			}
 
-			// Generate the info we need to gather for each project
-			foreach (var ProjectName in LocalizationTask.Batch.LocalizationProjectNames)
+			// Just a counter to help iterate through the set to build out the comma separated value 
+			int IterationCount = 0;
+			StringBuilder EnablePluginsBuilder = new StringBuilder();
+			foreach (string Plugin in PluginsToEnableSet)
 			{
-				LocalizationTask.ProjectInfos.Add(GenerateProjectInfo(LocalizationTask.RootLocalizationTargetDirectory, ProjectName, LocalizationStepNames));
+				EnablePluginsBuilder.Append(Plugin);
+				if (IterationCount < PluginsToEnableSet.Count - 1)
+				{
+					EnablePluginsBuilder.Append(",");
+				}
+				++IterationCount;
+			}
+			EnablePluginsNewValue = EnablePluginsBuilder.ToString();
+			Logger.LogInformation($"Appending following plugins to be enabled: {EnablePluginsNewValue}");
+			// if we already had a value of -EnablePlugins in AdditionalCommandletArguments, we'll need to replace that with the new values we've created.
+			if (EnablePluginsTokenIndex > -1)
+			{
+				AdditionalCommandletArguments = AdditionalCommandletArguments.Replace(EnablePluginsToken + EnablePluginsOldValue, EnablePluginsToken + EnablePluginsNewValue);
+			}
+			else
+			{
+				// The token doesn't exist, we'll add it to the end 
+				AdditionalCommandletArguments += " " + EnablePluginsToken + EnablePluginsNewValue;
 			}
 		}
 
-		// Hash the current PO files on disk so we can work out whether they actually change
-		Dictionary<string, byte[]> InitalPOFileHashes = null;
-		if (P4Enabled && !IsRunningInPreview)
+		if (!String.IsNullOrEmpty(AdditionalCommandletArguments))
 		{
-			InitalPOFileHashes = GetPOFileHashes(LocalizationBatches, UEProjectRoot);
+			EditorArguments += " " + AdditionalCommandletArguments;
 		}
+		return EditorArguments;
+	}
 
+	private void InitializeLocalizationProvider(List<LocalizationTask> LocalizationTasks)
+	{
 		foreach (var LocalizationTask in LocalizationTasks)
 		{
 			if (LocalizationTask.LocProvider != null)
@@ -325,149 +465,101 @@ class Localize : BuildCommand
 				}
 			}
 		}
+	}
 
-		// Download the latest translations from our localization provider
-		if (LocalizationStepNames.Contains("Download"))
+	private void DownloadFilesFromLocalizationProvider(List<LocalizationTask> LocalizationTasks)
+	{
+		foreach (var LocalizationTask in LocalizationTasks)
 		{
-			foreach (var LocalizationTask in LocalizationTasks)
+			if (LocalizationTask.LocProvider != null)
 			{
-				if (LocalizationTask.LocProvider != null)
+				foreach (var ProjectInfo in LocalizationTask.ProjectInfos)
 				{
-					foreach (var ProjectInfo in LocalizationTask.ProjectInfos)
-					{
-						Task DownloadTask = LocalizationTask.LocProvider.DownloadProjectFromLocalizationProvider(ProjectInfo.ProjectName, ProjectInfo.ImportInfo);
-						DownloadTask.Wait();
-					}
+					Task DownloadTask = LocalizationTask.LocProvider.DownloadProjectFromLocalizationProvider(ProjectInfo.ProjectName, ProjectInfo.ImportInfo);
+					DownloadTask.Wait();
 				}
 			}
 		}
+	}
 
-		// Begin the gather command for each task
-		// These can run in parallel when ParallelGather is enabled
+	private void UploadFilesToLocalizationProvider(List<LocalizationTask> LocalizationTasks)
+	{
+		foreach (var LocalizationTask in LocalizationTasks)
 		{
-			var EditorExe = CombinePaths(CmdEnv.LocalRoot, @"Engine/Binaries/Win64/UnrealEditor-Cmd.exe");
-			if (!File.Exists(EditorExe))
-			{
-				// Try using the debug .exe instead 
-				EditorExe = CombinePaths(CmdEnv.LocalRoot, @"Engine/Binaries/Win64/UnrealEditor-Win64-Debug-Cmd.exe");
-			}
-
-			// Set the common basic editor arguments
-			var EditorArguments = P4Enabled 
-				? String.Format("-SCCProvider=Perforce -P4Port={0} -P4User={1} -P4Client={2} -P4Passwd={3} -P4Changelist={4} -EnableSCC -DisableSCCSubmit", P4Env.ServerAndPort, P4Env.User, P4Env.Client, P4.GetAuthenticationToken(), PendingChangeList)
-				: "-SCCProvider=None";
-			if (IsBuildMachine)
-			{
-				EditorArguments += " -BuildMachine";
-			}
-			EditorArguments += " -Unattended";
-			EditorArguments += " -NoShaderCompile";
-			//EditorArguments += " -LogLocalizationConflicts";
-			if (EnableParallelGather)
-			{
-				EditorArguments += " -multiprocess";
-			}
-
-			// We append all the included plugins to -EnablePlugins if -EnableIncludedPlugins is enabled. This wil ensure that the plugin content and metadata will be loaded.
-			// @TODOLocalization: Ideally the enabling of plugins should be per batch, otherwise each instance of the editor is enabling a bunch of plugins it doesn't need 
-			if (!string.IsNullOrEmpty(AdditionalCommandletArguments) && bShouldEnableIncludedPlugins && IncludePlugins.Count > 0)
+			if (LocalizationTask.LocProvider != null)
 			{
-				HashSet<string> PluginsToEnableSet = IncludePlugins.Except(ExcludePlugins).ToHashSet();
-
-				// It's possible that there are already specified values for -EnabledPlugins, we willneed to try and parse them first.
-				string EnablePluginsToken = "-EnablePlugins=";
-				string EnablePluginsNewValue = "";
-
-				int EnablePluginsTokenIndex = AdditionalCommandletArguments.IndexOf(EnablePluginsToken);
-				string EnablePluginsOldValue = "";
-				if (EnablePluginsTokenIndex > -1)
+				// Upload all text to our localization provider
+				for (int ProjectIndex = 0; ProjectIndex < LocalizationTask.ProjectInfos.Count; ++ProjectIndex)
 				{
-					// -EnablePlugins token exists in the additional commandlet args. We need to process it 
-					int EnablePluginsValueStartIndex = EnablePluginsTokenIndex + EnablePluginsToken.Length;
-					// We try and find the end of the string where it's separated by a space  between the next token 
-					int EnablePluginsValueEndIndex = AdditionalCommandletArguments.IndexOf(" ", EnablePluginsValueStartIndex);
-					// We can't find a next space. THis means we're the last parameter in AdditionalCommandletArguments. The end index will be the length of the string 
-					if (EnablePluginsValueEndIndex== -1)
-					{
-						EnablePluginsValueEndIndex = AdditionalCommandletArguments.Length;
-					}
-					// Isolate the value of -EnablePlugins and add them to our list of plugins to enable 
-					EnablePluginsOldValue = AdditionalCommandletArguments.Substring(EnablePluginsValueStartIndex, EnablePluginsValueEndIndex - EnablePluginsValueStartIndex);
-					foreach (string Plugin in EnablePluginsOldValue.Split(','))
-					{
-						PluginsToEnableSet.Add(Plugin);
-					}
-				}
+					var ProjectInfo = LocalizationTask.ProjectInfos[ProjectIndex];
+					var RunResult = LocalizationTask.GatherProcessResults[ProjectIndex];
 
-				// Just a counter to help iterate through the set to build out the comma separated value 
-				int IterationCount = 0;
-				StringBuilder EnablePluginsBuilder = new StringBuilder();
-					foreach (string Plugin in PluginsToEnableSet)
-				{
-					EnablePluginsBuilder.Append(Plugin);
-					if (IterationCount< PluginsToEnableSet.Count - 1)
+					if (RunResult != null && RunResult.ExitCode == 0)
 					{
-						EnablePluginsBuilder.Append(",");
+						// Recalculate the split platform paths before doing the upload, as the export may have changed them
+						ProjectInfo.ExportInfo.CalculateSplitPlatformNames(LocalizationTask.RootLocalizationTargetDirectory);
+						Task UploadTask = LocalizationTask.LocProvider.UploadProjectToLocalizationProvider(ProjectInfo.ProjectName, ProjectInfo.ExportInfo);
+						UploadTask.Wait();
 					}
-					++IterationCount;
-				}
-				EnablePluginsNewValue = EnablePluginsBuilder.ToString();
-				Logger.LogInformation($"Appending following plugins to be enabled: {EnablePluginsNewValue}");
-					// if we already had a value of -EnablePlugins in AdditionalCommandletArguments, we'll need to replace that with the new values we've created.
-					if (EnablePluginsTokenIndex > -1)
-				{
-					AdditionalCommandletArguments = AdditionalCommandletArguments.Replace(EnablePluginsToken + EnablePluginsOldValue, EnablePluginsToken + EnablePluginsNewValue);
-				}
 					else
-				{
-					// The token doesn't exist, we'll add it to the end 
-					AdditionalCommandletArguments += " " + EnablePluginsToken + EnablePluginsNewValue;
+					{
+						Logger.LogWarning("Skipping upload to the localization provider for '{Arg0}' due to an earlier commandlet failure.", ProjectInfo.ProjectName);
+					}
 				}
 			}
+		}
+	}
 
-				if (!String.IsNullOrEmpty(AdditionalCommandletArguments))
-			{
-				EditorArguments += " " + AdditionalCommandletArguments;
-			}
+	private void StartGatherCommands(List<LocalizationTask> LocalizationTasks)
+	{
+		var EditorExe = CombinePaths(CmdEnv.LocalRoot, @"Engine/Binaries/Win64/UnrealEditor-Cmd.exe");
+		if (!File.Exists(EditorExe))
+		{
+			// Try using the debug .exe instead 
+			EditorExe = CombinePaths(CmdEnv.LocalRoot, @"Engine/Binaries/Win64/UnrealEditor-Win64-Debug-Cmd.exe");
+		}
 
-			// Set the common process run options
-			var CommandletRunOptions = ERunOptions.Default | ERunOptions.NoLoggingOfRunCommand; // Disable logging of the run command as it will print the exit code which GUBP can pick up as an error (we do that ourselves later)
-			if (EnableParallelGather)
-			{
-				CommandletRunOptions |= ERunOptions.NoWaitForExit;
-			}
+		// Set the common basic editor arguments
+		string EditorArguments = BuildEditorArguments();
 
-			foreach (var LocalizationTask in LocalizationTasks)
-			{
-				var ProjectArgument = String.IsNullOrEmpty(UEProjectName) ? "" : String.Format("\"{0}\"", Path.Combine(LocalizationTask.RootWorkingDirectory, String.Format("{0}.uproject", UEProjectName)));
+		// Set the common process run options
+		var CommandletRunOptions = ERunOptions.Default | ERunOptions.NoLoggingOfRunCommand; // Disable logging of the run command as it will print the exit code which GUBP can pick up as an error (we do that ourselves later)
+		if (bEnableParallelGather)
+		{
+			CommandletRunOptions |= ERunOptions.NoWaitForExit;
+		}
 
-				foreach (var ProjectInfo in LocalizationTask.ProjectInfos)
+		foreach (var LocalizationTask in LocalizationTasks)
+		{
+			var ProjectArgument = String.IsNullOrEmpty(UEProjectName) ? "" : String.Format("\"{0}\"", Path.Combine(LocalizationTask.RootWorkingDirectory, String.Format("{0}.uproject", UEProjectName)));
+
+			foreach (var ProjectInfo in LocalizationTask.ProjectInfos)
+			{
+				var LocalizationConfigFiles = new List<string>();
+				foreach (var LocalizationStep in ProjectInfo.LocalizationSteps)
 				{
-					var LocalizationConfigFiles = new List<string>();
-					foreach (var LocalizationStep in ProjectInfo.LocalizationSteps)
+					if (LocalizationStepNames.Contains(LocalizationStep.Name))
 					{
-						if (LocalizationStepNames.Contains(LocalizationStep.Name))
-						{
-							LocalizationConfigFiles.Add(LocalizationStep.LocalizationConfigFile);
-						}
+						LocalizationConfigFiles.Add(LocalizationStep.LocalizationConfigFile);
 					}
+				}
 
-					if (LocalizationConfigFiles.Count > 0)
-					{
-						var Arguments = String.Format("{0} -run=GatherText -config=\"{1}\" {2}", ProjectArgument, String.Join(";", LocalizationConfigFiles), EditorArguments);
-						Logger.LogInformation("Running localization commandlet for '{Arg0}': {Arguments}", ProjectInfo.ProjectName, Arguments);
-						LocalizationTask.GatherProcessResults.Add(Run(EditorExe, Arguments, null, CommandletRunOptions));
-					}
-					else
-					{
-						LocalizationTask.GatherProcessResults.Add(null);
-					}
+				if (LocalizationConfigFiles.Count > 0)
+				{
+					var Arguments = String.Format("{0} -run=GatherText -config=\"{1}\" {2}", ProjectArgument, String.Join(";", LocalizationConfigFiles), EditorArguments);
+					Logger.LogInformation("Running localization commandlet for '{Arg0}': {Arguments}", ProjectInfo.ProjectName, Arguments);
+					LocalizationTask.GatherProcessResults.Add(Run(EditorExe, Arguments, null, CommandletRunOptions));
+				}
+				else
+				{
+					LocalizationTask.GatherProcessResults.Add(null);
 				}
 			}
 		}
+	}
 
-		// Wait for each commandlet process to finish and report the result.
-		// This runs even for non-parallel execution to log the exit state of the process.
+	private void WaitForCommandletResults(List<LocalizationTask> LocalizationTasks)
+	{
 		foreach (var LocalizationTask in LocalizationTasks)
 		{
 			for (int ProjectIndex = 0; ProjectIndex < LocalizationTask.ProjectInfos.Count; ++ProjectIndex)
@@ -480,7 +572,7 @@ class Localize : BuildCommand
 					RunResult.WaitForExit();
 					RunResult.OnProcessExited();
 					RunResult.DisposeProcess();
-					
+
 					if (RunResult.ExitCode == 0)
 					{
 						Logger.LogInformation("The localization commandlet for '{Arg0}' exited with code 0.", ProjectInfo.ProjectName);
@@ -492,101 +584,64 @@ class Localize : BuildCommand
 				}
 			}
 		}
+	}
 
-		// If we are running in preview, we can go ahead and delete all generated preview files after the gather step is complete 
-		if (IsRunningInPreview)
-        {
-			var PreviewManifestFiles= GetPreviewManifestFilesToDelete(LocalizationBatches, UEProjectRoot);
-			foreach (var PreviewManifestFile in PreviewManifestFiles) 
+	private void CleanUpGeneratedPreviewFiles(List<LocalizationBatch> LocalizationBatches)
+	{
+		var PreviewManifestFiles = GetPreviewManifestFilesToDelete(LocalizationBatches, UEProjectRoot);
+		foreach (var PreviewManifestFile in PreviewManifestFiles)
+		{
+			Logger.LogInformation("Deleting preview manifest file {PreviewManifestFile}.", PreviewManifestFile);
+			try
 			{
-				Logger.LogInformation("Deleting preview manifest file {PreviewManifestFile}.", PreviewManifestFile);
-				try
-				{
-					File.Delete(PreviewManifestFile);
-				}
-				catch (Exception Ex)
-				{
-					Logger.LogInformation("[FAILED] Deleting preview file: '{PreviewManifestFile}' - {Ex}", PreviewManifestFile, Ex);
-				}
+				File.Delete(PreviewManifestFile);
+			}
+			catch (Exception Ex)
+			{
+				Logger.LogInformation("[FAILED] Deleting preview file: '{PreviewManifestFile}' - {Ex}", PreviewManifestFile, Ex);
 			}
 		}
+	}
 
-		// Upload the latest sources to our localization provider
-		if (LocalizationStepNames.Contains("Upload"))
+	private void RevertUnchangedFiles(List<LocalizationBatch> LocalizationBatches, Dictionary<string, Byte[]> InitalPOFileHashes)
+	{
+		if (!P4Enabled || bIsRunningInPreview)
 		{
-			foreach (var LocalizationTask in LocalizationTasks)
-			{
-				if (LocalizationTask.LocProvider != null)
-				{
-					// Upload all text to our localization provider
-					for (int ProjectIndex = 0; ProjectIndex < LocalizationTask.ProjectInfos.Count; ++ProjectIndex)
-					{
-						var ProjectInfo = LocalizationTask.ProjectInfos[ProjectIndex];
-						var RunResult = LocalizationTask.GatherProcessResults[ProjectIndex];
-
-						if (RunResult != null && RunResult.ExitCode == 0)
-						{
-							// Recalculate the split platform paths before doing the upload, as the export may have changed them
-							ProjectInfo.ExportInfo.CalculateSplitPlatformNames(LocalizationTask.RootLocalizationTargetDirectory);
-							Task UploadTask = LocalizationTask.LocProvider.UploadProjectToLocalizationProvider(ProjectInfo.ProjectName, ProjectInfo.ExportInfo);
-							UploadTask.Wait();
-						}
-						else
-						{
-							Logger.LogWarning("Skipping upload to the localization provider for '{Arg0}' due to an earlier commandlet failure.", ProjectInfo.ProjectName);
-						}
-					}
-				}
-			}
+			return;
 		}
 
-		// Clean-up the changelist so it only contains the changed files, and then submit it (if we were asked to)
-		if (P4Enabled && !IsRunningInPreview)
 		{
-			// Revert any PO files that haven't changed aside from their header
-			{
-				var POFilesToRevert = new List<string>();
+			var POFilesToRevert = new List<string>();
 
-				var CurrentPOFileHashes = GetPOFileHashes(LocalizationBatches, UEProjectRoot);
-				foreach (var CurrentPOFileHashPair in CurrentPOFileHashes)
+			var CurrentPOFileHashes = GetPOFileHashes(LocalizationBatches, UEProjectRoot);
+			foreach (var CurrentPOFileHashPair in CurrentPOFileHashes)
+			{
+				byte[] InitialPOFileHash;
+				if (InitalPOFileHashes.TryGetValue(CurrentPOFileHashPair.Key, out InitialPOFileHash) && InitialPOFileHash.SequenceEqual(CurrentPOFileHashPair.Value))
 				{
-					byte[] InitialPOFileHash;
-					if (InitalPOFileHashes.TryGetValue(CurrentPOFileHashPair.Key, out InitialPOFileHash) && InitialPOFileHash.SequenceEqual(CurrentPOFileHashPair.Value))
-					{
-						POFilesToRevert.Add(CurrentPOFileHashPair.Key);
-					}
+					POFilesToRevert.Add(CurrentPOFileHashPair.Key);
 				}
+			}
 
-				if (POFilesToRevert.Count > 0)
-				{
-					var P4RevertArgsFilename = CombinePaths(CmdEnv.LocalRoot, "Engine", "Intermediate", String.Format("LocalizationP4RevertArgs-{0}.txt", Guid.NewGuid().ToString()));
+			if (POFilesToRevert.Count > 0)
+			{
+				var P4RevertArgsFilename = CombinePaths(CmdEnv.LocalRoot, "Engine", "Intermediate", String.Format("LocalizationP4RevertArgs-{0}.txt", Guid.NewGuid().ToString()));
 
-					using (StreamWriter P4RevertArgsWriter = File.CreateText(P4RevertArgsFilename))
+				using (StreamWriter P4RevertArgsWriter = File.CreateText(P4RevertArgsFilename))
+				{
+					foreach (var POFileToRevert in POFilesToRevert)
 					{
-						foreach (var POFileToRevert in POFilesToRevert)
-						{
-							P4RevertArgsWriter.WriteLine(POFileToRevert);
-						}
+						P4RevertArgsWriter.WriteLine(POFileToRevert);
 					}
-
-					P4.LogP4(String.Format("-x {0}", P4RevertArgsFilename), "revert");
-					DeleteFile_NoExceptions(P4RevertArgsFilename);
 				}
-			}
-
-			// Revert any other unchanged files
-			P4.RevertUnchanged(PendingChangeList);
 
-			// Submit that single changelist now
-			if (AllowSubmit)
-			{
-				int SubmittedChangeList;
-				P4.Submit(PendingChangeList, out SubmittedChangeList);
+				P4.LogP4(String.Format("-x {0}", P4RevertArgsFilename), "revert");
+				DeleteFile_NoExceptions(P4RevertArgsFilename);
 			}
 		}
 
-		var RunDuration = (DateTime.UtcNow - StartTime).TotalMilliseconds;
-		Logger.LogInformation("Localize command finished in {Arg0} seconds", RunDuration / 1000);
+		// Revert any other unchanged files
+		P4.RevertUnchanged(PendingChangeList);
 	}
 
 	private ProjectInfo GenerateProjectInfo(string RootWorkingDirectory, string ProjectName, IReadOnlyList<string> LocalizationStepNames)
