Fix deserialization of baseline snapshot descriptors.

diff --git a/Engine/Source/Programs/AutomationTool/BuildGraph/Tasks/ZenExportSnapshotTask.cs b/Engine/Source/Programs/AutomationTool/BuildGraph/Tasks/ZenExportSnapshotTask.cs
--- a/Engine/Source/Programs/AutomationTool/BuildGraph/Tasks/ZenExportSnapshotTask.cs
+++ b/Engine/Source/Programs/AutomationTool/BuildGraph/Tasks/ZenExportSnapshotTask.cs
@@ -17,11 +17,26 @@ using EpicGames.Serialization;
 
 namespace AutomationTool.Tasks
 {
-	enum SnapshotStorageType
+	/// <summary>
+	/// Enumeration of different storage options for snapshots.
+	/// </summary>
+	public enum SnapshotStorageType
 	{
+		/// <summary>
+		/// A reserved non-valid storage type for snapshots.
+		/// </summary>
 		Invalid,
+		/// <summary>
+		/// Snapshot stored in cloud repositories such as Unreal Cloud DDC.
+		/// </summary>
 		Cloud,
+		/// <summary>
+		/// Snapshot stored in a zenserver.
+		/// </summary>
 		Zen,
+		/// <summary>
+		/// Snapshot stored as a file on disk.
+		/// </summary>
 		File,
 	}
 
@@ -103,19 +118,65 @@ namespace AutomationTool.Tasks
 	[TaskElement("ZenExportSnapshot", typeof(ZenExportSnapshotTaskParameters))]
 	public class ZenExportSnapshotTask : BgTaskImpl
 	{
-		private class SnapshotDescriptor
+		/// <summary>
+		/// Metadata about a snapshot
+		/// </summary>
+		public class SnapshotDescriptor
 		{
-			public string Name = null;
-			public SnapshotStorageType Type = SnapshotStorageType.Invalid;
-			public string TargetPlatform = null;
-
-			public string Host = null;
-			public string Namespace = null;
-			public string Bucket = null;
-			public string Key = null;
-
-			public string Directory = null;
-			public string Filename = null;
+			/// <summary>
+			/// Name of the snapshot
+			/// </summary>
+			public string Name { get; set; }
+
+			/// <summary>
+			/// Storage type used for the snapshot
+			/// </summary>
+			public SnapshotStorageType Type { get; set; }
+
+			/// <summary>
+			/// Target platform for this snapshot
+			/// </summary>
+			public string TargetPlatform { get; set; }
+			
+			/// <summary>
+			/// For cloud snapshots, the host they are stored on.
+			/// </summary>
+			public string Host { get; set; }
+
+			/// <summary>
+			/// For cloud snapshots, the namespace they are stored in.
+			/// </summary>
+			public string Namespace { get; set; }
+
+			/// <summary>
+			/// For cloud snapshots, the bucket they are stored in.
+			/// </summary>
+			public string Bucket { get; set; }
+
+			/// <summary>
+			/// For cloud snapshots, the key they are stored in.
+			/// </summary>
+			public string Key { get; set; }
+			
+			/// <summary>
+			/// For file snapshots, the directory it is stored in.
+			/// </summary>
+			public string Directory { get; set; }
+
+			/// <summary>
+			/// For file snapshots, the filename (not including path) that they are stored in.
+			/// </summary>
+			public string Filename { get; set; }
+		}
+		/// <summary>
+		/// A collection of one or more snapshot descriptors
+		/// </summary>
+		public class SnapshotDescriptorCollection
+		{
+			/// <summary>
+			/// The list of snapshots contained within this collection.
+			/// </summary>
+			public List<SnapshotDescriptor> Snapshots { get; set; }
 		}
 		private class ExportSourceData
 		{
@@ -299,17 +360,21 @@ namespace AutomationTool.Tasks
 					NewExportSource.ProjectId = ZenServerObject["projectid"].AsString();
 					NewExportSource.OplogId = ZenServerObject["oplogid"].AsString();
 					NewExportSource.TargetPlatform = Platform;
+					NewExportSource.SnapshotBaseDescriptor = null;
 
 					FileReference PlatformSnapshotBase = new FileReference(Parameters.SnapshotBaseDescriptorFile.FullName.Replace("{Platform}", Platform, StringComparison.InvariantCultureIgnoreCase));
 
-					SnapshotDescriptor? ParsedDescriptor = null;
-					if (TryLoadJson(PlatformSnapshotBase, out ParsedDescriptor))
+					SnapshotDescriptorCollection? ParsedDescriptorCollection = null;
+					if (TryLoadJson(PlatformSnapshotBase, out ParsedDescriptorCollection) && (ParsedDescriptorCollection != null) && (ParsedDescriptorCollection.Snapshots != null))
 					{
-						NewExportSource.SnapshotBaseDescriptor = ParsedDescriptor;
-					}
-					else
-					{
-						NewExportSource.SnapshotBaseDescriptor = null;
+						foreach (SnapshotDescriptor ParsedDescriptor in ParsedDescriptorCollection.Snapshots)
+						{
+							if (ParsedDescriptor.TargetPlatform == Platform)
+							{
+								NewExportSource.SnapshotBaseDescriptor = ParsedDescriptor;
+								break;
+							}
+						}
 					}
 
 					ExportSources.Add(NewExportSource);
