Localization:
Introduced the experimental LocalizationTargetEditor UAT command. This command allows users to create, update, delete and query localization targets in the Engine or project. The LocalizationTargetEditor UAT command is set up to have a series of sub-commands which can be executed with different parameters. This allows the LocalizationTargetEditor to be extensible to users for specialized localization target operations.
- Introduced the LocalizationTargetCommand abstract base class. This class is what all othe LocalizationTargetEditor sub-commands should be drived from. A scan of all sub-classes of LocalizationTargetCommand will be scanned for to determine all available localization target commands.
- Introduced the PluginLocalizationTargetCommand abstract base class. All commands to be operated on plugins should derive from this class.
- Introduced the NewPluginLocalizationTarget sub-command which will create localization config files for plugins and make the necessary updates to .uplugin files. All plugins under a provided directory can be operated on at a time to batch create the necessary localization config files.
- The NewPluginLocalizationTarget sub-command has perforce integration and can automatically add/check out all files that have been created/edited.
- Introduced a number of classes to encapsulate all the data in the various localization config files. This allows for generation of localization config files from code.
- Introduced the LocalizationConfigFileBuilder and LocalizationConfigFileGenerator base classes. These classes and their sub-classes are used to generate localization config files in the correct format provided the config file data. Current support is only for modular localization config files. However, the LocalizationConfigFileBuilder can also be used to generate Monolithic localization config files if necessary.also

#test Used the NewPluginLocalizationTarget sub-command to generate a number of plugins. The generated localization config files yield correct localization gather results and the plugin .locres file was successfully loaded at runtime to display the correct localized text.

diff --git a/Engine/Source/Programs/AutomationTool/Localization/LocalizationConfigFile.cs b/Engine/Source/Programs/AutomationTool/Localization/LocalizationConfigFile.cs
new file mode 100644
--- /dev/null
+++ b/Engine/Source/Programs/AutomationTool/Localization/LocalizationConfigFile.cs
@@ -0,0 +1,102 @@
+// Copyright Epic Games, Inc. All Rights Reserved.
+using EpicGames.Core;
+using System;
+using UnrealBuildTool;
+
+namespace EpicGames.Localization
+{
+
+	public enum LocalizationConfigFileFormat
+	{
+		LegacyMonolithic,
+		Modular,
+
+		// Add new values above this comment and upda te the value of Latest 
+		Latest = Modular
+	}
+
+	public class LocalizationConfigFile
+	{
+		private readonly ConfigFile _configFile;
+		public string Name { get; set; } = "";
+		public LocalizationConfigFile()
+		{
+			_configFile = new ConfigFile();
+		}
+
+		public LocalizationConfigFile(string fileName)
+		{
+			_configFile = new ConfigFile();
+			Name = fileName;
+		}
+
+		public LocalizationConfigFileSection FindOrAddSection(string sectionName)
+		{
+			ConfigFileSection configSection = _configFile.FindOrAddSection(sectionName);
+			return new LocalizationConfigFileSection(configSection);
+		}
+
+		public void Load(string filePath)
+		{
+			// @TODOLocalization: Load the config file and parse all the data into the various LocalizationConfigFileParam classes 
+		}
+
+		public void Write(FileReference destinationFilePath)
+		{
+			// Go through all the config values and replace \ directory separators with / for uniformity across platforms
+			foreach (string sectionName in _configFile.SectionNames)
+			{
+				ConfigFileSection section;
+				if (_configFile.TryGetSection(sectionName, out section))
+				{
+					foreach (ConfigLine line in section.Lines)
+					{
+						line.Value = line.Value.Replace('\\', '/');
+					}
+				}
+			}
+			_configFile.Write(destinationFilePath);
+		}
+	}
+
+	public class LocalizationConfigFileSection
+	{
+		private readonly ConfigFileSection _configFileSection;
+
+		public LocalizationConfigFileSection(ConfigFileSection configFileSection)
+		{
+			_configFileSection = configFileSection;
+		}
+
+		public void AddValue(string key, bool value)
+		{
+			if (!String.IsNullOrEmpty(key))
+			{
+				_configFileSection.Lines.Add(new ConfigLine(ConfigLineAction.Set, key, value ? "true" : "false"));
+			}
+		}
+
+		public void AddValue(string key, string value)
+		{
+			if (!String.IsNullOrEmpty(value) && !String.IsNullOrEmpty(key))
+			{
+				_configFileSection.Lines.Add(new ConfigLine(ConfigLineAction.Set, key, value));
+			}
+		}
+
+		public void AddValues(string key, string[] values)
+		{
+			if (String.IsNullOrEmpty(key) || values.Length == 0)
+			{
+				return;
+			}
+			foreach (string value in values)
+			{
+				if (!String.IsNullOrEmpty(value))
+				{
+					_configFileSection.Lines.Add(new ConfigLine(ConfigLineAction.Set, key, value));
+				}
+			}
+		}
+	}
+}
diff --git a/Engine/Source/Programs/AutomationTool/Localization/LocalizationConfigFileBuilder.cs b/Engine/Source/Programs/AutomationTool/Localization/LocalizationConfigFileBuilder.cs
new file mode 100644
--- /dev/null
+++ b/Engine/Source/Programs/AutomationTool/Localization/LocalizationConfigFileBuilder.cs
@@ -0,0 +1,129 @@
+// Copyright Epic Games, Inc. All Rights Reserved.
+
+namespace EpicGames.Localization
+{
+	public class LocalizationConfigFileBuilder
+	{
+		private LocalizationConfigFile _configFile;
+		private int _gatherTextStepCount;
+
+		public LocalizationConfigFileBuilder(string localizationConfigFileName)
+		{
+			_configFile = new LocalizationConfigFile(localizationConfigFileName);
+			_gatherTextStepCount = 0;
+		}
+
+		public void AddCommonSettingsSection(CommonSettingsParams commonSettings)
+		{
+			LocalizationConfigFileSection configSection = _configFile.FindOrAddSection("CommonSettings");
+			configSection.AddValues("ManifestDependencies", commonSettings.manifestDependencies.ToArray());
+			configSection.AddValue("SourcePath", commonSettings.SourcePath);
+			configSection.AddValue("DestinationPath", commonSettings.DestinationPath);
+			configSection.AddValue("ManifestName", commonSettings.ManifestName);
+			configSection.AddValue("ArchiveName", commonSettings.ArchiveName);
+			configSection.AddValue("PortableObjectName", commonSettings.PortableObjectName);
+			configSection.AddValue("ResourceName", commonSettings.ResourceName);
+			configSection.AddValue("bSkipSourceCheck", commonSettings.bSkipSourceCheck);
+			configSection.AddValue("bValidateFormatPatterns", commonSettings.bValidateFormatPatterns);
+			configSection.AddValue("bValidateSafeWhitespace", commonSettings.bValidateSafeWhitespace);
+			configSection.AddValue("NativeCulture", commonSettings.NativeCulture);
+			configSection.AddValues("CulturesToGenerate", commonSettings.CulturesToGenerate.ToArray());
+		}
+
+		public void AddGatherTextFromSourceSection(GatherTextFromSourceParams gatherTextParams)
+		{
+			LocalizationConfigFileSection configSection = _configFile.FindOrAddSection($"GatherTextStep{_gatherTextStepCount}");
+			++_gatherTextStepCount;
+
+			configSection.AddValue("CommandletClass", gatherTextParams.CommandletClass);
+			configSection.AddValues("SearchDirectoryPaths", gatherTextParams.SearchDirectoryPaths.ToArray());
+			configSection.AddValues("ExcludePathFilters", gatherTextParams.ExcludePathFilters.ToArray());
+			configSection.AddValues("FileNameFilters", gatherTextParams.FileNameFilters.ToArray());
+			configSection.AddValue("ShouldGatherFromEditorOnlyData", gatherTextParams.bShouldGatherFromEditorOnlyData);
+		}
+
+		public void AddGatherTextFromAssetsSection(GatherTextFromAssetsParams gatherTextParams)
+		{
+			LocalizationConfigFileSection configSection = _configFile.FindOrAddSection($"GatherTextStep{_gatherTextStepCount}");
+			++_gatherTextStepCount;
+
+			configSection.AddValue("CommandletClass", gatherTextParams.CommandletClass);
+			configSection.AddValues("IncludePathFilters", gatherTextParams.IncludePathFilters.ToArray());
+			configSection.AddValues("ExcludePathFilters", gatherTextParams.ExcludePathFilters.ToArray());
+			configSection.AddValues("PackageFileNameFilters", gatherTextParams.PackageFileNameFilters.ToArray());
+			configSection.AddValues("CollectionFilters", gatherTextParams.CollectionFilters.ToArray());
+
+			configSection.AddValue("ShouldExcludeDerivedClasses", gatherTextParams.bShouldExcludeDerivedClasses);
+			configSection.AddValue("ShouldGatherFromEditorOnlyData", gatherTextParams.bShouldGatherFromEditorOnlyData);
+			configSection.AddValue("SkipGatherCache", gatherTextParams.bSkipGatherCache);
+		}
+
+		public void AddGenerateGatherManifestSection(GenerateGatherManifestParams generateGatherManifestParams)
+		{
+			LocalizationConfigFileSection configSection = _configFile.FindOrAddSection($"GatherTextStep{_gatherTextStepCount}");
+			++_gatherTextStepCount;
+
+			configSection.AddValue("CommandletClass", generateGatherManifestParams.CommandletClass);
+		}
+
+		public void AddGenerateGatherArchiveSection(GenerateGatherArchiveParams generateGatherArchiveParams)
+		{
+			LocalizationConfigFileSection configSection = _configFile.FindOrAddSection($"GatherTextStep{_gatherTextStepCount}");
+			++_gatherTextStepCount;
+
+			configSection.AddValue("CommandletClass", generateGatherArchiveParams.CommandletClass);
+		}
+
+		public void AddExportPortableObjectFilesSection(ExportPortableObjectFilesParams exportParams)
+		{
+			LocalizationConfigFileSection configSection = _configFile.FindOrAddSection($"GatherTextStep{_gatherTextStepCount}");
+			++_gatherTextStepCount;
+
+			configSection.AddValue("CommandletClass", exportParams.CommandletClass);
+			configSection.AddValue("bExportLoc", exportParams.bExportLoc);
+			configSection.AddValue("LocalizedTextCollapseMode", exportParams.LocalizedTextCollapseMode);
+			configSection.AddValue("POFormat", exportParams.POFormat);
+			configSection.AddValue("ShouldPersistCommentsOnExport", exportParams.bShouldPersistCommentsOnExport);
+			configSection.AddValue("ShouldAddSourceLocationsAsComments", exportParams.bShouldAddSourceLocationsAsComments);
+		}
+
+		public void AddImportPortableObjectFilesSection(ImportPortableObjectFilesParams importParams)
+		{
+			LocalizationConfigFileSection configSection = _configFile.FindOrAddSection($"GatherTextStep{_gatherTextStepCount}");
+			++_gatherTextStepCount;
+
+			configSection.AddValue("CommandletClass", importParams.CommandletClass);
+			configSection.AddValue("bImportLoc", importParams.bImportLoc);
+			configSection.AddValue("LocalizedTextCollapseMode", importParams.LocalizedTextCollapseMode);
+			configSection.AddValue("POFormat", importParams.POFormat);
+		}
+
+		public void AddGenerateTextLocalizationResourceSection(GenerateTextLocalizationResourceParams localizationResourceParams)
+		{
+			LocalizationConfigFileSection configSection = _configFile.FindOrAddSection($"GatherTextStep{_gatherTextStepCount}");
+			++_gatherTextStepCount;
+
+			configSection.AddValue("CommandletClass", localizationResourceParams.CommandletClass);
+		}
+
+		public void AddGenerateTextLocalizationReportSection(GenerateTextLocalizationReportParams textReportParams)
+		{
+			LocalizationConfigFileSection configSection = _configFile.FindOrAddSection($"GatherTextStep{_gatherTextStepCount}");
+			++_gatherTextStepCount;
+
+			configSection.AddValue("CommandletClass", textReportParams.CommandletClass);
+			configSection.AddValue("bWordCountReport", textReportParams.bWordCountReport);
+			configSection.AddValue("WordCountReportName", textReportParams.WordCountReportName);
+			if (textReportParams.bConflictReport)
+			{
+				configSection.AddValue("bConflictReport", textReportParams.bConflictReport);
+				configSection.AddValue("ConflictReportName", textReportParams.ConflictReportName);
+			}
+		}
+
+		public LocalizationConfigFile Get()
+		{
+			return _configFile;
+		}
+	}
+}
diff --git a/Engine/Source/Programs/AutomationTool/Localization/LocalizationConfigFileGenerator.cs b/Engine/Source/Programs/AutomationTool/Localization/LocalizationConfigFileGenerator.cs
new file mode 100644
--- /dev/null
+++ b/Engine/Source/Programs/AutomationTool/Localization/LocalizationConfigFileGenerator.cs
@@ -0,0 +1,115 @@
+// Copyright Epic Games, Inc. All Rights Reserved.
+using System;
+using System.Collections.Generic;
+
+namespace EpicGames.Localization
+{
+	public class LocalizationConfigFileGeneratorParams
+	{
+		public string LocalizationTargetName { get; set; } = "";
+		public string LocalizationTargetRootDirectory { get; set; } = "";
+	}
+
+	public abstract class LocalizationConfigFileGenerator
+	{
+		public static LocalizationConfigFileGenerator GetGeneratorForFileFormat(LocalizationConfigFileFormat fileFormat)
+		{
+			switch (fileFormat)
+			{
+				case LocalizationConfigFileFormat.Modular:
+					return new ModularLocalizationConfigFileGenerator();
+				default:
+					throw new NotImplementedException("File format is not yet supported.");
+			}
+		}
+
+		public abstract List<LocalizationConfigFile> GenerateDefaultSettingsConfigFiles(LocalizationConfigFileGeneratorParams generatorParams);
+		
+	}
+
+	public class ModularLocalizationConfigFileGenerator : LocalizationConfigFileGenerator
+	{
+		public override List<LocalizationConfigFile> GenerateDefaultSettingsConfigFiles(LocalizationConfigFileGeneratorParams generatorParams)
+		{
+			List<LocalizationConfigFile> configFiles = new List<LocalizationConfigFile>();
+			configFiles.Add(GenerateLocalizationGatherConfigFile(generatorParams));
+			configFiles.Add(GenerateLocalizationExportConfigFile(generatorParams));
+			configFiles.Add(GenerateLocalizationImportConfigFile(generatorParams));
+			configFiles.Add(GenerateLocalizationCompileConfigFile(generatorParams));
+			configFiles.Add(GenerateLocalizationGenerateReportsConfigFile(generatorParams));
+			return configFiles;
+		}
+
+		private LocalizationConfigFile GenerateLocalizationGatherConfigFile(LocalizationConfigFileGeneratorParams generatorParam)
+		{
+			CommonSettingsParams commonSettings = CommonSettingsParams.CreateDefault(generatorParam.LocalizationTargetName, generatorParam.LocalizationTargetRootDirectory);
+			LocalizationConfigFileBuilder builder = new LocalizationConfigFileBuilder(generatorParam.LocalizationTargetName + ModularLocalizationConfigFileSuffixes.GatherSuffix + LocalizationFileExtensions.ModularLocalizationConfigFileExtension);
+			builder.AddCommonSettingsSection(commonSettings);
+
+			GatherTextFromSourceParams gatherTextFromSourceParams = GatherTextFromSourceParams.CreateDefault(generatorParam.LocalizationTargetName, generatorParam.LocalizationTargetRootDirectory); 
+			builder.AddGatherTextFromSourceSection(gatherTextFromSourceParams);
+
+			GatherTextFromAssetsParams gatherTextFromAssetsParams = GatherTextFromAssetsParams.CreateDefault(generatorParam.LocalizationTargetName, generatorParam.LocalizationTargetRootDirectory);
+			builder.AddGatherTextFromAssetsSection(gatherTextFromAssetsParams);
+
+			GenerateGatherManifestParams generateGatherManifestParams = GenerateGatherManifestParams.CreateDefault(generatorParam.LocalizationTargetName, generatorParam.LocalizationTargetRootDirectory);
+			builder.AddGenerateGatherManifestSection(generateGatherManifestParams);
+
+			GenerateGatherArchiveParams generateGatherArchiveParams = GenerateGatherArchiveParams.CreateDefault(generatorParam.LocalizationTargetName, generatorParam.LocalizationTargetRootDirectory);
+			builder.AddGenerateGatherArchiveSection(	generateGatherArchiveParams);
+
+			GenerateTextLocalizationReportParams generateTextLocalizationReportParams = GenerateTextLocalizationReportParams.CreateDefault(generatorParam.LocalizationTargetName, generatorParam.LocalizationTargetRootDirectory);
+			builder.AddGenerateTextLocalizationReportSection(generateTextLocalizationReportParams);
+
+			return builder.Get();
+		}
+
+		private LocalizationConfigFile GenerateLocalizationExportConfigFile(LocalizationConfigFileGeneratorParams generatorParams)
+		{
+			CommonSettingsParams commonSettings = CommonSettingsParams.CreateDefault(generatorParams.LocalizationTargetName, generatorParams.LocalizationTargetRootDirectory);
+			LocalizationConfigFileBuilder builder = new LocalizationConfigFileBuilder(generatorParams.LocalizationTargetName + ModularLocalizationConfigFileSuffixes.ExportSuffix + LocalizationFileExtensions.ModularLocalizationConfigFileExtension);
+			builder.AddCommonSettingsSection(commonSettings);
+
+			ExportPortableObjectFilesParams exportParams = ExportPortableObjectFilesParams.CreateDefault(generatorParams.LocalizationTargetName, generatorParams.LocalizationTargetRootDirectory);
+			builder.AddExportPortableObjectFilesSection(exportParams);
+			return builder.Get();
+		}
+
+		private LocalizationConfigFile GenerateLocalizationImportConfigFile(LocalizationConfigFileGeneratorParams generatorParams)
+		{
+			CommonSettingsParams commonSettings = CommonSettingsParams.CreateDefault(generatorParams.LocalizationTargetName, generatorParams.LocalizationTargetRootDirectory);
+			LocalizationConfigFileBuilder builder = new LocalizationConfigFileBuilder(generatorParams.LocalizationTargetName + ModularLocalizationConfigFileSuffixes.ImportSuffix + LocalizationFileExtensions.ModularLocalizationConfigFileExtension);
+			builder.AddCommonSettingsSection(commonSettings);
+
+			ImportPortableObjectFilesParams importParams = ImportPortableObjectFilesParams.CreateDefault(generatorParams.LocalizationTargetName, generatorParams.LocalizationTargetRootDirectory);
+			builder.AddImportPortableObjectFilesSection(importParams);
+			return builder.Get();
+		}
+
+		private LocalizationConfigFile GenerateLocalizationCompileConfigFile(LocalizationConfigFileGeneratorParams generatorParams)
+		{
+			CommonSettingsParams commonSettings = CommonSettingsParams.CreateDefault(generatorParams.LocalizationTargetName, generatorParams.LocalizationTargetRootDirectory);
+			LocalizationConfigFileBuilder builder = new LocalizationConfigFileBuilder(generatorParams.LocalizationTargetName + ModularLocalizationConfigFileSuffixes.CompileSuffix + LocalizationFileExtensions.ModularLocalizationConfigFileExtension);
+			builder.AddCommonSettingsSection(commonSettings);
+
+			GenerateTextLocalizationResourceParams generateTextLocalizationResourceParams = GenerateTextLocalizationResourceParams.CreateDefault(generatorParams.LocalizationTargetName, generatorParams.LocalizationTargetRootDirectory);
+			builder.AddGenerateTextLocalizationResourceSection(generateTextLocalizationResourceParams);
+			return builder.Get();
+		}
+
+		private LocalizationConfigFile GenerateLocalizationGenerateReportsConfigFile(LocalizationConfigFileGeneratorParams generatorParams)
+		{
+			CommonSettingsParams commonSettings = CommonSettingsParams.CreateDefault(generatorParams.LocalizationTargetName, generatorParams.LocalizationTargetRootDirectory);
+			LocalizationConfigFileBuilder builder = new LocalizationConfigFileBuilder(generatorParams.LocalizationTargetName + ModularLocalizationConfigFileSuffixes.GenerateReportsSuffix + LocalizationFileExtensions.ModularLocalizationConfigFileExtension);
+			builder.AddCommonSettingsSection(commonSettings);
+
+			GenerateTextLocalizationReportParams reportParams = GenerateTextLocalizationReportParams.CreateDefault(generatorParams.LocalizationTargetName, generatorParams.LocalizationTargetRootDirectory);
+			// We don't need the conflict report 
+			reportParams.bConflictReport = false;
+			reportParams.ConflictReportName = "";
+			builder.AddGenerateTextLocalizationReportSection(reportParams);
+			return builder.Get();
+		}
+	}
+}
+
diff --git a/Engine/Source/Programs/AutomationTool/Localization/LocalizationConfigFileParams.cs b/Engine/Source/Programs/AutomationTool/Localization/LocalizationConfigFileParams.cs
new file mode 100644
--- /dev/null
+++ b/Engine/Source/Programs/AutomationTool/Localization/LocalizationConfigFileParams.cs
@@ -0,0 +1,206 @@
+// Copyright Epic Games, Inc. All Rights Reserved.
+using System.Collections.Generic;
+using System.IO;
+
+namespace EpicGames.Localization
+{
+
+	public class CommonSettingsParams
+	{
+		public List<string> manifestDependencies { get; set; } = new List<string>();
+		public string SourcePath { get; set; } = "";
+		public string DestinationPath { get; set; } = "";
+		public string ManifestName { get; set; } = "";
+		public string ArchiveName { get; set; } = "";
+		public string ResourceName { get; set; } = "";
+		public string PortableObjectName { get; set; } = "";
+		public bool bSkipSourceCheck { get; set; } = false;
+		public bool bValidateFormatPatterns { get; set; } = true;
+		public bool bValidateSafeWhitespace { get; set; } = false;
+		public string NativeCulture { get; set; } = "";
+		public List<string> CulturesToGenerate { get; set; } = new List<string>();
+
+		public static CommonSettingsParams CreateDefault(string localizationTargetName, string localizationTargetRootDirectory)
+		{
+			CommonSettingsParams commonSettings = new CommonSettingsParams();
+			commonSettings.manifestDependencies.Add(Path.Combine("..", "Engine", "Content", "Localization", "Engine", $"Engine{LocalizationFileExtensions.ManifestFileExtension}"));
+			commonSettings.manifestDependencies.Add(Path.Combine("..", "Engine", "Content", "Localization", "Editor", $"Editor{LocalizationFileExtensions.ManifestFileExtension}"));
+
+			commonSettings.SourcePath = Path.Combine(localizationTargetRootDirectory, "Content", "Localization", localizationTargetName);
+			commonSettings.DestinationPath = Path.Combine(localizationTargetRootDirectory, "Content", "Localization", localizationTargetName);
+			commonSettings.ManifestName = localizationTargetName + LocalizationFileExtensions.ManifestFileExtension;
+			commonSettings.ArchiveName = localizationTargetName + LocalizationFileExtensions.ArchiveFileExtension;
+			commonSettings.PortableObjectName = localizationTargetName + LocalizationFileExtensions.PortableObjectFileExtension;
+			commonSettings.ResourceName = localizationTargetName + LocalizationFileExtensions.LocalizationResourceFileExtension;
+
+			commonSettings.NativeCulture = "en";
+			commonSettings.CulturesToGenerate = new List<string>() { "en", "fr", "de", "pl", "es-419", "es", "pt-BR", "it", "ru", "ko", "tr", "ar", "ja" };
+			commonSettings.CulturesToGenerate.Sort();
+			commonSettings.bSkipSourceCheck = false;
+			commonSettings.bValidateFormatPatterns = true;
+			commonSettings.bValidateSafeWhitespace = true;
+			return commonSettings;
+		}
+	}
+
+	public class GatherTextFromSourceParams
+	{
+		public string CommandletClass { get; } = "GatherTextFromSource";
+		public List<string> SearchDirectoryPaths { get; set; } = new List<string>();
+		public List<string> ExcludePathFilters { get; set; } = new List<string>();
+		public List<string> FileNameFilters { get; set; } = new List<string>();
+		public bool bShouldGatherFromEditorOnlyData { get; set; } = false;
+
+		public static GatherTextFromSourceParams CreateDefault(string localizationTargetName, string localizationTargetRootDirectory)
+		{
+			GatherTextFromSourceParams gatherTextFromSource = new GatherTextFromSourceParams();
+			gatherTextFromSource.SearchDirectoryPaths.Add(Path.Combine(localizationTargetRootDirectory, "Source"));
+			gatherTextFromSource.SearchDirectoryPaths.Add(Path.Combine(localizationTargetRootDirectory, "Config"));
+
+			gatherTextFromSource.ExcludePathFilters.Add(Path.Combine(localizationTargetRootDirectory, "Config", "Localization", "*"));
+			gatherTextFromSource.FileNameFilters = new List<string>() { "*.h", "*.cpp", "*.inl", "*.ini" };
+			gatherTextFromSource.bShouldGatherFromEditorOnlyData = false;
+			return gatherTextFromSource;
+		}
+	}
+
+	public class GatherTextFromAssetsParams
+	{
+		public string CommandletClass { get; } = "GatherTextFromAssets";
+		public List<string> IncludePathFilters { get; set; } = new List<string>();
+		public List<string> ExcludePathFilters { get; set; } = new List<string>();
+		public List<string> PackageFileNameFilters { get; set; } = new List<string>();
+		public List<string> CollectionFilters { get; set; } = new List<string>();
+		public bool bShouldExcludeDerivedClasses { get; set; } = false;
+		public bool bShouldGatherFromEditorOnlyData { get; set; } = false;
+		public bool bSkipGatherCache { get; set; } = false;
+
+		public static GatherTextFromAssetsParams CreateDefault(string localizationTargetName, string localizationTargetRootDirectory)
+		{
+			GatherTextFromAssetsParams gatherTextFromAssets = new GatherTextFromAssetsParams();
+			gatherTextFromAssets.IncludePathFilters.Add(Path.Combine(localizationTargetRootDirectory, "Content", "*"));
+
+			gatherTextFromAssets.ExcludePathFilters.Add(Path.Combine(localizationTargetRootDirectory, "Content", "Localization", "*"));
+			gatherTextFromAssets.ExcludePathFilters.Add(Path.Combine(localizationTargetRootDirectory, "Content", "L10N", "*"));
+			gatherTextFromAssets.ExcludePathFilters.Add(Path.Combine(localizationTargetRootDirectory, "Content", "Developers", "*"));
+			gatherTextFromAssets.ExcludePathFilters.Add(Path.Combine(localizationTargetRootDirectory, "Content", "*Test", "*"));
+			
+			gatherTextFromAssets.PackageFileNameFilters= new List<string>() { "*.umap", "*.uasset" };
+			//@TODOLocalization: Some licensees might not have a need for an Audit_InCook file...should remove this from defaults after config file generation with templates is implemented
+			gatherTextFromAssets.CollectionFilters = new List<string>() { "Audit_InCook" };
+
+			gatherTextFromAssets.bShouldExcludeDerivedClasses = false;
+			gatherTextFromAssets.bShouldGatherFromEditorOnlyData = false;
+			gatherTextFromAssets.bSkipGatherCache = false;
+			return gatherTextFromAssets;
+		}
+	}
+
+	public class GatherTextFromMetaDataParams
+	{
+		public string CommandletClass { get; } = "GatherTextFromMetaData";
+		public List<string> ModulesToPreload { get; set; } = new List<string>();
+		public List<string> IncludePathFilters { get; set; } = new List<string>();
+		public List<string> ExcludePathFilters { get; set; } = new List<string>();
+		public List<string> FieldTypesToInclude { get; set; } = new List<string>();
+		public List<string> FieldOwnerTypesToInclude { get; set; } = new List<string>();
+		public List<string> InputKeys { get; set; } = new List<string>();
+		public List<string> OutputKeys { get; set; } = new List<string>();
+		public List<string> InputNamespaceses { get; set; } = new List<string>();
+		public List<string> OutputNamespaces { get; set; } = new List<string>();
+		public bool bShouldGatherFromEditorOnlyData { get; set; } = false;
+
+		// @TODOLocalization: Create defaults for gathering from metada 
+	}
+
+	public class GenerateGatherManifestParams
+	{
+		public string CommandletClass { get; } = "GenerateGatherManifest";
+		public static GenerateGatherManifestParams CreateDefault(string localizationTargetName, string localizationTargetRootDirectory)
+		{
+			GenerateGatherManifestParams generateGatherManifestParams = new GenerateGatherManifestParams();
+
+			return generateGatherManifestParams;
+		}
+	}
+
+	public class GenerateGatherArchiveParams
+	{
+		public string CommandletClass { get; } = "GenerateGatherArchive";
+		public bool bPurgeOldEmptyEntries { get; set; } = false;
+
+		public static GenerateGatherArchiveParams CreateDefault(string localizationTargetName, string localizationTargetRootDirectory)
+		{
+			GenerateGatherArchiveParams generateGatherArchiveParams = new GenerateGatherArchiveParams();
+			generateGatherArchiveParams.bPurgeOldEmptyEntries = false;
+
+			return generateGatherArchiveParams;
+		}
+	}
+
+	public class ImportPortableObjectFilesParams
+	{
+		public string CommandletClass { get; } = "InternationalizationExport";
+		public bool bImportLoc { get; }  = true;
+		public string LocalizedTextCollapseMode { get; set; } = "";
+		public string POFormat { get; set; } = "";
+
+		public static ImportPortableObjectFilesParams CreateDefault(string localizationTargetName, string localizationTargetRootDirectory)
+		{
+			ImportPortableObjectFilesParams importPOFilesParams = new ImportPortableObjectFilesParams();
+			// @TODOLocalization: These should be enums not hard coded strings. 
+			importPOFilesParams.LocalizedTextCollapseMode = "ELocalizedTextCollapseMode::IdenticalTextIdAndSource";
+			importPOFilesParams.POFormat = "EPortableObjectFormat::Unreal";
+			return importPOFilesParams;
+		}
+	}
+
+	public class GenerateTextLocalizationResourceParams
+	{
+		public string CommandletClass { get; } = "GenerateTextLocalizationResource";
+		public static GenerateTextLocalizationResourceParams CreateDefault(string localizationTargetName, string localizationTargetRootDirectory)
+		{
+			return new GenerateTextLocalizationResourceParams();
+		}
+	}
+
+	public class ExportPortableObjectFilesParams
+	{
+		public string CommandletClass { get; } = "InternationalizationExport";
+		public bool bExportLoc { get; } = true;
+		public string LocalizedTextCollapseMode { get; set; } = "";
+		public string POFormat { get; set; } = "";
+		public bool bShouldPersistCommentsOnExport { get; set; } = true;
+		public bool bShouldAddSourceLocationsAsComments { get; set; } = true;
+
+		public static ExportPortableObjectFilesParams CreateDefault(string localizationTargetName, string localizationTargetRootDirectory)
+		{
+			ExportPortableObjectFilesParams exportPOFilesParams = new ExportPortableObjectFilesParams();
+			exportPOFilesParams.LocalizedTextCollapseMode = "ELocalizedTextCollapseMode::IdenticalTextIdAndSource";
+			exportPOFilesParams.POFormat = "EPortableObjectFormat::Unreal";
+			exportPOFilesParams.bShouldPersistCommentsOnExport = true;
+			exportPOFilesParams.bShouldAddSourceLocationsAsComments = true;
+			return exportPOFilesParams;
+		}
+	}
+
+	public class GenerateTextLocalizationReportParams
+	{
+		public string CommandletClass { get; } = "GenerateTextLocalizationReport";
+		public string DestinationPath { get; set; } = "";
+		public bool bWordCountReport { get; set; } = false;
+		public string WordCountReportName { get; set; } = "";
+		public bool bConflictReport { get; set; } = false;
+		public string ConflictReportName { get; set; } = "";
+
+		public static GenerateTextLocalizationReportParams CreateDefault(string localizationTargetName, string localizationTargetRootDirectory)
+		{
+			GenerateTextLocalizationReportParams generateTextLocalizationReportParams = new GenerateTextLocalizationReportParams();
+			generateTextLocalizationReportParams.bWordCountReport = true;
+			generateTextLocalizationReportParams.WordCountReportName = $"{localizationTargetName}.csv";
+			generateTextLocalizationReportParams.bConflictReport = true;
+			generateTextLocalizationReportParams.ConflictReportName = $"{localizationTargetName}_Conflicts.csv";
+			return generateTextLocalizationReportParams;
+		}
+	}
+}
diff --git a/Engine/Source/Programs/AutomationTool/Localization/LocalizationConstants.cs b/Engine/Source/Programs/AutomationTool/Localization/LocalizationConstants.cs
new file mode 100644
--- /dev/null
+++ b/Engine/Source/Programs/AutomationTool/Localization/LocalizationConstants.cs
@@ -0,0 +1,24 @@
+// Copyright Epic Games, Inc. All Rights Reserved.
+
+namespace EpicGames.Localization
+{
+	public static class LocalizationFileExtensions
+	{
+		public static string LegacyMonolithicLocalizationConfigFileExtension{ get; } = ".ini";
+		public static string ModularLocalizationConfigFileExtension { get; } = ".ini";
+		public static string ArchiveFileExtension { get; } = ".archive";
+		public static string ManifestFileExtension { get; } = ".manifest";
+		public static string PortableObjectFileExtension { get; } = ".po";
+		public static string LocalizationResourceFileExtension { get; } = ".locres";
+	}
+
+	public static class ModularLocalizationConfigFileSuffixes
+	{
+		public static string GatherSuffix { get; } = "_Gather";
+		public static string ExportSuffix { get; } = "_Export";
+		public static string ImportSuffix { get; } = "_Import";
+		public static string CompileSuffix { get; } = "_Compile";
+		public static string GenerateReportsSuffix { get; } = "_GenerateReports";
+	}
+
+}
\ No newline at end of file
diff --git a/Engine/Source/Programs/AutomationTool/Localization/LocalizationTargetEditor/Commands/LocalizationTargetCommand.cs b/Engine/Source/Programs/AutomationTool/Localization/LocalizationTargetEditor/Commands/LocalizationTargetCommand.cs
new file mode 100644
--- /dev/null
+++ b/Engine/Source/Programs/AutomationTool/Localization/LocalizationTargetEditor/Commands/LocalizationTargetCommand.cs
@@ -0,0 +1,79 @@
+// Copyright Epic Games, Inc. All Rights Reserved.
+using AutomationTool;
+using Microsoft.Extensions.Logging;
+using System;
+using System.Collections.Generic;
+using System.IO;
+using System.Linq;
+using System.Reflection;
+using UnrealBuildTool;
+
+using static AutomationTool.CommandUtils;
+
+namespace EpicGames.Localization
+{
+	public abstract class LocalizationTargetCommand
+	{
+		protected static BuildCommand _commandLineHelper;
+		private static Dictionary<string, LocalizationTargetCommand> _nameToLocalizationTargetCommandMap = new Dictionary<string, LocalizationTargetCommand>();
+
+		public string Name { get; protected set; } = "";
+		public string DisplayName { get; protected set; } = "";
+		public bool bIsExecutingInPreview { get; private set; } = false;
+
+		static LocalizationTargetCommand()
+		{
+			// look for all subclasses, and cache by their ProviderToken
+			foreach (Assembly assembly in ScriptManager.AllScriptAssemblies)
+			{
+				Type[] allTypesInAssembly = assembly.GetTypes();
+				foreach (Type typeInAssembly in allTypesInAssembly)
+				{
+					// we also guard against abstract classes as we can have abstract classes as children of LocalizationTargetEditorCommand e.g PluginLocalizationTargetCommand
+					if (typeof(LocalizationTargetCommand).IsAssignableFrom(typeInAssembly) && typeInAssembly != typeof(LocalizationTargetCommand) && !typeInAssembly.IsAbstract)
+					{
+						LocalizationTargetCommand provider = (LocalizationTargetCommand)Activator.CreateInstance(typeInAssembly);
+						_nameToLocalizationTargetCommandMap[provider.Name] = provider;
+					}
+				}
+			}
+		}
+
+		public static void Initialize(BuildCommand commandLineHelper)
+		{
+			_commandLineHelper = commandLineHelper;
+		}
+
+		public static LocalizationTargetCommand GetLocalizationTargetCommandFromName(string commandName)
+		{
+			LocalizationTargetCommand command = null;
+			_nameToLocalizationTargetCommandMap.TryGetValue(commandName, out command);
+			return command;
+		}
+
+
+		protected virtual bool ParseCommandLine()
+		{
+			bIsExecutingInPreview = _commandLineHelper.ParseParam("Preview");
+			return true;
+		}
+
+		public abstract bool Execute();
+
+		public void PrintHelpText()
+		{
+			Logger.LogInformation("Help:");
+			Logger.LogInformation($"{GetHelpText()}");
+		}
+
+		public virtual string GetHelpText()
+		{
+			return "";
+		}
+
+		public static List<LocalizationTargetCommand> GetAllCommands()
+		{
+			return _nameToLocalizationTargetCommandMap.Values.ToList();
+		}
+	}
+}
diff --git a/Engine/Source/Programs/AutomationTool/Localization/LocalizationTargetEditor/Commands/NewPluginLocalizationTargetCommand.cs b/Engine/Source/Programs/AutomationTool/Localization/LocalizationTargetEditor/Commands/NewPluginLocalizationTargetCommand.cs
new file mode 100644
--- /dev/null
+++ b/Engine/Source/Programs/AutomationTool/Localization/LocalizationTargetEditor/Commands/NewPluginLocalizationTargetCommand.cs
@@ -0,0 +1,241 @@
+// Copyright Epic Games, Inc. All Rights Reserved.
+using System;
+using System.Collections.Generic;
+using System.Linq;
+using System.Text;
+using System.IO;
+using AutomationTool;
+using UnrealBuildTool;
+using EpicGames.Core;
+using Microsoft.Extensions.Logging;
+
+using static AutomationTool.CommandUtils;
+using UnrealBuildBase;
+
+namespace EpicGames.Localization
+{
+	public class NewPluginLocalizationTargetCommand : PluginLocalizationTargetCommand
+	{
+		private LocalizationTargetDescriptorLoadingPolicy _loadingPolicy;
+		private string _localizationTargetNameSuffix = "";
+		private LocalizationConfigFileFormat _fileFormat = LocalizationConfigFileFormat.Latest;
+		private List<string> _filesToAdd = new List<string>();
+		private List<string> _filesToEdit = new List<string>();
+
+		public NewPluginLocalizationTargetCommand()
+		{
+			Name = "NewPluginLocalizationTarget";
+			DisplayName = "New Plugin Localization Target";
+		}
+
+		public override string GetHelpText()
+		{
+			StringBuilder helpTextBuilder = new StringBuilder();
+			helpTextBuilder.AppendLine("Creates new localization targets for plugins. This command generates the necessary localization config files for a plugin and updates the plugin descriptor with new localization target data.");
+			helpTextBuilder.AppendLine("Passing the -p4 flag to UAT will create a change list with all of the added and edited files.");
+			helpTextBuilder.AppendLine("If the localization config files already exist for the plugin or if the plugin descriptor already contains the localization target, this command will not edit any of the existing files.");
+
+			helpTextBuilder.AppendLine("Arguments:");
+			helpTextBuilder.AppendLine("UERootDirectory - Optional root directory for the Unreal Engine. This is usually one level above your project directory. Defaults to CmdEnv.LocalRoot");
+			helpTextBuilder.AppendLine("UEProjectDirectory - Required relative path from UERootDirectory to the project. This should be  your project directory that contains the Plugins directory.");
+			helpTextBuilder.AppendLine("UEProjectName - An optional name of the project the plugin is under. If blank, this implies the plugin is under Engine.");
+			helpTextBuilder.AppendLine("LocalizationTargetLoadingPolicy - A required string that determines when the plugin should load the localization data for the localization target. Valid values are Game, Always, Editor, PropertyNames, ToolTips. See LocalizationTargetDescriptorLoadingPolicy");
+			helpTextBuilder.AppendLine("_localizationTargetNameSuffix - An optional suffix to give to the plugin localization target. By default, the name of the plugin localization target would be the name of the plugin. This allows multiple localization targets to be created for plugins.");
+			// Include plugins 
+			helpTextBuilder.AppendLine("IncludePlugins - An optional comma separated list of plugins to create localization targets for. E.g PluginA,PluginB,PluginC");
+			helpTextBuilder.AppendLine("IncludePluginsDirectory - An optional relative directory to UEProjectDirectory. All plugins under this directory will have localization targets created if they are not excluded. E.g Plugins/PluginFolderA.");
+			helpTextBuilder.AppendLine("ExcludePlugins - A comma separated list of plugins to exclude from having localization targets created. E.g PluginA,BpluginB,PluginC");
+			helpTextBuilder.AppendLine("ExcludePluginsDirectory - An optional relative directory from UEProjectDirectory. All plugins under this directory will be excluded from localization target creation. E.g Plugin/DirectoryToExclude");
+
+			helpTextBuilder.AppendLine("Preview - An optional flag that will execute this command in preview mode. No files will be created or edited. No folders will be created. No files will be added or checked out of perforce.");
+			return helpTextBuilder.ToString();
+		}
+
+		protected override bool ParseCommandLine()
+		{
+			if (!base.ParseCommandLine())
+			{
+				return false;
+			}
+			_localizationTargetNameSuffix = _commandLineHelper.ParseParamValue("_localizationTargetNameSuffix");
+			// @TODOLocalization: Support an override for the localization config file format 
+			_fileFormat = LocalizationConfigFileFormat.Latest;
+			_loadingPolicy = _commandLineHelper.ParseRequiredEnumParamEnum<LocalizationTargetDescriptorLoadingPolicy>("LocalizationTargetLoadingPolicy");
+
+			return true;
+		}
+
+		public override bool Execute()
+		{
+			if (!ParseCommandLine())
+			{
+				return false;
+			}
+
+			string projectPluginDirectory = Path.Combine(UERootDirectory, UEProjectDirectory, "Plugins");
+
+			List<string> pluginsToGenerateFor = IncludePlugins.Except(ExcludePlugins).ToList();
+			if (pluginsToGenerateFor.Count == 0)
+			{
+				Logger.LogInformation("Based on provided plugins to include and exclude, no plugins require a localization target to be generated.");
+				return true;
+			}
+			DirectoryReference ueProjectDirectoryReference = new DirectoryReference(Path.Combine(UERootDirectory, UEProjectDirectory));
+			// Load all available plugins so we can find them 
+			PluginType pluginType = String.IsNullOrEmpty(UEProjectName) ? PluginType.Engine : PluginType.Project;
+			_ = pluginType == PluginType.Engine ? Plugins.ReadEnginePlugins(ueProjectDirectoryReference) : Plugins.ReadProjectPlugins(ueProjectDirectoryReference);
+
+			// loop through all plugins and generate the files 
+			foreach (string pluginName in pluginsToGenerateFor)
+			{
+				PluginInfo pluginInfo = Plugins.GetPlugin(pluginName);
+				if (pluginInfo is null)
+				{
+					Logger.LogWarning($"Cannot find {pluginName}. Skipping generation of localization config files for this plugin.");
+					continue;
+				}
+				string pluginLocalizationTargetName = pluginInfo.Name;
+				// we add on the suffix if a user provides it 
+				if (!String.IsNullOrEmpty(_localizationTargetNameSuffix))
+				{
+					pluginLocalizationTargetName += "_" + _localizationTargetNameSuffix;
+				}
+
+				GeneratePluginLocalizationConfigFiles(pluginInfo, pluginLocalizationTargetName, pluginInfo.Directory.MakeRelativeTo(ueProjectDirectoryReference));
+				// Update the .uplugin file with updated localization target descriptors 
+				UpdatePluginLocalizationTargetDescriptor(pluginInfo, pluginLocalizationTargetName);
+			}
+
+			// Create changelist with all edited/added files 
+			if (!bIsExecutingInPreview)
+			{
+				CreatePerforceChangeList();
+			}
+
+			return true;
+		}
+
+		private void GeneratePluginLocalizationConfigFiles(PluginInfo pluginInfo, string localizationTargetName, string localizationTargetRootDirectory)
+		{
+			LocalizationConfigFileGenerator generator = LocalizationConfigFileGenerator.GetGeneratorForFileFormat(_fileFormat);
+			// Create the necessary directories if they dont' exist for the plugin 
+			string pluginConfigLocalizationDirectory = Path.Combine(pluginInfo.Directory.FullName, "Config", "Localization");
+			if (!Directory.Exists(pluginConfigLocalizationDirectory) && !bIsExecutingInPreview)
+			{
+				Logger.LogInformation($"Plugin '{pluginInfo.Name}' does not have a localization configuration folder. Creating '{pluginConfigLocalizationDirectory}'");
+				Directory.CreateDirectory(pluginConfigLocalizationDirectory);
+			}
+
+			// Generate the config files 
+			LocalizationConfigFileGeneratorParams generatorParams = new LocalizationConfigFileGeneratorParams();
+			generatorParams.LocalizationTargetName = localizationTargetName;
+			generatorParams.LocalizationTargetRootDirectory = localizationTargetRootDirectory;
+			// @TODOLocalization: Create an option to generate the files based on a template localization config file
+			List<LocalizationConfigFile> generatedConfigFiles = generator.GenerateDefaultSettingsConfigFiles(generatorParams);
+			foreach (LocalizationConfigFile generatedConfigFile in generatedConfigFiles)
+			{
+				string saveFilePath = Path.Combine(pluginConfigLocalizationDirectory, generatedConfigFile.Name);
+				// We skip writing the file if it already exists. We only create new files 
+				if (File.Exists(saveFilePath)) 
+				{
+					Logger.LogInformation($"Localization config file {saveFilePath} already exists. File contents will not be overwritten. If you would like to edit the contents of the file, use the UpdatePluginLocalizationTarget command instead.");
+					continue;
+				}
+				_filesToAdd.Add(saveFilePath);
+				if (!bIsExecutingInPreview)
+				{
+					Logger.LogInformation($"Writing '{saveFilePath}'");
+					generatedConfigFile.Write(new FileReference(saveFilePath));
+				}
+			}
+		}
+
+		private void UpdatePluginLocalizationTargetDescriptor(PluginInfo pluginInfo, string localizationTargetName)
+		{
+			PluginDescriptor descriptor = pluginInfo.Descriptor;
+			if (descriptor.LocalizationTargets is not null)
+			{
+				foreach (LocalizationTargetDescriptor locTargetDescriptor in descriptor.LocalizationTargets)
+				{
+					// the localization target with the same name already exists. Early out and don't update anything 
+					if (locTargetDescriptor.Name == localizationTargetName)
+					{
+						Logger.LogWarning($"Plugin {pluginInfo.Name} already contains localization target descriptor with name {locTargetDescriptor.Name}. No updates to the plugin descriptor will be made.");
+						return;
+					}
+				}
+			}
+			LocalizationTargetDescriptor locTargetDescriptorToAdd = new LocalizationTargetDescriptor(localizationTargetName, _loadingPolicy);
+			if (descriptor.LocalizationTargets is null)
+			{
+				descriptor.LocalizationTargets = new LocalizationTargetDescriptor[] { locTargetDescriptorToAdd };
+			}
+			else
+			{
+				int numLocalizationTargets = descriptor.LocalizationTargets.Length;
+				LocalizationTargetDescriptor[] newLocalizationTargets = new LocalizationTargetDescriptor[numLocalizationTargets];
+				Array.Copy(descriptor.LocalizationTargets, newLocalizationTargets, numLocalizationTargets);
+				newLocalizationTargets[numLocalizationTargets] = locTargetDescriptorToAdd;
+				descriptor.LocalizationTargets = newLocalizationTargets;
+			}
+
+			Logger.LogInformation($"Updating localization target information for {pluginInfo.Name} plugin. Adding localization target {localizationTargetName} with loading policy {_loadingPolicy.ToString()}");
+
+			// Add descriptor to processing list 
+			_filesToEdit.Add(pluginInfo.File.FullName);
+			if (!bIsExecutingInPreview)
+			{
+				Logger.LogInformation($"Updating {pluginInfo.Name} plugin descriptor file '{pluginInfo.File.FullName}'");
+				// make the file writable and save the data
+				pluginInfo.File.ToFileInfo().IsReadOnly = false;
+				descriptor.Save2(pluginInfo.File.FullName);
+			}
+		}
+
+		private void CreatePerforceChangeList()
+		{
+			if (!P4Enabled)
+			{
+				Logger.LogInformation("Perfordce is not enabled. No files will be added to a change list.");
+				return;
+			}
+			StringBuilder clDescriptionBuilder= new StringBuilder();
+			clDescriptionBuilder.AppendLine("Generated localization config files with the NewPluginLocalizationTarget command of the LocalizationTargetEditor UAT command.");
+			clDescriptionBuilder.AppendLine("#rnx");
+			clDescriptionBuilder.AppendLine("#jira: none");
+			clDescriptionBuilder.AppendLine("#rb:");
+			clDescriptionBuilder.AppendLine("#preflight:");
+
+			int changeListNumber = P4.CreateChange(P4Env.Client, clDescriptionBuilder.ToString());
+
+			if (P4Enabled && changeListNumber != -1)
+			{
+				if (_filesToAdd.Count > 0)
+				{
+					P4.Add(changeListNumber, _filesToAdd);
+				}
+
+				if (_filesToEdit.Count > 0)
+				{
+					P4.Edit(changeListNumber, _filesToEdit);
+				}
+			}
+			else
+			{
+				// p4 is not enabled somehow. We list all the files that should be added or edited.
+				Logger.LogWarning("Failed to create a perforce changelist for files to be added/edited.");
+				Logger.LogInformation("Files to add:");
+				foreach (string fileToAdd in _filesToAdd)
+				{
+					Logger.LogInformation(fileToAdd);
+				}
+
+				Logger.LogInformation("Files to edit:");
+				foreach (string fileToEdit in _filesToEdit)
+				{
+					Logger.LogInformation(fileToEdit);
+				}
+			}
+		}
+	}
+}
diff --git a/Engine/Source/Programs/AutomationTool/Localization/LocalizationTargetEditor/Commands/PluginLocalizationTargetCommand.cs b/Engine/Source/Programs/AutomationTool/Localization/LocalizationTargetEditor/Commands/PluginLocalizationTargetCommand.cs
new file mode 100644
--- /dev/null
+++ b/Engine/Source/Programs/AutomationTool/Localization/LocalizationTargetEditor/Commands/PluginLocalizationTargetCommand.cs
@@ -0,0 +1,82 @@
+// Copyright Epic Games, Inc. All Rights Reserved.
+using AutomationTool;
+using Microsoft.Extensions.Logging;
+using System;
+using System.Collections.Generic;
+using System.IO;
+using System.Linq;
+using System.Reflection;
+using UnrealBuildTool;
+
+using static AutomationTool.CommandUtils;
+
+namespace EpicGames.Localization
+{
+
+	public abstract class PluginLocalizationTargetCommand : LocalizationTargetCommand
+	{
+		public string UERootDirectory { get; protected set; } = "";
+		public string UEProjectDirectory { get; protected set; } = "";
+		public string UEProjectName { get; protected set; } = ""; 
+		
+		public HashSet<string> IncludePlugins { get; protected set; } = new HashSet<string>();
+		public HashSet<string> ExcludePlugins { get; protected set; } = new HashSet<string>();
+
+
+		protected override bool ParseCommandLine()
+		{
+			if (!base.ParseCommandLine())
+			{
+				return false;
+			}
+			UERootDirectory = _commandLineHelper.ParseParamValue("UERootDirectory");
+			if (String.IsNullOrEmpty(UERootDirectory))
+			{
+				UERootDirectory = CommandUtils.CmdEnv.LocalRoot;
+			}
+			UEProjectName = _commandLineHelper.ParseParamValue("UEProject");
+			UEProjectDirectory = _commandLineHelper.ParseRequiredStringParam("UEProjectDirectory");
+			
+			// getting the plugins 
+			string includePluginsString = _commandLineHelper.ParseParamValue("IncludePlugins");
+			if (!String.IsNullOrEmpty(includePluginsString))
+			{
+				foreach (string plugin in includePluginsString.Split(","))
+				{
+					IncludePlugins.Add(plugin.Trim());
+				}
+			}
+
+			string excludePluginsString = _commandLineHelper.ParseParamValue("ExcludePlugins");
+			if (!String.IsNullOrEmpty(excludePluginsString))
+			{
+				foreach (string plugin in excludePluginsString.Split(","))
+				{
+					ExcludePlugins.Add(plugin.Trim());
+				}
+			}
+
+			string projectPluginDirectory = Path.Combine(UERootDirectory, UEProjectDirectory, "Plugins");
+			PluginType pluginType = String.IsNullOrEmpty(UEProjectName) ? PluginType.Engine : PluginType.Project;
+			// The provided path should be relative to the project directory 
+			string includePluginsDirectoryString = _commandLineHelper.ParseParamValue("IncludePluginsDirectory");
+			if (!String.IsNullOrEmpty(includePluginsDirectoryString))
+			{
+				// @TODOLocalization: May need to manipulate passed in string to replace directory separators on various platforms 
+				string IncludePluginsDirectoryAbsolutePath = Path.Combine(UERootDirectory, UEProjectDirectory, includePluginsDirectoryString);
+				Logger.LogInformation($"Including plugins under '{IncludePluginsDirectoryAbsolutePath}' directory.");
+				IncludePlugins.UnionWith(LocalizationUtilities.GetPluginNamesUnderDirectory(IncludePluginsDirectoryAbsolutePath, projectPluginDirectory, pluginType));
+			}
+
+			string excludePluginsDirectoryString = _commandLineHelper.ParseParamValue("ExcludePluginsDirectory");
+			if (!String.IsNullOrEmpty(excludePluginsDirectoryString))
+			{
+				string excludePluginsDirectoryAbsolutePath = Path.Combine(UERootDirectory, UEProjectDirectory, excludePluginsDirectoryString);
+				Logger.LogInformation($"Excluding plugins under '{excludePluginsDirectoryAbsolutePath}' directory.");
+				ExcludePlugins.UnionWith(LocalizationUtilities.GetPluginNamesUnderDirectory(excludePluginsDirectoryAbsolutePath, projectPluginDirectory, pluginType));
+			}
+			
+			return true;
+		}
+	}
+}
diff --git a/Engine/Source/Programs/AutomationTool/Localization/LocalizationTargetEditor/LocalizationTargetEditor.Automation.cs b/Engine/Source/Programs/AutomationTool/Localization/LocalizationTargetEditor/LocalizationTargetEditor.Automation.cs
new file mode 100644
--- /dev/null
+++ b/Engine/Source/Programs/AutomationTool/Localization/LocalizationTargetEditor/LocalizationTargetEditor.Automation.cs
@@ -0,0 +1,54 @@
+// Copyright Epic Games, Inc. All Rights Reserved.
+using System;
+using System.Text;
+using AutomationTool;
+using Microsoft.Extensions.Logging;
+
+namespace EpicGames.Localization
+{
+	[Help("Allows users to create, update, delete or query localization targets for the Engine or projects. This command has sub-commands and more info for each command can be retrieved with the -Info flag.")]
+	[Help("Command", "The sub-command you want to run. For a full list of valid command, use LocalizationTargetEditor -Info")]
+	[Help("Info", "An optional flag that can be passed to any sub-command. It will print out help text for the command and its associated parameters.")]
+	class LocalizationTargetEditor : BuildCommand
+	{
+		public override void ExecuteBuild()
+		{
+			LocalizationTargetCommand.Initialize(this);
+			string commandName = ParseParamValue("Command");
+			if (String.IsNullOrEmpty(commandName))
+			{
+				if (ParseParam("Info"))
+				{
+					StringBuilder helpTextBuilder = new StringBuilder();
+					helpTextBuilder.AppendLine("The Localization Target Editor is a means for you to create new localization targets. Use this command to create, edit, delete or query localization targets for projects, plugisn and platforms.");
+					helpTextBuilder.AppendLine("Below are all available commands under the Localization Target Editor. Add the -Info flag after any of these commands for more information.");
+					foreach (LocalizationTargetCommand helpCommand in LocalizationTargetCommand.GetAllCommands())
+					{
+						helpTextBuilder.AppendLine($"{helpCommand.Name}");
+					}
+					Logger.LogInformation($"{helpTextBuilder.ToString()}");
+					return;
+				}
+			}
+			LocalizationTargetCommand command = LocalizationTargetCommand.GetLocalizationTargetCommandFromName(commandName);
+
+			if (command is null)
+			{
+				Logger.LogError($"{commandName} is not a supported localization target command.");
+				return;
+			}
+			// If there's the -Info flag, we just display the help text and exit
+			// We have to use -Info as -Help is handled globally via UAT 
+			if (ParseParam("Info"))
+			{
+				command.PrintHelpText();
+				return;
+			}
+			Logger.LogInformation($"Starting execution of '{command.Name}' localization target editor command.");
+			if (!command.Execute())
+			{
+				Logger.LogError($"Localization Target Editor '{command.Name}' failed.");
+			}
+		}
+	}
+}
