Fix bug in ZenExportSnapshotTask when NOT using a baseline snapshot descriptor.

diff --git a/Engine/Source/Programs/AutomationTool/BuildGraph/Tasks/ZenExportSnapshotTask.cs b/Engine/Source/Programs/AutomationTool/BuildGraph/Tasks/ZenExportSnapshotTask.cs
--- a/Engine/Source/Programs/AutomationTool/BuildGraph/Tasks/ZenExportSnapshotTask.cs
+++ b/Engine/Source/Programs/AutomationTool/BuildGraph/Tasks/ZenExportSnapshotTask.cs
@@ -362,17 +362,20 @@ namespace AutomationTool.Tasks
 					NewExportSource.TargetPlatform = Platform;
 					NewExportSource.SnapshotBaseDescriptor = null;
 
-					FileReference PlatformSnapshotBase = new FileReference(Parameters.SnapshotBaseDescriptorFile.FullName.Replace("{Platform}", Platform, StringComparison.InvariantCultureIgnoreCase));
-
-					SnapshotDescriptorCollection? ParsedDescriptorCollection = null;
-					if (TryLoadJson(PlatformSnapshotBase, out ParsedDescriptorCollection) && (ParsedDescriptorCollection != null) && (ParsedDescriptorCollection.Snapshots != null))
+					if (Parameters.SnapshotBaseDescriptorFile != null)
 					{
-						foreach (SnapshotDescriptor ParsedDescriptor in ParsedDescriptorCollection.Snapshots)
+						FileReference PlatformSnapshotBase = new FileReference(Parameters.SnapshotBaseDescriptorFile.FullName.Replace("{Platform}", Platform, StringComparison.InvariantCultureIgnoreCase));
+
+						SnapshotDescriptorCollection? ParsedDescriptorCollection = null;
+						if (TryLoadJson(PlatformSnapshotBase, out ParsedDescriptorCollection) && (ParsedDescriptorCollection != null) && (ParsedDescriptorCollection.Snapshots != null))
 						{
-							if (ParsedDescriptor.TargetPlatform == Platform)
+							foreach (SnapshotDescriptor ParsedDescriptor in ParsedDescriptorCollection.Snapshots)
 							{
-								NewExportSource.SnapshotBaseDescriptor = ParsedDescriptor;
-								break;
+								if (ParsedDescriptor.TargetPlatform == Platform)
+								{
+									NewExportSource.SnapshotBaseDescriptor = ParsedDescriptor;
+									break;
+								}
 							}
 						}
 					}
