Accepting SkipInstall, SkipDeploy and SkipCopy arguments to skip build installation.
Uses AutoParam functions to set an AppConfig flag.
This change affects all platforms that in some form implement the "skip install" behavior.

diff --git a/Engine/Source/Programs/AutomationTool/Gauntlet/Platform/Android/Gauntlet.TargetDeviceAndroid.cs b/Engine/Source/Programs/AutomationTool/Gauntlet/Platform/Android/Gauntlet.TargetDeviceAndroid.cs
--- a/Engine/Source/Programs/AutomationTool/Gauntlet/Platform/Android/Gauntlet.TargetDeviceAndroid.cs
+++ b/Engine/Source/Programs/AutomationTool/Gauntlet/Platform/Android/Gauntlet.TargetDeviceAndroid.cs
@@ -926,11 +926,9 @@ namespace Gauntlet
 				RunAdbDeviceCommand(string.Format("shell rm -r {0}/Android/data/{1}/*", StorageLocation, Build.AndroidPackageName));
 				RunAdbDeviceCommand(string.Format("shell rm -r {0}/Android/obb/{1}/*", StorageLocation, Build.AndroidPackageName));
 				RunAdbDeviceCommand(string.Format("shell rm -r {0}/Download/*", StorageLocation));
-			}
-
-				bool SkipDeploy = Globals.Params.ParseParam("SkipDeploy");
+			}			
 
-			if (SkipDeploy == false)
+			if (!AppConfig.SkipInstall)
 			{
 				if (Globals.Params.ParseParam("cleandevice")
 					|| Globals.Params.ParseParam("fullclean"))
@@ -1010,7 +1008,7 @@ namespace Gauntlet
                 }
             }
 
-			if (SkipDeploy == false)
+			if (!AppConfig.SkipInstall)
 			{
 				// obb files need to be named based on APK version (grrr), so find that out. This should return something like
 				// versionCode=2 minSdk=21 targetSdk=21
diff --git a/Engine/Source/Programs/AutomationTool/Gauntlet/Platform/Linux/Gauntlet.TargetDeviceLinux.cs b/Engine/Source/Programs/AutomationTool/Gauntlet/Platform/Linux/Gauntlet.TargetDeviceLinux.cs
--- a/Engine/Source/Programs/AutomationTool/Gauntlet/Platform/Linux/Gauntlet.TargetDeviceLinux.cs
+++ b/Engine/Source/Programs/AutomationTool/Gauntlet/Platform/Linux/Gauntlet.TargetDeviceLinux.cs
@@ -345,8 +345,6 @@ namespace Gauntlet
 
 		protected IAppInstall InstallStagedBuild(UnrealAppConfig AppConfig, StagedBuild InBuild)
 		{
-			bool SkipDeploy = Globals.Params.ParseParam("SkipDeploy");
-
 			string BuildPath = InBuild.BuildPath;
 
 			if (CanRunFromPath(BuildPath) == false)
@@ -354,7 +352,7 @@ namespace Gauntlet
 				string SubDir = string.IsNullOrEmpty(AppConfig.Sandbox) ? AppConfig.ProjectName : AppConfig.Sandbox;
 				string DestPath = Path.Combine(this.LocalCachePath, SubDir, AppConfig.ProcessType.ToString());
 
-				if (!SkipDeploy)
+				if (!AppConfig.SkipInstall)
 				{
 					DestPath = StagedBuild.InstallBuildParallel(AppConfig, InBuild, BuildPath, DestPath, ToString());
 				}
diff --git a/Engine/Source/Programs/AutomationTool/Gauntlet/Platform/Windows/Gauntlet.TargetDeviceWindows.cs b/Engine/Source/Programs/AutomationTool/Gauntlet/Platform/Windows/Gauntlet.TargetDeviceWindows.cs
--- a/Engine/Source/Programs/AutomationTool/Gauntlet/Platform/Windows/Gauntlet.TargetDeviceWindows.cs
+++ b/Engine/Source/Programs/AutomationTool/Gauntlet/Platform/Windows/Gauntlet.TargetDeviceWindows.cs
@@ -481,8 +481,6 @@ namespace Gauntlet
 
 		protected IAppInstall InstallStagedBuild(UnrealAppConfig AppConfig, StagedBuild InBuild)
 		{
-			bool SkipDeploy = Globals.Params.ParseParam("SkipDeploy");
-
 			string BuildPath = InBuild.BuildPath;
 
 			if (CanRunFromPath(BuildPath) == false)
@@ -490,7 +488,7 @@ namespace Gauntlet
 				string SubDir = string.IsNullOrEmpty(AppConfig.Sandbox) ? AppConfig.ProjectName : AppConfig.Sandbox;
 				string DestPath = Path.Combine(this.LocalCachePath, SubDir, AppConfig.ProcessType.ToString());
 
-				if (!SkipDeploy)
+				if (!AppConfig.SkipInstall)
 				{
 					DestPath = StagedBuild.InstallBuildParallel(AppConfig, InBuild, BuildPath, DestPath, ToString());
 				}
diff --git a/Engine/Source/Programs/AutomationTool/Gauntlet/Unreal/Base/Gauntlet.UnrealAppConfig.cs b/Engine/Source/Programs/AutomationTool/Gauntlet/Unreal/Base/Gauntlet.UnrealAppConfig.cs
--- a/Engine/Source/Programs/AutomationTool/Gauntlet/Unreal/Base/Gauntlet.UnrealAppConfig.cs
+++ b/Engine/Source/Programs/AutomationTool/Gauntlet/Unreal/Base/Gauntlet.UnrealAppConfig.cs
@@ -102,6 +102,10 @@ namespace Gauntlet
 		// new system
 		public IBuild Build { get; set; }
 
+		// Prevents installing a build on device
+		[AutoParamWithNames(false, "SkipInstall", "SkipDeploy", "SkipCopy")]
+		public bool SkipInstall { get; set; }
+
 		/// <summary>
 		/// Constructor that sets some required values to defaults
 		/// </summary>
@@ -112,6 +116,7 @@ namespace Gauntlet
 			CommandLine = "";
 			Configuration = UnrealTargetConfiguration.Development;
 			Sandbox = "Gauntlet";
+			AutoParam.ApplyParamsAndDefaults(this, Globals.Params.AllArguments);
 		}
 	}
 }
\ No newline at end of file
