Localization: Added profiling to Smartling downloads.
- The time it takes to download each individual file from Smartling is now logged
- Download time for an entire localization project is also logged.
#test locally compiled everything and performed a download of Smartling files for a project. The times were indeed reported correctly.

diff --git a/Engine/Source/Programs/AutomationTool/SmartlingLocalization/SmartlingLocalizationProvider.cs b/Engine/Source/Programs/AutomationTool/SmartlingLocalization/SmartlingLocalizationProvider.cs
--- a/Engine/Source/Programs/AutomationTool/SmartlingLocalization/SmartlingLocalizationProvider.cs
+++ b/Engine/Source/Programs/AutomationTool/SmartlingLocalization/SmartlingLocalizationProvider.cs
@@ -1,5 +1,6 @@
 // Copyright Epic Games, Inc. All Rights Reserved.
 using System;
+using System.Diagnostics;
 using System.Collections.Generic;
 using System.IO;
 using System.Threading.Tasks;
@@ -13,6 +14,7 @@ using AutomationTool;
 using UnrealBuildTool;
 using EpicGames.Localization;
 using System.Threading;
+using Microsoft.Extensions.Logging;
 
 
 #pragma warning disable SYSLIB0014
@@ -87,6 +89,8 @@ namespace EpicGames.SmartlingLocalization
 
 		public async override Task DownloadProjectFromLocalizationProvider(string ProjectName, ProjectImportExportInfo ProjectImportInfo)
 		{
+			Console.WriteLine($"Starting Smartling download for {ProjectName} project files.");
+			Stopwatch Watch = Stopwatch.StartNew();
 			// Get the latest files for each culture.
 			foreach (var Culture in ProjectImportInfo.CulturesToGenerate)
 			{
@@ -102,6 +106,8 @@ namespace EpicGames.SmartlingLocalization
 					await DownloadLatestPOFile(Culture, Platform, ProjectImportInfo);
 				}
 			}
+			Watch.Stop();
+			Console.WriteLine($"Completed Smartling download for {ProjectName} project files in {Watch.ElapsedMilliseconds / 1000} seconds.");
 		}
 
 		private async Task DownloadLatestPOFile(string EpicLocale, string Platform, ProjectImportExportInfo ProjectImportInfo)
@@ -137,6 +143,8 @@ namespace EpicGames.SmartlingLocalization
 			// The base for which all the exponential back off will be derived from. By default HttpClient has a default timeout of 100s 
 			int InitialTimeOut = 150;
 			int CurrentTimeOut = CurrentTries * InitialTimeOut;
+			// Measures how long it takes from sending a response to successfully getting a response 
+			Stopwatch DownloadStopWatch = new Stopwatch();
 			while (true)
 			{
 				CurrentTimeOut = CurrentTries * InitialTimeOut;
@@ -145,13 +153,15 @@ namespace EpicGames.SmartlingLocalization
 				using var CancellationToken = new CancellationTokenSource(Timeout);
 				try
 				{
+					DownloadStopWatch.Start(); 
 					DownloadResponse = await Client.GetAsync(DownloadUriBuilder.Uri, CancellationToken.Token);
 					if (DownloadResponse.StatusCode == HttpStatusCode.Unauthorized)
 					{
 						++CurrentTries;
 						if (CurrentTries > MaxTries)
 						{
-							Console.WriteLine($"[FAILED] Exporting: '{ExportFile.FullName}' ({EpicLocale}) exhausted all retries.");
+							DownloadStopWatch.Stop();
+							Console.WriteLine($"[FAILED] Exporting: '{ExportFile.FullName}' ({EpicLocale}) exhausted all retries after {DownloadStopWatch.ElapsedMilliseconds / 1000}.");
 							return;
 						}
 						Console.WriteLine($"Encountered HTTP Status Code 401. Authentication most likely expired. Retrying {CurrentTries}/{MaxTries} times with refreshed authentication token.");
@@ -165,7 +175,8 @@ namespace EpicGames.SmartlingLocalization
 					++CurrentTries;
 					if (CurrentTries > MaxTries)
 					{
-						Console.WriteLine($"[FAILED] Exporting: '{ExportFile.FullName}' ({EpicLocale}) exhausted all retries. - {Ex}");
+						DownloadStopWatch.Stop();
+						Console.WriteLine($"[FAILED] Exporting: '{ExportFile.FullName}' ({EpicLocale}) exhausted all retries in {DownloadStopWatch.ElapsedMilliseconds / 1000} seconds. - {Ex}");
 						return;
 					}
 					Console.WriteLine($"Failed to get download response. Retrying {CurrentTries}/{MaxTries} times.");
@@ -195,8 +206,8 @@ namespace EpicGames.SmartlingLocalization
 						await DownloadStream.CopyToAsync(DownloadFileStream);
 					}
 				}
-
-				Console.WriteLine($"[SUCCESS] Exporting: '{SmartlingFileUri}' as '{ExportFile.FullName}' ({EpicLocale})");
+				DownloadStopWatch.Stop();
+				Console.WriteLine($"[SUCCESS] Exporting: '{SmartlingFileUri}' as '{ExportFile.FullName}' ({EpicLocale}) in {DownloadStopWatch.ElapsedMilliseconds / 1000} seconds.");
 				// Reset the write status of the file
 				if (ExportFileWasReadOnly)
 				{
@@ -233,7 +244,8 @@ namespace EpicGames.SmartlingLocalization
 			else
 			{
 				// The file may not currently exist in Smartling and will need to be uploaded first via the Upload step later on. 
-				Console.WriteLine($"[FAILED] Exporting: '{ExportFile.FullName}' ({EpicLocale}. The file may need to be uploaded first.)");
+				DownloadStopWatch.Stop();
+				Console.WriteLine($"[FAILED] Exporting: '{ExportFile.FullName}' ({EpicLocale} in {DownloadStopWatch.ElapsedMilliseconds / 1000} seconds. The file may need to be uploaded first.)");
 				await PrintRequestErrors(DownloadResponse);
 			}
 		}
