Automation: LowLevelTestRunner: Modified the test runner to deploy each test application into its own directory to avoid the risk of one test reading other tests' data.

More specifically, all low level tests were deployed into LowLevelTests directory/sandbox. It meanst that if we ran OnlineTests, which deploy .ini scripts, a subsequent run of FoundationTests would have these .ini as well, which would unintentionally affect the behavior of FoundationTests.

Now, Online will be deployed into LowLevelTests-OnlineTests and Foundation into LowLevelTests-Foundation avoiding the clash.

#rnx

diff --git a/Engine/Source/Programs/AutomationTool/LowLevelTests/RunLowLevelTests.cs b/Engine/Source/Programs/AutomationTool/LowLevelTests/RunLowLevelTests.cs
--- a/Engine/Source/Programs/AutomationTool/LowLevelTests/RunLowLevelTests.cs
+++ b/Engine/Source/Programs/AutomationTool/LowLevelTests/RunLowLevelTests.cs
@@ -730,7 +730,7 @@ namespace LowLevelTests
 				CachedConfig.Platform = Platform;
 				CachedConfig.Configuration = UnrealTargetConfiguration.Development;
 				CachedConfig.Build = DiscoveredBuild;
-				CachedConfig.Sandbox = "LowLevelTests";
+				CachedConfig.Sandbox = $"LowLevelTests-{TestApp}";
 				CachedConfig.FilesToCopy = new List<UnrealFileToCopy>();
 				if (InContainerized)
 				{
