ReplayRun gauntlet job - Incremental reporting for replayrun (upload CSVs as they come in) via new -incrementalReporting param
- Move log processing out of the IsTestRunning accessor and into TickTest where it should be.
- Tidy up some legacy comparison test stuff. Will remove soon

diff --git a/Engine/Source/Programs/AutomationTool/Gauntlet/Framework/Utils/Gauntlet.ReportGenUtils.cs b/Engine/Source/Programs/AutomationTool/Gauntlet/Framework/Utils/Gauntlet.ReportGenUtils.cs
--- a/Engine/Source/Programs/AutomationTool/Gauntlet/Framework/Utils/Gauntlet.ReportGenUtils.cs
+++ b/Engine/Source/Programs/AutomationTool/Gauntlet/Framework/Utils/Gauntlet.ReportGenUtils.cs
@@ -325,12 +325,17 @@ namespace Gauntlet
 			return "\"" + inStr + "\"";
 		}
 
-		public static bool BuildVersionValidationCheck(string InBuildVersion, string CsvDirectory)
+		public static bool BuildVersionValidationCheck(string InBuildVersion, string CsvDirectory, string SearchPattern = null)
 		{
-			Dictionary<string, string> Metadata = ReadCsvMetadataFromDirectory(CsvDirectory);
+			Dictionary<string, string> Metadata = ReadCsvMetadataFromDirectory(CsvDirectory, false, SearchPattern);
 			return BuildVersionValidationCheck(InBuildVersion, Metadata);
 		}
 
+		public static bool BuildVersionValidationCheckForCsv(string InBuildVersion, string CsvFilename)
+		{
+			Dictionary<string, string> Metadata = ReadCsvMetadata(CsvFilename);
+			return BuildVersionValidationCheck(InBuildVersion, Metadata);
+		}
 
 		public static string GetMetadataQueryForLastNDays(int DurationInDays)
 		{
