- More fixes for Mac packaging/archiving issues
  * Set the no-encryption setting in the template Mac plist to make uploading to TestFlight/AppStore easier
  * Fixed editor to not compile as universal when using -distribution, only the client should be universal
  * Further improvements to stripping out nested .app's in a final .app (remove them from the staged directory helps)
  * Changed how the Mac app name is displayed, since the .app name itself is shown in the Finder,  unlike IOS where the CFBundleName is shown (the archived .app name in the .xcarchive is named by a project setting, falling back to the .uproject name if not set)
  * Disabled the SignExecutables function on Modern because they attempt to sign the wrong .apps, and one is no longer (previously it was uselessly signing .apps, but now it throws an error due to changes for the third item in this list)
[CONSIDER5.3]

diff --git a/Engine/Source/Developer/Mac/MacTargetPlatform/Classes/XcodeProjectSettings.h b/Engine/Source/Developer/Mac/MacTargetPlatform/Classes/XcodeProjectSettings.h
--- a/Engine/Source/Developer/Mac/MacTargetPlatform/Classes/XcodeProjectSettings.h
+++ b/Engine/Source/Developer/Mac/MacTargetPlatform/Classes/XcodeProjectSettings.h
@@ -66,6 +66,14 @@ public:
 	UPROPERTY(EditAnywhere, config, Category=Xcode, meta = (EditCondition="bUseModernXcode", DisplayName = "Bundle ID Prefix"))
 	FString CodeSigningPrefix;
 	
+	/**
+	 * The name of the Mac .app when making an archived build (for uploading to App Store, etc). The Finder shows Mac apps by their .app name, and we don't name the .app  with
+	 * "pretty" names during development. When packaging for distribution (or using Archive menu in Xcode) this will become the name of the .app, and will be what end users
+	 * will have on their Mac. If this is not set, the .app will have the name of the .uproject file
+	 */
+	UPROPERTY(EditAnywhere, config, Category=Xcode, meta = (EditCondition="bUseModernXcode", DisplayName = "Mac: Published App Name"))
+	FString ApplicationDisplayName;
+
 	/**
 	 * The App Category that will be used for Apple App Store submissions
 	 */
diff --git a/Engine/Source/Programs/AutomationTool/Apple/ApplePlatform.Automation.cs b/Engine/Source/Programs/AutomationTool/Apple/ApplePlatform.Automation.cs
--- a/Engine/Source/Programs/AutomationTool/Apple/ApplePlatform.Automation.cs
+++ b/Engine/Source/Programs/AutomationTool/Apple/ApplePlatform.Automation.cs
@@ -92,6 +92,14 @@ public abstract class ApplePlatform : Platform
 
 	public override void PostStagingFileCopy(ProjectParams Params, DeploymentContext SC)
 	{
+		// staging will put binaries into Staged/<game>/Binaries/<platform> and they aren't needed, and when we pull this into a .app, it 
+		// messes with the resulting .app. So, we remove the game binary now (leaving in helper .app's and raw .dylibs, etc)
+		// they come from BuildProducts, and we could maybe remove from that list, but it could cause issues with Horde/buildmachines
+		FileReference BinaryPath = FileReference.Combine(SC.StageDirectory, Params.IsCodeBasedProject ? SC.ShortProjectName : "Engine", "Binaries", SC.PlatformDir, SC.StageExecutables[0]);
+		DirectoryReference AppPath = new DirectoryReference(BinaryPath.FullName + ".app");
+		InternalUtils.SafeDeleteFile(BinaryPath.FullName, true);
+		InternalUtils.SafeDeleteDirectory(AppPath.FullName, true);
+
 		if (AppleExports.UseModernXcode(Params.RawProjectPath))
 		{
 			// now reset the envvar so the following build will process the staged data
diff --git a/Engine/Source/Programs/AutomationTool/Mac/MacPlatform.Automation.cs b/Engine/Source/Programs/AutomationTool/Mac/MacPlatform.Automation.cs
--- a/Engine/Source/Programs/AutomationTool/Mac/MacPlatform.Automation.cs
+++ b/Engine/Source/Programs/AutomationTool/Mac/MacPlatform.Automation.cs
@@ -105,7 +105,9 @@ public class MacPlatform : ApplePlatform
 		foreach (UnrealBuild.BuildTarget Target in Agenda.Targets)
 		{
 			// if building for Distribution, and no arch is already specified, then get the distro architectures and use that for this build
-			if (Params.Distribution && !Target.UBTArgs.ToLower().Contains("-architecture="))
+			// editors aren't usually distributed, so if we are doing distribution, it's probably not for the editor target, and we don't want to make universal 
+			// editors just to make a distribution client
+			if (Params.Distribution && !Target.TargetName.Contains("Editor") && !Target.UBTArgs.ToLower().Contains("-architecture="))
 			{
 				UnrealArchitectures DistroArches = UnrealArchitectureConfig.ForPlatform(UnrealTargetPlatform.Mac).DistributionArchitectures(Params.RawProjectPath, Target.TargetName);
 				Target.UBTArgs += " -architecture=" + DistroArches.ToString();
@@ -356,7 +358,6 @@ public class MacPlatform : ApplePlatform
 			IProcessResult CommandResult = Run("otool", "-l \"" + ExePath + "\"", null, ERunOptions.None);
 			if (CommandResult.ExitCode == 0)
 			{
-				bool bModifiedWithInstallNameTool = false;
 				StringReader Reader = new StringReader(CommandResult.Output);
 				Regex RPathPattern = new Regex(@"^\s+path (?<rpath>.+)\s\(offset");
 				string ToRemovePattern = Params.CreateAppBundle ? "/../../../" : "@loader_path/../UE4/";
@@ -375,18 +376,10 @@ public class MacPlatform : ApplePlatform
 							if (RPath.Contains(ToRemovePattern))
 							{
 								Run("xcrun", "install_name_tool -delete_rpath \"" + RPath + "\" \"" + ExePath + "\"", null, ERunOptions.NoStdOutCapture);
-								bModifiedWithInstallNameTool = true;
 							}
 						}
 					}
 				}
-
-				// Now re-sign the exectuables if they were previously signed during staging. Modifying the binaries
-				// via install_name_tool _after_ signing them invalidates the signature.
-				if (bModifiedWithInstallNameTool && (Params.bCodeSign && !Params.SkipStage))
-				{
-					SC.StageTargetPlatform.SignExecutables(SC, Params);
-				}
 			}
 		}
 	}
@@ -649,28 +642,32 @@ public class MacPlatform : ApplePlatform
 				RemoveExtraRPaths(Params, SC);
 			}
 
-			// Sign everything we built
-			List<FileReference> FilesToSign = GetExecutableNames(SC);
-			Logger.LogInformation("{Text}", "RuntimeProjectRootDir: " + SC.RuntimeProjectRootDir);
-			foreach (var Exe in FilesToSign)
+			// with modern, the .app we'd want to sign is in the root of the staging directory, so this doesn't do anything,
+			// and that one is already signed. note this is only done when Staging, so it doesn't affect codesigning the editor
+			if (!AppleExports.UseModernXcode(Params.RawProjectPath))
 			{
-				Logger.LogInformation("{Text}", "Signing: " + Exe);
-				string AppBundlePath = "";
-				if (Exe.IsUnderDirectory(DirectoryReference.Combine(SC.RuntimeProjectRootDir, "Binaries", SC.PlatformDir)))
+				// Sign everything we built
+				List<FileReference> FilesToSign = GetExecutableNames(SC);
+				Logger.LogInformation("{Text}", "RuntimeProjectRootDir: " + SC.RuntimeProjectRootDir);
+				foreach (var Exe in FilesToSign)
 				{
-					Logger.LogInformation("Starts with Binaries");
-					AppBundlePath = CombinePaths(SC.RuntimeProjectRootDir.FullName, "Binaries", SC.PlatformDir, Path.GetFileNameWithoutExtension(Exe.FullName) + ".app");
-				}
-				else if (Exe.IsUnderDirectory(DirectoryReference.Combine(SC.RuntimeRootDir, "Engine/Binaries", SC.PlatformDir)))
-				{
-					Logger.LogInformation("Starts with Engine/Binaries");
-					AppBundlePath = CombinePaths("Engine/Binaries", SC.PlatformDir, Path.GetFileNameWithoutExtension(Exe.FullName) + ".app");
-				}
+					Logger.LogInformation("{Text}", "Signing: " + Exe);
+					string AppBundlePath = "";
+					if (Exe.IsUnderDirectory(DirectoryReference.Combine(SC.RuntimeProjectRootDir, "Binaries", SC.PlatformDir)))
+					{
+						Logger.LogInformation("Starts with Binaries");
+						AppBundlePath = CombinePaths(SC.RuntimeProjectRootDir.FullName, "Binaries", SC.PlatformDir, Path.GetFileNameWithoutExtension(Exe.FullName) + ".app");
+					}
+					else if (Exe.IsUnderDirectory(DirectoryReference.Combine(SC.RuntimeRootDir, "Engine/Binaries", SC.PlatformDir)))
+					{
+						Logger.LogInformation("Starts with Engine/Binaries");
+						AppBundlePath = CombinePaths("Engine/Binaries", SC.PlatformDir, Path.GetFileNameWithoutExtension(Exe.FullName) + ".app");
+					}
 
-				Logger.LogInformation("{Text}", "Signing: " + AppBundlePath);
-				CodeSign.SignMacFileOrFolder(AppBundlePath);
+					Logger.LogInformation("{Text}", "Signing: " + AppBundlePath);
+					CodeSign.SignMacFileOrFolder(AppBundlePath);
+				}
 			}
-
 		}
 		return true;
 	}
diff --git a/Engine/Source/Programs/UnrealBuildTool/ProjectFiles/Xcode/XcodeProject.cs b/Engine/Source/Programs/UnrealBuildTool/ProjectFiles/Xcode/XcodeProject.cs
--- a/Engine/Source/Programs/UnrealBuildTool/ProjectFiles/Xcode/XcodeProject.cs
+++ b/Engine/Source/Programs/UnrealBuildTool/ProjectFiles/Xcode/XcodeProject.cs
@@ -5,6 +5,7 @@ using System.Collections.Generic;
 using System.IO;
 using System.Linq;
 using System.Text;
+using System.Xml.Linq;
 using EpicGames.Core;
 using Microsoft.CodeAnalysis;
 using Microsoft.Extensions.Logging;
@@ -1412,11 +1413,22 @@ namespace UnrealBuildTool.XcodeProjectXcconfig
 					});
 				}
 
+				// when we bring stated data into the .app, we have to skip some temp stuff that went into it
+				string[] Exclusions =
+				{
+					"-/Info.plist",
+					"-/Manifest_*",
+					$"/*.app", // remove the staged .app from the root dir, it's hard to do by name due to ProjectName in staging, and TargetName, etc here
+				};
+
+				// make a string like --exclude=/Info.plist --exclude=/Manifest_* ...
+				string ExcludeString = string.Join(" ", Exclusions.Select(x => (x[0] == '+' ? "--include" : "--exclude") + $"=\\\"{x.Substring(1)}\\\""));
+
 				CopyScript.AddRange(new string[]
 				{
 					"",
 					$"echo \\\"Syncing ${{STAGED_DIR}}{SyncSourceSubdir} to ${{CONFIGURATION_BUILD_DIR}}/${{CONTENTS_FOLDER_PATH}}{SyncDestSubdir}\\\"",
-					$"rsync -a --delete --exclude=/Info.plist --exclude=/Manifest_* --exclude=${{{ExecutableKey}}}.app --exclude=${{PRODUCT_NAME}} --exclude=${{{ExecutableKey}}} --exclude=${{PRODUCT_NAME}}.app \\\"${{STAGED_DIR}}{SyncSourceSubdir}/\\\" \\\"${{CONFIGURATION_BUILD_DIR}}/${{CONTENTS_FOLDER_PATH}}{SyncDestSubdir}\\\"",
+					$"rsync -a --delete {ExcludeString} \\\"${{STAGED_DIR}}{SyncSourceSubdir}/\\\" \\\"${{CONFIGURATION_BUILD_DIR}}/${{CONTENTS_FOLDER_PATH}}{SyncDestSubdir}\\\"",
 				});
 			}
 
@@ -1699,6 +1711,7 @@ namespace UnrealBuildTool.XcodeProjectXcconfig
 			string? SupportedDevices = null;
 			string? MarketingVersion = null;
 			string? BundleIdentifier;
+			string? ApplicationDisplayName = null;
 			List<string> ExtraConfigLines = new();
 
 			// get signing settings
@@ -1710,6 +1723,7 @@ namespace UnrealBuildTool.XcodeProjectXcconfig
 			PlatformIni.TryGetValue("/Script/MacTargetPlatform.XcodeProjectSettings", $"{Platform}ProvisioningProfile", out ProvisioningProfile);
 			PlatformIni.TryGetValue("/Script/MacTargetPlatform.XcodeProjectSettings", $"{Platform}SigningIdentity", out SigningIdentity);
 			PlatformIni.TryGetValue("/Script/MacTargetPlatform.XcodeProjectSettings", "AppCategory", out AppCategory);
+			PlatformIni.TryGetValue("/Script/MacTargetPlatform.XcodeProjectSettings", "ApplicationDisplayName", out ApplicationDisplayName);
 
 
 			if (Platform == UnrealTargetPlatform.Mac)
@@ -1984,7 +1998,20 @@ namespace UnrealBuildTool.XcodeProjectXcconfig
 				ConfigXcconfig.AppendLine($"UE_TARGET_CONFIG = {Config.BuildConfig}");
 				ConfigXcconfig.AppendLine($"UE_UBT_BINARY_SUBPATH = {ExetuableSubPath}");
 				ConfigXcconfig.AppendLine($"{ExecutableKey} = {ExecutableName}");
-				ConfigXcconfig.AppendLine($"PRODUCT_NAME = {ProductName}");
+				if (Platform == UnrealTargetPlatform.Mac)
+				{
+					// on Mac, we need to name the .app nicely before pushing to App store, otherwise distributing, so use the ini setting if it's there ("Unreal Match 3"), otherwise use the uproject name (ie "Lyra" instead of "LyraGame")
+					ConfigXcconfig.AppendLine($"PRODUCT_NAME_build = {ProductName}");
+					string ArchivedName = ApplicationDisplayName ?? (UnrealData.UProjectFileLocation == null ? ProductName : UnrealData.UProjectFileLocation!.GetFileNameWithoutAnyExtensions());
+					ConfigXcconfig.AppendLine($"PRODUCT_NAME_install = {ArchivedName}");
+
+					// this will choose the proper PRODUCT_NAME when archiving vs normal building
+					ConfigXcconfig.AppendLine("PRODUCT_NAME = $(PRODUCT_NAME_$(ACTION))");
+				}
+				else
+				{
+					ConfigXcconfig.AppendLine($"PRODUCT_NAME = {ProductName}");
+				}
 				if (EntitlementsMetadata != null && EntitlementsMetadata.Mode == MetadataMode.UsePremade)
 				{
 					ConfigXcconfig.AppendLine($"CODE_SIGN_ENTITLEMENTS = {EntitlementsMetadata.XcodeProjectRelative}");
