Functionality to call VerifyLogin on devices of type IOnlineServiceLogin

diff --git a/Engine/Build/LowLevelTests.xml b/Engine/Build/LowLevelTests.xml
--- a/Engine/Build/LowLevelTests.xml
+++ b/Engine/Build/LowLevelTests.xml
@@ -91,6 +91,7 @@
 	<Option Name="Device" DefaultValue="" Description="Device address to run tests on"/>
 	<Option Name="Configuration" DefaultValue="Development" Description="Configuration type: development, debug, test, shipping etc."/>
 	<Option Name="AttachToDebugger" DefaultValue="false" Description="Wait for debugger to attach."/>
+	<Option Name="EnableDeviceLoginVerification" DefaultValue="$(IsBuildMachine)" Description="Verify a user is logged in on Devices that supports it, try to sign in if not."/>
 
 	<Property Name="TestNodes" Value=""/>
 	<Property Name="AfterTestNodes" Value=""/>
@@ -142,6 +143,7 @@
 		<Property Name="TestArgsExt" Value="$(TestArgsExt) $($(TestName)GauntletArgs)"/>
 		<Property Name="TestArgsExt" Value="$(TestArgsExt) -extra-args=&quot;$($(TestName)ExtraArgs)&quot;"/>
 		<Property Name="TestArgsExt" Value="$(TestArgsExt) $($(Platform)ExtraArgs)"/>
+		<Property Name="TestArgsExt" Value="$(TestArgsExt) -VerifyLogin" If="$(EnableDeviceLoginVerification)"/>
 		<Do If="'$(Device)' != ''">
 			<Property Name="TestArgsExt" Value="$(TestArgsExt) -device=$(Device) "/>
 		</Do>
diff --git a/Engine/Source/Programs/AutomationTool/Gauntlet/Framework/Base/Gauntlet.OnlineServiceLogin.cs b/Engine/Source/Programs/AutomationTool/Gauntlet/Framework/Base/Gauntlet.OnlineServiceLogin.cs
--- a/Engine/Source/Programs/AutomationTool/Gauntlet/Framework/Base/Gauntlet.OnlineServiceLogin.cs
+++ b/Engine/Source/Programs/AutomationTool/Gauntlet/Framework/Base/Gauntlet.OnlineServiceLogin.cs
@@ -6,7 +6,7 @@ namespace Gauntlet
 	/// TargetDevice extension for platforms that contain an online service.
 	/// Enables more robust user sign in verification.
 	/// </summary>
-	internal interface IOnlineServiceLogin
+	public interface IOnlineServiceLogin
 	{
 		/// <summary>
 		/// Verifies the device has connected to its relevant platform network.
diff --git a/Engine/Source/Programs/AutomationTool/LowLevelTests/RunLowLevelTests.cs b/Engine/Source/Programs/AutomationTool/LowLevelTests/RunLowLevelTests.cs
--- a/Engine/Source/Programs/AutomationTool/LowLevelTests/RunLowLevelTests.cs
+++ b/Engine/Source/Programs/AutomationTool/LowLevelTests/RunLowLevelTests.cs
@@ -165,6 +165,8 @@ namespace LowLevelTests
 
 		public bool AttachToDebugger;
 
+		public bool VerifyLogin;
+
 		[AutoParam(false)]
 		public bool SkipStage;
 
@@ -234,6 +236,7 @@ namespace LowLevelTests
 
 			Tags = Params.ParseValue("tags=", null);
 			AttachToDebugger = Params.ParseParam("attachtodebugger");
+			VerifyLogin = Params.ParseParam("verifylogin");
 
 			TestExtraArgs = Params.ParseValue("extra-args=", null);
 
@@ -269,6 +272,7 @@ namespace LowLevelTests
 		private string Tags { get; set; }
 		private int Sleep { get; set; }
 		private bool AttachToDebugger { get; set; }
+		private bool VerifyLogin { get; set; }
 		private string ReportType { get; set; }
 		private int PerTestTimeout { get; set; }
 		private string TestExtraArgs { get; set; }
@@ -277,17 +281,18 @@ namespace LowLevelTests
 
 		public UnrealDeviceReservation UnrealDeviceReservation { get; private set; }
 
-		public LowLevelTestsSession(LowLevelTestsBuildSource InBuildSource, string InTags, int InSleep, bool InAttachToDebugger, string InReportType, int InPerTestTimeout = 0, string InTestExtraArgs = null, bool InContainerized = false)
+		public LowLevelTestsSession(LowLevelTestsBuildSource InBuildSource, LowLevelTestExecutorOptions InOptions)
 		{
 			BuildSource = InBuildSource;
-			Tags = InTags;
-			Sleep = InSleep;
-			AttachToDebugger = InAttachToDebugger;
-			ReportType = InReportType;
-			PerTestTimeout = InPerTestTimeout;
+			Tags = InOptions.Tags;
+			Sleep = InOptions.Sleep;
+			AttachToDebugger = InOptions.AttachToDebugger;
+			VerifyLogin = InOptions.VerifyLogin;
+			ReportType = InOptions.ReportType;
+			PerTestTimeout = InOptions.Timeout;
 			UnrealDeviceReservation = new UnrealDeviceReservation();
-			TestExtraArgs = InTestExtraArgs;
-			Containerized = InContainerized;
+			TestExtraArgs = InOptions.TestExtraArgs;
+			Containerized = InOptions.Containerized;
 		}
 
 		public bool TryReserveDevices()
@@ -337,6 +342,16 @@ namespace LowLevelTests
 			IDeviceUsageReporter.RecordStart(Device.Name, (UnrealTargetPlatform)Device.Platform, IDeviceUsageReporter.EventType.Install, IDeviceUsageReporter.EventState.Success, BuildSource.BuildName);
 			try
 			{
+				if (VerifyLogin && Device is IOnlineServiceLogin)
+				{
+					Log.Info("\nVerifying device login...");
+					if (!(Device as IOnlineServiceLogin).VerifyLogin())
+					{
+						throw new AutomationException("Unable to secure login to an online platform account!");
+					}
+					Log.Info("Success! User signed-in.\n");
+				}
+
 				Install = Device.InstallApplication(AppConfig);
 				InstallSuccess = true;
 				IDeviceUsageReporter.RecordEnd(Device.Name, (UnrealTargetPlatform)Device.Platform, IDeviceUsageReporter.EventType.Install, IDeviceUsageReporter.EventState.Success);
diff --git a/Engine/Source/Programs/AutomationTool/LowLevelTests/Tests/Gauntlet.LowLevelTestNode.cs b/Engine/Source/Programs/AutomationTool/LowLevelTests/Tests/Gauntlet.LowLevelTestNode.cs
--- a/Engine/Source/Programs/AutomationTool/LowLevelTests/Tests/Gauntlet.LowLevelTestNode.cs
+++ b/Engine/Source/Programs/AutomationTool/LowLevelTests/Tests/Gauntlet.LowLevelTestNode.cs
@@ -50,7 +50,7 @@ namespace LowLevelTests
 		{
 			if (LowLevelTestsApp == null)
 			{
-				LowLevelTestsApp = new LowLevelTestsSession(Context.BuildInfo, Context.Options.Tags, Context.Options.Sleep, Context.Options.AttachToDebugger, Context.Options.ReportType, Context.Options.Timeout, Context.Options.TestExtraArgs, Context.Options.Containerized);
+				LowLevelTestsApp = new LowLevelTestsSession(Context.BuildInfo, Context.Options);
 			}
 
 			return LowLevelTestsApp.TryReserveDevices();
