Add option to support override build dir.

diff --git a/Engine/Source/Programs/AutomationTool/BuildGraph/Tasks/ZenExportSnapshotTask.cs b/Engine/Source/Programs/AutomationTool/BuildGraph/Tasks/ZenExportSnapshotTask.cs
--- a/Engine/Source/Programs/AutomationTool/BuildGraph/Tasks/ZenExportSnapshotTask.cs
+++ b/Engine/Source/Programs/AutomationTool/BuildGraph/Tasks/ZenExportSnapshotTask.cs
@@ -110,6 +110,13 @@ namespace AutomationTool.Tasks
 		/// </summary>
 		[TaskParameter(Optional = true)]
 		public string DestinationFileName;
+
+		/// <summary>
+		/// Optional. Where to look for the .projectstore
+		/// The pattern {Platform} can be used for exporting multiple platforms at once.
+		/// </summary>
+		[TaskParameter(Optional = true)]
+		public string OverridePlatformCookedDir;
 	}
 
 	/// <summary>
@@ -336,7 +343,15 @@ namespace AutomationTool.Tasks
 			List<ExportSourceData> ExportSources = new List<ExportSourceData>();
 			foreach (string Platform in Parameters.Platform.Split('+'))
 			{
-				DirectoryReference PlatformCookedDirectory = DirectoryReference.Combine(ProjectFile.Directory, "Saved", "Cooked", Platform);
+				DirectoryReference PlatformCookedDirectory;
+				if (string.IsNullOrEmpty(Parameters.OverridePlatformCookedDir))
+				{
+					PlatformCookedDirectory = DirectoryReference.Combine(ProjectFile.Directory, "Saved", "Cooked", Platform);
+				}
+				else
+				{
+					PlatformCookedDirectory = new DirectoryReference(Parameters.OverridePlatformCookedDir.Replace("{Platform}", Platform, StringComparison.InvariantCultureIgnoreCase));
+				}
 				if (!DirectoryReference.Exists(PlatformCookedDirectory))
 				{
 					throw new AutomationException("Cook output directory not found ({0})", PlatformCookedDirectory.FullName);
