BuildGraph: Fix EnvVar lookup on non Windows platforms

Just like the Engine, swap - for _ as - is not valid on non Windows platform

#rnx

diff --git a/Engine/Source/Programs/AutomationTool/BuildGraph/BgScriptReader.cs b/Engine/Source/Programs/AutomationTool/BuildGraph/BgScriptReader.cs
--- a/Engine/Source/Programs/AutomationTool/BuildGraph/BgScriptReader.cs
+++ b/Engine/Source/Programs/AutomationTool/BuildGraph/BgScriptReader.cs
@@ -910,7 +910,14 @@ namespace AutomationTool
 				string name = ReadAttribute(element, "Name");
 				if (ValidateName(element, name))
 				{
-					string value = Environment.GetEnvironmentVariable(name) ?? "";
+					string envVarName = name;
+					if (!RuntimePlatform.IsWindows)
+					{
+						// Non-windows platforms don't allow dashes in variable names. The engine platform layer substitutes underscores for them.
+						envVarName = envVarName.Replace("-", "_");
+					}
+
+					string value = Environment.GetEnvironmentVariable(envVarName) ?? "";
 					SetPropertyValue(element, name, value);
 				}
 			}
