Fixing ZenLaunch spawning failing on non-Windows Platforms

no need to wrap in quotes

#rnx

diff --git a/Engine/Source/Programs/AutomationTool/BuildGraph/Tasks/ZenExportSnapshotTask.cs b/Engine/Source/Programs/AutomationTool/BuildGraph/Tasks/ZenExportSnapshotTask.cs
--- a/Engine/Source/Programs/AutomationTool/BuildGraph/Tasks/ZenExportSnapshotTask.cs
+++ b/Engine/Source/Programs/AutomationTool/BuildGraph/Tasks/ZenExportSnapshotTask.cs
@@ -234,7 +234,7 @@ namespace AutomationTool.Tasks
 			StringBuilder ZenLaunchCommandline = new StringBuilder();
 			ZenLaunchCommandline.AppendFormat("{0} -SponsorProcessID={1}", CommandUtils.MakePathSafeToUseWithCommandLine(ProjectFile.FullName), Environment.ProcessId);
 
-			CommandUtils.RunAndLog(CommandUtils.CmdEnv, CommandUtils.MakePathSafeToUseWithCommandLine(ZenLaunchExe.FullName), ZenLaunchCommandline.ToString(), Options: CommandUtils.ERunOptions.Default);
+			CommandUtils.RunAndLog(CommandUtils.CmdEnv, ZenLaunchExe.FullName, ZenLaunchCommandline.ToString(), Options: CommandUtils.ERunOptions.Default);
 		}
 		
 		static JsonSerializerOptions GetDefaultJsonSerializerOptions()
