Remove use of CommandUtils.MakePathSafeToUseWithCommandLine on the executable name when performing zen snapshot exports as it makes the execution fail on Mac.

diff --git a/Engine/Source/Programs/AutomationTool/BuildGraph/Tasks/ZenExportSnapshotTask.cs b/Engine/Source/Programs/AutomationTool/BuildGraph/Tasks/ZenExportSnapshotTask.cs
--- a/Engine/Source/Programs/AutomationTool/BuildGraph/Tasks/ZenExportSnapshotTask.cs
+++ b/Engine/Source/Programs/AutomationTool/BuildGraph/Tasks/ZenExportSnapshotTask.cs
@@ -464,7 +464,7 @@ namespace AutomationTool.Tasks
 								return null;
 							});
 						ExportSingleSourceCommandline.AppendFormat(" {0} --embedloosefiles --key {1} {2} {3} {4}", HostUrlArg, DestinationKeyHash.ToString().ToLowerInvariant(), BaseKeyArg, ExportSource.ProjectId, ExportSource.OplogId);
-						CommandUtils.RunAndLog(CommandUtils.CmdEnv, CommandUtils.MakePathSafeToUseWithCommandLine(ZenExe.FullName), ExportSingleSourceCommandline.ToString(), MaxSuccessCode: int.MaxValue, Options: CommandUtils.ERunOptions.Default, SpewFilterCallback: SilentOutputFilter);
+						CommandUtils.RunAndLog(CommandUtils.CmdEnv, ZenExe.FullName, ExportSingleSourceCommandline.ToString(), MaxSuccessCode: int.MaxValue, Options: CommandUtils.ERunOptions.Default, SpewFilterCallback: SilentOutputFilter);
 
 						ExportIndex = ExportIndex + 1;
 					}
@@ -511,7 +511,7 @@ namespace AutomationTool.Tasks
 						}
 						ExportSingleSourceCommandline.AppendFormat(" --file {0} --name {1} {2} {3} {4}", CommandUtils.MakePathSafeToUseWithCommandLine(PlatformDestinationFileDir.FullName), DestinationFileName, BaseNameArg, ProjectId, ExportSource.OplogId);
 
-						CommandUtils.RunAndLog(CommandUtils.CmdEnv, CommandUtils.MakePathSafeToUseWithCommandLine(ZenExe.FullName), ExportSingleSourceCommandline.ToString(), Options: CommandUtils.ERunOptions.Default);
+						CommandUtils.RunAndLog(CommandUtils.CmdEnv, ZenExe.FullName, ExportSingleSourceCommandline.ToString(), Options: CommandUtils.ERunOptions.Default);
 
 						ExportIndex = ExportIndex + 1;
 					}
