Implement FullClean method in TargetDevice for all platforms

diff --git a/Engine/Source/Programs/AutomationTool/Gauntlet/Platform/Android/Gauntlet.TargetDeviceAndroid.cs b/Engine/Source/Programs/AutomationTool/Gauntlet/Platform/Android/Gauntlet.TargetDeviceAndroid.cs
--- a/Engine/Source/Programs/AutomationTool/Gauntlet/Platform/Android/Gauntlet.TargetDeviceAndroid.cs
+++ b/Engine/Source/Programs/AutomationTool/Gauntlet/Platform/Android/Gauntlet.TargetDeviceAndroid.cs
@@ -1104,6 +1104,11 @@ namespace Gauntlet
 		public void CopyAppConfigurationFiles(UnrealAppConfig AppConfiguration)
 		{
 
+		}
+		
+		public void FullClean()
+		{
+
 		}
 
 		public IAppInstance Run(IAppInstall App)
diff --git a/Engine/Source/Programs/AutomationTool/Gauntlet/Platform/IOS/Gauntlet.TargetDeviceIOS.cs b/Engine/Source/Programs/AutomationTool/Gauntlet/Platform/IOS/Gauntlet.TargetDeviceIOS.cs
--- a/Engine/Source/Programs/AutomationTool/Gauntlet/Platform/IOS/Gauntlet.TargetDeviceIOS.cs
+++ b/Engine/Source/Programs/AutomationTool/Gauntlet/Platform/IOS/Gauntlet.TargetDeviceIOS.cs
@@ -593,6 +593,11 @@ namespace Gauntlet
 
 		}
 
+		public void FullClean(UnrealAppConfig AppConfiguration)
+		{
+
+		}
+		
 		public void PopulateDirectoryMappings(string ProjectDir)
 		{
 			LocalDirectoryMappings.Add(EIntendedBaseCopyDirectory.Build, Path.Combine(ProjectDir, "Build"));
