Add ability for gauntlet to parse BP script callstack from error messages matching regex \[ScriptStack\]\s*.+

diff --git a/Engine/Source/Programs/AutomationTool/Gauntlet/Unreal/Utils/Gauntlet.UnrealLogParser.cs b/Engine/Source/Programs/AutomationTool/Gauntlet/Unreal/Utils/Gauntlet.UnrealLogParser.cs
--- a/Engine/Source/Programs/AutomationTool/Gauntlet/Unreal/Utils/Gauntlet.UnrealLogParser.cs
+++ b/Engine/Source/Programs/AutomationTool/Gauntlet/Unreal/Utils/Gauntlet.UnrealLogParser.cs
@@ -838,9 +838,21 @@ namespace Gauntlet
 
 							LinesWithoutBacktrace = 0;
 						}
-						else
+						else 
 						{
-							LinesWithoutBacktrace++;
+							// Callstack may be prefixed (or interspersed?) with callstack from blueprint script VM 
+							Match ScriptCSMatch = Regex.Match(Line, @"\[ScriptStack\]\s*(.+)$", RegexOptions.IgnoreCase);
+
+							if (ScriptCSMatch.Success)
+							{
+								string Stack = ScriptCSMatch.Groups[1];	
+								Backtrace.Add(Stack);
+								LinesWithoutBacktrace = 0;
+							}
+							else
+							{
+								LinesWithoutBacktrace++;
+							}
 						}
 
 						SearchContent = SearchContent.Substring(EOL + 1);
