- Fixed various blocking issues with making Mac/IOS builds that can be uploaded to AppStore in Xcode
  * Missing ThinApp.sh
  * Missing some quotes around a variable in GenerateUniversalDSYM.sh
  * Correctly filter out the stub .app's and such when copying the Staged data directory into the .app
  * Disabled the Mac's Sign To Run Locally option by default, as that causes Xcode's Validation/Upload to not have the team name embedded in it, causing a hassle while pushing up to AppStore/TestFlight
  * Fixed the PRODUCT_NAME for BP projects
  * Made CrashReportClient be a sandboxed app that inherits from parent
  * Fix Hybrid apps to check all platforms before in the project generator, so that if a project needs a temp target for IOS (via a plugin)
  * Disabled the LoginFlow module from OnlineFramework plugin - this was causing issues with having CEF and EpicWebHelper embedded into a sandboxed .app, and LoginFlow isn't seemingly actually used by anything
  * Set the no-encryption setting in the template Mac plist to make uploading to TestFlight/AppStore easier
  * Fixed editor to not compile as universal when using -distribution, only the client should be universal
  * Further improvements to stripping out nested .app's in a final .app (remove them from the staged directory helps)
  * Changed how the Mac app name is displayed, since the .app name itself is shown in the Finder,  unlike IOS where the CFBundleName is shown (the archived .app name in the .xcarchive is named by a project setting, falling back to the .uproject name if not set)
  * Disabled the SignExecutables function on Modern because they attempt to sign the wrong .apps, and one is no longer (previously it was uselessly signing .apps, but now it throws an error due to changes for the third item in this list)
  * Legacy Xcode mode is forced with Window's Remote Mac builds, so use the old legacy code when running under windows
  * Added a third party .cpp file to get copied to remote macs
  * Fix Tools -> Open Xcode not working with modern Xcode

diff --git a/Engine/Build/InstalledEngineFilters.xml b/Engine/Build/InstalledEngineFilters.xml
--- a/Engine/Build/InstalledEngineFilters.xml
+++ b/Engine/Build/InstalledEngineFilters.xml
@@ -339,6 +339,7 @@
 		Engine/Build/BatchFiles/Mac/SetupDotnet.sh
 		Engine/Build/BatchFiles/Mac/XcodeBuild.sh
 		Engine/Build/BatchFiles/Mac/GenerateUniversalDSYM.sh
+		Engine/Build/BatchFiles/Mac/ThinApp.sh
 		Engine/Build/BatchFiles/Mac/UpdateVersionAfterBuild.sh
 		Engine/Build/BatchFiles/Mac/ThinApp.sh
 		Engine/Build/BatchFiles/RunDotnet.sh
diff --git a/Engine/Source/Programs/AutomationTool/IOS/IOSPlatform.Automation.cs b/Engine/Source/Programs/AutomationTool/IOS/IOSPlatform.Automation.cs
--- a/Engine/Source/Programs/AutomationTool/IOS/IOSPlatform.Automation.cs
+++ b/Engine/Source/Programs/AutomationTool/IOS/IOSPlatform.Automation.cs
@@ -2000,15 +2000,25 @@ public class IOSPlatform : ApplePlatform
 			// Params.StageDirectory will be something like /.../Saved/StagedBuilds
 			// ClientApp will be something like /.../Saved/StagedBuilds/IOSClient/QAGame/Binaries/IOS/QAGameClient
 			// so we want the first directory in the CLientApp after Params.StageDirectory	
-			string TargetStageDir = ClientApp.Substring(0, ClientApp.IndexOf('/', Params.BaseStageDirectory.Length + 1));
-			if (AppleExports.UseModernXcode(Params.RawProjectPath))
+
+			// Currently (2023/09), Remote Mac builds from a Windows System are forced to Legacy mode and uses a different ClientApp value
+			if (OperatingSystem.IsMacOS())
 			{
-				// modern mode runs from the staged dir
-				PlistFile = Path.Combine(TargetStageDir, Path.GetFileNameWithoutExtension(ClientApp) + ".app", "Info.plist");
+				string TargetStageDir = ClientApp.Substring(0, ClientApp.IndexOf('/', Params.BaseStageDirectory.Length + 1));
+				if (AppleExports.UseModernXcode(Params.RawProjectPath))
+				{
+					// modern mode runs from the staged dir
+					PlistFile = Path.Combine(TargetStageDir, Path.GetFileNameWithoutExtension(ClientApp) + ".app", "Info.plist");
+				}
+				else
+				{
+					PlistFile = Path.Combine(TargetStageDir, "Info.plist");
+				}
 			}
 			else
 			{
-				PlistFile = Path.Combine(TargetStageDir, "Info.plist");
+				PlistFile = Path.Combine(Params.BaseStageDirectory, PlatformName);
+				PlistFile = Path.Combine(PlistFile, "Info.plist");
 			}
 
 			Logger.LogWarning("LOOKING IN PLIST {PListfile}", PlistFile);
@@ -2029,7 +2039,12 @@ public class IOSPlatform : ApplePlatform
 			Arguments += " --detach";
 			Arguments = GetLibimobileDeviceNetworkedArgument(Arguments, Params.DeviceNames[0]);
 			Arguments += " run '" + BundleIdentifier + "'";
-			Arguments += " " + ClientCmdLine;
+			
+			// ClientCmdLine is only relevant when running on a Mac
+			if (OperatingSystem.IsMacOS())
+			{
+				Arguments += " " + ClientCmdLine;
+			}
 
 			IProcessResult ClientProcess = Run(Program, Arguments, null, ClientRunFlags);
 			if (ClientProcess.ExitCode == -1)
