Gauntlet device configuration profiles

diff --git a/Engine/Source/Programs/AutomationTool/Gauntlet/Framework/Base/Gauntlet.DeviceConfigurationProfile.cs b/Engine/Source/Programs/AutomationTool/Gauntlet/Framework/Base/Gauntlet.DeviceConfigurationProfile.cs
new file mode 100644
--- /dev/null
+++ b/Engine/Source/Programs/AutomationTool/Gauntlet/Framework/Base/Gauntlet.DeviceConfigurationProfile.cs
@@ -0,0 +1,231 @@
+using System;
+using System.Collections.Generic;
+using UnrealBuildTool;
+using EpicGames.Core;
+using System.Text.Json;
+using System.Text.Json.Serialization;
+using System.Linq;
+using System.Text;
+using System.Threading.Tasks;
+using System.IO;
+
+namespace Gauntlet
+{
+	enum EConfigurationNamespace
+	{
+		Engine,
+		Snapshot,
+		Project
+	}
+
+	/// <summary>
+	/// Base class for a platform configuration profile.
+	/// This class needs to be implemented for each configurable platform
+	/// </summary>
+	public abstract class PlatformConfigurationBase
+	{
+		/// <summary>
+		/// Name of the config profile.
+		/// This name doesn't need to be unique across namespaces.
+		/// </summary>
+		public string ProfileName { get; set; }
+
+		/// <summary>
+		/// Namespace of the profile. Profiles have to have a unique name inside a namespace.
+		/// </summary>
+		public string Namespace { get; set; }
+
+		/// <summary>
+		/// Platform this configuration profile is for
+		/// </summary>
+		public UnrealTargetPlatform Platform { get; set; }
+	}
+
+	/// <summary>
+	/// Base interface for a platform specific configuration reader
+	/// </summary>
+	public interface IPlatformConfigurationReader
+	{
+		bool SupportsPlatform(UnrealTargetPlatform? Platform);
+
+		/// <summary>
+		/// Returns the supported platform config file extension
+		/// </summary>
+		/// <returns></returns>
+		string ConfigFileExtension();
+
+		/// <summary>
+		/// Reads the configuration from the passed in file location
+		/// </summary>
+		/// <param name="Location">An absolute path the to a file with platform configuration</param>
+		/// <returns></returns>
+		PlatformConfigurationBase ReadConfiguration(FileReference Location);
+	}
+
+	/// <summary>
+	/// A singleton class that encapsulates a cache of device config profiles 
+	/// </summary>
+	public class DeviceConfigurationCache
+	{
+
+		/// <summary>
+		/// Key by which a configuration profile is identified in the cache.
+		/// </summary>
+		class ConfigurationCacheKey
+		{
+			public UnrealTargetPlatform Platform;
+			public string Namespace;
+			public string ProfileName;
+
+			public ConfigurationCacheKey(UnrealTargetPlatform Platform, string ProjectName, string ProfileName)
+			{
+				this.Platform = Platform;
+				this.Namespace = ProjectName;
+				this.ProfileName = ProfileName;
+			}
+
+			public ConfigurationCacheKey(PlatformConfigurationBase Configuration)
+			{
+				this.Platform = Configuration.Platform;
+				this.Namespace = Configuration.Namespace;
+				this.ProfileName = Configuration.ProfileName;
+			}
+
+			public override bool Equals(object Other)
+			{
+				ConfigurationCacheKey OtherKey = Other as ConfigurationCacheKey;
+				return OtherKey != null &&
+					OtherKey.Platform == Platform &&
+					OtherKey.ProfileName == ProfileName &&
+					OtherKey.Namespace == Namespace;
+
+			}
+
+			public override int GetHashCode()
+			{
+				return (Platform, Namespace, ProfileName).GetHashCode();
+			}
+		}
+
+		Object LockObject = new Object();
+
+		Dictionary<ConfigurationCacheKey, PlatformConfigurationBase> ConfigurationCache = new();
+
+		private static DeviceConfigurationCache _Instance;
+
+		protected DeviceConfigurationCache()
+		{
+			if (_Instance == null)
+			{
+				_Instance = this;
+			}
+		}
+
+		public static DeviceConfigurationCache Instance
+		{
+			get
+			{
+				if (_Instance == null)
+				{
+					_Instance = new DeviceConfigurationCache();
+				}
+
+				return _Instance;
+			}
+		}
+
+		/// <summary>
+		/// Scans the SettingDir for configuration profiles
+		/// </summary>
+		/// <param name="Platform">Platform for which to look for config files</param>
+		/// <param name="Namespace">The namespace in which to put the configuration profiles</param>
+		/// <param name="SettingDir"></param>
+		/// <returns></returns>
+		public bool DiscoverConfigurationProfiles(UnrealTargetPlatform? Platform, string Namespace, string SettingDir)
+		{
+			if (Platform == null || !Directory.Exists(SettingDir))
+			{
+				return false;
+			}
+
+			IPlatformConfigurationReader ConfigReader = Gauntlet.Utils.InterfaceHelpers.FindImplementations<IPlatformConfigurationReader>(true).ToList().Find(D => D.SupportsPlatform(Platform));
+			if (ConfigReader == null)
+			{
+				Log.Info("Couldn't find a configuration reader for {0}", Platform);
+				return false;
+			}
+
+			foreach (string File in Directory.EnumerateFiles(SettingDir, ConfigReader.ConfigFileExtension(), SearchOption.AllDirectories))
+			{
+				FileReference FileRef = new FileReference(File);
+				PlatformConfigurationBase DeviceProfile = ConfigReader.ReadConfiguration(FileRef);
+				if (DeviceProfile != null)
+				{
+					DeviceProfile.Platform = Platform.Value;
+					DeviceProfile.Namespace = Namespace;
+					DeviceProfile.ProfileName = FileRef.GetFileNameWithoutAnyExtensions();
+
+					lock (LockObject)
+					{
+						ConfigurationCacheKey Key = new ConfigurationCacheKey(Platform.Value, Namespace, DeviceProfile.ProfileName);
+						if (ConfigurationCache.ContainsKey(Key))
+						{
+							Log.Verbose("Device configuration profile {0} already exists", FileRef.FullName);
+							continue;
+						}
+
+						ConfigurationCache.Add(Key, DeviceProfile);
+					}
+				}
+			}
+
+			return true;
+		}
+
+		public void CacheConfigurationSnapshot(PlatformConfigurationBase Configuration, bool Overwrite = false)
+		{
+			lock (LockObject)
+			{
+				ConfigurationCacheKey Key = new ConfigurationCacheKey(Configuration.Platform, "Snapshot", Configuration.ProfileName);
+				if (ConfigurationCache.ContainsKey(Key) && !Overwrite)
+				{
+					return;
+				}
+
+				ConfigurationCache.Add(Key, Configuration);
+			}
+		}
+
+		public PlatformConfigurationBase GetConfigurationSnapshot(UnrealTargetPlatform? Platform, string DeviceName)
+		{
+			return GetConfiguration(Platform, "Snapshot", DeviceName);
+		}
+
+		public void ClearSnapshot(PlatformConfigurationBase Snapshot)
+		{
+			lock (LockObject)
+			{
+				ConfigurationCacheKey Key = new ConfigurationCacheKey(Snapshot);
+				ConfigurationCache.Remove(Key);
+			}
+		}
+
+		public PlatformConfigurationBase GetConfiguration(UnrealTargetPlatform? Platform, string Namespace, string ProfileName)
+		{
+			if (Platform == null)
+			{
+				return null;
+			}
+
+			PlatformConfigurationBase Res = null;
+
+			lock (LockObject)
+			{
+				ConfigurationCacheKey Key = new ConfigurationCacheKey(Platform.Value, Namespace, ProfileName);
+				ConfigurationCache.TryGetValue(Key, out Res);
+			}
+
+			return Res;
+		}
+	}
+}
diff --git a/Engine/Source/Programs/AutomationTool/Gauntlet/Framework/Base/Gauntlet.TargetDevice.cs b/Engine/Source/Programs/AutomationTool/Gauntlet/Framework/Base/Gauntlet.TargetDevice.cs
--- a/Engine/Source/Programs/AutomationTool/Gauntlet/Framework/Base/Gauntlet.TargetDevice.cs
+++ b/Engine/Source/Programs/AutomationTool/Gauntlet/Framework/Base/Gauntlet.TargetDevice.cs
@@ -113,7 +113,7 @@ namespace Gauntlet
 		/// Returns true if there were any crash dumps for the run, and false otherwise
 		/// </summary>
 		bool CopyCrashDumps() { return false; }
-		
+
 	};
 
 	/// <summary>
@@ -209,4 +209,23 @@ namespace Gauntlet
 
 		public virtual bool NeedBuildDeployed() => true;
 	}
+
+	/// <summary>
+	/// Represents a device that is configurable
+	/// </summary>
+	public interface IConfigurableDevice
+	{
+		/// <summary>
+		/// Returns a configuration profile with the current device settings
+		/// </summary>
+		/// <returns></returns>
+		PlatformConfigurationBase GetCurrentConfigurationSnapshot();
+
+		/// <summary>
+		/// Applies a configuration profile to the device. If required it will reboot the device
+		/// </summary>
+		/// <param name="Configuration"></param>
+		/// <returns>false if applying the config profile fails</returns>
+		bool ApplyConfiguration(PlatformConfigurationBase Configuration);
+	}
 }
\ No newline at end of file
diff --git a/Engine/Source/Programs/AutomationTool/Gauntlet/Framework/Devices/Gauntlet.DevicePool.cs b/Engine/Source/Programs/AutomationTool/Gauntlet/Framework/Devices/Gauntlet.DevicePool.cs
--- a/Engine/Source/Programs/AutomationTool/Gauntlet/Framework/Devices/Gauntlet.DevicePool.cs
+++ b/Engine/Source/Programs/AutomationTool/Gauntlet/Framework/Devices/Gauntlet.DevicePool.cs
@@ -564,6 +564,38 @@ namespace Gauntlet
 			}
 		}
 
+		void RevertDeviceConfiguration(ITargetDevice Device)
+		{
+			IConfigurableDevice ConfigurableDevice = Device as IConfigurableDevice;
+			if (ConfigurableDevice != null)
+			{
+				var Snapshot = DeviceConfigurationCache.Instance.GetConfigurationSnapshot(Device.Platform, Device.Name);
+				if (Snapshot == null)
+				{
+					return;
+				}
+
+				// Connect temporarily to be able to revert the device's configuration
+				// if the device was disconnected entering here disconnect it after this is over
+				bool bNeedsDisconnect = false;
+				if (!Device.IsConnected)
+				{
+					Device.Connect();
+					bNeedsDisconnect = true;
+				}
+				
+				if (ConfigurableDevice.ApplyConfiguration(Snapshot))
+				{
+					DeviceConfigurationCache.Instance.ClearSnapshot(Snapshot);
+				}
+
+				if (bNeedsDisconnect)
+				{
+					Device.Disconnect();
+				}
+			}
+		}
+
 		/// <summary>
 		/// Explicitly release all device reservations
 		/// </summary>
@@ -603,6 +635,8 @@ namespace Gauntlet
 				List<ITargetDevice> KeepList = new List<ITargetDevice>();
 				foreach (ITargetDevice Device in DeviceList)
 				{
+					RevertDeviceConfiguration(Device);
+
 					if (Device.IsConnected && InitialConnectState[Device] == false)
 					{
 						Device.Disconnect();
diff --git a/Engine/Source/Programs/AutomationTool/Gauntlet/Unreal/Base/Gauntlet.UnrealSession.cs b/Engine/Source/Programs/AutomationTool/Gauntlet/Unreal/Base/Gauntlet.UnrealSession.cs
--- a/Engine/Source/Programs/AutomationTool/Gauntlet/Unreal/Base/Gauntlet.UnrealSession.cs
+++ b/Engine/Source/Programs/AutomationTool/Gauntlet/Unreal/Base/Gauntlet.UnrealSession.cs
@@ -4,6 +4,7 @@ using AutomationTool;
 using EpicGames.Core;
 using System;
 using System.Collections.Generic;
+using System.Configuration;
 using System.IO;
 using System.Linq;
 using System.Text.RegularExpressions;
@@ -690,6 +691,28 @@ namespace Gauntlet
 			}
 		}
 
+		public bool TryConfigureDevice(ITargetDevice Device, string ProjectName, string ProfileName)
+		{
+			IConfigurableDevice ConfigurableDevice = Device as IConfigurableDevice;
+			if (ConfigurableDevice == null)
+			{
+				Log.Info("Trying to set a config profile on a device that isn't configurable: {0}", Device.Platform);
+				return false;
+			}
+
+			var Configuration = DeviceConfigurationCache.Instance.GetConfiguration(Device.Platform, ProjectName, ProfileName);
+			if (Configuration == null)
+			{
+				Log.Info("Device profile {0}.{1} for {2} doesn't exist", ProjectName, ProfileName, Device.Platform);
+				return false;
+			}
+
+			var ConfigSnapshot = ConfigurableDevice.GetCurrentConfigurationSnapshot();
+			DeviceConfigurationCache.Instance.CacheConfigurationSnapshot(ConfigSnapshot);
+			
+			return ConfigurableDevice.ApplyConfiguration(Configuration);
+		}
+
 		/// <summary>
 		/// Check that all the current roles can be performed by our build source
 		/// </summary>
@@ -791,8 +814,26 @@ namespace Gauntlet
 					// create a config from the build source (this also applies the role options)
 					UnrealAppConfig AppConfig = BuildSource.CreateConfiguration(Role, OtherRoles);
 
+					// device profile comes in the format <namespace>.<profilename>. If no namespace is given, assume Engine
+					string ProfileName = Globals.Params.ParseValue("deviceprofile", "");
+					string Namespace = "Engine";
+					int DotIndex = ProfileName.IndexOf('.');
+					if (DotIndex != -1)
+					{
+						Namespace = ProfileName.Substring(0, DotIndex);
+						ProfileName = ProfileName.Substring(DotIndex + 1);
+					}
+					
+					if (!string.IsNullOrWhiteSpace(ProfileName))
+					{
+						if (!TryConfigureDevice(Device, Namespace, ProfileName))
+						{
+							Log.Info("Failed to apply {0} configuration profile to device {1}", ProfileName, Device.Name);
+						}
+					}
+
 					//Verify the device's OS version, and update if necessary
-					if(Globals.Params.ParseParam("TryFirmwareUpdate") && AppConfig.Platform != null)
+					if (Globals.Params.ParseParam("TryFirmwareUpdate") && AppConfig.Platform != null)
 					{
 						List<IPlatformFirmwareHandler> PlatformFirmwareHandlers = Gauntlet.Utils.InterfaceHelpers.FindImplementations<IPlatformFirmwareHandler>(true).ToList();
 
diff --git a/Engine/Source/Programs/AutomationTool/Gauntlet/Unreal/RunUnreal.cs b/Engine/Source/Programs/AutomationTool/Gauntlet/Unreal/RunUnreal.cs
--- a/Engine/Source/Programs/AutomationTool/Gauntlet/Unreal/RunUnreal.cs
+++ b/Engine/Source/Programs/AutomationTool/Gauntlet/Unreal/RunUnreal.cs
@@ -180,6 +180,8 @@ namespace Gauntlet
 
 				UnrealTargetPlatform PlatformType = UnrealTargetPlatform.Parse(PlatformString);
 
+				SetupPlatformConfigurationProfiles(PlatformType, ContextOptions);
+
 				if (!InitializedDevices)
 				{
 					// Setup the devices and assign them to the executor
@@ -450,13 +452,15 @@ namespace Gauntlet
 
 		protected void SetupDevices(UnrealTargetPlatform DefaultPlatform, UnrealTestOptions Options)
 		{
-
 			Reservation.ReservationDetails = Options.JobDetails;
 
 			DevicePool.Instance.SetLocalOptions(Options.TempDir, Options.Parallel > 1, Options.DeviceURL);
 			DevicePool.Instance.AddLocalDevices(Options.MaxLocalDevices);
 			DevicePool.Instance.AddVirtualDevices(Options.MaxVirtualDevices);
 
+			string EngineDeviceConfigDir = Path.Combine(Globals.UnrealRootDir, "Engine", "Build", "DeviceConfigProfiles");
+			string ProjectDeviceConfigDir = Path.Combine(Options.ProjectPath.Directory.FullName, "Build", "DeviceConfigProfiles");
+
 			foreach (var DeviceWithParams in Options.DeviceList)
 			{
 				UnrealTargetPlatform Platform = DefaultPlatform;
@@ -473,5 +477,14 @@ namespace Gauntlet
 				DevicePool.Instance.AddDevices(Platform, DeviceWithParams.Argument);
 			}
 		}
+
+		protected void SetupPlatformConfigurationProfiles(UnrealTargetPlatform PlatformType, UnrealTestOptions Options)
+		{
+			string EngineDeviceConfigDir = Path.Combine(Globals.UnrealRootDir, "Engine", "Build", "DeviceConfigProfiles");
+			string ProjectDeviceConfigDir = Path.Combine(Options.ProjectPath.Directory.FullName, "Build", "DeviceConfigProfiles");
+
+			DeviceConfigurationCache.Instance.DiscoverConfigurationProfiles(PlatformType, "Engine", EngineDeviceConfigDir);
+			DeviceConfigurationCache.Instance.DiscoverConfigurationProfiles(PlatformType, Options.Project, ProjectDeviceConfigDir);
+		}
 	}
 }
