[Backout] - CL28792672
Original CL Desc
-----------------------------------------------------------------
Gauntlet device configuration profiles

diff --git a/Engine/Source/Programs/AutomationTool/Gauntlet/Framework/Base/Gauntlet.DeviceConfigurationProfile.cs b/Engine/Source/Programs/AutomationTool/Gauntlet/Framework/Base/Gauntlet.DeviceConfigurationProfile.cs
deleted file mode 100644
diff --git a/Engine/Source/Programs/AutomationTool/Gauntlet/Framework/Base/Gauntlet.TargetDevice.cs b/Engine/Source/Programs/AutomationTool/Gauntlet/Framework/Base/Gauntlet.TargetDevice.cs
--- a/Engine/Source/Programs/AutomationTool/Gauntlet/Framework/Base/Gauntlet.TargetDevice.cs
+++ b/Engine/Source/Programs/AutomationTool/Gauntlet/Framework/Base/Gauntlet.TargetDevice.cs
@@ -113,7 +113,7 @@ namespace Gauntlet
 		/// Returns true if there were any crash dumps for the run, and false otherwise
 		/// </summary>
 		bool CopyCrashDumps() { return false; }
-
+		
 	};
 
 	/// <summary>
@@ -209,23 +209,4 @@ namespace Gauntlet
 
 		public virtual bool NeedBuildDeployed() => true;
 	}
-
-	/// <summary>
-	/// Represents a device that is configurable
-	/// </summary>
-	public interface IConfigurableDevice
-	{
-		/// <summary>
-		/// Returns a configuration profile with the current device settings
-		/// </summary>
-		/// <returns></returns>
-		PlatformConfigurationBase GetCurrentConfigurationSnapshot();
-
-		/// <summary>
-		/// Applies a configuration profile to the device. If required it will reboot the device
-		/// </summary>
-		/// <param name="Configuration"></param>
-		/// <returns>false if applying the config profile fails</returns>
-		bool ApplyConfiguration(PlatformConfigurationBase Configuration);
-	}
 }
\ No newline at end of file
diff --git a/Engine/Source/Programs/AutomationTool/Gauntlet/Framework/Devices/Gauntlet.DevicePool.cs b/Engine/Source/Programs/AutomationTool/Gauntlet/Framework/Devices/Gauntlet.DevicePool.cs
--- a/Engine/Source/Programs/AutomationTool/Gauntlet/Framework/Devices/Gauntlet.DevicePool.cs
+++ b/Engine/Source/Programs/AutomationTool/Gauntlet/Framework/Devices/Gauntlet.DevicePool.cs
@@ -564,38 +564,6 @@ namespace Gauntlet
 			}
 		}
 
-		void RevertDeviceConfiguration(ITargetDevice Device)
-		{
-			IConfigurableDevice ConfigurableDevice = Device as IConfigurableDevice;
-			if (ConfigurableDevice != null)
-			{
-				var Snapshot = DeviceConfigurationCache.Instance.GetConfigurationSnapshot(Device.Platform, Device.Name);
-				if (Snapshot == null)
-				{
-					return;
-				}
-
-				// Connect temporarily to be able to revert the device's configuration
-				// if the device was disconnected entering here disconnect it after this is over
-				bool bNeedsDisconnect = false;
-				if (!Device.IsConnected)
-				{
-					Device.Connect();
-					bNeedsDisconnect = true;
-				}
-				
-				if (ConfigurableDevice.ApplyConfiguration(Snapshot))
-				{
-					DeviceConfigurationCache.Instance.ClearSnapshot(Snapshot);
-				}
-
-				if (bNeedsDisconnect)
-				{
-					Device.Disconnect();
-				}
-			}
-		}
-
 		/// <summary>
 		/// Explicitly release all device reservations
 		/// </summary>
@@ -635,8 +603,6 @@ namespace Gauntlet
 				List<ITargetDevice> KeepList = new List<ITargetDevice>();
 				foreach (ITargetDevice Device in DeviceList)
 				{
-					RevertDeviceConfiguration(Device);
-
 					if (Device.IsConnected && InitialConnectState[Device] == false)
 					{
 						Device.Disconnect();
diff --git a/Engine/Source/Programs/AutomationTool/Gauntlet/Unreal/Base/Gauntlet.UnrealSession.cs b/Engine/Source/Programs/AutomationTool/Gauntlet/Unreal/Base/Gauntlet.UnrealSession.cs
--- a/Engine/Source/Programs/AutomationTool/Gauntlet/Unreal/Base/Gauntlet.UnrealSession.cs
+++ b/Engine/Source/Programs/AutomationTool/Gauntlet/Unreal/Base/Gauntlet.UnrealSession.cs
@@ -4,7 +4,6 @@ using AutomationTool;
 using EpicGames.Core;
 using System;
 using System.Collections.Generic;
-using System.Configuration;
 using System.IO;
 using System.Linq;
 using System.Text.RegularExpressions;
@@ -691,28 +690,6 @@ namespace Gauntlet
 			}
 		}
 
-		public bool TryConfigureDevice(ITargetDevice Device, string ProjectName, string ProfileName)
-		{
-			IConfigurableDevice ConfigurableDevice = Device as IConfigurableDevice;
-			if (ConfigurableDevice == null)
-			{
-				Log.Info("Trying to set a config profile on a device that isn't configurable: {0}", Device.Platform);
-				return false;
-			}
-
-			var Configuration = DeviceConfigurationCache.Instance.GetConfiguration(Device.Platform, ProjectName, ProfileName);
-			if (Configuration == null)
-			{
-				Log.Info("Device profile {0}.{1} for {2} doesn't exist", ProjectName, ProfileName, Device.Platform);
-				return false;
-			}
-
-			var ConfigSnapshot = ConfigurableDevice.GetCurrentConfigurationSnapshot();
-			DeviceConfigurationCache.Instance.CacheConfigurationSnapshot(ConfigSnapshot);
-			
-			return ConfigurableDevice.ApplyConfiguration(Configuration);
-		}
-
 		/// <summary>
 		/// Check that all the current roles can be performed by our build source
 		/// </summary>
@@ -814,26 +791,8 @@ namespace Gauntlet
 					// create a config from the build source (this also applies the role options)
 					UnrealAppConfig AppConfig = BuildSource.CreateConfiguration(Role, OtherRoles);
 
-					// device profile comes in the format <namespace>.<profilename>. If no namespace is given, assume Engine
-					string ProfileName = Globals.Params.ParseValue("deviceprofile", "");
-					string Namespace = "Engine";
-					int DotIndex = ProfileName.IndexOf('.');
-					if (DotIndex != -1)
-					{
-						Namespace = ProfileName.Substring(0, DotIndex);
-						ProfileName = ProfileName.Substring(DotIndex + 1);
-					}
-					
-					if (!string.IsNullOrWhiteSpace(ProfileName))
-					{
-						if (!TryConfigureDevice(Device, Namespace, ProfileName))
-						{
-							Log.Info("Failed to apply {0} configuration profile to device {1}", ProfileName, Device.Name);
-						}
-					}
-
 					//Verify the device's OS version, and update if necessary
-					if (Globals.Params.ParseParam("TryFirmwareUpdate") && AppConfig.Platform != null)
+					if(Globals.Params.ParseParam("TryFirmwareUpdate") && AppConfig.Platform != null)
 					{
 						List<IPlatformFirmwareHandler> PlatformFirmwareHandlers = Gauntlet.Utils.InterfaceHelpers.FindImplementations<IPlatformFirmwareHandler>(true).ToList();
 
diff --git a/Engine/Source/Programs/AutomationTool/Gauntlet/Unreal/RunUnreal.cs b/Engine/Source/Programs/AutomationTool/Gauntlet/Unreal/RunUnreal.cs
--- a/Engine/Source/Programs/AutomationTool/Gauntlet/Unreal/RunUnreal.cs
+++ b/Engine/Source/Programs/AutomationTool/Gauntlet/Unreal/RunUnreal.cs
@@ -180,8 +180,6 @@ namespace Gauntlet
 
 				UnrealTargetPlatform PlatformType = UnrealTargetPlatform.Parse(PlatformString);
 
-				SetupPlatformConfigurationProfiles(PlatformType, ContextOptions);
-
 				if (!InitializedDevices)
 				{
 					// Setup the devices and assign them to the executor
@@ -452,15 +450,13 @@ namespace Gauntlet
 
 		protected void SetupDevices(UnrealTargetPlatform DefaultPlatform, UnrealTestOptions Options)
 		{
+
 			Reservation.ReservationDetails = Options.JobDetails;
 
 			DevicePool.Instance.SetLocalOptions(Options.TempDir, Options.Parallel > 1, Options.DeviceURL);
 			DevicePool.Instance.AddLocalDevices(Options.MaxLocalDevices);
 			DevicePool.Instance.AddVirtualDevices(Options.MaxVirtualDevices);
 
-			string EngineDeviceConfigDir = Path.Combine(Globals.UnrealRootDir, "Engine", "Build", "DeviceConfigProfiles");
-			string ProjectDeviceConfigDir = Path.Combine(Options.ProjectPath.Directory.FullName, "Build", "DeviceConfigProfiles");
-
 			foreach (var DeviceWithParams in Options.DeviceList)
 			{
 				UnrealTargetPlatform Platform = DefaultPlatform;
@@ -477,14 +473,5 @@ namespace Gauntlet
 				DevicePool.Instance.AddDevices(Platform, DeviceWithParams.Argument);
 			}
 		}
-
-		protected void SetupPlatformConfigurationProfiles(UnrealTargetPlatform PlatformType, UnrealTestOptions Options)
-		{
-			string EngineDeviceConfigDir = Path.Combine(Globals.UnrealRootDir, "Engine", "Build", "DeviceConfigProfiles");
-			string ProjectDeviceConfigDir = Path.Combine(Options.ProjectPath.Directory.FullName, "Build", "DeviceConfigProfiles");
-
-			DeviceConfigurationCache.Instance.DiscoverConfigurationProfiles(PlatformType, "Engine", EngineDeviceConfigDir);
-			DeviceConfigurationCache.Instance.DiscoverConfigurationProfiles(PlatformType, Options.Project, ProjectDeviceConfigDir);
-		}
 	}
 }
