[Gauntlet] Adds MSTest unit testing project to gauntlet for self testing

diff --git a/Engine/Source/Programs/AutomationTool/Gauntlet.Tests/Gauntlet.Tests.Automation.csproj b/Engine/Source/Programs/AutomationTool/Gauntlet.Tests/Gauntlet.Tests.Automation.csproj
new file mode 100644
--- /dev/null
+++ b/Engine/Source/Programs/AutomationTool/Gauntlet.Tests/Gauntlet.Tests.Automation.csproj
@@ -0,0 +1,46 @@
+<Project Sdk="Microsoft.NET.Sdk">
+
+  <PropertyGroup Label="Globals">
+    <SccProjectName />
+    <SccProvider />
+    <SccAuxPath />
+    <SccLocalPath />
+  </PropertyGroup>
+
+  <PropertyGroup>
+    <TargetFramework>net6.0</TargetFramework>
+    <ImplicitUsings>enable</ImplicitUsings>
+    <Nullable>enable</Nullable>
+    <Configurations>Debug;Release;Development</Configurations>
+    <IsPackable>false</IsPackable>
+  </PropertyGroup>
+
+  <PropertyGroup Condition=" '$(Configuration)|$(Platform)' == 'Development|AnyCPU' ">
+    <DefineConstants>$(DefineConstants);TRACE</DefineConstants>
+    <Optimize>true</Optimize>
+  </PropertyGroup>
+  <PropertyGroup Condition=" '$(Configuration)|$(Platform)' == 'Release|AnyCPU' ">
+    <DefineConstants>$(DefineConstants);TRACE</DefineConstants>
+    <Optimize>true</Optimize>
+  </PropertyGroup>
+  <PropertyGroup Condition=" '$(Configuration)|$(Platform)' == 'Debug|AnyCPU' ">
+    <DefineConstants>$(DefineConstants);DEBUG;TRACE</DefineConstants>
+    <DebugSymbols>true</DebugSymbols>
+    <DebugType>full</DebugType>
+  </PropertyGroup>
+
+  <ItemGroup>
+    <PackageReference Include="Microsoft.NET.Test.Sdk" Version="16.5.0" />
+    <PackageReference Include="MSTest.TestAdapter" Version="2.1.0" />
+    <PackageReference Include="MSTest.TestFramework" Version="2.1.0" />
+    <PackageReference Include="coverlet.collector" Version="1.2.0" />
+  </ItemGroup>
+
+  <ItemGroup>
+    <ProjectReference Include="..\..\UnrealBuildTool\UnrealBuildTool.csproj" />
+    <ProjectReference Include="..\AutomationUtils\AutomationUtils.Automation.csproj" />
+    <ProjectReference Include="..\Gauntlet\Gauntlet.Automation.csproj" />
+    <ProjectReference Include="..\Scripts\AutomationScripts.Automation.csproj" />
+  </ItemGroup>
+
+</Project>
diff --git a/Engine/Source/Programs/AutomationTool/Gauntlet.Tests/Gauntlet.Tests.cs b/Engine/Source/Programs/AutomationTool/Gauntlet.Tests/Gauntlet.Tests.cs
new file mode 100644
--- /dev/null
+++ b/Engine/Source/Programs/AutomationTool/Gauntlet.Tests/Gauntlet.Tests.cs
@@ -0,0 +1,14 @@
+global using Microsoft.VisualStudio.TestTools.UnitTesting;
+
+namespace Gauntlet.Tests
+{
+	[TestClass]
+	public class TestFramework
+	{
+		[TestMethod]
+		public void TestAssert()
+		{
+			Assert.IsTrue(1 + 1 == 2);
+		}
+	}
+}
\ No newline at end of file
