Updating the NotarizeTask to now use the new "notarytool" as Apple is deprecating "altool --notarize-app".

diff --git a/Engine/Source/Programs/AutomationTool/BuildGraph/Tasks/NotarizeTask.cs b/Engine/Source/Programs/AutomationTool/BuildGraph/Tasks/NotarizeTask.cs
--- a/Engine/Source/Programs/AutomationTool/BuildGraph/Tasks/NotarizeTask.cs
+++ b/Engine/Source/Programs/AutomationTool/BuildGraph/Tasks/NotarizeTask.cs
@@ -97,111 +97,78 @@ namespace AutomationTool.Tasks
 
 			int ExitCode = 0;
 			Logger.LogInformation("Uploading {Arg0} to the notarization server...", Dmg.FullName);
-			string CommandLine = string.Format("altool --notarize-app --primary-bundle-id \"{0}\" --username \"{1}\" --password \"@keychain:{2}\" --file \"{3}\"", Parameters.BundleID, Parameters.UserName, Parameters.KeyChainID, Dmg.FullName);
-			string Output = "";
+			
+			// The notarytool will timeout after 5 retries or 1 hour. Whichever comes first.
 			const int MaxNumRetries = 5;
+			const int MaxTimeoutInMilliseconds = 3600000;
+			long TimeoutInMilliseconds = MaxTimeoutInMilliseconds;
+			string Output = "";
+
+			System.Diagnostics.Stopwatch TimeoutStopwatch = System.Diagnostics.Stopwatch.StartNew();
 			
-			for(int NumRetries = 0;;NumRetries++)
+			for (int NumRetries = 0; NumRetries < MaxNumRetries; NumRetries++)
 			{
+				string CommandLine = string.Format("notarytool submit \"{0}\" --keychain-profile \"{1}\" --wait --timeout \"{2}\"", Dmg.FullName, Parameters.KeyChainID, TimeoutInMilliseconds);
 				Output = CommandUtils.RunAndLog("xcrun", CommandLine, out ExitCode);
 				
-				if(ExitCode == 0)
+				if (ExitCode == 0)
 				{
 					break;
 				}
 				
-				if (NumRetries < MaxNumRetries)
+				if (TimeoutStopwatch.ElapsedMilliseconds >= TimeoutInMilliseconds)
 				{
-					Logger.LogInformation("--notarize-app failed with exit {ExitCode} attempting retry {NumRetries} of {MaxNumRetries}", ExitCode, NumRetries, MaxNumRetries);
+					Logger.LogInformation("notarytool timed out after {TimeoutInMilliseconds}ms.", TimeoutInMilliseconds);
+					TimeoutStopwatch.Stop();
+				}
+				else if (NumRetries < MaxNumRetries)
+				{
+					Logger.LogInformation("notarytool failed with exit {ExitCode} attempting retry {NumRetries} of {MaxNumRetries}", ExitCode, NumRetries, MaxNumRetries);
 					Thread.Sleep(2000);
+					TimeoutInMilliseconds = MaxTimeoutInMilliseconds - TimeoutStopwatch.ElapsedMilliseconds;
 					continue;
 				}
-				
+
+
 				Logger.LogInformation("Retries have been exhausted");
-				throw new AutomationException("--notarize-app failed with exit {0}", ExitCode);
+				throw new AutomationException("notarytool failed with exit {0}", ExitCode);
 			}
 
 			// Grab the UUID from the log
 			string RequestUUID = null;
 			try
 			{
-				RequestUUID = Regex.Match(Output, "RequestUUID = ([a-zA-Z0-9]{8}-[a-zA-Z0-9]{4}-[a-zA-Z0-9]{4}-[a-zA-Z0-9]{4}-[a-zA-Z0-9]{12})").Groups[1].Value.Trim();
+				RequestUUID = Regex.Match(Output, "id: ([a-zA-Z0-9]{8}-[a-zA-Z0-9]{4}-[a-zA-Z0-9]{4}-[a-zA-Z0-9]{4}-[a-zA-Z0-9]{12})").Groups[1].Value.Trim();
 			}
-			catch(Exception Ex)
+			catch (Exception Ex)
 			{
 				throw new AutomationException(Ex, "Couldn't get UUID from the log output {0}", Output);
 			}
 
-			// Wait 2 minutes for the server to associate this build with the UUID it passes back.
-			// Trying instantly just returns back and says it can't find it.
-			Logger.LogInformation("Waiting 3 minutes for the UUID to propagate...");
-			Thread.Sleep(180000);
-
-			// Repeat for an hour until we get something back.
 			try
 			{
-				int Timeout = 0;
-				int WaitTime = 30000;
-				int MaxTimeout = 120;
-				while (Timeout < MaxTimeout)
-				{
-					CommandLine = string.Format("altool --notarization-info {0} -u \"{1}\" -p \"@keychain:{2}\"", RequestUUID, Parameters.UserName, Parameters.KeyChainID);
-					Output = CommandUtils.RunAndLog("xcrun", CommandLine, out ExitCode);
-					if (ExitCode != 0)
-					{
-						throw new AutomationException("--notarization-info failed with exit {0}", ExitCode);
-					}
+				MatchCollection StatusMatches = Regex.Matches(Output, "(?<=status: ).+");
+				// The last status update is the right one.
+				string Status = StatusMatches[StatusMatches.Count - 1].Value.ToLower();
 
-					Match StatusMatches = (new Regex("(?<=Status: ).+")).Match(Output);
-					string Status = StatusMatches.Value.ToLower();
-					Match LogFileURLMatches = (new Regex("(?<=LogFileURL: ).+")).Match(Output);
-					string LogFileUrl = LogFileURLMatches.Value;
-
-					if (Status == "invalid")
-					{
-						Logger.LogInformation("{Text}", GetLogFile(LogFileUrl));
-						throw new AutomationException("Could not notarize the app. See log output above.");
-					}
-					else if (Status == "in progress")
-					{
-						Logger.LogInformation("Notarization still in progress, waiting 30 seconds...");
-						Thread.Sleep(WaitTime);
-					}
-					else if (Status == "success")
+				if (Status == "accepted")
+				{
+					if(Parameters.RequireStapling)
 					{
-						if (LogFileUrl == "(null)")
+						// once we have a log file, print it out, staple, and we're done.
+						Logger.LogInformation("{Text}", GetRequestLogs(RequestUUID));
+						string CommandLine = string.Format("stapler staple {0}", Dmg.FullName);
+						Output = CommandUtils.RunAndLog("xcrun", CommandLine, out ExitCode);
+						if (ExitCode != 0)
 						{
-							Logger.LogInformation("Notarization success but no log file has been generated, waiting 30 seconds...");
-							Thread.Sleep(WaitTime);
+							throw new AutomationException("stapler failed with exit {0}", ExitCode);
 						}
-						else if(Parameters.RequireStapling)
-						{
-							// once we have a log file, print it out, staple, and we're done.
-							Logger.LogInformation("{Text}", GetLogFile(LogFileUrl));
-							CommandLine = string.Format("stapler staple {0}", Dmg.FullName);
-							Output = CommandUtils.RunAndLog("xcrun", CommandLine, out ExitCode);
-							if (ExitCode != 0)
-							{
-								throw new AutomationException("stapler failed with exit {0}", ExitCode);
-							}
-							break;
-						}
-						else
-						{
-							// Success, no need to run the stapler.
-							break;
-						}
-					}
-					else
-					{
-						Logger.LogInformation("Status is? {Status}", Status);
-						Thread.Sleep(WaitTime);
 					}
-					Timeout++;
 				}
-				if(Timeout == MaxTimeout)
+				else
 				{
-					throw new AutomationException("Did not get a response back from the notarization server after an hour. Aborting!");
+					Logger.LogError("{Text}", GetRequestLogs(RequestUUID));
+					throw new AutomationException($"Could not notarize the app. Request status: {0}. See log output above.", Status);
 				}
 			}
 			catch (Exception Ex)
@@ -212,38 +179,28 @@ namespace AutomationTool.Tasks
 				}
 				else
 				{
-					throw new AutomationException(Ex, "Querying for the notarization progress failed, output: {0}", Output);
+					throw new AutomationException(Ex, "Querying for the notarization result failed, output: {0}", Output);
 				}
 			}
 
 			return Task.CompletedTask;
 		}
 
-		private string GetLogFile(string Url)
+		private string GetRequestLogs(string RequestUUID)
 		{
-			HttpWebRequest Request = (HttpWebRequest)WebRequest.Create(Url);
-			Request.Method = "GET";
 			try
 			{
-				WebResponse Response = Request.GetResponse();
+				string LogCommand = string.Format("notarytool log {0} --keychain-profile \"{1}\"", RequestUUID, Parameters.KeyChainID);
+				IProcessResult LogResult = CommandUtils.Run("xcrun", LogCommand);
+
 				string ResponseContent = null;
-				using (StreamReader ResponseReader = new System.IO.StreamReader(Response.GetResponseStream(), Encoding.Default))
+				if (LogResult.bExitCodeSuccess)
 				{
-					ResponseContent = ResponseReader.ReadToEnd();
+					ResponseContent = LogResult.Output;
 				}
+
 				return ResponseContent;
 			}
-			catch (WebException Ex)
-			{
-				if (Ex.Response != null)
-				{
-					throw new AutomationException(Ex, string.Format("Request returned status: {0}, message: {1}", ((HttpWebResponse)Ex.Response).StatusCode, Ex.Message));
-				}
-				else
-				{
-					throw new AutomationException(Ex, string.Format("Request returned message: {0}", Ex.InnerException.Message));
-				}
-			}
 			catch (Exception Ex)
 			{
 				throw new AutomationException(Ex, string.Format("Couldn't complete the request, error: {0}", Ex.Message));
