Multi client launcher:
Add option to disable deleting all logs (this was breaking if there's a client running which wasn't launched by the tool as it couldn't delete its log file)
Pass in the full client number to the command line (I had issues with the existing 2 digit system because my bot accounts cross a 3 digit boundary)
#tests built and ran tool locally using the new args

diff --git a/Engine/Source/Programs/AutomationTool/Scripts/MultiClientLauncher.Automation.cs b/Engine/Source/Programs/AutomationTool/Scripts/MultiClientLauncher.Automation.cs
--- a/Engine/Source/Programs/AutomationTool/Scripts/MultiClientLauncher.Automation.cs
+++ b/Engine/Source/Programs/AutomationTool/Scripts/MultiClientLauncher.Automation.cs
@@ -30,6 +30,7 @@ namespace MultiClientLauncher.Automation
 	[ParamHelp("ClientSessionCompleted", "Log message indicating that a client has completed a game session and may be terminated", ParamType = typeof(string))]
 	[ParamHelp("ClientFailed", "Log message indicating that a client failed to connect to the server", ParamType = typeof(string))]
 	[ParamHelp("ClientConnected", "Log message indicating that a client connected to the server", ParamType = typeof(string))]
+	[ParamHelp("DeleteExistingLogs", "Whether to clear the log directory before launching clients, defaults to true", ParamType = typeof(bool))]
 	public class MultiClientLauncher : BuildCommand
 	{
 		private bool CancelClientProcesses = false;
@@ -45,6 +46,7 @@ namespace MultiClientLauncher.Automation
 		private bool NullRHI = false;
 		private bool GridLayout = true;
 		private bool NoTimeouts = false;
+		private bool DeleteExistingLogs = true;
 
 		private int SleepTimeBetweenLaunches = 100;
 		private int MaxRunAttemptsPerClient = 3;
@@ -90,13 +92,14 @@ namespace MultiClientLauncher.Automation
 				ClientArgs += string.Format(" -buildidoverride={0} ", BuildIdOverride);
 			}
 
-			FirstClientNumber = int.Parse(ParseParamValue("FirstClientNumber", "1"));
+			FirstClientNumber = int.Parse(ParseParamValue("FirstClientNumber", "0"));
 
 			SleepTimeBetweenLaunches = ParseParamInt("SleepTimeBetweenLaunches", -1);
 			MaxRunAttemptsPerClient = ParseParamInt("MaxRunAttemptsPerClient", 3);
 			NullRHI = ParseParamBool("NullRHI", NullRHI);
 			GridLayout = ParseParamBool("GridLayout", GridLayout);
 			NoTimeouts = ParseParamBool("NoTimeouts", NoTimeouts);
+			DeleteExistingLogs = ParseParamBool("DeleteExistingLogs", DeleteExistingLogs);
 
 			// disable grid layout for nullrhi
 			if (NullRHI || ClientArgs.Contains("-nullrhi"))
@@ -138,14 +141,17 @@ namespace MultiClientLauncher.Automation
 				KillProcesses(ClientProcesses);
 			};
 
-			// Delete all previous log files
-			Console.WriteLine("Deleting all existing log files in the log directory {0}", ClientLogDir);			
-			string[] LogFiles = Directory.GetFiles(ClientLogDir);
-			foreach (string Filename in LogFiles)
+			if (DeleteExistingLogs)
 			{
-				if (Filename.EndsWith(".log"))
+				// Delete all previous log files
+				Console.WriteLine("Deleting all existing log files in the log directory {0}", ClientLogDir);
+				string[] LogFiles = Directory.GetFiles(ClientLogDir);
+				foreach (string Filename in LogFiles)
 				{
-					File.Delete(Filename);
+					if (Filename.EndsWith(".log"))
+					{
+						File.Delete(Filename);
+					}
 				}
 			}
 
@@ -162,8 +168,8 @@ namespace MultiClientLauncher.Automation
 
 					Console.WriteLine("Spawning client {0}...", ClientIdx);
 
-					string ClientNumber = (FirstClientNumber + ClientIdx).ToString("00.##");
-					string CurrentClientArgs = ClientArgs.Replace("#REPLACED_WITH_TWO_DIGIT_CLIENT_ID#", ClientNumber);
+					string ClientNumber = (FirstClientNumber + ClientIdx).ToString("00000.##");
+					string CurrentClientArgs = ClientArgs.Replace("#REPLACED_WITH_FIVE_DIGIT_CLIENT_ID#", ClientNumber);
 
 					string CurrentClientLog = CommandUtils.CombinePaths(ClientLogDir, ClientLogFilenameGuess);
 					if (ClientIdx > 0)
