Fixing multiple blocking modern xcode issues

diff --git a/Engine/Programs/CrashReportClient/Config/DefaultEngine.ini b/Engine/Programs/CrashReportClient/Config/DefaultEngine.ini
--- a/Engine/Programs/CrashReportClient/Config/DefaultEngine.ini
+++ b/Engine/Programs/CrashReportClient/Config/DefaultEngine.ini
@@ -13,8 +13,3 @@ n.VerifyPeer=false
 [/Script/UdpMessaging.UdpMessagingSettings]
 EnableTransport=True
 MulticastTimeToLive=0
-
-[/Script/MacTargetPlatform.XcodeProjectSettings]
-PremadeMacEntitlements=(FilePath="/Engine/Build/Mac/Resources/Sandbox.EmbeddedApp.entitlements")
-ShippingSpecificMacEntitlements=(FilePath="/Engine/Build/Mac/Resources/Sandbox.EmbeddedApp.entitlements")
-bUseEntitlementsForPrograms=true
diff --git a/Engine/Source/Programs/AutomationTool/Apple/ApplePlatform.Automation.cs b/Engine/Source/Programs/AutomationTool/Apple/ApplePlatform.Automation.cs
--- a/Engine/Source/Programs/AutomationTool/Apple/ApplePlatform.Automation.cs
+++ b/Engine/Source/Programs/AutomationTool/Apple/ApplePlatform.Automation.cs
@@ -142,11 +142,13 @@ public abstract class ApplePlatform : Platform
 		// go through each folder, starting at most recent, looking for an archive for the target
 		foreach (DirectoryReference DateDir in DateDirs)
 		{
-			Logger.LogInformation("Looking in Xcode archive dir {0}...", DateDir);
-
 			// find the most recent archive for this target (based on name of target, this ignores Development vs Shipping, but 
 			// since Distribution is meant only for Shipping it's ok
-			string Wildcard = $"{Target.Receipt.TargetName} *.xcarchive";
+			string Wildcard = AppleExports.MakeBinaryFileName(SC.ShortProjectName, Target.Receipt.Platform, UnrealTargetConfiguration.Development,
+				Target.Receipt.Architectures, UnrealTargetConfiguration.Development, null) + " *.xcarchive";
+
+			Logger.LogInformation("Looking in Xcode archive dir {0} for {1}", DateDir, Wildcard);
+
 			List<DirectoryReference> XcArchives = DirectoryReference.EnumerateDirectories(DateDir, Wildcard).ToList();
 			if (XcArchives.Count > 0)
 			{
@@ -290,7 +292,7 @@ public abstract class ApplePlatform : Platform
 			Logger.LogInformation("Copying {Type} package {ArchiveSource} to archive directory {ArchiveDir}", Params.Distribution ? "Distribution" : "Development", ArchiveSource, SC.ArchiveDirectory);
 			Logger.LogInformation("=====================================================================================");
 
-			SC.ArchiveFiles(ArchiveSource.FullName, NewPath: ArchiveSource.GetDirectoryName());
+			Utils.RunLocalProcessAndReturnStdOut("/usr/bin/env", $"ditto \"{ArchiveSource}\" \"{SC.ArchiveDirectory}/{ArchiveSource.GetDirectoryName()}\"", null);
 		}
 	}
 }
diff --git a/Engine/Source/Programs/UnrealBuildTool/Platform/Mac/MacToolChain.cs b/Engine/Source/Programs/UnrealBuildTool/Platform/Mac/MacToolChain.cs
--- a/Engine/Source/Programs/UnrealBuildTool/Platform/Mac/MacToolChain.cs
+++ b/Engine/Source/Programs/UnrealBuildTool/Platform/Mac/MacToolChain.cs
@@ -506,7 +506,7 @@ namespace UnrealBuildTool
 				LipoAction.CommandPath = FileReference.Combine(Info.Clang.Directory, "lipo");
 				LipoAction.CommandDescription = (bBuildImportLibraryOnly ? "LipoStub" : "Lipo");
 				LipoAction.CommandVersion = Info.ClangVersionString;
-				LipoAction.CommandArguments = String.Join(" ", PerArchOutputFiles) + $" -create -output {OutputFileItem.AbsolutePath}";
+				LipoAction.CommandArguments = String.Join(" ", PerArchOutputFiles.Select(x => $"\"{x}\"")) + $" -create -output \"{OutputFileItem.AbsolutePath}\"";
 				LipoAction.StatusDescription = Path.GetFileName(OutputFileItem.AbsolutePath);
 				LipoAction.bCanExecuteRemotely = false;
 			}
diff --git a/Engine/Source/Programs/UnrealBuildTool/ProjectFiles/Xcode/XcodeProject.cs b/Engine/Source/Programs/UnrealBuildTool/ProjectFiles/Xcode/XcodeProject.cs
--- a/Engine/Source/Programs/UnrealBuildTool/ProjectFiles/Xcode/XcodeProject.cs
+++ b/Engine/Source/Programs/UnrealBuildTool/ProjectFiles/Xcode/XcodeProject.cs
@@ -1418,7 +1418,7 @@ namespace UnrealBuildTool.XcodeProjectXcconfig
 				{
 					"-/Info.plist",
 					"-/Manifest_*",
-					$"/*.app", // remove the staged .app from the root dir, it's hard to do by name due to ProjectName in staging, and TargetName, etc here
+					$"-/*.app", // remove the staged .app from the root dir, it's hard to do by name due to ProjectName in staging, and TargetName, etc here
 				};
 
 				// make a string like --exclude=/Info.plist --exclude=/Manifest_* ...
