- Fixed uncooking of maps after recent fix for .ini files

diff --git a/Engine/Source/Programs/AutomationTool/AutomationUtils/DeploymentContext.cs b/Engine/Source/Programs/AutomationTool/AutomationUtils/DeploymentContext.cs
--- a/Engine/Source/Programs/AutomationTool/AutomationUtils/DeploymentContext.cs
+++ b/Engine/Source/Programs/AutomationTool/AutomationUtils/DeploymentContext.cs
@@ -1211,4 +1211,42 @@ public class DeploymentContext //: ProjectParams
 
 		throw new Exception();
 	}
+	public static FileReference UnmakeRelativeStagedReference(DeploymentContext SC, StagedFileReference Ref)
+	{
+		// paths will be in the form "Engine/Foo" or "{ProjectName}/Foo" or "RemappedPlugins/{PluginName}/Foo
+		// Anything else we don't handle.
+		// So, replace the Engine/ with {EngineDir} and {ProjectName}/ with {ProjectDir}, or change PluginDir to RemappedPlugins/{PluginName}
+		// with the plugin path from AdditionalPluginDirectories, and then append Foo
+
+		string RemappedPluginsStr = "RemappedPlugins/";
+		if (Ref.Name.StartsWith(RemappedPluginsStr, StringComparison.CurrentCultureIgnoreCase))
+		{
+			int PluginEndIndex = Ref.Name.IndexOf("/", RemappedPluginsStr.Length);
+			if (PluginEndIndex >= 0 && PluginEndIndex < Ref.Name.Length - 1)
+			{
+				string PluginName = Ref.Name.Substring(RemappedPluginsStr.Length, PluginEndIndex - RemappedPluginsStr.Length);
+				foreach (DirectoryReference AdditionalPluginDir in SC.AdditionalPluginDirectories)
+				{
+					DirectoryReference PossiblePluginDir = DirectoryReference.Combine(AdditionalPluginDir, PluginName);
+					if (System.IO.Directory.Exists(PossiblePluginDir.FullName))
+					{
+						return FileReference.Combine(PossiblePluginDir, Ref.Name.Substring(PluginEndIndex+1));
+					}
+				}
+			}
+		}
+
+		if (Ref.Name.StartsWith("Engine/", StringComparison.CurrentCultureIgnoreCase))
+		{
+			// skip over "Engine/" which is 7 chars long
+			return FileReference.Combine(SC.EngineRoot, Ref.Name.Substring(7));
+		}
+
+		if (Ref.Name.StartsWith(SC.ShortProjectName + "/", StringComparison.CurrentCultureIgnoreCase))
+		{
+			return FileReference.Combine(SC.ProjectRoot, Ref.Name.Substring(SC.ShortProjectName.Length + 1));
+		}
+
+		throw new Exception($"Don't know how to convert staged file {Ref.Name} to its original editor path, because it is not in a recognized root directory.");
+	}
 }
diff --git a/Engine/Source/Programs/AutomationTool/CookedEditor/CookedEditor.Automation.cs b/Engine/Source/Programs/AutomationTool/CookedEditor/CookedEditor.Automation.cs
--- a/Engine/Source/Programs/AutomationTool/CookedEditor/CookedEditor.Automation.cs
+++ b/Engine/Source/Programs/AutomationTool/CookedEditor/CookedEditor.Automation.cs
@@ -288,6 +288,13 @@ public class ModifyStageContext
 		}
 	}
 
+	private void AddCookedUFSFilesToList(List<FileReference> FileList, string Extension, DeploymentContext SC)
+	{
+		// look in SC and UFSFiles
+		FileList.AddRange(SC.FilesToStage.UFSFiles.Keys.Where(x => x.Name.EndsWith(Extension, StringComparison.OrdinalIgnoreCase)).Select(y => DeploymentContext.UnmakeRelativeStagedReference(SC, y)));
+		FileList.AddRange(UFSFilesToStage.Where(x => x.FullName.EndsWith(Extension, StringComparison.InvariantCultureIgnoreCase)));
+	}
+
 	private void AddUFSFilesToList(List<FileReference> FileList, string Extension, DeploymentContext SC)
 	{
 		// look in SC and UFSFiles
@@ -310,8 +317,8 @@ public class ModifyStageContext
 		}
 		else if (CookedMapMode == "uncooked")
 		{
-			// remove maps from SC and Context (SC has path to the cooked map, so we have to come back from Staged refernece that doesn't have the Cooked dir in it)
-			AddUFSFilesToList(FilesToUncook, ".umap", SC);
+			// remove maps from SC and Context (SC has path to the cooked map, so we have to come back from Staged reference that doesn't have the Cooked dir in it)
+			AddCookedUFSFilesToList(FilesToUncook, ".umap", SC);
 		}
 		else if (CookedMapMode == "none")
 		{
