CookedEditor.Automation: Don't stage .cs files

#rnx

diff --git a/Engine/Source/Programs/AutomationTool/CookedEditor/CookedEditor.Automation.cs b/Engine/Source/Programs/AutomationTool/CookedEditor/CookedEditor.Automation.cs
--- a/Engine/Source/Programs/AutomationTool/CookedEditor/CookedEditor.Automation.cs
+++ b/Engine/Source/Programs/AutomationTool/CookedEditor/CookedEditor.Automation.cs
@@ -440,7 +440,9 @@ public class MakeCookedEditor : BuildCommand
 		// engine shaders
 		if (Context.bStageShaderDirs)
 		{
-			Context.NonUFSFilesToStage.AddRange(DirectoryReference.EnumerateFiles(DirectoryReference.Combine(Unreal.EngineDirectory, "Shaders"), "*", SearchOption.AllDirectories));
+			IEnumerable<FileReference> ShaderFiles = DirectoryReference.EnumerateFiles(DirectoryReference.Combine(Unreal.EngineDirectory, "Shaders"), "*", SearchOption.AllDirectories)
+				.Where(x => !x.GetExtension().Equals(".cs", StringComparison.OrdinalIgnoreCase));
+			Context.NonUFSFilesToStage.AddRange(ShaderFiles);
 			GatherTargetDependencies(Params, SC, Context, "ShaderCompileWorker");
 		}
 		if (bIsCookedCooker)
