Added a helper function to resolve the UNC share path for a mapped network drive

#rnx

diff --git a/Engine/Source/Programs/AutomationTool/AutomationUtils/Utils.cs b/Engine/Source/Programs/AutomationTool/AutomationUtils/Utils.cs
--- a/Engine/Source/Programs/AutomationTool/AutomationUtils/Utils.cs
+++ b/Engine/Source/Programs/AutomationTool/AutomationUtils/Utils.cs
@@ -14,6 +14,7 @@ using UnrealBuildBase;
 using Microsoft.Extensions.Logging;
 
 using static AutomationTool.CommandUtils;
+using System.Management;
 
 namespace AutomationTool
 {
@@ -921,6 +922,61 @@ namespace AutomationTool
 	        return bFound;
 	    }
 
+		/// <summary>
+		/// Attempts to get the shared network path for a given mapped network drive
+		/// e.g.  X:\Shared\Path\   -> \\MyServer\Root\Shared\Path\
+		/// </summary>
+		/// <param name="InPath">Source path</param>
+		/// <param name="UNCPath">Receives the UNC path</param>
+		/// <returns>true on success</returns>
+		public static bool TryResolveMappedNetworkPath( string InPath, out string UNCPath )
+		{
+			UNCPath = "";
+
+			// only works on Windows at the moment
+			if (HostPlatform.Current.HostEditorPlatform != UnrealTargetPlatform.Win64)
+			{
+				return false;
+			}
+
+			// path is already a network shared path
+			if (InPath.StartsWith("\\\\"))
+			{
+				UNCPath = InPath;
+				return true;
+			}
+
+			// path has no drive
+			if (!Path.IsPathRooted(InPath))
+			{
+				return false;
+			}
+
+			// see if the path is on a network drive & try to resolve the target
+			try
+			{
+				DriveInfo Drive = new(InPath);
+				if (Drive.DriveType == DriveType.Network)
+				{
+					using (ManagementObject ManObj = new($"Win32_LogicalDisk='{Drive.Name.TrimEnd(Path.DirectorySeparatorChar)}'")) // Win32_LogicalDisk='X:'
+					{
+						string UNCRoot = ManObj["ProviderName"].ToString(); // e.g. \\MyServer\Root
+						string SharedPathFragment = InPath.Replace(Drive.Name, "", StringComparison.InvariantCultureIgnoreCase);
+
+						UNCPath = Path.Combine( UNCRoot, SharedPathFragment);
+						return true;
+					}
+				}
+			}
+			catch
+			{
+			}
+
+			// something failed, or the drive isn't a network drive
+			return false;
+		}
+
+
 		/// <summary>
 		/// Efficient iterator for walking over a string line by line.
 		/// </summary>
