Added extraction of BuildVersion to UnrealLogParser BuildInfo class

diff --git a/Engine/Source/Programs/AutomationTool/Gauntlet/Unreal/Utils/Gauntlet.UnrealLogParser.cs b/Engine/Source/Programs/AutomationTool/Gauntlet/Unreal/Utils/Gauntlet.UnrealLogParser.cs
--- a/Engine/Source/Programs/AutomationTool/Gauntlet/Unreal/Utils/Gauntlet.UnrealLogParser.cs
+++ b/Engine/Source/Programs/AutomationTool/Gauntlet/Unreal/Utils/Gauntlet.UnrealLogParser.cs
@@ -129,6 +129,7 @@ namespace Gauntlet
 		/// </summary>
 		public class BuildInfo
 		{
+			public string BuildVersion;
 			public string BranchName;
 			public int Changelist;
 		}
@@ -407,6 +408,13 @@ namespace Gauntlet
 				Info.Changelist = Convert.ToInt32(M.Groups[1].ToString());
 			}
 
+			M = Regex.Match(Content, @"LogInit.+Build:\s*(\+.*)", RegexOptions.IgnoreCase);
+
+			if (M.Success)
+			{
+				Info.BuildVersion = M.Groups[1].ToString();
+			}
+
 			return Info;
 		}
 
