Use new windowskit debugging tools to support indexing > 4Gb pdbs

diff --git a/Engine/Source/Programs/AutomationTool/Win/WinPlatform.Automation.cs b/Engine/Source/Programs/AutomationTool/Win/WinPlatform.Automation.cs
--- a/Engine/Source/Programs/AutomationTool/Win/WinPlatform.Automation.cs
+++ b/Engine/Source/Programs/AutomationTool/Win/WinPlatform.Automation.cs
@@ -417,6 +417,30 @@ public class Win64Platform : Platform
 	[SupportedOSPlatform("windows")]
     private static FileReference GetSymStoreExe()
     {
+		// Trying first to look for auto sdk latest WindowsKits debugger tools
+		DirectoryReference HostAutoSdkDir = null;
+		if (UEBuildPlatformSDK.TryGetHostPlatformAutoSDKDir(out HostAutoSdkDir))
+		{
+			DirectoryReference WindowsKitsDebuggersDirAutoSdk = DirectoryReference.Combine(HostAutoSdkDir, "Win64", "Windows Kits", "Debuggers");
+
+			if (DirectoryReference.Exists(WindowsKitsDebuggersDirAutoSdk))
+			{
+				// Defaulting to the x86 because of a known issue with the latest x64 version
+				// x64 version gets the errorcode STATUS_ENTRYPOINT_NOT_FOUND on some configurations
+				FileReference SymStoreExe32 = FileReference.Combine(WindowsKitsDebuggersDirAutoSdk, "x86", "SymStore.exe");
+				if (FileReference.Exists(SymStoreExe32))
+				{
+					return SymStoreExe32;
+				}
+
+				FileReference SymStoreExe64 = FileReference.Combine(WindowsKitsDebuggersDirAutoSdk, "x64", "SymStore.exe");
+				if (FileReference.Exists(SymStoreExe64))
+				{
+					return SymStoreExe64;
+				}
+			}
+		}
+
 		List<KeyValuePair<string, DirectoryReference>> WindowsSdkDirs = WindowsExports.GetWindowsSdkDirs();
 		foreach (DirectoryReference WindowsSdkDir in WindowsSdkDirs.Select(x => x.Value))
 		{
@@ -438,6 +462,32 @@ public class Win64Platform : Platform
 	[SupportedOSPlatform("windows")]
 	public static bool TryGetPdbCopyLocation(out FileReference OutLocation)
 	{
+		// Trying first to look for auto sdk latest WindowsKits debugger tools
+		DirectoryReference HostAutoSdkDir = null;
+		if (UEBuildPlatformSDK.TryGetHostPlatformAutoSDKDir(out HostAutoSdkDir))
+		{
+			DirectoryReference WindowsKitsDebuggersDirAutoSdk = DirectoryReference.Combine(HostAutoSdkDir, "Win64", "Windows Kits", "Debuggers");
+
+			if (DirectoryReference.Exists(WindowsKitsDebuggersDirAutoSdk))
+			{
+				// Defaulting to the x86 because of a known issue with the latest x64 version
+				// x64 version gets the errorcode STATUS_ENTRYPOINT_NOT_FOUND on some configurations
+				FileReference PdbCopyExe32 = FileReference.Combine(WindowsKitsDebuggersDirAutoSdk, "x86", "PdbCopy.exe");
+				if (FileReference.Exists(PdbCopyExe32))
+				{
+					OutLocation = PdbCopyExe32;
+					return true;
+				}
+
+				FileReference PdbCopyExe64 = FileReference.Combine(WindowsKitsDebuggersDirAutoSdk, "x64", "PdbCopy.exe");
+				if (FileReference.Exists(PdbCopyExe64))
+				{
+					OutLocation = PdbCopyExe64;
+					return true;
+				}
+			}
+		}
+
 		// Try to find an installation of the Windows 10 SDK
 		List<KeyValuePair<string, DirectoryReference>> WindowsSdkDirs = WindowsExports.GetWindowsSdkDirs();
 		foreach (DirectoryReference WindowsSdkDir in WindowsSdkDirs.Select(x => x.Value))
@@ -753,21 +803,6 @@ public class Win64Platform : Platform
 		FileInfo PdbInfo = new FileInfo(PdbFile.FullName);
 		FileInfo IniInfo = new FileInfo(SrcSrvIni.FullName);
 
-		const long Size4GB = 4L * 1024 * 1024 * 1024;
-
-		// PdbStr.exe tool can't handle pdb's larger than 4GB (linked with /PDBPAGESIZE:8192 switch).
-		// It's a known limitation and it may happen for large game projects, especially in development configurations.
-		// We choose to emit a warning for this known issue rather than failing the whole build with an exception, if we run into this case.
-		if (PdbInfo.Length + IniInfo.Length >= Size4GB)
-		{
-			lock (State)
-			{
-				Logger.LogInformation("Skipping indexing of {Arg0} because it would make the file larger than 4GB (not supported by pdbstr.exe tool).", PdbFile.GetFileName());
-			}
-
-			return;
-		}
-
 		using (Process Process = new Process())
 		{
 			List<string> Messages = new List<string>();
@@ -812,6 +847,31 @@ public class Win64Platform : Platform
 	static FileReference GetPdbStrExe()
 	{
 		List<KeyValuePair<string, DirectoryReference>> WindowsSdkDirs = WindowsExports.GetWindowsSdkDirs();
+
+		// Trying first to look for auto sdk latest WindowsKits debugger tools
+		DirectoryReference HostAutoSdkDir = null;
+		if (UEBuildPlatformSDK.TryGetHostPlatformAutoSDKDir(out HostAutoSdkDir))
+		{
+			DirectoryReference WindowsKitsDebuggersDirAutoSdk = DirectoryReference.Combine(HostAutoSdkDir, "Win64", "Windows Kits", "Debuggers");
+
+			if (DirectoryReference.Exists(WindowsKitsDebuggersDirAutoSdk))
+			{
+				// Defaulting to the x86 because of a known issue with the latest x64 version
+				// x64 version gets the errorcode STATUS_ENTRYPOINT_NOT_FOUND on some configurations
+				FileReference CheckPdbStrExe32 = FileReference.Combine(WindowsKitsDebuggersDirAutoSdk, "x86", "SrcSrv", "PdbStr.exe");
+				if (FileReference.Exists(CheckPdbStrExe32))
+				{
+					return CheckPdbStrExe32;
+				}
+
+				FileReference CheckPdbStrExe64 = FileReference.Combine(WindowsKitsDebuggersDirAutoSdk, "x64", "SrcSrv", "PdbStr.exe");
+				if (FileReference.Exists(CheckPdbStrExe64))
+				{
+					return CheckPdbStrExe64;
+				}
+			}
+		}
+
 		foreach (DirectoryReference WindowsSdkDir in WindowsSdkDirs.Select(x => x.Value))
 		{
 			FileReference CheckPdbStrExe64 = FileReference.Combine(WindowsSdkDir, "Debuggers", "x64", "SrcSrv", "PdbStr.exe");
