Use Gauntlet command for low level tests on incremental and complete LLTs instead of spawning tests directly.
This should allow better detection for hangs.

diff --git a/Engine/Source/Programs/AutomationTool/LowLevelTests/RunLowLevelTests.cs b/Engine/Source/Programs/AutomationTool/LowLevelTests/RunLowLevelTests.cs
--- a/Engine/Source/Programs/AutomationTool/LowLevelTests/RunLowLevelTests.cs
+++ b/Engine/Source/Programs/AutomationTool/LowLevelTests/RunLowLevelTests.cs
@@ -568,6 +568,15 @@ namespace LowLevelTests
 					continue;
 				}
 
+				if ((InPlatform == UnrealTargetPlatform.Mac
+					|| InPlatform == UnrealTargetPlatform.Linux
+					|| InPlatform == UnrealTargetPlatform.LinuxArm64)
+					&& !string.IsNullOrEmpty(Path.GetExtension(Executable)))
+				{
+					// Mac & Linux executable candidates should have no extension
+					continue;
+				}
+
 				// Development executable does not contain configuration or platform name
 				Log.Verbose("Config type: {0}", InConfiguration);
 				if (InConfiguration == UnrealTargetConfiguration.Development)
@@ -640,13 +649,13 @@ namespace LowLevelTests
 		{
 			if (InPlatform.IsInGroup(UnrealPlatformGroup.Windows))
 			{
-				return @"\w+(Tests)?(?:-\w+)?(?:-\w+)?.exe$";
+				return @"[A-Za-z0-9_]+(Tests)?(?:-[A-Za-z0-9_]+)?(?:-[A-Za-z0-9_]+)?.exe$";
 			}
 			else if (InPlatform == UnrealTargetPlatform.Linux ||
 					 InPlatform == UnrealTargetPlatform.LinuxArm64 ||
 					 InPlatform == UnrealTargetPlatform.Mac)
 			{
-				return @"\w+(Tests)?$";
+				return @"[A-Za-z0-9_]+(Tests)?(?:-[A-Za-z0-9_]+)?(?:-[A-Za-z0-9_]+)?$";
 			}
 			else
 			{
