Add support for -DestLocalInstallDir for Gauntlet to override where local builds are installed to.

 - this can be used to persist a local build across multiple nodes to avoid it being cleaned from the artifact directory

#rnx

diff --git a/Engine/Source/Programs/AutomationTool/Gauntlet/Platform/Linux/Gauntlet.TargetDeviceLinux.cs b/Engine/Source/Programs/AutomationTool/Gauntlet/Platform/Linux/Gauntlet.TargetDeviceLinux.cs
--- a/Engine/Source/Programs/AutomationTool/Gauntlet/Platform/Linux/Gauntlet.TargetDeviceLinux.cs
+++ b/Engine/Source/Programs/AutomationTool/Gauntlet/Platform/Linux/Gauntlet.TargetDeviceLinux.cs
@@ -350,7 +350,8 @@ namespace Gauntlet
 			if (CanRunFromPath(BuildPath) == false)
 			{
 				string SubDir = string.IsNullOrEmpty(AppConfig.Sandbox) ? AppConfig.ProjectName : AppConfig.Sandbox;
-				string DestPath = Path.Combine(this.LocalCachePath, SubDir, AppConfig.ProcessType.ToString());
+				string BasePath = string.IsNullOrEmpty(AppConfig.DestLocalInstallDir) ? this.LocalCachePath : AppConfig.DestLocalInstallDir;
+				string DestPath = Path.Combine(BasePath, SubDir, AppConfig.ProcessType.ToString());
 
 				if (!AppConfig.SkipInstall)
 				{
diff --git a/Engine/Source/Programs/AutomationTool/Gauntlet/Platform/Mac/Gauntlet.TargetDeviceMac.cs b/Engine/Source/Programs/AutomationTool/Gauntlet/Platform/Mac/Gauntlet.TargetDeviceMac.cs
--- a/Engine/Source/Programs/AutomationTool/Gauntlet/Platform/Mac/Gauntlet.TargetDeviceMac.cs
+++ b/Engine/Source/Programs/AutomationTool/Gauntlet/Platform/Mac/Gauntlet.TargetDeviceMac.cs
@@ -254,7 +254,8 @@ namespace Gauntlet
 			if (BuildVolume.Equals(LocalRoot, StringComparison.OrdinalIgnoreCase) == false)
 			{
 				string SubDir = string.IsNullOrEmpty(AppConfig.Sandbox) ? AppConfig.ProjectName : AppConfig.Sandbox;
-				string DestPath = Path.Combine(this.LocalCachePath, SubDir, AppConfig.ProcessType.ToString());
+				string BasePath = string.IsNullOrEmpty(AppConfig.DestLocalInstallDir) ? this.LocalCachePath : AppConfig.DestLocalInstallDir;
+				string DestPath = Path.Combine(BasePath, SubDir, AppConfig.ProcessType.ToString());
 
 				if (!SkipDeploy)
 				{
diff --git a/Engine/Source/Programs/AutomationTool/Gauntlet/Platform/Windows/Gauntlet.TargetDeviceWindows.cs b/Engine/Source/Programs/AutomationTool/Gauntlet/Platform/Windows/Gauntlet.TargetDeviceWindows.cs
--- a/Engine/Source/Programs/AutomationTool/Gauntlet/Platform/Windows/Gauntlet.TargetDeviceWindows.cs
+++ b/Engine/Source/Programs/AutomationTool/Gauntlet/Platform/Windows/Gauntlet.TargetDeviceWindows.cs
@@ -486,7 +486,8 @@ namespace Gauntlet
 			if (CanRunFromPath(BuildPath) == false)
 			{
 				string SubDir = string.IsNullOrEmpty(AppConfig.Sandbox) ? AppConfig.ProjectName : AppConfig.Sandbox;
-				string DestPath = Path.Combine(this.LocalCachePath, SubDir, AppConfig.ProcessType.ToString());
+				string BasePath = string.IsNullOrEmpty(AppConfig.DestLocalInstallDir) ? this.LocalCachePath : AppConfig.DestLocalInstallDir;
+				string DestPath = Path.Combine(BasePath, SubDir, AppConfig.ProcessType.ToString());
 
 				if (!AppConfig.SkipInstall)
 				{
diff --git a/Engine/Source/Programs/AutomationTool/Gauntlet/Unreal/Base/Gauntlet.UnrealAppConfig.cs b/Engine/Source/Programs/AutomationTool/Gauntlet/Unreal/Base/Gauntlet.UnrealAppConfig.cs
--- a/Engine/Source/Programs/AutomationTool/Gauntlet/Unreal/Base/Gauntlet.UnrealAppConfig.cs
+++ b/Engine/Source/Programs/AutomationTool/Gauntlet/Unreal/Base/Gauntlet.UnrealAppConfig.cs
@@ -99,6 +99,11 @@ namespace Gauntlet
 		/// </summary>
 		public string Sandbox { get; set; }
 
+		/// <summary>
+		/// Optional directory for builds that need to be run locally. This can be used to persist a locally-install build between build graph nodes.
+		/// </summary>
+		public string DestLocalInstallDir { get; set; }
+
 		// new system
 		public IBuild Build { get; set; }
 
diff --git a/Engine/Source/Programs/AutomationTool/Gauntlet/Unreal/Base/Gauntlet.UnrealSession.cs b/Engine/Source/Programs/AutomationTool/Gauntlet/Unreal/Base/Gauntlet.UnrealSession.cs
--- a/Engine/Source/Programs/AutomationTool/Gauntlet/Unreal/Base/Gauntlet.UnrealSession.cs
+++ b/Engine/Source/Programs/AutomationTool/Gauntlet/Unreal/Base/Gauntlet.UnrealSession.cs
@@ -589,6 +589,11 @@ namespace Gauntlet
 		/// </summary>
 		public string Sandbox { get; set; }
 
+		/// <summary>
+		/// Optional directory where local builds will be installed
+		/// </summary>
+		public string DestLocalInstallDir { get; set; }
+
 		/// <summary>
 		/// Whether or not we should retain our devices this pass
 		/// </summary>
@@ -895,6 +900,12 @@ namespace Gauntlet
 					// todo - should this be elsewhere?
 					AppConfig.Sandbox = Sandbox;
 
+					if (!string.IsNullOrEmpty(DestLocalInstallDir) && BuildHostPlatform.Current.Platform == Device.Platform)
+					{
+						AppConfig.DestLocalInstallDir = Path.Combine(DestLocalInstallDir, Device.Platform.ToString(), Device.Name );
+					}
+
+
 					Device.ClearSavedDirectory(AppConfig);
 
 					IAppInstall Install = null;
diff --git a/Engine/Source/Programs/AutomationTool/Gauntlet/Unreal/Base/Gauntlet.UnrealTestContext.cs b/Engine/Source/Programs/AutomationTool/Gauntlet/Unreal/Base/Gauntlet.UnrealTestContext.cs
--- a/Engine/Source/Programs/AutomationTool/Gauntlet/Unreal/Base/Gauntlet.UnrealTestContext.cs
+++ b/Engine/Source/Programs/AutomationTool/Gauntlet/Unreal/Base/Gauntlet.UnrealTestContext.cs
@@ -140,6 +140,12 @@ namespace Gauntlet
 		[AutoParam("")]
 		public string Sandbox { get; set; }
 
+		/// <summary>
+		/// Optional directory where local builds will be installed
+		/// </summary>
+		[AutoParam("")]
+		public string DestLocalInstallDir;
+
 		/// <summary>
 		/// Skip any check or copying of builds. Mostly useful when debugging
 		/// </summary>
diff --git a/Engine/Source/Programs/AutomationTool/Gauntlet/Unreal/Base/Gauntlet.UnrealTestNode.cs b/Engine/Source/Programs/AutomationTool/Gauntlet/Unreal/Base/Gauntlet.UnrealTestNode.cs
--- a/Engine/Source/Programs/AutomationTool/Gauntlet/Unreal/Base/Gauntlet.UnrealTestNode.cs
+++ b/Engine/Source/Programs/AutomationTool/Gauntlet/Unreal/Base/Gauntlet.UnrealTestNode.cs
@@ -682,7 +682,11 @@ namespace Gauntlet
 				}
 			}
 
-			UnrealApp = new UnrealSession(Context.BuildInfo, SessionRoles) { Sandbox = Context.Options.Sandbox };
+			UnrealApp = new UnrealSession(Context.BuildInfo, SessionRoles) 
+			{ 
+				Sandbox = Context.Options.Sandbox, 
+				DestLocalInstallDir = Context.Options.DestLocalInstallDir 
+			};
 
 			return true;
 		}
