Accounts for read-only files when attempting to delete source artifacts.

diff --git a/Engine/Source/Programs/AutomationTool/Gauntlet/Unreal/Base/Gauntlet.UnrealSession.cs b/Engine/Source/Programs/AutomationTool/Gauntlet/Unreal/Base/Gauntlet.UnrealSession.cs
--- a/Engine/Source/Programs/AutomationTool/Gauntlet/Unreal/Base/Gauntlet.UnrealSession.cs
+++ b/Engine/Source/Programs/AutomationTool/Gauntlet/Unreal/Base/Gauntlet.UnrealSession.cs
@@ -1195,6 +1195,19 @@ namespace Gauntlet
 					// By default we delete the source artifacts, but we keep them if requested
 					if (!bRetainArtifacts)
 					{
+						// Account for any read-only files.
+						void SetAttributesNormal(DirectoryInfo Directory)
+						{
+							foreach(FileInfo File in Directory.GetFiles())
+							{
+								File.Attributes = FileAttributes.Normal;
+							}
+							foreach(DirectoryInfo SubDirectory in Directory.GetDirectories())
+							{
+								SetAttributesNormal(SubDirectory);
+							}
+						};
+						SetAttributesNormal(SourceDirectory);
 						try
 						{
 							SourceDirectory.Delete(true);
@@ -1568,4 +1581,4 @@ namespace Gauntlet
 			return LongFilePath;
 		}
 	}
-}
\ No newline at end of file
+}
