- Fix warning about PluginBase.ini not being explicitly staged or not

diff --git a/Engine/Source/Programs/AutomationTool/Scripts/CopyBuildToStagingDirectory.Automation.cs b/Engine/Source/Programs/AutomationTool/Scripts/CopyBuildToStagingDirectory.Automation.cs
--- a/Engine/Source/Programs/AutomationTool/Scripts/CopyBuildToStagingDirectory.Automation.cs
+++ b/Engine/Source/Programs/AutomationTool/Scripts/CopyBuildToStagingDirectory.Automation.cs
@@ -1861,7 +1861,7 @@ namespace AutomationScripts
 				}
 
 				// this is needed for all targets
-				if (NormalizedPath == "configredirects.ini")
+				if (NormalizedPath == "configredirects.ini" || NormalizedPath == "pluginbase.ini")
 				{
 					return true;
 				}
