Add additional lines to be removed when displaying a simple command in horde.

diff --git a/Engine/Source/Programs/AutomationTool/Gauntlet/Unreal/Base/Gauntlet.UnrealTestNode.cs b/Engine/Source/Programs/AutomationTool/Gauntlet/Unreal/Base/Gauntlet.UnrealTestNode.cs
--- a/Engine/Source/Programs/AutomationTool/Gauntlet/Unreal/Base/Gauntlet.UnrealTestNode.cs
+++ b/Engine/Source/Programs/AutomationTool/Gauntlet/Unreal/Base/Gauntlet.UnrealTestNode.cs
@@ -911,6 +911,8 @@ namespace Gauntlet
 					"ECBranch",
 					"ECChangelist",
 					"PreFlightChange",
+					"AssetRegistryCacheRootFolder",
+					"DeactivatedTestConfigPath",
 			}.Select(I => I.ToLower()).ToArray();
 
 			bool ShouldArgBeDisplayed(string InArg)
