- Only apply UE_RELATIVE_BASE_DIR if it exists (ie has not been staged to a differnet dir structure)
- SubmitTool now will set the projectdir to the Restricted location only if not-staged
- Disabled some settings in SubmitTool to reduce unused staged files (dlls, etc)
- Disabled VULKAN_HAS_DEBUGGING_ENABLED when IS_PROGRAM, and no longer stage all the VkLayers when staging a Program (helps save over 20mb from standalone programs on Windows)
- No longer force ShaderFormats to true when building a program that is TargetPlatform enabled

diff --git a/Engine/Source/Programs/AutomationTool/AutomationUtils/DeploymentContext.cs b/Engine/Source/Programs/AutomationTool/AutomationUtils/DeploymentContext.cs
--- a/Engine/Source/Programs/AutomationTool/AutomationUtils/DeploymentContext.cs
+++ b/Engine/Source/Programs/AutomationTool/AutomationUtils/DeploymentContext.cs
@@ -932,16 +932,16 @@ public class DeploymentContext //: ProjectParams
 			StageCrashReporterFile(FileType, InputFile, StagedFile);
 		}
 	}
-
-	public void StageVulkanValidationLayerFiles(StagedFileType FileType, DirectoryReference InputDir, StageFilesSearch Option)
+	
+	public void StageVulkanValidationLayerFiles(ProjectParams Params, StagedFileType FileType, DirectoryReference InputDir, StageFilesSearch Option)
 	{
-			StageVulkanValidationLayerFiles(FileType, InputDir, Option, new StagedDirectoryReference(InputDir.MakeRelativeTo(LocalRoot)));
-		}
+		StageVulkanValidationLayerFiles(Params, FileType, InputDir, Option, new StagedDirectoryReference(InputDir.MakeRelativeTo(LocalRoot)));
+	}
 
-	public void StageVulkanValidationLayerFiles(StagedFileType FileType, DirectoryReference InputDir, StageFilesSearch Option, StagedDirectoryReference OutputDir)
+	public void StageVulkanValidationLayerFiles(ProjectParams Params, StagedFileType FileType, DirectoryReference InputDir, StageFilesSearch Option, StagedDirectoryReference OutputDir)
 	{
 		// This needs to match the c++ define VULKAN_HAS_DEBUGGING_ENABLED to avoid mismatched functionality/files
-		bool bShouldStageVulkanLayers = StageTargetConfigurations.Contains(UnrealTargetConfiguration.Debug) || StageTargetConfigurations.Contains(UnrealTargetConfiguration.Development);
+		bool bShouldStageVulkanLayers = !Params.IsProgramTarget && (StageTargetConfigurations.Contains(UnrealTargetConfiguration.Debug) || StageTargetConfigurations.Contains(UnrealTargetConfiguration.Development));
 		if (bShouldStageVulkanLayers)
 		{
 			List<FileReference> InputFiles = FindFilesToStage(InputDir, Option);
diff --git a/Engine/Source/Programs/AutomationTool/Scripts/CopyBuildToStagingDirectory.Automation.cs b/Engine/Source/Programs/AutomationTool/Scripts/CopyBuildToStagingDirectory.Automation.cs
--- a/Engine/Source/Programs/AutomationTool/Scripts/CopyBuildToStagingDirectory.Automation.cs
+++ b/Engine/Source/Programs/AutomationTool/Scripts/CopyBuildToStagingDirectory.Automation.cs
@@ -1426,7 +1426,7 @@ namespace AutomationScripts
 						}
 					}
 
-					SC.StageVulkanValidationLayerFiles(StagedFileType.NonUFS, DirectoryReference.Combine(SC.EngineRoot, "Binaries", "ThirdParty", "Vulkan", SC.PlatformDir), StageFilesSearch.AllDirectories);
+					SC.StageVulkanValidationLayerFiles(Params, StagedFileType.NonUFS, DirectoryReference.Combine(SC.EngineRoot, "Binaries", "ThirdParty", "Vulkan", SC.PlatformDir), StageFilesSearch.AllDirectories);
 
 					// CrashReportClient is a standalone slate app that does not look in the generated pak file, so it needs the Content/Slate and Shaders/StandaloneRenderer folders Non-UFS
 					// @todo Make CrashReportClient more portable so we don't have to do this
diff --git a/Engine/Source/Programs/UnrealBuildTool/Platform/Windows/UEBuildWindows.cs b/Engine/Source/Programs/UnrealBuildTool/Platform/Windows/UEBuildWindows.cs
--- a/Engine/Source/Programs/UnrealBuildTool/Platform/Windows/UEBuildWindows.cs
+++ b/Engine/Source/Programs/UnrealBuildTool/Platform/Windows/UEBuildWindows.cs
@@ -1513,7 +1513,7 @@ namespace UnrealBuildTool
 		{
 			bool bBuildShaderFormats = Target.bForceBuildShaderFormats;
 
-			if (!Target.bBuildRequiresCookedData)
+			if (!Target.bBuildRequiresCookedData && Target.Type != TargetType.Program)
 			{
 				if (ModuleName == "TargetPlatform")
 				{
diff --git a/Engine/Source/Runtime/Core/Private/Windows/WindowsPlatformProcess.cpp b/Engine/Source/Runtime/Core/Private/Windows/WindowsPlatformProcess.cpp
--- a/Engine/Source/Runtime/Core/Private/Windows/WindowsPlatformProcess.cpp
+++ b/Engine/Source/Runtime/Core/Private/Windows/WindowsPlatformProcess.cpp
@@ -1119,7 +1119,12 @@ const TCHAR* FWindowsPlatformProcess::BaseDir()
 
 			FString CollapseResult(Result);
 #ifdef UE_RELATIVE_BASE_DIR
-			CollapseResult /= UE_RELATIVE_BASE_DIR;
+			// this may have been defined at compile time because we are in Restricted, but then we have been staged as a program, and then remapped out of Restricted
+			// so if we are already in a Binaries/Win64 directory
+			if (IFileManager::Get().DirectoryExists(*FPaths::Combine(CollapseResult, UE_RELATIVE_BASE_DIR)))
+			{
+				CollapseResult = FPaths::Combine(CollapseResult, UE_RELATIVE_BASE_DIR);
+			}
 #endif
 			FPaths::CollapseRelativeDirectories(CollapseResult);
 			FCString::Strcpy(Result, *CollapseResult);
diff --git a/Engine/Source/Runtime/VulkanRHI/Public/VulkanConfiguration.h b/Engine/Source/Runtime/VulkanRHI/Public/VulkanConfiguration.h
--- a/Engine/Source/Runtime/VulkanRHI/Public/VulkanConfiguration.h
+++ b/Engine/Source/Runtime/VulkanRHI/Public/VulkanConfiguration.h
@@ -24,7 +24,7 @@ struct VkAllocationCallbacks;
 #endif
 
 #ifndef VULKAN_HAS_DEBUGGING_ENABLED
-	#define VULKAN_HAS_DEBUGGING_ENABLED						(UE_BUILD_DEBUG || (UE_BUILD_DEVELOPMENT && VULKAN_SHOULD_DEBUG_IN_DEVELOPMENT))
+	#define VULKAN_HAS_DEBUGGING_ENABLED						(!IS_PROGRAM && (UE_BUILD_DEBUG || (UE_BUILD_DEVELOPMENT && VULKAN_SHOULD_DEBUG_IN_DEVELOPMENT)))
 #endif
 
 // default value of r.Vulkan.EnableValidation
