add device related arguments to filter out list when generating reproducing line

diff --git a/Engine/Source/Programs/AutomationTool/Gauntlet/Unreal/Base/Gauntlet.UnrealTestNode.cs b/Engine/Source/Programs/AutomationTool/Gauntlet/Unreal/Base/Gauntlet.UnrealTestNode.cs
--- a/Engine/Source/Programs/AutomationTool/Gauntlet/Unreal/Base/Gauntlet.UnrealTestNode.cs
+++ b/Engine/Source/Programs/AutomationTool/Gauntlet/Unreal/Base/Gauntlet.UnrealTestNode.cs
@@ -913,6 +913,16 @@ namespace Gauntlet
 					"PreFlightChange",
 					"AssetRegistryCacheRootFolder",
 					"DeactivatedTestConfigPath",
+					"SkipInstall",
+					"destlocalinstalldir",
+					"deviceurl",
+					"devicepool",
+					"VerifyLogin",
+					"cleardevices",
+					"reboot",
+					"fullclean",
+					"BuildName",
+					"PerfReportServer",
 			}.Select(I => I.ToLower()).ToArray();
 
 			bool ShouldArgBeDisplayed(string InArg)
