Added more Verbose logging to Gauntlet.TargetDevicePS5 and UnrealBuildSource

#rnx

diff --git a/Engine/Source/Programs/AutomationTool/Gauntlet/Unreal/BuildSource/Gauntlet.UnrealBuildSource.cs b/Engine/Source/Programs/AutomationTool/Gauntlet/Unreal/BuildSource/Gauntlet.UnrealBuildSource.cs
--- a/Engine/Source/Programs/AutomationTool/Gauntlet/Unreal/BuildSource/Gauntlet.UnrealBuildSource.cs
+++ b/Engine/Source/Programs/AutomationTool/Gauntlet/Unreal/BuildSource/Gauntlet.UnrealBuildSource.cs
@@ -522,6 +522,8 @@ namespace Gauntlet
 				throw new AutomationException("No build found that can support a role of {0}.", Role);
 			}
 
+			Log.Info("Selected build {Build} for test run.", Config.Build.ToString());
+
 			if (Role.Options != null)
 			{
 				UnrealTestConfiguration ConfigOptions = Role.Options as UnrealTestConfiguration;
