Adjusting FixPerforceCase boilerplate tags

#rnx

diff --git a/Engine/Source/Programs/AutomationTool/Scripts/FixPerforceCase.cs b/Engine/Source/Programs/AutomationTool/Scripts/FixPerforceCase.cs
--- a/Engine/Source/Programs/AutomationTool/Scripts/FixPerforceCase.cs
+++ b/Engine/Source/Programs/AutomationTool/Scripts/FixPerforceCase.cs
@@ -15,7 +15,7 @@ namespace AutomationTool
 	[Help("Files", "Pattern for files to match. Should be a full depot path with the correct case. May end with a wildcard.")]
 	class FixPerforceCase : BuildCommand
 	{
-		const string BoilerplateText = "\n\n#rb none\n#rnx\n#preflight skip";
+		const string BoilerplateText = "\n\n#rb none\n#rnx\n#preflight skip\n#jira none\n#submittool FixPerforceCase\n#okforgithub ignore";
 
 		/// <summary>
 		/// Main entry point for the command
