Gauntlet: Support recovery with a new device when using a device reservation block, make it so don't need to specify fullclean/skipinstall when running in a block vs individually
#rnx

diff --git a/Engine/Source/Programs/AutomationTool/AutomationUtils/DeviceReservation.cs b/Engine/Source/Programs/AutomationTool/AutomationUtils/DeviceReservation.cs
--- a/Engine/Source/Programs/AutomationTool/AutomationUtils/DeviceReservation.cs
+++ b/Engine/Source/Programs/AutomationTool/AutomationUtils/DeviceReservation.cs
@@ -38,6 +38,9 @@ namespace AutomationTool.DeviceReservation
 		private Reservation ActiveReservation;
 		private List<Device> ReservedDevices;
 
+		// Whether the reservation requires an installation
+		public bool? InstallRequired => ActiveReservation?.InstallRequired;
+
 		public IReadOnlyList<Device> Devices
 		{
 			get
@@ -244,6 +247,7 @@ namespace AutomationTool.DeviceReservation
 		public TimeSpan Duration { get; set; }
 		public Guid Guid { get; set; }
 		public static string ReservationDetails = "";
+		public bool? InstallRequired { get; set; } = null;
 
 		private sealed class CreateReservationData
 		{
diff --git a/Engine/Source/Programs/AutomationTool/Gauntlet/Framework/Devices/Gauntlet.DevicePool.cs b/Engine/Source/Programs/AutomationTool/Gauntlet/Framework/Devices/Gauntlet.DevicePool.cs
--- a/Engine/Source/Programs/AutomationTool/Gauntlet/Framework/Devices/Gauntlet.DevicePool.cs
+++ b/Engine/Source/Programs/AutomationTool/Gauntlet/Framework/Devices/Gauntlet.DevicePool.cs
@@ -51,7 +51,7 @@ namespace Gauntlet
 
 		public string Available { get; set; }
 
-		public bool RemoveOnShutdown { get; set; }
+		public bool RemoveOnShutdown { get; set; }		
 
 		public override string ToString()
 		{
@@ -965,6 +965,22 @@ namespace Gauntlet
 					return false;
 				}
 
+				if (DeviceReservation.InstallRequired == true)
+				{
+					UnrealAppConfig.ForceSkipInstall = false;
+					UnrealAppConfig.ForceFullClean = true;
+				}
+				else if (DeviceReservation.InstallRequired == false)
+				{
+					UnrealAppConfig.ForceSkipInstall = true;
+					UnrealAppConfig.ForceFullClean = false;
+				}
+				else
+				{
+					UnrealAppConfig.ForceSkipInstall = null;
+					UnrealAppConfig.ForceFullClean = null;
+				}
+
 				// Add target devices from reservation
 				List<ITargetDevice> ReservedDevices = new List<ITargetDevice>();
 				foreach (var Device in DeviceReservation.Devices)
@@ -974,7 +990,7 @@ namespace Gauntlet
 					Def.Name = Device.Name;
 					Def.Platform = DeviceMap.FirstOrDefault(Entry => Entry.Value == Device.Type.Replace("-DevKit", "", StringComparison.OrdinalIgnoreCase)).Key;
 					Def.DeviceData = Device.DeviceData;
-					Def.Model = Device.Model;
+					Def.Model = Device.Model;					
 
 					EPerfSpec Out = EPerfSpec.Unspecified;
 					if (!String.IsNullOrEmpty(Device.PerfSpec) && !Enum.TryParse<EPerfSpec>(Device.PerfSpec, true, out Out))
diff --git a/Engine/Source/Programs/AutomationTool/Gauntlet/Unreal/Base/Gauntlet.UnrealAppConfig.cs b/Engine/Source/Programs/AutomationTool/Gauntlet/Unreal/Base/Gauntlet.UnrealAppConfig.cs
--- a/Engine/Source/Programs/AutomationTool/Gauntlet/Unreal/Base/Gauntlet.UnrealAppConfig.cs
+++ b/Engine/Source/Programs/AutomationTool/Gauntlet/Unreal/Base/Gauntlet.UnrealAppConfig.cs
@@ -108,8 +108,22 @@ namespace Gauntlet
 		public IBuild Build { get; set; }
 
 		// Prevents installing a build on device
+		public bool SkipInstall => ForceSkipInstall.HasValue ? ForceSkipInstall.Value : _SkipInstall;
+
+		// Performs a full clean on the device before installing
+		public bool FullClean => ForceFullClean.HasValue ? ForceFullClean.Value : _FullClean;
+
+		// Force a full clean
+		public static bool? ForceFullClean = null;
+
+		// Force a skip install
+		public static bool? ForceSkipInstall = null;
+		
 		[AutoParamWithNames(false, "SkipInstall", "SkipDeploy", "SkipCopy")]
-		public bool SkipInstall { get; set; }
+		private bool _SkipInstall { get; set; }
+		
+		[AutoParamWithNames(false, "FullClean")]
+		private bool _FullClean { get; set; }
 
 		/// <summary>
 		/// Constructor that sets some required values to defaults
