Gauntlet: Android platform use AppConfig for fullclean instead of looking in globals params
#rnx

diff --git a/Engine/Source/Programs/AutomationTool/Gauntlet/Platform/Android/Gauntlet.TargetDeviceAndroid.cs b/Engine/Source/Programs/AutomationTool/Gauntlet/Platform/Android/Gauntlet.TargetDeviceAndroid.cs
--- a/Engine/Source/Programs/AutomationTool/Gauntlet/Platform/Android/Gauntlet.TargetDeviceAndroid.cs
+++ b/Engine/Source/Programs/AutomationTool/Gauntlet/Platform/Android/Gauntlet.TargetDeviceAndroid.cs
@@ -918,8 +918,8 @@ namespace Gauntlet
 			// clear all file store paths between installs:
 			RunAdbDeviceCommand(string.Format("shell rm -r {0}", DeviceExternalStorageSavedPath));
 			RunAdbDeviceCommand(string.Format("shell rm -r {0}", DeviceExternalFilesSavedPath));
-
-			if (Globals.Params.ParseParam("fullclean"))
+		
+			if (AppConfig.FullClean)
 			{
 				Log.Info("Fully cleaning console before install...");
 				RunAdbDeviceCommand(string.Format("shell rm -r {0}/UnrealGame/*", StorageLocation));
@@ -931,7 +931,7 @@ namespace Gauntlet
 			if (!AppConfig.SkipInstall)
 			{
 				if (Globals.Params.ParseParam("cleandevice")
-					|| Globals.Params.ParseParam("fullclean"))
+					|| AppConfig.FullClean)
 				{
 					Log.Info("Cleaning previous builds due to presence of -cleandevice");
 
