IAS - added UAT option -Upload=<args> for uploading ondemand content to S3 compatible endpoint

Use -Upload=LocalZen to upload to local Zen server

diff --git a/Engine/Source/Programs/AutomationTool/AutomationUtils/ProjectParams.cs b/Engine/Source/Programs/AutomationTool/AutomationUtils/ProjectParams.cs
--- a/Engine/Source/Programs/AutomationTool/AutomationUtils/ProjectParams.cs
+++ b/Engine/Source/Programs/AutomationTool/AutomationUtils/ProjectParams.cs
@@ -444,6 +444,7 @@ namespace AutomationTool
 			this.TraceFile = InParams.TraceFile;
 			this.SessionLabel = InParams.SessionLabel;
 			this.ProjectDescriptor = InParams.ProjectDescriptor;
+			this.Upload = InParams.Upload;
 		}
 
 		/// <summary>
@@ -612,8 +613,9 @@ namespace AutomationTool
 			string TraceFile = null,
 			string SessionLabel = null,
 			ParamList<string> InMapsToRebuildLightMaps = null,
-            ParamList<string> InMapsToRebuildHLOD = null,
-            ParamList<string> TitleID = null
+			ParamList<string> InMapsToRebuildHLOD = null,
+			ParamList<string> TitleID = null,
+			string Upload = null
 			)
 		{
 			//
@@ -801,6 +803,10 @@ namespace AutomationTool
 			this.Compressed = GetParamValueIfNotSpecified(Command, Compressed, this.Compressed, "compressed");
 			this.ForceUncompressed = GetParamValueIfNotSpecified(Command, ForceUncompressed, this.ForceUncompressed, "ForceUncompressed");
 			this.AdditionalPakOptions = ParseParamValueIfNotSpecified(Command, AdditionalPakOptions, "AdditionalPakOptions");
+			if (!string.IsNullOrEmpty(this.NoZenAutoLaunch))
+			{
+				this.AdditionalPakOptions += string.Format(" -NoZenAutoLaunch={0}", this.NoZenAutoLaunch);
+			}
 			this.AdditionalIoStoreOptions = ParseParamValueIfNotSpecified(Command, AdditionalIoStoreOptions, "AdditionalIoStoreOptions");
 			this.ForceOodleDllVersion = ParseParamValueIfNotSpecified(Command, ForceOodleDllVersion, "ForceOodleDllVersion");
 			this.IterativeCooking = GetParamValueIfNotSpecified(Command, IterativeCooking, this.IterativeCooking, new string[] { "iterativecooking", "iterate" });
@@ -1037,7 +1043,9 @@ namespace AutomationTool
 				this.SessionLabel += "=" + SessionLabel;	
 			}
 
-				if (ClientConfigsToBuild == null)
+			this.Upload = Command.ParseParamValue("upload");
+
+			if (ClientConfigsToBuild == null)
 			{
 				if (Command != null)
 				{
@@ -2330,6 +2338,9 @@ namespace AutomationTool
 		[Help("sessionlabel", "A label to pass to analytics")]
 		public string SessionLabel { get; set; }
 
+		[Help("upload", "Arguments for uploading on demand content")]
+		public string Upload { get; set; }
+
 		private List<SingleTargetProperties> DetectedTargets;
 		private Dictionary<UnrealTargetPlatform, ConfigHierarchy> LoadedEngineConfigs;
 		private Dictionary<UnrealTargetPlatform, ConfigHierarchy> LoadedGameConfigs;
diff --git a/Engine/Source/Programs/AutomationTool/Scripts/CopyBuildToStagingDirectory.Automation.cs b/Engine/Source/Programs/AutomationTool/Scripts/CopyBuildToStagingDirectory.Automation.cs
--- a/Engine/Source/Programs/AutomationTool/Scripts/CopyBuildToStagingDirectory.Automation.cs
+++ b/Engine/Source/Programs/AutomationTool/Scripts/CopyBuildToStagingDirectory.Automation.cs
@@ -3626,6 +3626,57 @@ namespace AutomationScripts
 				RunIoStore(Params, SC, IoStoreCommandsFileName, FinalOrderFiles, AdditionalArgs);
 			}
 
+			if (!string.IsNullOrWhiteSpace(Params.Upload))
+			{
+				StringBuilder UploadArgs = new StringBuilder();
+				if (Params.Upload.ToLower() == "localzen")
+				{
+					string ServiceUrl = "http://127.0.0.1:8558";
+					if (!string.IsNullOrWhiteSpace(Params.NoZenAutoLaunch))
+					{
+						ServiceUrl = Params.NoZenAutoLaunch.Trim();
+						if (!ServiceUrl.StartsWith("http://"))
+						{
+							int Sep = ServiceUrl.LastIndexOf(':');
+							if (Sep < 0)
+							{
+								ServiceUrl = "http://" + ServiceUrl + ":8558";
+							}
+							else
+							{
+								ServiceUrl = "http://" + ServiceUrl;
+							}
+						}
+					}
+
+					string PakPath = CombinePaths(SC.StageDirectory.FullName, SC.ShortProjectName, "Content", "Paks");
+					string CloudDir = CombinePaths(SC.StageDirectory.FullName, "Cloud");
+					string ConfigFilePath = CombinePaths(CloudDir, "IoStoreOnDemand.ini");
+
+					UploadArgs.Append(string.Format("-Upload={0}", PakPath));
+					UploadArgs.Append(string.Format(" -ServiceUrl={0}/obj/bucket", ServiceUrl));
+					UploadArgs.Append(" -Bucket=" + SC.ShortProjectName.ToLower());
+					UploadArgs.Append(" -ConfigFilePath=\"" + ConfigFilePath + "\"");
+					UploadArgs.Append(" -TargetPlatform=" + SC.StageTargetPlatform.PlatformType.ToString());
+					UploadArgs.Append(" -KeepContainerFiles");
+					UploadArgs.Append(" -KeepPakFiles");
+					UploadArgs.Append(" -BuildVersion=0");
+					UploadArgs.Append(" -WriteTocToDisk");
+					UploadArgs.Append(Params.AdditionalPakOptions);
+				}
+				else
+				{
+					UploadArgs.Append(Params.Upload);
+				}
+
+				if (UploadArgs.Length > 0)
+				{
+					string PakArgs = UploadArgs.ToString();
+					Logger.LogInformation("Running UnrealPak with arguments: {CommandletParams}", PakArgs);
+					RunAndLog(CmdEnv, GetUnrealPakLocation().FullName, PakArgs, Options: ERunOptions.Default | ERunOptions.UTF8Output);
+				}
+			}
+
 			// Do any additional processing on the command output
 			for (int Idx = 0; Idx < PakParamsList.Count; Idx++)
 			{
