Automation: Source indexing: Fixed source indexing working incorrectly for branches using the 'import' functionality on P4.

On Windows the code assumed that all source files come from the game's branch, which may not be true when importing parts of other branches. Modified the code to query the locally synced files using 'p4 have' which returns a list of correct local - depot mappings. Minor but related fixes on other platforms.

diff --git a/Engine/Source/Programs/AutomationTool/AutomationUtils/P4Utils.cs b/Engine/Source/Programs/AutomationTool/AutomationUtils/P4Utils.cs
--- a/Engine/Source/Programs/AutomationTool/AutomationUtils/P4Utils.cs
+++ b/Engine/Source/Programs/AutomationTool/AutomationUtils/P4Utils.cs
@@ -225,6 +225,20 @@ namespace AutomationTool
 		public string Path;
 	}
 
+	public class P4HaveRecord
+	{
+		public string DepotFile;
+		public string ClientFile;
+		public int    Revision;
+
+		public P4HaveRecord(string DepotFile, string ClientFile, int Revision)
+		{
+			this.DepotFile = DepotFile;
+			this.ClientFile = ClientFile;
+			this.Revision = Revision;
+		}
+	}
+
 	public class P4Spec
 	{
 		public List<KeyValuePair<string, string>> Sections;
@@ -4092,6 +4106,58 @@ namespace AutomationTool
 			return Records;
 		}
 
+		static readonly Regex HaveFilesOutputPattern = new Regex(@"(?<depot>.+)#(?<revision>\d+) - (?<client>.+)", RegexOptions.Compiled);
+
+		/// <summary>
+		/// Get a information about synced files (p4 have) matching the given pattern.
+		/// The information for each file includes the local path, depot path and the synced revision number.
+		/// </summary>
+		/// <param name="Pattern">Local or depot path e.g. .../Source/...</param>
+		/// <returns>Collection of information about synced files.</returns>
+		public List<P4HaveRecord> HaveFiles(string Pattern)
+		{
+			List<P4HaveRecord> Files = new ();
+
+			string CommandToRun = $"have \"{Pattern}\"";
+			string CommandLine  = $"{GlobalOptions} {CommandToRun}";
+
+			// We run the process manually because 'p4 have' may generate a very long output (even hundreds of MB)
+			// and gathering all of this into a single string is wasteful. Instead, we process the output line by line
+			// as it arrives. Neither P4Output nor CommandUtils.Run allow achieving this.
+			Process Process = new Process();
+
+			Process.StartInfo.FileName = HostPlatform.Current.P4Exe;
+			Process.StartInfo.Arguments = CommandLine;
+			Process.StartInfo.UseShellExecute = false;
+			Process.StartInfo.RedirectStandardOutput = true;
+			Process.StartInfo.StandardOutputEncoding = Encoding.UTF8;
+			Process.StartInfo.RedirectStandardError = false;
+			Process.StartInfo.RedirectStandardInput = false;
+			Process.StartInfo.CreateNoWindow = true;
+			Process.OutputDataReceived += (s, e) =>
+			{
+				if (e.Data != null)
+				{
+					Match Match = HaveFilesOutputPattern.Match(e.Data);
+
+					if (Match.Success)
+					{
+						Files.Add(new P4HaveRecord(Match.Groups["depot"].Value, Match.Groups["client"].Value, Int32.Parse(Match.Groups["revision"].Value)));
+					}
+				}
+			};
+			Process.Start();
+			Process.BeginOutputReadLine();
+			Process.WaitForExit();
+
+			if (Process.ExitCode != 0)
+			{
+				throw new P4Exception($"p4.exe {CommandLine} failed with code {Process.ExitCode}.");
+			}
+
+			return Files;
+		}
+
 		/// <summary>
 		/// Parse a nullable boolean
 		/// </summary>
diff --git a/Engine/Source/Programs/AutomationTool/Win/WinPlatform.Automation.cs b/Engine/Source/Programs/AutomationTool/Win/WinPlatform.Automation.cs
--- a/Engine/Source/Programs/AutomationTool/Win/WinPlatform.Automation.cs
+++ b/Engine/Source/Programs/AutomationTool/Win/WinPlatform.Automation.cs
@@ -747,6 +747,33 @@ public class Win64Platform : Platform
 		return false;
 	}
 
+	/// <summary>
+	/// Build a database of source code files in the current Perforce workspace.
+	/// By relying on the standard layout for UE projects i.e. the fact that source code is in Source directories,
+	/// we may very significantly reduce the about of data sent to us from the server.
+	/// </summary>
+	/// <param name="Pattern">Perforce pattern path to query e.g. //UE/Branch/.../Source/...</param>
+	/// <returns></returns>
+	protected static Dictionary<string, P4HaveRecord> BuildSourceDatabase(string Pattern)
+	{
+		List<P4HaveRecord> Files = null;
+
+		P4Connection DefaultConnection = new P4Connection(User: null, Client: null, ServerAndPort: null);
+
+		try
+		{
+			Files = DefaultConnection.HaveFiles(Pattern);
+		}
+		catch (P4Exception e)
+		{
+			Logger.LogError("Failed to fetch source code information from Perforce for '{Pattern}' ({Message}).", Pattern, e.Message);
+
+			return null;
+		}
+
+		return Files.ToDictionary(file => file.ClientFile, file => file, StringComparer.InvariantCultureIgnoreCase);
+	}
+
 	/// <summary>
 	/// 
 	/// </summary>
@@ -759,6 +786,15 @@ public class Win64Platform : Platform
 	{
 		Logger.LogInformation("Adding source control information to PDB files...");
 
+		string DepotFilter = ".../Source/...";
+
+		Dictionary<string, P4HaveRecord> SourceDatabase = BuildSourceDatabase(DepotFilter);
+
+		if (SourceDatabase == null)
+		{
+			throw new AutomationException($"Failed to query the source code information for '{DepotFilter}'.");
+		}
+
 		// Get the PDBSTR.EXE path, using the latest SDK version we can find.
 		FileReference PdbStrExe = GetPdbStrExe();
 
@@ -769,21 +805,36 @@ public class Win64Platform : Platform
 		// Generate the SRCSRV.INI file
 		using (StreamWriter Writer = new StreamWriter(SrcSrvIni.FullName))
 		{
+			int MissingFilesCount = 0;
+
 			Writer.WriteLine("SRCSRV: ini------------------------------------------------");
 			Writer.WriteLine("VERSION=1");
 			Writer.WriteLine("VERCTRL=Perforce");
 			Writer.WriteLine("SRCSRV: variables------------------------------------------");
 			Writer.WriteLine("SRCSRVTRG=%sdtrg%");
 			Writer.WriteLine("SRCSRVCMD=%sdcmd%");
-			Writer.WriteLine("SDCMD=p4.exe print -o %srcsrvtrg% \"{0}/%var2%@{1}\"", Branch.TrimEnd('/'), Change);
-			Writer.WriteLine("SDTRG=%targ%\\{0}\\{1}\\%fnbksl%(%var2%)", Branch.Replace('/', '+'), Change);
+			Writer.WriteLine("SDCMD=p4.exe print -o %srcsrvtrg% \"//%var2%#%var3%\"");
+			Writer.WriteLine("SDTRG=%targ%\\%fnbksl%(%var2%)#%var3%");
 			Writer.WriteLine("SRCSRV: source files ---------------------------------------");
 			foreach (FileReference SourceFile in SourceFiles)
 			{
-				string RelativeSourceFile = SourceFile.MakeRelativeTo(Unreal.RootDirectory);
-				Writer.WriteLine("{0}*{1}", SourceFile.FullName, RelativeSourceFile.Replace('\\', '/'));
+				P4HaveRecord SourceInfo;
+
+				if (SourceDatabase.TryGetValue(SourceFile.FullName, out SourceInfo))
+				{
+					Writer.WriteLine("{0}*{1}*{2}", SourceFile.FullName, SourceInfo.DepotFile.Replace("//", ""), SourceInfo.Revision);
+				}
+				else
+				{
+					++MissingFilesCount;
+				}
 			}
 			Writer.WriteLine("SRCSRV: end------------------------------------------------");
+
+			if (MissingFilesCount > 0)
+			{
+				Logger.LogInformation("Skipped {MissingFilesCount} files (out of {SourceFileCount}) for which source control files couldn't be located.", MissingFilesCount, SourceFiles.Count());
+			}
 		}
 
 		// Execute PDBSTR on the PDB files in parallel.
