Extend the -EnableStereoTestVariants Gauntlet argument to affect UE.TargetAutomation steps as well as UE.EditorAutomation.

diff --git a/Engine/Source/Programs/AutomationTool/Gauntlet/Unreal/Automation/UE.Automation.cs b/Engine/Source/Programs/AutomationTool/Gauntlet/Unreal/Automation/UE.Automation.cs
--- a/Engine/Source/Programs/AutomationTool/Gauntlet/Unreal/Automation/UE.Automation.cs
+++ b/Engine/Source/Programs/AutomationTool/Gauntlet/Unreal/Automation/UE.Automation.cs
@@ -244,8 +244,16 @@ namespace UE
 
 				if (ConfigRole.RoleType.IsClient())
 				{
-					// have the client list the tests it knows about. useful for troubleshooting discrepencies
-					AppConfig.CommandLine += string.Format(" -sessionid={0} -messaging -log -TcpMessagingConnect={1}:6666 -ExecCmds=\"Automation list\"", SessionID, HostIP);
+					// Have the client list the tests it knows about. Useful for troubleshooting discrepencies
+					string ClientAutomationTestArgument = "List;";
+
+					// Make sure the stereo test setting propogates to the client
+					if (EnableStereoTestVariants)
+					{
+						ClientAutomationTestArgument += "EnableStereoTests;";
+					}
+
+					AppConfig.CommandLine += string.Format(" -sessionid={0} -messaging -log -TcpMessagingConnect={1}:6666 -ExecCmds=\"Automation {2}\"", SessionID, HostIP, ClientAutomationTestArgument);
 				}
 				else if (ConfigRole.RoleType.IsEditor())
 				{
