Skip PAK file index encryption when -SkipEncryption is set

We have an option to skip encryption for PAK files but is it not consistently applied. This shows up for example when building the Content Worker: it sets the -SkipEncrpytion parameter in its packaging stage, yet the PAK tool is invoked with -encryptindex. This is relevant for modular builds, since those might not have encryption keys loaded at the time that the first PAK files are loaded.

diff --git a/Engine/Source/Programs/AutomationTool/Scripts/CopyBuildToStagingDirectory.Automation.cs b/Engine/Source/Programs/AutomationTool/Scripts/CopyBuildToStagingDirectory.Automation.cs
--- a/Engine/Source/Programs/AutomationTool/Scripts/CopyBuildToStagingDirectory.Automation.cs
+++ b/Engine/Source/Programs/AutomationTool/Scripts/CopyBuildToStagingDirectory.Automation.cs
@@ -3444,8 +3444,8 @@ namespace AutomationScripts
 								PakParams.bOnDemand,
 								bCompressContainers,
 								StageLooseFileRootPath,
-								CryptoSettings,
-								PakParams.EncryptionKeyGuid,
+								Params.SkipEncryption ? null : CryptoSettings,
+								Params.SkipEncryption ? "" : PakParams.EncryptionKeyGuid,
 								ContainerPatchSourcePath,
 								bGenerateDiffPatch,
 								Params.HasDLCName));
@@ -3489,7 +3489,15 @@ namespace AutomationScripts
 				}
 
 				StringBuilder Arguments = new StringBuilder(MakePathSafeToUseWithCommandLine(Params.RawProjectPath.FullName));
-				Arguments.AppendFormat(" {0}", GetCommonUnrealPakArguments(PrimaryOrderFiles, CommonAdditionalArgs, CryptoSettings, CryptoKeysCacheFilename, SecondaryOrderFiles, Params.Unattended));
+				string CommonArguments = GetCommonUnrealPakArguments(
+					PrimaryOrderFiles,
+					CommonAdditionalArgs,
+					Params.SkipEncryption ? null : CryptoSettings,
+					Params.SkipEncryption ? null : CryptoKeysCacheFilename,
+					SecondaryOrderFiles,
+					Params.Unattended
+				);
+				Arguments.AppendFormat(" {0}", CommonArguments);
 				Arguments.AppendFormat(" -CreateMultiple={0}", CommandUtils.MakePathSafeToUseWithCommandLine(PakCommandsFileName));
 				RunAndLog(CmdEnv, GetUnrealPakLocation().FullName, Arguments.ToString(), Options: ERunOptions.Default | ERunOptions.UTF8Output);
 			}
