Adding SplitFirst and SplitLast StringOps to BuildGraph, which return the first and last element (respectively) of the split input string.

diff --git a/Engine/Source/Programs/AutomationTool/BuildGraph/BgScriptReader.cs b/Engine/Source/Programs/AutomationTool/BuildGraph/BgScriptReader.cs
--- a/Engine/Source/Programs/AutomationTool/BuildGraph/BgScriptReader.cs
+++ b/Engine/Source/Programs/AutomationTool/BuildGraph/BgScriptReader.cs
@@ -869,9 +869,9 @@ namespace AutomationTool
 				string input = ReadAttribute(element, "Input");
 				string method = ReadAttribute(element, "Method");
 				string output = ReadAttribute(element, "Output");
-
+				
 				string operationResult = string.Empty;
-
+				
 				string[] arguments = { };
 
 				const string ArgumentsName = "Arguments";
@@ -893,6 +893,20 @@ namespace AutomationTool
 						}
 						operationResult = input.Replace(arguments[0], arguments[1]);
 						break;
+					case "SplitFirst": 
+						if (arguments.Length != 1)
+						{
+							throw new AutomationException($"String operation 'SplitFirst' requires exactly 1 argument.");
+						}
+						operationResult = input.Split(arguments[0]).First();
+						break;
+					case "SplitLast": 
+						if (arguments.Length != 1)
+						{
+							throw new AutomationException($"String operation 'SplitLast' requires exactly 1 argument.");
+						}
+						operationResult = input.Split(arguments[0]).Last();
+						break;
 					default: throw new AutomationException($"String operation '{method}' not available.");
 				}
 				SetPropertyValue(element, output, operationResult);
