Fix an issue where a Gauntlet test would fails to run when running an EditorGame, EditorServer or no Engine role

diff --git a/Engine/Source/Programs/AutomationTool/Gauntlet/Unreal/Base/Gauntlet.UnrealTestConfiguration.cs b/Engine/Source/Programs/AutomationTool/Gauntlet/Unreal/Base/Gauntlet.UnrealTestConfiguration.cs
--- a/Engine/Source/Programs/AutomationTool/Gauntlet/Unreal/Base/Gauntlet.UnrealTestConfiguration.cs
+++ b/Engine/Source/Programs/AutomationTool/Gauntlet/Unreal/Base/Gauntlet.UnrealTestConfiguration.cs
@@ -874,24 +874,31 @@ namespace Gauntlet
 		/// <returns></returns>
 		public UnrealTestRole GetMainRequiredRole()
 		{
-			var ClientRoles = GetRequiredRoles(UnrealTargetRole.Client);
-			if (ClientRoles.Any())
-			{
-				return ClientRoles.First();
-			}
-			var ServerRoles = GetRequiredRoles(UnrealTargetRole.Server);
-			if (ServerRoles.Any())
-			{
-				return ServerRoles.First();
+			var PriorityList = new UnrealTargetRole[] {
+				UnrealTargetRole.Client,
+				UnrealTargetRole.EditorGame,
+				UnrealTargetRole.Server,
+				UnrealTargetRole.EditorServer,
+				UnrealTargetRole.Editor,
+				UnrealTargetRole.CookedEditor
+			};
+			foreach (UnrealTargetRole TargetRole in PriorityList)
+			{
+				IEnumerable<UnrealTestRole> Roles = GetRequiredRoles(TargetRole);
+				if (Roles.Any())
+				{
+					return Roles.First();
+				}
 			}
-			var EditorRoles = GetRequiredRoles(UnrealTargetRole.Editor);
-			if (EditorRoles.Any())
+
+			if (RequiredRoles.Any())
 			{
-				return EditorRoles.First();
+				var RoleEnumerator = RequiredRoles.Values.GetEnumerator();
+				RoleEnumerator.MoveNext();
+				return RoleEnumerator.Current.First();
 			}
-			var RoleEnumerator = RequiredRoles.Values.GetEnumerator();
-			RoleEnumerator.MoveNext();
-			return RoleEnumerator.Current.First();
+
+			return new UnrealTestRole(UnrealTargetRole.Unknown, null);
 		}
 
 		/// <summary>
