Use the ZenProject from the export source rather than the autogenerated value when available.

diff --git a/Engine/Source/Programs/AutomationTool/BuildGraph/Tasks/ZenExportSnapshotTask.cs b/Engine/Source/Programs/AutomationTool/BuildGraph/Tasks/ZenExportSnapshotTask.cs
--- a/Engine/Source/Programs/AutomationTool/BuildGraph/Tasks/ZenExportSnapshotTask.cs
+++ b/Engine/Source/Programs/AutomationTool/BuildGraph/Tasks/ZenExportSnapshotTask.cs
@@ -471,7 +471,7 @@ namespace AutomationTool.Tasks
 
 					break;
 				case SnapshotStorageType.File:
-					string ProjectId = ProjectUtils.GetProjectPathId(ProjectFile);
+					string DefaultProjectId = ProjectUtils.GetProjectPathId(ProjectFile);
 					ExportIndex = 0;
 					foreach (ExportSourceData ExportSource in ExportSources)
 					{
@@ -488,6 +488,7 @@ namespace AutomationTool.Tasks
 							DestinationFileName = Parameters.DestinationFileName.Replace("{Platform}", ExportSource.TargetPlatform, StringComparison.InvariantCultureIgnoreCase);
 						}
 
+						string ProjectId = string.IsNullOrEmpty(ExportSource.ProjectId) ? DefaultProjectId : ExportSource.ProjectId;
 						string BaseNameArg = string.Empty;
 						DirectoryReference PlatformDestinationFileDir = new DirectoryReference(Parameters.DestinationFileDir.FullName.Replace("{Platform}", ExportSource.TargetPlatform, StringComparison.InvariantCultureIgnoreCase));
 						if ((ExportSource.SnapshotBaseDescriptor != null) && !string.IsNullOrEmpty(ExportSource.SnapshotBaseDescriptor.Directory) && !string.IsNullOrEmpty(ExportSource.SnapshotBaseDescriptor.Filename))
