Add ParamKeyValueDelimiter field to ParamHelpAttribute to be used with flags that use alternative strings to separate the key value pairs of arguments. For example, timeout /t 10

diff --git a/Engine/Source/Programs/AutomationTool/AutomationUtils/HelpUtils.cs b/Engine/Source/Programs/AutomationTool/AutomationUtils/HelpUtils.cs
--- a/Engine/Source/Programs/AutomationTool/AutomationUtils/HelpUtils.cs
+++ b/Engine/Source/Programs/AutomationTool/AutomationUtils/HelpUtils.cs
@@ -247,6 +247,15 @@ namespace AutomationTool
 				}
 			}
 		}
+
+		/// <summary>
+		/// This string/character separates the flag from the value; ex. -foo=bar
+		/// </summary>
+		public string ParamKeyValueDelimiter
+		{
+			get;
+			set;
+		} = "=";
 	}
 
 /// <summary>
