UAT - use correct host adress when deploying to non host platforms

diff --git a/Engine/Source/Programs/AutomationTool/Scripts/CopyBuildToStagingDirectory.Automation.cs b/Engine/Source/Programs/AutomationTool/Scripts/CopyBuildToStagingDirectory.Automation.cs
--- a/Engine/Source/Programs/AutomationTool/Scripts/CopyBuildToStagingDirectory.Automation.cs
+++ b/Engine/Source/Programs/AutomationTool/Scripts/CopyBuildToStagingDirectory.Automation.cs
@@ -103,6 +103,66 @@ namespace AutomationScripts
 			public int Priority;
 		};
 
+		private static List<string> GetHostAddresses(DeploymentContext SC)
+		{
+			List<string> HostAddresses = new List<string>();
+
+			// Add localhost first for host platforms and skip it completely for other platforms.
+			// Any Platform can implement ModifyFileHostAddresses to tweak this default behavior.
+			string LocalHost = "127.0.0.1";
+			if (BuildHostPlatform.Current.Platform == SC.StageTargetPlatform.PlatformType)
+			{
+				HostAddresses.Add(LocalHost);
+			}
+			
+			NetworkInterface[] Interfaces = NetworkInterface.GetAllNetworkInterfaces();
+			foreach (NetworkInterface Adapter in Interfaces)
+			{
+				if (BuildHostPlatform.Current.Platform != SC.StageTargetPlatform.PlatformType && Adapter.NetworkInterfaceType == NetworkInterfaceType.Loopback)
+				{
+					continue;
+				}
+				
+				if (BuildHostPlatform.Current.Platform == UnrealTargetPlatform.Mac)
+				{
+					if (Adapter.NetworkInterfaceType == NetworkInterfaceType.Loopback)
+					{
+						continue;
+					}
+				}
+				else
+				{
+					if (Adapter.OperationalStatus != OperationalStatus.Up)
+					{
+						continue;
+					}
+				}
+
+				IPInterfaceProperties IP = Adapter.GetIPProperties();
+				foreach (UnicastIPAddressInformation UnicastAddress in IP.UnicastAddresses)
+				{
+					if (!InternalUtils.IsDnsEligible(UnicastAddress))
+					{
+						continue;
+					}
+
+					if (UnicastAddress.Address.AddressFamily != System.Net.Sockets.AddressFamily.InterNetwork)
+					{
+						continue;
+					}
+
+					string HostAddress = UnicastAddress.Address.ToString();
+					if (HostAddress == LocalHost)
+					{
+						continue;
+					}
+					HostAddresses.Add(HostAddress);
+				}
+			}
+
+			return HostAddresses.ToList();
+		}
+
 		/// <returns>The path for the BuildPatchTool executable depending on host platform.</returns>
 		private static string GetBuildPatchToolExecutable()
 		{
@@ -3639,7 +3699,15 @@ namespace AutomationScripts
 				StringBuilder UploadArgs = new StringBuilder();
 				if (Params.Upload.ToLower() == "localzen")
 				{
-					string ServiceUrl = "http://127.0.0.1:8558";
+					int Port = 8558;
+					string ServiceUrl = string.Format("http://127.0.0.1:{0}", Port);
+					List<string> Addresses = GetHostAddresses(SC);
+					if (Addresses.Count > 0)
+					{
+						//TODO: Support multiple addresses
+						ServiceUrl = string.Format("http://{0}:{1}", Addresses[0], Port);
+					}
+
 					if (!string.IsNullOrWhiteSpace(Params.NoZenAutoLaunch))
 					{
 						ServiceUrl = Params.NoZenAutoLaunch.Trim();
@@ -3648,7 +3716,7 @@ namespace AutomationScripts
 							int Sep = ServiceUrl.LastIndexOf(':');
 							if (Sep < 0)
 							{
-								ServiceUrl = "http://" + ServiceUrl + ":8558";
+								ServiceUrl = string.Format("http://{0}:{1}", ServiceUrl, Port);
 							}
 							else
 							{
