[Gauntlet] Create TargetDeviceDesktopCommon to abstract re-used behavior for Windows, Linux, and Mac

diff --git a/Engine/Source/Programs/AutomationTool/Gauntlet/Framework/Base/Gauntlet.TargetDeviceDesktopCommon.cs b/Engine/Source/Programs/AutomationTool/Gauntlet/Framework/Base/Gauntlet.TargetDeviceDesktopCommon.cs
new file mode 100644
--- /dev/null
+++ b/Engine/Source/Programs/AutomationTool/Gauntlet/Framework/Base/Gauntlet.TargetDeviceDesktopCommon.cs
@@ -0,0 +1,345 @@
+// Copyright Epic Games, Inc. All Rights Reserved.
+
+using AutomationTool;
+using EpicGames.Core;
+using System;
+using System.Collections.Generic;
+using System.IO;
+using System.Threading;
+using UnrealBuildTool;
+
+namespace Gauntlet
+{
+	public abstract class TargetDeviceDesktopCommon : ITargetDevice
+	{
+		#region ITargetDevice
+		public string Name { get; protected set; }
+		public UnrealTargetPlatform? Platform { get; protected set; }
+		public CommandUtils.ERunOptions RunOptions { get; set; }
+		public bool IsAvailable => true;
+		public bool IsConnected => true;
+		public bool IsOn => true;
+		public bool Connect() => true;
+		public bool Disconnect(bool bForce = false) => true;
+		public bool PowerOn() => true;
+		public bool PowerOff() => true;
+		public bool Reboot() => true;
+		public Dictionary<EIntendedBaseCopyDirectory, string> GetPlatformDirectoryMappings() => LocalDirectoryMappings;
+		public override string ToString() => Name;
+		#endregion
+
+		#region IDisposable Support
+		private bool DisposedValue = false; // To detect redundant calls
+
+		protected virtual void Dispose(bool Disposing)
+		{
+			if (!DisposedValue)
+			{
+				if (Disposing)
+				{
+					// TODO: dispose managed state (managed objects).
+				}
+
+				// TODO: free unmanaged resources (unmanaged objects) and override a finalizer below.
+				// TODO: set large fields to null.
+
+				DisposedValue = true;
+			}
+		}
+
+		// This code added to correctly implement the disposable pattern.
+		public void Dispose()
+		{
+			// Do not change this code. Put cleanup code in Dispose(bool disposing) above.
+			Dispose(true);
+			// TODO: uncomment the following line if the finalizer is overridden above.
+			// GC.SuppressFinalize(this);
+		}
+		#endregion
+
+		public string LocalCachePath { get; protected set; }
+
+		public string UserDir { get; protected set; }
+
+		// TODO - move this be part of ITargetDevice
+		protected Dictionary<EIntendedBaseCopyDirectory, string> LocalDirectoryMappings { get; set; }
+
+		public TargetDeviceDesktopCommon(string InName, string InCacheDir)
+		{
+			Name = InName;
+			LocalCachePath = InCacheDir;
+			UserDir = Path.Combine(LocalCachePath, "UserDir");
+
+			LocalDirectoryMappings = new Dictionary<EIntendedBaseCopyDirectory, string>();
+		}
+
+		public virtual IAppInstall InstallApplication(UnrealAppConfig AppConfig)
+		{
+			switch (AppConfig.Build)
+			{
+				case NativeStagedBuild:
+					return InstallNativeStagedBuild(AppConfig, AppConfig.Build as NativeStagedBuild);
+
+				case StagedBuild:
+					return InstallStagedBuild(AppConfig, AppConfig.Build as StagedBuild);
+
+				case EditorBuild:
+					return InstallEditorBuild(AppConfig, AppConfig.Build as EditorBuild);
+
+				default:
+					throw new AutomationException("{0} is an invalid build type!", AppConfig.Build.ToString());
+			}
+		}
+
+		public virtual void ClearSavedDirectory(UnrealAppConfig AppConfiguration) { }
+
+		public virtual void CopyAppConfigurationFiles(UnrealAppConfig AppConfiguration) { }
+
+		public abstract IAppInstance Run(IAppInstall Install);
+
+		public virtual void PopulateDirectoryMappings(string BaseDirectory)
+		{
+			LocalDirectoryMappings.Clear();
+
+			LocalDirectoryMappings.Add(EIntendedBaseCopyDirectory.Build, Path.Combine(BaseDirectory, "Build"));
+			LocalDirectoryMappings.Add(EIntendedBaseCopyDirectory.Binaries, Path.Combine(BaseDirectory, "Binaries"));
+			LocalDirectoryMappings.Add(EIntendedBaseCopyDirectory.Config, Path.Combine(BaseDirectory, "Saved", "Config"));
+            LocalDirectoryMappings.Add(EIntendedBaseCopyDirectory.Content, Path.Combine(BaseDirectory, "Content"));
+            LocalDirectoryMappings.Add(EIntendedBaseCopyDirectory.Demos, Path.Combine(BaseDirectory, "Saved", "Demos"));
+			LocalDirectoryMappings.Add(EIntendedBaseCopyDirectory.PersistentDownloadDir, Path.Combine(BaseDirectory, "Saved", "PersistentDownloadDir"));
+			LocalDirectoryMappings.Add(EIntendedBaseCopyDirectory.Profiling, Path.Combine(BaseDirectory, "Saved", "Profiling"));
+            LocalDirectoryMappings.Add(EIntendedBaseCopyDirectory.Saved, Path.Combine(BaseDirectory, "Saved"));
+		}
+
+		protected void CopyAdditionalFiles(UnrealAppConfig AppConfig)
+		{
+			if (AppConfig.FilesToCopy != null)
+			{
+				foreach (UnrealFileToCopy FileToCopy in AppConfig.FilesToCopy)
+				{
+					string PathToCopyTo = Path.Combine(LocalDirectoryMappings[FileToCopy.TargetBaseDirectory], FileToCopy.TargetRelativeLocation);
+					if (File.Exists(FileToCopy.SourceFileLocation))
+					{
+						FileInfo SrcInfo = new FileInfo(FileToCopy.SourceFileLocation);
+						SrcInfo.IsReadOnly = false;
+						string DirectoryToCopyTo = Path.GetDirectoryName(PathToCopyTo);
+						if (!Directory.Exists(DirectoryToCopyTo))
+						{
+							Directory.CreateDirectory(DirectoryToCopyTo);
+						}
+						if (File.Exists(PathToCopyTo))
+						{
+							FileInfo ExistingFile = new FileInfo(PathToCopyTo);
+							ExistingFile.IsReadOnly = false;
+						}
+						SrcInfo.CopyTo(PathToCopyTo, true);
+						Log.Info("Copying {0} to {1}", FileToCopy.SourceFileLocation, PathToCopyTo);
+					}
+					else
+					{
+						Log.Warning(KnownLogEvents.Gauntlet_DeviceEvent, "File to copy {File} not found", FileToCopy);
+					}
+				}
+			}
+		}
+
+		// TODO - b.lienau: implement these at desktop level and remove implementations from each desktop
+		protected abstract IAppInstall InstallNativeStagedBuild(UnrealAppConfig AppConfig, NativeStagedBuild Build);
+		protected abstract IAppInstall InstallStagedBuild(UnrealAppConfig AppConfig, StagedBuild Build);
+		protected abstract IAppInstall InstallEditorBuild(UnrealAppConfig AppConfig, EditorBuild Build);
+	}
+
+	public abstract class DesktopCommonAppInstall<DesktopTargetDevice> : IAppInstall where DesktopTargetDevice : TargetDeviceDesktopCommon
+	{
+		public string Name { get; }
+
+		public ITargetDevice Device => DesktopDevice;
+
+		public DesktopTargetDevice DesktopDevice { get; protected set; }
+
+		public string ArtifactPath;
+
+		public string ExecutablePath;
+
+		public string ProjectName;
+
+		public string WorkingDirectory;
+
+		public virtual string CommandArguments { get; set; }
+
+		public CommandUtils.ERunOptions RunOptions { get; set; }
+
+		public DesktopCommonAppInstall(string InName, string InProjectName, DesktopTargetDevice InDevice)
+		{
+			Name = InName;
+			ProjectName = InProjectName;
+			DesktopDevice = InDevice;
+			CommandArguments = string.Empty;
+			RunOptions = CommandUtils.ERunOptions.NoWaitForExit;
+		}
+
+		public virtual IAppInstance Run()
+		{
+			return Device.Run(this);
+		}
+
+		public virtual void CleanDeviceArtifacts()
+		{
+			if (!string.IsNullOrEmpty(ArtifactPath) && Directory.Exists(ArtifactPath))
+			{
+				try
+				{
+					Log.Info("Clearing device artifacts path {0} for {1}", ArtifactPath, Device.Name);
+					Directory.Delete(ArtifactPath, true);
+				}
+				catch (Exception Ex)
+				{
+					Log.Info(KnownLogEvents.Gauntlet_DeviceEvent, "First attempt at clearing artifact path {0} failed - trying again", ArtifactPath);
+					if (!ForceCleanDeviceArtifacts())
+					{
+						Log.Warning(KnownLogEvents.Gauntlet_DeviceEvent, "Failed to delete {File}. {Exception}", ArtifactPath, Ex.Message);
+					}
+				}
+			}
+		}
+
+		public virtual bool ForceCleanDeviceArtifacts()
+		{
+			DirectoryInfo ClientTempDirInfo = new DirectoryInfo(ArtifactPath) { Attributes = FileAttributes.Normal };
+			Log.Info(KnownLogEvents.Gauntlet_DeviceEvent, "Setting files in device artifacts {0} to have normal attributes (no longer read-only).", ArtifactPath);
+			foreach (FileSystemInfo info in ClientTempDirInfo.GetFileSystemInfos("*", SearchOption.AllDirectories))
+			{
+				info.Attributes = FileAttributes.Normal;
+			}
+			try
+			{
+				Log.Info(KnownLogEvents.Gauntlet_DeviceEvent, "Clearing device artifact path {0} (force)", ArtifactPath);
+				Directory.Delete(ArtifactPath, true);
+			}
+			catch (Exception Ex)
+			{
+				Log.Warning(KnownLogEvents.Gauntlet_DeviceEvent, "Failed to force delete artifact path {File}. {Exception}", ArtifactPath, Ex.Message);
+				return false;
+			}
+			return true;
+		}
+	}
+
+	public abstract class DesktopCommonAppInstance<DesktopAppInstall, DesktopTargetDevice> : LocalAppProcess
+		where DesktopAppInstall : DesktopCommonAppInstall<DesktopTargetDevice>
+		where DesktopTargetDevice : TargetDeviceDesktopCommon
+	{
+		public override string ArtifactPath => Install.ArtifactPath;
+
+		public override ITargetDevice Device => Install.Device;
+
+		protected DesktopAppInstall Install;
+
+		public DesktopCommonAppInstance(DesktopAppInstall InInstall, IProcessResult InProcess, string ProcessLogFile = null)
+			: base(InProcess, InInstall.CommandArguments, ProcessLogFile)
+		{
+			Install = InInstall;
+		}
+	}
+
+	public abstract class LocalAppProcess : IAppInstance
+	{
+		public IProcessResult ProcessResult { get; private set; }
+
+		public bool HasExited { get { return ProcessResult.HasExited; } }
+
+		public bool WasKilled { get; protected set; }
+
+		public string StdOut { get { return string.IsNullOrEmpty(ProcessLogFile) ? ProcessResult.Output : ProcessLogOutput; } }
+
+		public int ExitCode { get { return ProcessResult.ExitCode; } }
+
+		public string CommandLine { get; private set; }
+
+		public LocalAppProcess(IProcessResult InProcess, string InCommandLine, string InProcessLogFile = null)
+		{
+			this.CommandLine = InCommandLine;
+			this.ProcessResult = InProcess;
+			this.ProcessLogFile = InProcessLogFile;
+
+			// start reader thread if logging to a file
+			if (!string.IsNullOrEmpty(InProcessLogFile))
+			{
+				new Thread(LogFileReaderThread).Start();
+			}
+		}
+
+		public int WaitForExit()
+		{
+			if (!HasExited)
+			{
+				ProcessResult.WaitForExit();
+			}
+
+			return ExitCode;
+		}
+
+		virtual public void Kill()
+		{
+			if (!HasExited)
+			{
+				WasKilled = true;
+				ProcessResult.ProcessObject.Kill(true);
+			}
+		}
+
+		/// <summary>
+		/// Reader thread when logging to file
+		/// </summary>
+		void LogFileReaderThread()
+		{
+			// Wait for the processes log file to be created
+			while (!File.Exists(ProcessLogFile) && !HasExited)
+			{
+				Thread.Sleep(2000);
+			}
+
+			// Check whether the process exited before log file was created (this can happen for example if a server role exits and forces client to shutdown)
+			if (!File.Exists(ProcessLogFile))
+			{
+				ProcessLogOutput += "Process exited before log file created";
+				return;
+			}
+
+			Thread.Sleep(1000);
+
+			using (FileStream ProcessLog = File.Open(ProcessLogFile, FileMode.Open, FileAccess.Read, FileShare.ReadWrite))
+			{
+				StreamReader LogReader = new StreamReader(ProcessLog);
+				// Read until the process has exited
+				do
+				{
+					Thread.Sleep(250);
+
+					while (!LogReader.EndOfStream)
+					{
+						string Output = LogReader.ReadToEnd();
+
+						if (!string.IsNullOrEmpty(Output))
+						{
+							ProcessLogOutput += Output;
+						}
+					}
+				}
+				while (!HasExited);
+
+				LogReader.Close();
+				ProcessLog.Close();
+				ProcessLog.Dispose();
+			}
+		}
+
+
+		public abstract string ArtifactPath { get; }
+
+		public abstract ITargetDevice Device { get; }
+
+		string ProcessLogFile;
+		string ProcessLogOutput = "";
+	}
+}
\ No newline at end of file
diff --git a/Engine/Source/Programs/AutomationTool/Gauntlet/Platform/Linux/Gauntlet.TargetDeviceLinux.cs b/Engine/Source/Programs/AutomationTool/Gauntlet/Platform/Linux/Gauntlet.TargetDeviceLinux.cs
--- a/Engine/Source/Programs/AutomationTool/Gauntlet/Platform/Linux/Gauntlet.TargetDeviceLinux.cs
+++ b/Engine/Source/Programs/AutomationTool/Gauntlet/Platform/Linux/Gauntlet.TargetDeviceLinux.cs
@@ -13,216 +13,21 @@ using static AutomationTool.ProcessResult;
 
 namespace Gauntlet
 {
-	class LinuxAppInstance : LocalAppProcess
-	{
-		protected LinuxAppInstall Install;
-
-		public LinuxAppInstance(LinuxAppInstall InInstall, IProcessResult InProcess, string ProcessLogFile = null)
-			: base(InProcess, InInstall.CommandArguments, ProcessLogFile)
-		{
-			Install = InInstall;
-		}
-
-		public override string ArtifactPath
-		{
-			get
-			{
-				return Install.ArtifactPath;
-			}
-		}
-
-		public override ITargetDevice Device
-		{
-			get
-			{
-				return Install.Device;
-			}
-		}
-	}
-
-	class LinuxAppContainerInstall : LinuxAppInstall, IContainerized
-	{
-		public ContainerInfo ContainerInfo { get; set; }
-
-		public LinuxAppContainerInstall(string InName, string InProjectName, ContainerInfo InContainerInfo, TargetDeviceLinux InDevice) : base(InName, InProjectName, InDevice)
-		{
-			ContainerInfo = InContainerInfo;
-		}
-	}
-
-	class LinuxAppInstall : IAppInstall
-	{
-		public string Name { get; private set; }
-
-		public string WorkingDirectory;
-
-		public string ExecutablePath;
-
-		public string CommandArguments;
-
-		public string ArtifactPath;
-
-		public string ProjectName;
-
-		public TargetDeviceLinux LinuxDevice { get; private set; }
-
-		public ITargetDevice Device { get { return LinuxDevice; } }
-
-		public CommandUtils.ERunOptions RunOptions { get; set; }
-
-		public LinuxAppInstall(string InName, string InProjectName, TargetDeviceLinux InDevice)
-		{
-			Name = InName;
-			ProjectName = InProjectName;
-			LinuxDevice = InDevice;
-			CommandArguments = "";
-			this.RunOptions = CommandUtils.ERunOptions.NoWaitForExit;
-		}
-
-		public IAppInstance Run()
-		{
-			return Device.Run(this);
-		}
-
-		public bool ForceCleanDeviceArtifacts()
-		{
-			DirectoryInfo ClientTempDirInfo = new DirectoryInfo(ArtifactPath) { Attributes = FileAttributes.Normal };
-			Log.Info(KnownLogEvents.Gauntlet_DeviceEvent, "Setting files in device artifacts {0} to have normal attributes (no longer read-only).", ArtifactPath);
-			foreach (FileSystemInfo info in ClientTempDirInfo.GetFileSystemInfos("*", SearchOption.AllDirectories))
-			{
-				info.Attributes = FileAttributes.Normal;
-			}
-			try
-			{
-				Log.Info(KnownLogEvents.Gauntlet_DeviceEvent, "Clearing device artifact path {0} (force)", ArtifactPath);
-				Directory.Delete(ArtifactPath, true);
-			}
-			catch (Exception Ex)
-			{
-				Log.Warning(KnownLogEvents.Gauntlet_DeviceEvent, "Failed to force delete artifact path {File}. {Exception}", ArtifactPath, Ex.Message);
-				return false;
-			}
-			return true;
-		}
-
-		public virtual void CleanDeviceArtifacts()
-		{
-			if (!string.IsNullOrEmpty(ArtifactPath) && Directory.Exists(ArtifactPath))
-			{
-				try
-				{
-					Log.Info("Clearing device artifacts path {0} for {1}", ArtifactPath, Device.Name);
-					Directory.Delete(ArtifactPath, true);
-				}
-				catch (Exception Ex)
-				{
-					Log.Info(KnownLogEvents.Gauntlet_DeviceEvent, "First attempt at clearing artifact path {0} failed - trying again", ArtifactPath);
-					if (!ForceCleanDeviceArtifacts())
-					{
-						Log.Warning(KnownLogEvents.Gauntlet_DeviceEvent, "Failed to delete {File}. {Exception}", ArtifactPath, Ex.Message);
-					}
-				}
-			}
-		}
-	}
-
-	public class LinuxDeviceFactory : IDeviceFactory
-	{
-		public bool CanSupportPlatform(UnrealTargetPlatform? Platform)
-		{
-			return Platform == UnrealTargetPlatform.Linux;
-		}
-
-		public ITargetDevice CreateDevice(string InRef, string InCachePath, string InParam = null)
-		{
-			return new TargetDeviceLinux(InRef, InCachePath);
-		}
-	}
-
-	public class LinuxArm64DeviceFactory : IDeviceFactory
-	{
-		public bool CanSupportPlatform(UnrealTargetPlatform? Platform)
-		{
-			return Platform == UnrealTargetPlatform.LinuxArm64;
-		}
-
-		public ITargetDevice CreateDevice(string InRef, string InCachePath, string InParam = null)
-		{
-			return new TargetDeviceLinux(InRef, InCachePath) { IsArm64 = true };
-		}
-	}
-
 	/// <summary>
 	/// Linux implementation of a device to run applications
 	/// </summary>
-	public class TargetDeviceLinux : ITargetDevice
+	public class TargetDeviceLinux : TargetDeviceDesktopCommon
 	{
-		public string Name { get; protected set; }
-
-		protected string UserDir { get; set; }
-
 		public bool IsArm64 { get; set; }
 
-		/// <summary>
-		/// Our mappings of Intended directories to where they actually represent on this platform.
-		/// </summary>
-		protected Dictionary<EIntendedBaseCopyDirectory, string> LocalDirectoryMappings { get; set; }
-
 		public TargetDeviceLinux(string InName, string InCacheDir)
+			: base(InName, InCacheDir)
 		{
-			Name = InName;
-			LocalCachePath = InCacheDir;
+			Platform = UnrealTargetPlatform.Linux;
 			RunOptions = CommandUtils.ERunOptions.NoWaitForExit | CommandUtils.ERunOptions.NoLoggingOfRunCommand;
-
-			UserDir = Path.Combine(LocalCachePath, string.Format("{0}_UserDir", Name));
-            LocalDirectoryMappings = new Dictionary<EIntendedBaseCopyDirectory, string>();
-		}
-
-		#region IDisposable Support
-		private bool disposedValue = false; // To detect redundant calls
-
-		protected virtual void Dispose(bool disposing)
-		{
-			if (!disposedValue)
-			{
-				if (disposing)
-				{
-					// TODO: dispose managed state (managed objects).
-				}
-
-				// TODO: free unmanaged resources (unmanaged objects) and override a finalizer below.
-				// TODO: set large fields to null.
-
-				disposedValue = true;
-			}
-		}
-
-		// This code added to correctly implement the disposable pattern.
-		public void Dispose()
-		{
-			// Do not change this code. Put cleanup code in Dispose(bool disposing) above.
-			Dispose(true);
-			// TODO: uncomment the following line if the finalizer is overridden above.
-			// GC.SuppressFinalize(this);
-		}
-		#endregion
-
-		public CommandUtils.ERunOptions RunOptions { get; set; }
-
-        // We care about UserDir in Linux as some of the roles may require files going into user instead of build dir.
-        public void PopulateDirectoryMappings(string BasePath, string UserDir)
-		{
-			LocalDirectoryMappings.Add(EIntendedBaseCopyDirectory.Build, Path.Combine(BasePath, "Build"));
-			LocalDirectoryMappings.Add(EIntendedBaseCopyDirectory.Binaries, Path.Combine(BasePath, "Binaries"));
-			LocalDirectoryMappings.Add(EIntendedBaseCopyDirectory.Config, Path.Combine(BasePath, "Saved", "Config"));
-            LocalDirectoryMappings.Add(EIntendedBaseCopyDirectory.Content, Path.Combine(BasePath, "Content"));
-            LocalDirectoryMappings.Add(EIntendedBaseCopyDirectory.Demos, Path.Combine(UserDir, "Saved", "Demos"));
-			LocalDirectoryMappings.Add(EIntendedBaseCopyDirectory.PersistentDownloadDir, Path.Combine(BasePath, "Saved", "PersistentDownloadDir"));
-			LocalDirectoryMappings.Add(EIntendedBaseCopyDirectory.Profiling, Path.Combine(BasePath, "Saved", "Profiling"));
-            LocalDirectoryMappings.Add(EIntendedBaseCopyDirectory.Saved, Path.Combine(BasePath, "Saved"));
 		}
 
-        public IAppInstance Run(IAppInstall App)
+		public override IAppInstance Run(IAppInstall App)
 		{
 			LinuxAppInstall LinuxApp = App as LinuxAppInstall;
 
@@ -266,7 +71,7 @@ namespace Gauntlet
 				Result = CommandUtils.Run(bAppContainerized ? "docker" : LinuxApp.ExecutablePath,
 					CmdLine,
 					Options: LinuxApp.RunOptions,
-					SpewFilterCallback: new SpewFilterCallbackType(delegate(string M) { return bAllowSpew ? M : null; }) /* make sure stderr does not spew in the stdout */,
+					SpewFilterCallback: new SpewFilterCallbackType(delegate (string M) { return bAllowSpew ? M : null; }) /* make sure stderr does not spew in the stdout */,
 					WorkingDir: LinuxApp.WorkingDirectory);
 
 				if (Result.HasExited && Result.ExitCode != 0)
@@ -280,39 +85,7 @@ namespace Gauntlet
 			return new LinuxAppInstance(LinuxApp, Result);
 		}
 
-		private void CopyAdditionalFiles(UnrealAppConfig AppConfig)
-		{
-			if (AppConfig.FilesToCopy != null)
-			{
-				foreach (UnrealFileToCopy FileToCopy in AppConfig.FilesToCopy)
-				{
-					string PathToCopyTo = Path.Combine(LocalDirectoryMappings[FileToCopy.TargetBaseDirectory], FileToCopy.TargetRelativeLocation);
-					if (File.Exists(FileToCopy.SourceFileLocation))
-					{
-						FileInfo SrcInfo = new FileInfo(FileToCopy.SourceFileLocation);
-						SrcInfo.IsReadOnly = false;
-						string DirectoryToCopyTo = Path.GetDirectoryName(PathToCopyTo);
-						if (!Directory.Exists(DirectoryToCopyTo))
-						{
-							Directory.CreateDirectory(DirectoryToCopyTo);
-						}
-						if (File.Exists(PathToCopyTo))
-						{
-							FileInfo ExistingFile = new FileInfo(PathToCopyTo);
-							ExistingFile.IsReadOnly = false;
-						}
-						SrcInfo.CopyTo(PathToCopyTo, true);
-						Log.Info("Copying {0} to {1}", FileToCopy.SourceFileLocation, PathToCopyTo);
-					}
-					else
-					{
-						Log.Warning(KnownLogEvents.Gauntlet_DeviceEvent, "File to copy {File} not found", FileToCopy);
-					}
-				}
-			}
-		}
-
-		protected IAppInstall InstallNativeStagedBuild(UnrealAppConfig AppConfig, NativeStagedBuild InBuild)
+		protected override IAppInstall InstallNativeStagedBuild(UnrealAppConfig AppConfig, NativeStagedBuild InBuild)
 		{
 			LinuxAppInstall LinuxApp;
 			if (AppConfig.ContainerInfo != null)
@@ -343,11 +116,11 @@ namespace Gauntlet
 			return LinuxApp;
 		}
 
-		protected IAppInstall InstallStagedBuild(UnrealAppConfig AppConfig, StagedBuild InBuild)
+		protected override IAppInstall InstallStagedBuild(UnrealAppConfig AppConfig, StagedBuild InBuild)
 		{
 			string BuildPath = InBuild.BuildPath;
 
-			if (CanRunFromPath(BuildPath) == false)
+			if (Utils.SystemHelpers.IsNetworkPath(BuildPath))
 			{
 				string SubDir = string.IsNullOrEmpty(AppConfig.Sandbox) ? AppConfig.ProjectName : AppConfig.Sandbox;
 				string BasePath = string.IsNullOrEmpty(AppConfig.DestLocalInstallDir) ? this.LocalCachePath : AppConfig.DestLocalInstallDir;
@@ -390,14 +163,14 @@ namespace Gauntlet
 			// clear artifact path
 			LinuxApp.CleanDeviceArtifacts();
 
-            if (LocalDirectoryMappings.Count == 0)
-            {
-                PopulateDirectoryMappings(Path.Combine(BuildPath, AppConfig.ProjectName), UserDir);
-            }
+			if (LocalDirectoryMappings.Count == 0)
+			{
+				PopulateDirectoryMappings(Path.Combine(BuildPath, AppConfig.ProjectName));
+			}
 
 			CopyAdditionalFiles(AppConfig);
 
-            if (Path.IsPathRooted(InBuild.ExecutablePath))
+			if (Path.IsPathRooted(InBuild.ExecutablePath))
 			{
 				LinuxApp.ExecutablePath = InBuild.ExecutablePath;
 			}
@@ -432,83 +205,80 @@ namespace Gauntlet
 			return LinuxApp;
 		}
 
-		public void ClearSavedDirectory(UnrealAppConfig AppConfiguration)
+		protected override IAppInstall InstallEditorBuild(UnrealAppConfig AppConfig, EditorBuild Build)
 		{
-
-		}
-
-		public IAppInstall InstallApplication(UnrealAppConfig AppConfig)
-		{
-			if (AppConfig.Build is NativeStagedBuild)
-			{
-				return InstallNativeStagedBuild(AppConfig, AppConfig.Build as NativeStagedBuild);
-			}
-			else if (AppConfig.Build is StagedBuild)
-			{
-				return InstallStagedBuild(AppConfig, AppConfig.Build as StagedBuild);
-			}
-
-			EditorBuild EditorBuild = AppConfig.Build as EditorBuild;
-
-			if (EditorBuild == null)
-			{
-				throw new AutomationException("Invalid build type!");
-			}
-
 			LinuxAppInstall LinuxApp = new LinuxAppInstall(AppConfig.Name, AppConfig.ProjectName, this);
 
-			LinuxApp.WorkingDirectory = Path.GetDirectoryName(EditorBuild.ExecutablePath);
+			LinuxApp.WorkingDirectory = Path.GetDirectoryName(Build.ExecutablePath);
 			LinuxApp.RunOptions = RunOptions;
 
 			// Force this to stop logs and other artifacts going to different places
 			LinuxApp.CommandArguments = AppConfig.CommandLine + string.Format(" -userdir=\"{0}\"", UserDir);
 			LinuxApp.ArtifactPath = Path.Combine(UserDir, @"Saved");
-			LinuxApp.ExecutablePath = EditorBuild.ExecutablePath;
+			LinuxApp.ExecutablePath = Build.ExecutablePath;
 
 			if (LocalDirectoryMappings.Count == 0)
 			{
-				PopulateDirectoryMappings(AppConfig.ProjectFile.Directory.FullName, AppConfig.ProjectFile.Directory.FullName);
+				PopulateDirectoryMappings(AppConfig.ProjectFile.Directory.FullName);
 			}
 
 			CopyAdditionalFiles(AppConfig);
 
 			return LinuxApp;
 		}
+	}
 
-		public void CopyAppConfigurationFiles(UnrealAppConfig AppConfiguration)
-		{
+	public class LinuxAppInstall : DesktopCommonAppInstall<TargetDeviceLinux>
+	{
+		[Obsolete("Will be removed in a future release. Use 'DesktopDevice' instead.")]
+		public TargetDeviceLinux LinuxDevice { get; private set; }
 
-		}
+		public LinuxAppInstall(string InName, string InProjectName, TargetDeviceLinux InDevice)
+			: base(InName, InProjectName, InDevice)
+		{ }
+	}
+
+	public class LinuxAppContainerInstall : LinuxAppInstall, IContainerized
+	{
+		public ContainerInfo ContainerInfo { get; set; }
 
-		public bool CanRunFromPath(string InPath)
+		public LinuxAppContainerInstall(string InName, string InProjectName, ContainerInfo InContainerInfo, TargetDeviceLinux InDevice)
+			: base(InName, InProjectName, InDevice)
 		{
-			return !Utils.SystemHelpers.IsNetworkPath(InPath);
+			ContainerInfo = InContainerInfo;
 		}
+	}
 
-		public UnrealTargetPlatform? Platform { get { return IsArm64 ? UnrealTargetPlatform.LinuxArm64 : UnrealTargetPlatform.Linux; } }
+	public class LinuxAppInstance : DesktopCommonAppInstance<LinuxAppInstall, TargetDeviceLinux>
+	{
+		public LinuxAppInstance(LinuxAppInstall InInstall, IProcessResult InProcess, string ProcessLogFile = null)
+			: base(InInstall, InProcess, ProcessLogFile)
+		{ }
+	}
 
-		public string LocalCachePath { get; private set; }
-		public bool IsAvailable { get { return true; } }
-		public bool IsConnected { get { return true; } }
-		public bool IsOn { get { return true; } }
-		public bool PowerOn() { return true; }
-		public bool PowerOff() { return true; }
-		public bool Reboot() { return true; }
-		public bool Connect() { return true; }
-		public bool Disconnect(bool bForce = false) { return true; }
+	public class LinuxDeviceFactory : IDeviceFactory
+	{
+		public bool CanSupportPlatform(UnrealTargetPlatform? Platform)
+		{
+			return Platform == UnrealTargetPlatform.Linux;
+		}
 
-		public override string ToString()
+		public ITargetDevice CreateDevice(string InRef, string InCachePath, string InParam = null)
 		{
-			return Name;
+			return new TargetDeviceLinux(InRef, InCachePath);
 		}
+	}
 
-		public Dictionary<EIntendedBaseCopyDirectory, string> GetPlatformDirectoryMappings()
+	public class LinuxArm64DeviceFactory : IDeviceFactory
+	{
+		public bool CanSupportPlatform(UnrealTargetPlatform? Platform)
 		{
-			if (LocalDirectoryMappings.Count == 0)
-			{
-				Log.Warning("Platform directory mappings have not been populated for this platform! This should be done within InstallApplication()");
-			}
-			return LocalDirectoryMappings;
+			return Platform == UnrealTargetPlatform.LinuxArm64;
+		}
+
+		public ITargetDevice CreateDevice(string InRef, string InCachePath, string InParam = null)
+		{
+			return new TargetDeviceLinux(InRef, InCachePath) { IsArm64 = true };
 		}
 	}
-}
+}
\ No newline at end of file
diff --git a/Engine/Source/Programs/AutomationTool/Gauntlet/Platform/Mac/Gauntlet.TargetDeviceMac.cs b/Engine/Source/Programs/AutomationTool/Gauntlet/Platform/Mac/Gauntlet.TargetDeviceMac.cs
--- a/Engine/Source/Programs/AutomationTool/Gauntlet/Platform/Mac/Gauntlet.TargetDeviceMac.cs
+++ b/Engine/Source/Programs/AutomationTool/Gauntlet/Platform/Mac/Gauntlet.TargetDeviceMac.cs
@@ -12,270 +12,70 @@ using static AutomationTool.ProcessResult;
 
 namespace Gauntlet
 {
-
-	public class MacDeviceFactory : IDeviceFactory
-	{
-		public bool CanSupportPlatform(UnrealTargetPlatform? Platform)
-		{
-			return Platform == UnrealTargetPlatform.Mac;
-		}
-
-		public ITargetDevice CreateDevice(string InRef, string InCachePath, string InParam = null)
-		{
-			return new TargetDeviceMac(InRef, InCachePath);
-		}
-	}
-
-	class MacAppInstall : IAppInstall
-	{
-		public string Name { get; private set; }
-
-		public string LocalPath;
-
-		public string WorkingDirectory;
-
-		public string ExecutablePath;
-
-		public string CommandArguments;
-
-		public string ArtifactPath;
-
-		public ITargetDevice Device { get; protected set; }
-
-		public CommandUtils.ERunOptions RunOptions { get; set; }
-
-		public MacAppInstall(string InName, TargetDeviceMac InDevice)
-		{
-			Name = InName;
-			Device = InDevice;
-			CommandArguments = "";
-			this.RunOptions = CommandUtils.ERunOptions.NoWaitForExit;
-		}
-
-		public IAppInstance Run()
-		{
-			return Device.Run(this);
-		}
-
-		public bool ForceCleanDeviceArtifacts()
-		{
-			DirectoryInfo ClientTempDirInfo = new DirectoryInfo(ArtifactPath) { Attributes = FileAttributes.Normal };
-			Log.Info(KnownLogEvents.Gauntlet_DeviceEvent, "Setting files in device artifacts {0} to have normal attributes (no longer read-only).", ArtifactPath);
-			foreach (FileSystemInfo info in ClientTempDirInfo.GetFileSystemInfos("*", SearchOption.AllDirectories))
-			{
-				info.Attributes = FileAttributes.Normal;
-			}
-			try
-			{
-				Log.Info(KnownLogEvents.Gauntlet_DeviceEvent, "Clearing device artifact path {0} (force)", ArtifactPath);
-				Directory.Delete(ArtifactPath, true);
-			}
-			catch (Exception Ex)
-			{
-				Log.Warning(KnownLogEvents.Gauntlet_DeviceEvent, "Failed to force delete artifact path {File}. {Exception}", ArtifactPath, Ex.Message);
-				return false;
-			}
-			return true;
-		}
-
-		public virtual void CleanDeviceArtifacts()
-		{
-			if (!string.IsNullOrEmpty(ArtifactPath) && Directory.Exists(ArtifactPath))
-			{
-				try
-				{
-					Log.Info("Clearing device artifacts path {0} for {1}", ArtifactPath, Device.Name);
-					Directory.Delete(ArtifactPath, true);
-				}
-				catch (Exception Ex)
-				{
-					Log.Info(KnownLogEvents.Gauntlet_DeviceEvent, "First attempt at clearing artifact path {0} failed - trying again", ArtifactPath);
-					if (!ForceCleanDeviceArtifacts())
-					{
-						Log.Warning(KnownLogEvents.Gauntlet_DeviceEvent, "Failed to delete {File}. {Exception}", ArtifactPath, Ex.Message);
-					}
-				}
-			}
-		}
-	}
-
-	class MacAppInstance : LocalAppProcess
-	{
-		protected MacAppInstall Install;
-
-		public MacAppInstance(MacAppInstall InInstall, IProcessResult InProcess)
-			: base(InProcess, InInstall.CommandArguments)
-		{
-			Install = InInstall;
-		}
-
-		public override string ArtifactPath
-		{
-			get
-			{
-				return Install.ArtifactPath;
-			}
-		}
-
-		public override ITargetDevice Device
-		{
-			get
-			{
-				return Install.Device;
-			}
-		}
-	}
-
-	public class TargetDeviceMac : ITargetDevice
+	public class TargetDeviceMac : TargetDeviceDesktopCommon
 	{
-		public string Name { get; protected set; }
-
-		public UnrealTargetPlatform? Platform { get { return UnrealTargetPlatform.Mac; } }
-
-		public bool IsAvailable { get { return true; } }
-		public bool IsConnected { get { return true; } }
-		public bool IsOn { get { return true; } }
-		public bool PowerOn() { return true; }
-		public bool PowerOff() { return true; }
-		public bool Reboot() { return true; }
-		public bool Connect() { return true; }
-		public bool Disconnect(bool bForce = false) { return true; }
-
-		protected string UserDir { get; set; }
-
-		protected string LocalCachePath { get; set; }
-
-		public CommandUtils.ERunOptions RunOptions { get; set; }
-
-		protected Dictionary<EIntendedBaseCopyDirectory, string> LocalDirectoryMappings { get; set; }
-
 		public TargetDeviceMac(string InName, string InCachePath)
+			: base(InName, InCachePath)
 		{
-			Name = InName;
-			LocalCachePath = InCachePath;
-			UserDir = Path.Combine(LocalCachePath, "UserDir");
-			this.RunOptions = CommandUtils.ERunOptions.NoWaitForExit;
-			LocalDirectoryMappings = new Dictionary<EIntendedBaseCopyDirectory, string>();
+			Platform = UnrealTargetPlatform.Mac;
+			RunOptions = CommandUtils.ERunOptions.NoWaitForExit;
 		}
 
-		#region IDisposable Support
-		private bool disposedValue = false; // To detect redundant calls
-
-		protected virtual void Dispose(bool disposing)
+		public override IAppInstall InstallApplication(UnrealAppConfig AppConfig)
 		{
-			if (!disposedValue)
+			switch (AppConfig.Build)
 			{
-				if (disposing)
-				{
-					// TODO: dispose managed state (managed objects).
-				}
+				case NativeStagedBuild:
+					return InstallNativeStagedBuild(AppConfig, AppConfig.Build as NativeStagedBuild);
 
-				// TODO: free unmanaged resources (unmanaged objects) and override a finalizer below.
-				// TODO: set large fields to null.
-
-				disposedValue = true;
-			}
-		}
-
-		// This code added to correctly implement the disposable pattern.
-		public void Dispose()
-		{
-			// Do not change this code. Put cleanup code in Dispose(bool disposing) above.
-			Dispose(true);
-			// TODO: uncomment the following line if the finalizer is overridden above.
-			// GC.SuppressFinalize(this);
-		}
-		#endregion
+				case StagedBuild:
+					return InstallStagedBuild(AppConfig, AppConfig.Build as StagedBuild);
 
-		public override string ToString()
-		{
-			return Name;
-		}
+				case MacPackagedBuild:
+					return InstallPackagedBuild(AppConfig, AppConfig.Build as MacPackagedBuild);
 
-		protected string GetVolumeName(string InPath)
-		{
-			Match M = Regex.Match(InPath, @"/Volumes/(.+?)/");
+				case EditorBuild:
+					return InstallEditorBuild(AppConfig, AppConfig.Build as EditorBuild);
 
-			if (M.Success)
-			{
-				return M.Groups[1].ToString();
+				default:
+					throw new AutomationException("{0} is an invalid build type!", AppConfig.Build.ToString());
 			}
-
-			return "";
 		}
 
-		public void PopulateDirectoryMappings(string ProjectDir)
+		public override IAppInstance Run(IAppInstall App)
 		{
-			LocalDirectoryMappings.Add(EIntendedBaseCopyDirectory.Build, Path.Combine(ProjectDir, "Build"));
-			LocalDirectoryMappings.Add(EIntendedBaseCopyDirectory.Binaries, Path.Combine(ProjectDir, "Binaries"));
-			LocalDirectoryMappings.Add(EIntendedBaseCopyDirectory.Config, Path.Combine(ProjectDir, "Saved", "Config"));
-			LocalDirectoryMappings.Add(EIntendedBaseCopyDirectory.Content, Path.Combine(ProjectDir, "Content"));
-			LocalDirectoryMappings.Add(EIntendedBaseCopyDirectory.Demos, Path.Combine(ProjectDir, "Saved", "Demos"));
-			LocalDirectoryMappings.Add(EIntendedBaseCopyDirectory.PersistentDownloadDir, Path.Combine(ProjectDir, "Saved", "PersistentDownloadDir"));
-			LocalDirectoryMappings.Add(EIntendedBaseCopyDirectory.Profiling, Path.Combine(ProjectDir, "Saved", "Profiling"));
-			LocalDirectoryMappings.Add(EIntendedBaseCopyDirectory.Saved, Path.Combine(ProjectDir, "Saved"));
-		}
+			MacAppInstall MacInstall = App as MacAppInstall;
 
-		/// <summary>
-		/// If the path is to Foo.app this returns the actual executable to use (e.g. Foo.app/Contents/MacOS/Foo).
-		/// </summary>
-		/// <param name="InBundlePath"></param>
-		/// <returns></returns>
-		protected string GetExecutableIfBundle(string InBundlePath)
-		{
-			if (Path.GetExtension(InBundlePath).Equals(".app", StringComparison.OrdinalIgnoreCase))
+			if (MacInstall == null)
 			{
-				// Technically we should look at the plist, but for now...
-				string BaseName = Path.GetFileNameWithoutExtension(InBundlePath);
-				return Path.Combine(InBundlePath, "Contents", "MacOS", BaseName);
+				throw new AutomationException("Invalid install type!");
 			}
 
-			return InBundlePath;
-		}
-
-
-		/// <summary>
-		/// Copies a build folder (either a package.app or a folder with a staged built) to a local path if necessary.
-		/// Necessary is defined as not being on locally attached storage
-		/// </summary>
-		/// <param name="AppConfig"></param>
-		/// <param name="InBuildPath"></param>
-		/// <returns></returns>
-		protected string CopyBuildIfNecessary(UnrealAppConfig AppConfig, string InBuildPath)
-		{
-			bool SkipDeploy = Globals.Params.ParseParam("SkipDeploy");
-
-			string OutBuildPath = InBuildPath;
-
-			string BuildVolume = GetVolumeName(OutBuildPath);
-			string LocalRoot = GetVolumeName(Environment.CurrentDirectory);
+			IProcessResult Result = null;
 
-			// Must be on our volume to run
-			if (BuildVolume.Equals(LocalRoot, StringComparison.OrdinalIgnoreCase) == false)
+			lock (Globals.MainLock)
 			{
-				string SubDir = string.IsNullOrEmpty(AppConfig.Sandbox) ? AppConfig.ProjectName : AppConfig.Sandbox;
-				string BasePath = string.IsNullOrEmpty(AppConfig.DestLocalInstallDir) ? this.LocalCachePath : AppConfig.DestLocalInstallDir;
-				string DestPath = Path.Combine(BasePath, SubDir, AppConfig.ProcessType.ToString());
+				string OldWD = Environment.CurrentDirectory;
+				Environment.CurrentDirectory = MacInstall.WorkingDirectory;
 
-				if (!SkipDeploy)
-				{
-					Utils.SystemHelpers.CopyDirectory(InBuildPath, DestPath, Utils.SystemHelpers.CopyOptions.Mirror);
-				}
+				Log.Info("Launching {0} on {1}", App.Name, ToString());
+				Log.Verbose("\t{0}", MacInstall.CommandArguments);
 
-				else
+				bool bAllowSpew = MacInstall.RunOptions.HasFlag(CommandUtils.ERunOptions.AllowSpew);
+				Result = CommandUtils.Run(GetExecutableIfBundle(MacInstall.ExecutablePath), MacInstall.CommandArguments, Options: MacInstall.RunOptions, SpewFilterCallback: new SpewFilterCallbackType(delegate (string M) { return bAllowSpew ? M : null; }) /* make sure stderr does not spew in the stdout */);
+
+				if (Result.HasExited && Result.ExitCode != 0)
 				{
-					Log.Info("Skipping install of {0} (-skipdeploy)", InBuildPath);
+					throw new AutomationException("Failed to launch {0}. Error {1}", MacInstall.ExecutablePath, Result.ExitCode);
 				}
 
-				Utils.SystemHelpers.MarkDirectoryForCleanup(DestPath);
-
-				OutBuildPath = DestPath;
+				Environment.CurrentDirectory = OldWD;
 			}
 
-			return OutBuildPath;
+			return new MacAppInstance(MacInstall, Result);
 		}
 
-		protected IAppInstall InstallNativeStagedBuild(UnrealAppConfig AppConfig, NativeStagedBuild InBuild)
+		protected override IAppInstall InstallNativeStagedBuild(UnrealAppConfig AppConfig, NativeStagedBuild InBuild)
 		{
 			MacAppInstall MacApp = new MacAppInstall(AppConfig.Name, this);
 
@@ -295,30 +95,48 @@ namespace Gauntlet
 			return MacApp;
 		}
 
-		protected IAppInstall InstallBuild(UnrealAppConfig AppConfig, IBuild InBuild)
+		protected override IAppInstall InstallEditorBuild(UnrealAppConfig AppConfig, EditorBuild Build)
 		{
-			// Full path to the build
-			string BuildPath;
+			MacAppInstall MacApp = new MacAppInstall(AppConfig.Name, this);
 
-			//  Full path to the build.app to use. This will be under BuildPath for staged builds, and the build path itself for packaged builds
-			string BundlePath;
+			MacApp.WorkingDirectory = Path.GetFullPath(Build.ExecutablePath);
+			MacApp.CommandArguments = AppConfig.CommandLine;
+			MacApp.RunOptions = RunOptions;
 
-			if (InBuild is StagedBuild)
-			{
-				StagedBuild InStagedBuild = InBuild as StagedBuild;
-				BuildPath = CopyBuildIfNecessary(AppConfig, InStagedBuild.BuildPath);
-				BundlePath = Path.Combine(BuildPath, InStagedBuild.ExecutablePath);
-			}
-			else
+			// Mac always forces this to stop logs and other artifacts going to different places
+			MacApp.CommandArguments += string.Format(" -userdir=\"{0}\"", UserDir);
+			MacApp.ArtifactPath = Path.Combine(UserDir, @"Saved");
+
+			// now turn the Foo.app into Foo/Content/MacOS/Foo
+			string AppPath = Path.GetDirectoryName(Build.ExecutablePath);
+
+			MacApp.ExecutablePath = GetExecutableIfBundle(Build.ExecutablePath);
+
+			if (LocalDirectoryMappings.Count == 0)
 			{
-				MacPackagedBuild InPackagedBuild = InBuild as MacPackagedBuild;
-				BuildPath = CopyBuildIfNecessary(AppConfig, InPackagedBuild.BuildPath);
-				BundlePath = BuildPath;
+				PopulateDirectoryMappings(AppPath);
 			}
+			return MacApp;
+		}
+
+		protected override IAppInstall InstallStagedBuild(UnrealAppConfig AppConfig, StagedBuild Build)
+		{
+			string BuildPath = CopyBuildIfNecessary(AppConfig, Build.BuildPath);
+			string BundlePath = Path.Combine(BuildPath, Build.ExecutablePath);
+			return InstallBuild(AppConfig, BuildPath, BundlePath);
+		}
 
+		protected IAppInstall InstallPackagedBuild(UnrealAppConfig AppConfig, MacPackagedBuild Build)
+		{
+			string BuildPath = CopyBuildIfNecessary(AppConfig, Build.BuildPath);
+			return InstallBuild(AppConfig, BuildPath, BuildPath);
+		}
+
+		protected IAppInstall InstallBuild(UnrealAppConfig AppConfig, string BuildPath, string BundlePath)
+		{
 			MacAppInstall MacApp = new MacAppInstall(AppConfig.Name, this);
-			MacApp.LocalPath = BuildPath;
-			MacApp.WorkingDirectory = MacApp.LocalPath;
+
+			MacApp.WorkingDirectory = BuildPath;
 			MacApp.RunOptions = RunOptions;
 
 			MacApp.ExecutablePath = GetExecutableIfBundle(BundlePath);
@@ -338,19 +156,19 @@ namespace Gauntlet
 
 			// clear artifact path
 			MacApp.CleanDeviceArtifacts();
-			
+
 			// check for a local newer executable
-			if (Globals.Params.ParseParam("dev") 
+			if (Globals.Params.ParseParam("dev")
 				&& AppConfig.ProcessType.UsesEditor() == false
 				&& AppConfig.ProjectFile != null)
 			{
 				// Get project properties
-				ProjectProperties Props = ProjectUtils.GetProjectProperties(AppConfig.ProjectFile, 
+				ProjectProperties Props = ProjectUtils.GetProjectProperties(AppConfig.ProjectFile,
 					new List<UnrealTargetPlatform>(new[] { AppConfig.Platform.Value }),
 					new List<UnrealTargetConfiguration>(new[] { AppConfig.Configuration }));
 
 				// Would this executable be built under Engine or a Project?
-				DirectoryReference  WorkingDir = Props.bIsCodeBasedProject ? AppConfig.ProjectFile.Directory : Unreal.EngineDirectory;
+				DirectoryReference WorkingDir = Props.bIsCodeBasedProject ? AppConfig.ProjectFile.Directory : Unreal.EngineDirectory;
 
 				// The bundlepath may be under Binaries/Mac for a staged build, or it could be in any folder for a packaged build so just use the name and
 				// build the path ourselves
@@ -375,98 +193,109 @@ namespace Gauntlet
 			return MacApp;
 		}
 
-		public void ClearSavedDirectory(UnrealAppConfig AppConfiguration)
+		protected string GetVolumeName(string InPath)
 		{
+			Match M = Regex.Match(InPath, @"/Volumes/(.+?)/");
 
-		}
-
-		public IAppInstall InstallApplication(UnrealAppConfig AppConfig)
-		{
-			if (AppConfig.Build is NativeStagedBuild)
-			{
-				return InstallNativeStagedBuild(AppConfig, (NativeStagedBuild)AppConfig.Build);
-			}
-			if (AppConfig.Build is StagedBuild || AppConfig.Build is MacPackagedBuild)
+			if (M.Success)
 			{
-				return InstallBuild(AppConfig, AppConfig.Build);
+				return M.Groups[1].ToString();
 			}
 
-			EditorBuild EditorBuild = AppConfig.Build as EditorBuild;
+			return "";
+		}
 
-			if (EditorBuild == null)
+		/// <summary>
+		/// If the path is to Foo.app this returns the actual executable to use (e.g. Foo.app/Contents/MacOS/Foo).
+		/// </summary>
+		/// <param name="InBundlePath"></param>
+		/// <returns></returns>
+		protected string GetExecutableIfBundle(string InBundlePath)
+		{
+			if (Path.GetExtension(InBundlePath).Equals(".app", StringComparison.OrdinalIgnoreCase))
 			{
-				throw new AutomationException("Invalid build type!");
+				// Technically we should look at the plist, but for now...
+				string BaseName = Path.GetFileNameWithoutExtension(InBundlePath);
+				return Path.Combine(InBundlePath, "Contents", "MacOS", BaseName);
 			}
 
-			MacAppInstall MacApp = new MacAppInstall(AppConfig.Name, this);
-
-			MacApp.WorkingDirectory = Path.GetFullPath(EditorBuild.ExecutablePath);
-			MacApp.CommandArguments = AppConfig.CommandLine;
-			MacApp.RunOptions = RunOptions;
-
-			// Mac always forces this to stop logs and other artifacts going to different places
-			MacApp.CommandArguments += string.Format(" -userdir=\"{0}\"", UserDir);
-			MacApp.ArtifactPath = Path.Combine(UserDir, @"Saved");
-
-			// now turn the Foo.app into Foo/Content/MacOS/Foo
-			string AppPath = Path.GetDirectoryName(EditorBuild.ExecutablePath);
-
-			MacApp.ExecutablePath = GetExecutableIfBundle(EditorBuild.ExecutablePath);
-
-			if (LocalDirectoryMappings.Count == 0)
-			{
-				PopulateDirectoryMappings(AppPath);
-			}
-			return MacApp;
+			return InBundlePath;
 		}
 
-		public void CopyAppConfigurationFiles(UnrealAppConfig AppConfiguration)
-		{
-
-		}
 
-		public IAppInstance Run(IAppInstall App)
+		/// <summary>
+		/// Copies a build folder (either a package.app or a folder with a staged built) to a local path if necessary.
+		/// Necessary is defined as not being on locally attached storage
+		/// </summary>
+		/// <param name="AppConfig"></param>
+		/// <param name="InBuildPath"></param>
+		/// <returns></returns>
+		protected string CopyBuildIfNecessary(UnrealAppConfig AppConfig, string InBuildPath)
 		{
-			MacAppInstall MacInstall = App as MacAppInstall;
+			bool SkipDeploy = Globals.Params.ParseParam("SkipDeploy");
 
-			if (MacInstall == null)
-			{
-				throw new AutomationException("Invalid install type!");
-			}
+			string OutBuildPath = InBuildPath;
 
-			IProcessResult Result = null;
+			string BuildVolume = GetVolumeName(OutBuildPath);
+			string LocalRoot = GetVolumeName(Environment.CurrentDirectory);
 
-			lock (Globals.MainLock)
+			// Must be on our volume to run
+			if (BuildVolume.Equals(LocalRoot, StringComparison.OrdinalIgnoreCase) == false)
 			{
-				string NewWorkingDir = string.IsNullOrEmpty(MacInstall.WorkingDirectory) ? MacInstall.LocalPath : MacInstall.WorkingDirectory;
-				string OldWD = Environment.CurrentDirectory;
-				Environment.CurrentDirectory = NewWorkingDir;
-
-				Log.Info("Launching {0} on {1}", App.Name, ToString());
-				Log.Verbose("\t{0}", MacInstall.CommandArguments);
+				string SubDir = string.IsNullOrEmpty(AppConfig.Sandbox) ? AppConfig.ProjectName : AppConfig.Sandbox;
+				string BasePath = string.IsNullOrEmpty(AppConfig.DestLocalInstallDir) ? this.LocalCachePath : AppConfig.DestLocalInstallDir;
+				string DestPath = Path.Combine(BasePath, SubDir, AppConfig.ProcessType.ToString());
 
-				bool bAllowSpew = MacInstall.RunOptions.HasFlag(CommandUtils.ERunOptions.AllowSpew);
-				Result = CommandUtils.Run(GetExecutableIfBundle(MacInstall.ExecutablePath), MacInstall.CommandArguments, Options: MacInstall.RunOptions, SpewFilterCallback: new SpewFilterCallbackType(delegate (string M) { return bAllowSpew ? M : null; }) /* make sure stderr does not spew in the stdout */);
+				if (!SkipDeploy)
+				{
+					Utils.SystemHelpers.CopyDirectory(InBuildPath, DestPath, Utils.SystemHelpers.CopyOptions.Mirror);
+				}
 
-				if (Result.HasExited && Result.ExitCode != 0)
+				else
 				{
-					throw new AutomationException("Failed to launch {0}. Error {1}", MacInstall.ExecutablePath, Result.ExitCode);
+					Log.Info("Skipping install of {0} (-skipdeploy)", InBuildPath);
 				}
 
-				Environment.CurrentDirectory = OldWD;
+				Utils.SystemHelpers.MarkDirectoryForCleanup(DestPath);
+
+				OutBuildPath = DestPath;
 			}
 
-			return new MacAppInstance(MacInstall, Result);
+			return OutBuildPath;
 		}
+	}
 
-		public Dictionary<EIntendedBaseCopyDirectory, string> GetPlatformDirectoryMappings()
+	public class MacAppInstall : DesktopCommonAppInstall<TargetDeviceMac>
+	{
+		[Obsolete("Will be removed in a future release. Use WorkingDirectory instead.")]
+		public string LocalPath
 		{
-			if (LocalDirectoryMappings.Count == 0)
-			{
-				Log.Warning("Platform directory mappings have not been populated for this platform! This should be done within InstallApplication()");
-			}
-			return LocalDirectoryMappings;
+			get => WorkingDirectory;
+			set { } // intentional nop
 		}
+
+		public MacAppInstall(string InName, TargetDeviceMac InDevice)
+			: base(InName, string.Empty, InDevice)
+		{ }
+	}
+
+	public class MacAppInstance : DesktopCommonAppInstance<MacAppInstall, TargetDeviceMac>
+	{
+		public MacAppInstance(MacAppInstall InInstall, IProcessResult InProcess, string ProcessLogFile = null)
+			: base(InInstall, InProcess, ProcessLogFile)
+		{ }
 	}
 
-}
+	public class MacDeviceFactory : IDeviceFactory
+	{
+		public bool CanSupportPlatform(UnrealTargetPlatform? Platform)
+		{
+			return Platform == UnrealTargetPlatform.Mac;
+		}
+
+		public ITargetDevice CreateDevice(string InRef, string InCachePath, string InParam = null)
+		{
+			return new TargetDeviceMac(InRef, InCachePath);
+		}
+	}
+}
\ No newline at end of file
diff --git a/Engine/Source/Programs/AutomationTool/Gauntlet/Platform/Windows/Gauntlet.TargetDeviceWindows.cs b/Engine/Source/Programs/AutomationTool/Gauntlet/Platform/Windows/Gauntlet.TargetDeviceWindows.cs
--- a/Engine/Source/Programs/AutomationTool/Gauntlet/Platform/Windows/Gauntlet.TargetDeviceWindows.cs
+++ b/Engine/Source/Programs/AutomationTool/Gauntlet/Platform/Windows/Gauntlet.TargetDeviceWindows.cs
@@ -1,10 +1,7 @@
 // Copyright Epic Games, Inc. All Rights Reserved.
 
 using System;
-using System.Collections.Generic;
 using System.IO;
-using System.Linq;
-using System.Threading;
 using AutomationTool;
 using UnrealBuildTool;
 using System.Text.RegularExpressions;
@@ -13,321 +10,40 @@ using static AutomationTool.ProcessResult;
 
 namespace Gauntlet
 {
-
-	public abstract class LocalAppProcess : IAppInstance
-	{
-		public IProcessResult ProcessResult { get; private set; }
-
-		public bool HasExited { get { return ProcessResult.HasExited; } }
-
-		public bool WasKilled { get; protected set; }
-
-		public string StdOut { get { return string.IsNullOrEmpty(ProcessLogFile) ? ProcessResult.Output : ProcessLogOutput; } }
-
-		public int ExitCode { get { return ProcessResult.ExitCode; } }
-
-		public string CommandLine { get; private set; }
-
-		public LocalAppProcess(IProcessResult InProcess, string InCommandLine, string InProcessLogFile = null)
-		{
-			this.CommandLine = InCommandLine;
-			this.ProcessResult = InProcess;
-			this.ProcessLogFile = InProcessLogFile;
-
-			// start reader thread if logging to a file
-			if (!string.IsNullOrEmpty(InProcessLogFile))
-			{
-				new System.Threading.Thread(LogFileReaderThread).Start();
-			}
-		}
-
-		public int WaitForExit()
-		{
-			if (!HasExited)
-			{
-				ProcessResult.WaitForExit();
-			}
-
-			return ExitCode;
-		}
-
-		virtual public void Kill()
-		{
-			if (!HasExited)
-			{
-				WasKilled = true;
-				ProcessResult.ProcessObject.Kill(true);
-			}
-		}
-
-		/// <summary>
-		/// Reader thread when logging to file
-		/// </summary>
-		void LogFileReaderThread()
-		{
-			// Wait for the processes log file to be created
-			while (!File.Exists(ProcessLogFile) && !HasExited)
-			{
-				Thread.Sleep(2000);
-			}
-
-			// Check whether the process exited before log file was created (this can happen for example if a server role exits and forces client to shutdown)
-			if (!File.Exists(ProcessLogFile))
-			{
-				ProcessLogOutput += "Process exited before log file created";
-				return;
-			}
-
-			Thread.Sleep(1000);
-
-			using (FileStream ProcessLog = File.Open(ProcessLogFile, FileMode.Open, FileAccess.Read, FileShare.ReadWrite))
-			{
-				StreamReader LogReader = new StreamReader(ProcessLog);
-				// Read until the process has exited
-				do
-				{
-					Thread.Sleep(250);
-
-					while (!LogReader.EndOfStream)
-					{
-						string Output = LogReader.ReadToEnd();
-
-						if (!string.IsNullOrEmpty(Output))
-						{
-							ProcessLogOutput += Output;
-						}
-					}
-				}
-				while (!HasExited);
-
-				LogReader.Close();
-				ProcessLog.Close();
-				ProcessLog.Dispose();
-			}
-		}
-
-
-		public abstract string ArtifactPath { get; }
-
-		public abstract ITargetDevice Device { get; }
-
-		string ProcessLogFile;
-		string ProcessLogOutput = "";		
-	}
-
-	class WindowsAppInstance : LocalAppProcess
-	{
-		protected WindowsAppInstall Install;
-
-		public WindowsAppInstance(WindowsAppInstall InInstall, IProcessResult InProcess, string ProcessLogFile = null)
-			: base(InProcess, InInstall.CommandArguments, ProcessLogFile)
-		{
-			Install = InInstall;
-		}
-
-		public override string ArtifactPath
-		{
-			get
-			{
-				return Install.ArtifactPath;
-			}
-		}
-
-		public override ITargetDevice Device
-		{
-			get
-			{
-				return Install.Device;
-			}
-		}
-	}
-
-
-	public class WindowsAppInstall : IAppInstall, IAppInstall.IDynamicCommandLine
-	{
-		public string Name { get; private set; }
-
-		public string WorkingDirectory;
-
-		public string ExecutablePath;
-
-		public bool CanAlterCommandArgs = true;
-
-		public string CommandArguments
-		{
-			get { return CommandArgumentsPrivate; }
-			set
-			{
-				if (CanAlterCommandArgs || string.IsNullOrEmpty(CommandArgumentsPrivate))
-				{
-					CommandArgumentsPrivate = value;
-				}
-				else
-				{
-					Log.Info("Skipped setting command AppInstall line when CanAlterCommandArgs = false");
-				}
-			}
-		}
-		private string CommandArgumentsPrivate;
-
-		public string ArtifactPath;
-
-		public string ProjectName;
-
-		public TargetDeviceWindows WinDevice { get; private set; }
-
-		public ITargetDevice Device { get { return WinDevice; } }
-
-		public CommandUtils.ERunOptions RunOptions { get; set; }
-
-		public WindowsAppInstall(string InName, string InProjectName, TargetDeviceWindows InDevice)
-		{
-			Name = InName;
-			ProjectName = InProjectName;
-			WinDevice = InDevice;
-			CommandArguments = "";
-			this.RunOptions = CommandUtils.ERunOptions.NoWaitForExit;
-		}
-
-		public IAppInstance Run()
-		{
-			return Device.Run(this);
-		}
-
-		public void AppendCommandline(string AdditionalCommandline)
-		{
-			CommandArguments += AdditionalCommandline;
-		}
-
-		public bool ForceCleanDeviceArtifacts()
-		{
-			DirectoryInfo ClientTempDirInfo = new DirectoryInfo(ArtifactPath) { Attributes = FileAttributes.Normal };
-			Log.Info(KnownLogEvents.Gauntlet_DeviceEvent, "Setting files in device artifacts {0} to have normal attributes (no longer read-only).", ArtifactPath);
-			foreach (FileSystemInfo info in ClientTempDirInfo.GetFileSystemInfos("*", SearchOption.AllDirectories))
-			{
-				info.Attributes = FileAttributes.Normal;
-			}
-			try
-			{
-				Log.Info(KnownLogEvents.Gauntlet_DeviceEvent, "Clearing device artifact path {0} (force)", ArtifactPath);
-				Directory.Delete(ArtifactPath, true);
-			}
-			catch (Exception Ex)
-			{
-				Log.Warning(KnownLogEvents.Gauntlet_DeviceEvent, "Failed to force delete artifact path {File}. {Exception}", ArtifactPath, Ex.Message);
-				return false;
-			}
-			return true;
-		}
-
-		public virtual void CleanDeviceArtifacts()
-		{
-			if (!string.IsNullOrEmpty(ArtifactPath) && Directory.Exists(ArtifactPath))
-			{
-				try
-				{
-					Log.Info("Clearing device artifacts path {0} for {1}", ArtifactPath, Device.Name);
-					Directory.Delete(ArtifactPath, true);
-				}
-				catch (Exception Ex)
-				{
-					Log.Info(KnownLogEvents.Gauntlet_DeviceEvent, "First attempt at clearing artifact path {0} failed - trying again", ArtifactPath);
-					if (!ForceCleanDeviceArtifacts())
-					{
-						Log.Warning(KnownLogEvents.Gauntlet_DeviceEvent, "Failed to delete {File}. {Exception}", ArtifactPath, Ex.Message);
-					}
-				}
-			}
-		}
-	}
-
-	public interface IWindowsSelfInstallingBuild
-	{
-		WindowsAppInstall Install(TargetDeviceWindows TargetDevice, UnrealAppConfig AppConfig, out string BasePath);
-	}
-
-	public class Win64DeviceFactory : IDeviceFactory
-	{
-		public bool CanSupportPlatform(UnrealTargetPlatform? Platform)
-		{
-			return Platform == UnrealTargetPlatform.Win64;
-		}
-
-		public ITargetDevice CreateDevice(string InRef, string InCachePath, string InParam = null)
-		{
-			return new TargetDeviceWindows(InRef, InCachePath);
-		}
-	}
-
 	/// <summary>
 	/// Win32/64 implementation of a device to run applications
 	/// </summary>
-	public class TargetDeviceWindows : ITargetDevice
+	public class TargetDeviceWindows : TargetDeviceDesktopCommon
 	{
-		public string Name { get; protected set; }
-
-		protected string UserDir { get; set; }
-
-		/// <summary>
-		/// Our mappings of Intended directories to where they actually represent on this platform.
-		/// </summary>
-		protected Dictionary<EIntendedBaseCopyDirectory, string> LocalDirectoryMappings { get; set; }
-
 		public TargetDeviceWindows(string InName, string InCacheDir)
+			: base(InName, InCacheDir)
 		{
-			Name = InName;
-			LocalCachePath = InCacheDir;
+			Platform = UnrealTargetPlatform.Win64;
 			RunOptions = CommandUtils.ERunOptions.NoWaitForExit | CommandUtils.ERunOptions.NoLoggingOfRunCommand;
-
-			UserDir = Path.Combine(LocalCachePath, "UserDir");
-            LocalDirectoryMappings = new Dictionary<EIntendedBaseCopyDirectory, string>();
 		}
 
-		#region IDisposable Support
-		private bool disposedValue = false; // To detect redundant calls
-
-		protected virtual void Dispose(bool disposing)
+		public override IAppInstall InstallApplication(UnrealAppConfig AppConfig)
 		{
-			if (!disposedValue)
+			switch (AppConfig.Build)
 			{
-				if (disposing)
-				{
-					// TODO: dispose managed state (managed objects).
-				}
+				case NativeStagedBuild:
+					return InstallNativeStagedBuild(AppConfig, AppConfig.Build as NativeStagedBuild);
 
-				// TODO: free unmanaged resources (unmanaged objects) and override a finalizer below.
-				// TODO: set large fields to null.
-
-				disposedValue = true;
-			}
-		}
+				case StagedBuild:
+					return InstallStagedBuild(AppConfig, AppConfig.Build as StagedBuild);
 
-		// This code added to correctly implement the disposable pattern.
-		public void Dispose()
-		{
-			// Do not change this code. Put cleanup code in Dispose(bool disposing) above.
-			Dispose(true);
-			// TODO: uncomment the following line if the finalizer is overridden above.
-			// GC.SuppressFinalize(this);
-		}
-		#endregion
+				case EditorBuild:
+					return InstallEditorBuild(AppConfig, AppConfig.Build as EditorBuild);
 
-		public CommandUtils.ERunOptions RunOptions { get; set; }
+				case IWindowsSelfInstallingBuild:
+					return InstallSelfInstallingBuild(AppConfig, AppConfig.Build as IWindowsSelfInstallingBuild);
 
-        // We care about UserDir in windows as some of the roles may require files going into user instead of build dir.
-        public void PopulateDirectoryMappings(string BasePath, string UserDir)
-		{
-			LocalDirectoryMappings.Add(EIntendedBaseCopyDirectory.Build, Path.Combine(BasePath, "Build"));
-			LocalDirectoryMappings.Add(EIntendedBaseCopyDirectory.Binaries, Path.Combine(BasePath, "Binaries"));
-			LocalDirectoryMappings.Add(EIntendedBaseCopyDirectory.Config, Path.Combine(BasePath, "Saved", "Config"));
-            LocalDirectoryMappings.Add(EIntendedBaseCopyDirectory.Content, Path.Combine(BasePath, "Content"));
-            LocalDirectoryMappings.Add(EIntendedBaseCopyDirectory.Demos, Path.Combine(UserDir, "Saved", "Demos"));
-			LocalDirectoryMappings.Add(EIntendedBaseCopyDirectory.PersistentDownloadDir, Path.Combine(BasePath, "Saved", "PersistentDownloadDir"));
-			LocalDirectoryMappings.Add(EIntendedBaseCopyDirectory.Profiling, Path.Combine(BasePath, "Saved", "Profiling"));
-            LocalDirectoryMappings.Add(EIntendedBaseCopyDirectory.Saved, Path.Combine(BasePath, "Saved"));
+				default:
+					throw new AutomationException("{0} is an invalid build type!", AppConfig.Build.ToString());
+			}
 		}
 
-        public IAppInstance Run(IAppInstall App)
+		public override IAppInstance Run(IAppInstall App)
 		{
 			WindowsAppInstall WinApp = App as WindowsAppInstall;
 
@@ -410,7 +126,7 @@ namespace Gauntlet
 				Result = CommandUtils.Run(WinApp.ExecutablePath,
 					CmdLine,
 					Options: FinalRunOptions,
-					SpewFilterCallback: new SpewFilterCallbackType(M => { return ProcessLogFile == null ? M : null ; }) /* make sure stderr does not spew in the stdout */,
+					SpewFilterCallback: new SpewFilterCallbackType(M => { return ProcessLogFile == null ? M : null; }) /* make sure stderr does not spew in the stdout */,
 					WorkingDir: WinApp.WorkingDirectory);
 
 				if (Result.HasExited && Result.ExitCode != 0)
@@ -424,39 +140,7 @@ namespace Gauntlet
 			return new WindowsAppInstance(WinApp, Result, ProcessLogFile);
 		}
 
-		private void CopyAdditionalFiles(UnrealAppConfig AppConfig)
-		{
-			if (AppConfig.FilesToCopy != null)
-			{
-				foreach (UnrealFileToCopy FileToCopy in AppConfig.FilesToCopy)
-				{
-					string PathToCopyTo = Path.Combine(LocalDirectoryMappings[FileToCopy.TargetBaseDirectory], FileToCopy.TargetRelativeLocation);
-					if (File.Exists(FileToCopy.SourceFileLocation))
-					{
-						FileInfo SrcInfo = new FileInfo(FileToCopy.SourceFileLocation);
-						SrcInfo.IsReadOnly = false;
-						string DirectoryToCopyTo = Path.GetDirectoryName(PathToCopyTo);
-						if (!Directory.Exists(DirectoryToCopyTo))
-						{
-							Directory.CreateDirectory(DirectoryToCopyTo);
-						}
-						if (File.Exists(PathToCopyTo))
-						{
-							FileInfo ExistingFile = new FileInfo(PathToCopyTo);
-							ExistingFile.IsReadOnly = false;
-						}
-						SrcInfo.CopyTo(PathToCopyTo, true);
-						Log.Info("Copying {0} to {1}", FileToCopy.SourceFileLocation, PathToCopyTo);
-					}
-					else
-					{
-						Log.Warning(KnownLogEvents.Gauntlet_DeviceEvent, "File to copy {File} not found", FileToCopy);
-					}
-				}
-			}
-		}
-
-		protected IAppInstall InstallNativeStagedBuild(UnrealAppConfig AppConfig, NativeStagedBuild InBuild)
+		protected override IAppInstall InstallNativeStagedBuild(UnrealAppConfig AppConfig, NativeStagedBuild InBuild)
 		{
 			WindowsAppInstall WinApp = new WindowsAppInstall(AppConfig.Name, AppConfig.ProjectName, this);
 			WinApp.CanAlterCommandArgs = AppConfig.CanAlterCommandArgs;
@@ -479,11 +163,11 @@ namespace Gauntlet
 			return WinApp;
 		}
 
-		protected IAppInstall InstallStagedBuild(UnrealAppConfig AppConfig, StagedBuild InBuild)
+		protected override IAppInstall InstallStagedBuild(UnrealAppConfig AppConfig, StagedBuild InBuild)
 		{
 			string BuildPath = InBuild.BuildPath;
 
-			if (CanRunFromPath(BuildPath) == false)
+			if (Utils.SystemHelpers.IsNetworkPath(BuildPath))
 			{
 				string SubDir = string.IsNullOrEmpty(AppConfig.Sandbox) ? AppConfig.ProjectName : AppConfig.Sandbox;
 				string BasePath = string.IsNullOrEmpty(AppConfig.DestLocalInstallDir) ? this.LocalCachePath : AppConfig.DestLocalInstallDir;
@@ -527,14 +211,14 @@ namespace Gauntlet
 			// clear artifact path
 			WinApp.CleanDeviceArtifacts();
 
-            if (LocalDirectoryMappings.Count == 0)
-            {
-                PopulateDirectoryMappings(Path.Combine(BuildPath, AppConfig.ProjectName), UserDir);
-            }
+			if (LocalDirectoryMappings.Count == 0)
+			{
+				PopulateDirectoryMappings(Path.Combine(BuildPath, AppConfig.ProjectName));
+			}
 
 			CopyAdditionalFiles(AppConfig);
 
-            if (Path.IsPathRooted(InBuild.ExecutablePath))
+			if (Path.IsPathRooted(InBuild.ExecutablePath))
 			{
 				WinApp.ExecutablePath = InBuild.ExecutablePath;
 			}
@@ -569,9 +253,26 @@ namespace Gauntlet
 			return WinApp;
 		}
 
-		public void ClearSavedDirectory(UnrealAppConfig AppConfiguration)
+		protected override IAppInstall InstallEditorBuild(UnrealAppConfig AppConfig, EditorBuild Build)
 		{
+			WindowsAppInstall WinApp = new WindowsAppInstall(AppConfig.Name, AppConfig.ProjectName, this);
+
+			WinApp.WorkingDirectory = Path.GetDirectoryName(Build.ExecutablePath);
+			WinApp.RunOptions = RunOptions;
+
+			// Force this to stop logs and other artifacts going to different places
+			WinApp.CommandArguments = AppConfig.CommandLine + string.Format(" -userdir=\"{0}\"", UserDir);
+			WinApp.ArtifactPath = Path.Combine(UserDir, @"Saved");
+			WinApp.ExecutablePath = Build.ExecutablePath;
+
+			if (LocalDirectoryMappings.Count == 0)
+			{
+				PopulateDirectoryMappings(AppConfig.ProjectFile.Directory.FullName);
+			}
+
+			CopyAdditionalFiles(AppConfig);
 
+			return WinApp;
 		}
 
 		protected IAppInstall InstallSelfInstallingBuild(UnrealAppConfig AppConfig, IWindowsSelfInstallingBuild Build)
@@ -597,90 +298,68 @@ namespace Gauntlet
 
 			if (LocalDirectoryMappings.Count == 0)
 			{
-				PopulateDirectoryMappings(Path.Combine(BasePath, AppConfig.ProjectName), UserDir);
+				PopulateDirectoryMappings(Path.Combine(BasePath, AppConfig.ProjectName));
 			}
 
 			CopyAdditionalFiles(AppConfig);
 			return WinApp;
 		}
+	}
 
+	public class WindowsAppInstall : DesktopCommonAppInstall<TargetDeviceWindows>, IAppInstall.IDynamicCommandLine
+	{
+		public bool CanAlterCommandArgs;
 
-		public IAppInstall InstallApplication(UnrealAppConfig AppConfig)
-		{
-			if (AppConfig.Build is NativeStagedBuild)
-			{
-				return InstallNativeStagedBuild(AppConfig, AppConfig.Build as NativeStagedBuild);
-			}
-			else if (AppConfig.Build is StagedBuild)
-			{
-				return InstallStagedBuild(AppConfig, AppConfig.Build as StagedBuild);
-			}
-			else if (AppConfig.Build is IWindowsSelfInstallingBuild)
-			{
-				return InstallSelfInstallingBuild(AppConfig, AppConfig.Build as IWindowsSelfInstallingBuild);
-			}
-
-			EditorBuild EditorBuild = AppConfig.Build as EditorBuild;
-
-			if (EditorBuild == null)
-			{
-				throw new AutomationException("Invalid build type!");
-			}
-
-			WindowsAppInstall WinApp = new WindowsAppInstall(AppConfig.Name, AppConfig.ProjectName, this);
-
-			WinApp.WorkingDirectory = Path.GetDirectoryName(EditorBuild.ExecutablePath);
-			WinApp.RunOptions = RunOptions;
+		[Obsolete("Will be removed in a future release. Use 'DesktopDevice' instead.")]
+		public TargetDeviceWindows WinDevice => DesktopDevice;
 
-			// Force this to stop logs and other artifacts going to different places
-			WinApp.CommandArguments = AppConfig.CommandLine + string.Format(" -userdir=\"{0}\"", UserDir);
-			WinApp.ArtifactPath = Path.Combine(UserDir, @"Saved");
-			WinApp.ExecutablePath = EditorBuild.ExecutablePath;
-
-			if (LocalDirectoryMappings.Count == 0)
+		public override string CommandArguments
+		{
+			get { return CommandArgumentsPrivate; }
+			set
 			{
-				PopulateDirectoryMappings(AppConfig.ProjectFile.Directory.FullName, AppConfig.ProjectFile.Directory.FullName);
+				if (CanAlterCommandArgs || string.IsNullOrEmpty(CommandArgumentsPrivate))
+				{
+					CommandArgumentsPrivate = value;
+				}
+				else
+				{
+					Log.Info("Skipped setting command AppInstall line when CanAlterCommandArgs = false");
+				}
 			}
-
-			CopyAdditionalFiles(AppConfig);
-
-			return WinApp;
 		}
 
-		public void CopyAppConfigurationFiles(UnrealAppConfig AppConfiguration)
-		{
+		private string CommandArgumentsPrivate;
 
+		public WindowsAppInstall(string InName, string InProjectName, TargetDeviceWindows InDevice)
+			: base(InName, InProjectName, InDevice)
+		{
+			CanAlterCommandArgs = true;
 		}
 
-		public bool CanRunFromPath(string InPath)
+		public void AppendCommandline(string AdditionalCommandline)
 		{
-			return !Utils.SystemHelpers.IsNetworkPath(InPath);
+			CommandArguments += AdditionalCommandline;
 		}
+	}
 
-		public UnrealTargetPlatform? Platform { get { return UnrealTargetPlatform.Win64; } }
-
-		public string LocalCachePath { get; private set; }
-		public bool IsAvailable { get { return true; } }
-		public bool IsConnected { get { return true; } }
-		public bool IsOn { get { return true; } }
-		public bool PowerOn() { return true; }
-		public bool PowerOff() { return true; }
-		public bool Reboot() { return true; }
-		public bool Connect() { return true; }
-		public bool Disconnect(bool bForce = false) { return true; }
+	public class WindowsAppInstance : DesktopCommonAppInstance<WindowsAppInstall, TargetDeviceWindows>
+	{
+		public WindowsAppInstance(WindowsAppInstall InInstall, IProcessResult InProcess, string ProcessLogFile = null)
+			: base(InInstall, InProcess, ProcessLogFile)
+		{ }
+	}
 
-		public override string ToString()
+	public class Win64DeviceFactory : IDeviceFactory
+	{
+		public bool CanSupportPlatform(UnrealTargetPlatform? Platform)
 		{
-			return Name;
+			return Platform == UnrealTargetPlatform.Win64;
 		}
 
-		public Dictionary<EIntendedBaseCopyDirectory, string> GetPlatformDirectoryMappings()
+		public ITargetDevice CreateDevice(string InRef, string InCachePath, string InParam = null)
 		{
-			if (LocalDirectoryMappings.Count == 0)
-			{
-				Log.Warning("Platform directory mappings have not been populated for this platform! This should be done within InstallApplication()");
-			}
-			return LocalDirectoryMappings;
+			return new TargetDeviceWindows(InRef, InCachePath);
 		}
 	}
 }
\ No newline at end of file
diff --git a/Engine/Source/Programs/AutomationTool/Gauntlet/Platform/Windows/Gauntlet.WindowsBuildSources.cs b/Engine/Source/Programs/AutomationTool/Gauntlet/Platform/Windows/Gauntlet.WindowsBuildSources.cs
--- a/Engine/Source/Programs/AutomationTool/Gauntlet/Platform/Windows/Gauntlet.WindowsBuildSources.cs
+++ b/Engine/Source/Programs/AutomationTool/Gauntlet/Platform/Windows/Gauntlet.WindowsBuildSources.cs
@@ -13,4 +13,9 @@ namespace Gauntlet
 
 		public override string PlatformFolderPrefix { get { return "Windows"; } }
 	}
+
+	public interface IWindowsSelfInstallingBuild
+	{
+		WindowsAppInstall Install(TargetDeviceWindows TargetDevice, UnrealAppConfig AppConfig, out string BasePath);
+	}
 }
diff --git a/Engine/Source/Programs/AutomationTool/Gauntlet/SelfTest/Unreal/Gauntlet.SelfTest.TestUnrealSession.cs b/Engine/Source/Programs/AutomationTool/Gauntlet/SelfTest/Unreal/Gauntlet.SelfTest.TestUnrealSession.cs
--- a/Engine/Source/Programs/AutomationTool/Gauntlet/SelfTest/Unreal/Gauntlet.SelfTest.TestUnrealSession.cs
+++ b/Engine/Source/Programs/AutomationTool/Gauntlet/SelfTest/Unreal/Gauntlet.SelfTest.TestUnrealSession.cs
@@ -225,8 +225,8 @@ namespace Gauntlet.SelfTest
 
 			TargetDeviceWindows ClientDevice = (TargetDeviceWindows)Session.UnrealDeviceReservation.ReservedDevices[0];
 			TargetDeviceWindows ServerDevice = (TargetDeviceWindows)Session.UnrealDeviceReservation.ReservedDevices[1];
-			ClientDevice.PopulateDirectoryMappings(Path.Combine(BuildPath, "WindowsClient", ProjectName), Path.Combine(ClientDevice.LocalCachePath, "UserDir"));
-			ServerDevice.PopulateDirectoryMappings(Path.Combine(BuildPath, "WindowsServer", ProjectName), Path.Combine(ServerDevice.LocalCachePath, "UserDir"));
+			ClientDevice.PopulateDirectoryMappings(Path.Combine(BuildPath, "WindowsClient", ProjectName));
+			ServerDevice.PopulateDirectoryMappings(Path.Combine(BuildPath, "WindowsServer", ProjectName));
 
 			// SessionInstance
 			UnrealSessionInstance.RoleInstance ClientRoleInstance = new(Client, GetAppInstanceForDevice(ClientDevice, Build.BuildName));
