[SubmitTool]
- Add CheckBalancedMacros validator and functionality to the script to accept a file list
- Added some extra error checking for CheckCopyrightNotices validator
- Add CheckUnacceptableWords validator and functionality to the script to accept a file list
- Fix an unacceptable word in SubmitTool.ini

diff --git a/Engine/Source/Programs/AutomationTool/Scripts/CheckBalancedMacros.cs b/Engine/Source/Programs/AutomationTool/Scripts/CheckBalancedMacros.cs
--- a/Engine/Source/Programs/AutomationTool/Scripts/CheckBalancedMacros.cs
+++ b/Engine/Source/Programs/AutomationTool/Scripts/CheckBalancedMacros.cs
@@ -16,6 +16,7 @@ namespace AutomationTool
 	[Help("Checks that all source files have balanced macros for enabling/disabling optimization, warnings, etc...")]
 	[Help("Project=<Path>", "Path to an additional project file to consider")]
 	[Help("File=<Path>", "Path to a file to parse in isolation, for testing")]
+	[Help("OverrideFileList=<Path>", "Path to a text file with paths to the files you want to parse")]
 	[Help("Ignore=<Name>", "File name (without path) to exclude from testing")]
 	class CheckBalancedMacros : BuildCommand
 	{
@@ -106,7 +107,14 @@ namespace AutomationTool
 
 			// Check if we want to just parse a single file
 			string FileParam = ParseParamValue("File");
-			if(FileParam != null)
+			string OverrideFileList = ParseParamValue("OverrideFileList=", null); // Specify a file list of individual files you want to check instead of the entire directory
+
+			if (FileParam != null && OverrideFileList != null)
+			{
+				throw new AutomationException("File and OverrideFileList parameters cannot be passed at the same time.");
+			}
+
+			if (FileParam != null)
 			{
 				// Check the file exists
 				FileReference File = new FileReference(FileParam);
@@ -116,6 +124,42 @@ namespace AutomationTool
 				}
 				CheckSourceFile(File, IdentifierToIndex, new object());
 			}
+			else if (OverrideFileList != null)
+			{
+				Logger.LogInformation("Finding files from OverrideFileList {File}", OverrideFileList);
+
+				FileReference FileListToCheck = new FileReference(OverrideFileList);
+				if (!FileReference.Exists(FileListToCheck))
+				{
+					throw new AutomationException("FileList '{0}' does not exist", FileListToCheck);
+				}
+
+				string[] FilesToCheck = FileReference.ReadAllLines(FileListToCheck);
+				List<FileReference> SourceFiles = new List<FileReference>();
+
+				foreach (string File in FilesToCheck.Where(x => !String.IsNullOrWhiteSpace(x) && (x.EndsWith(".h") || x.EndsWith(".cpp"))))
+				{
+					SourceFiles.Add(new FileReference(File));
+				}
+
+				// Loop through all the source files
+				using (ThreadPoolWorkQueue Queue = new ThreadPoolWorkQueue())
+				{
+					object LogLock = new object();
+					foreach (FileReference SourceFile in SourceFiles)
+					{
+						Queue.Enqueue(() => CheckSourceFile(SourceFile, IdentifierToIndex, LogLock));
+					}
+
+					using (LogStatusScope Scope = new LogStatusScope("Checking source files..."))
+					{
+						while (!Queue.Wait(10 * 1000))
+						{
+							Scope.SetProgress("{0}/{1}", SourceFiles.Count - Queue.NumRemaining, SourceFiles.Count);
+						}
+					}
+				}
+			}
 			else
 			{
 				// Add the additional files to be ignored
