Downgrade errors parsing JSON log output to Debug verbosity.

diff --git a/Engine/Source/Programs/AutomationTool/Gauntlet/Unreal/Utils/Gauntlet.UnrealLogParser.cs b/Engine/Source/Programs/AutomationTool/Gauntlet/Unreal/Utils/Gauntlet.UnrealLogParser.cs
--- a/Engine/Source/Programs/AutomationTool/Gauntlet/Unreal/Utils/Gauntlet.UnrealLogParser.cs
+++ b/Engine/Source/Programs/AutomationTool/Gauntlet/Unreal/Utils/Gauntlet.UnrealLogParser.cs
@@ -307,7 +307,7 @@ namespace Gauntlet
 					}
 					catch (Exception ex)
 					{
-						EpicGames.Core.Log.Logger.LogInformation(ex, "Unable to parse log line: {Line}, Exception: {Ex}", Line, ex.ToString());
+						EpicGames.Core.Log.Logger.LogDebug(ex, "Unable to parse log line: {Line}, Exception: {Ex}", Line, ex.ToString());
 
 						int MinIdx = Math.Max(BaseIdx - 2048, 0);
 						int MaxIdx = Math.Min(BaseIdx + 2048, InContent.Length);
@@ -315,7 +315,7 @@ namespace Gauntlet
 						string[] Context = InContent.Substring(MinIdx, MaxIdx - MinIdx).Split('\n');
 						for (int idx = 1; idx < Context.Length - 1; idx++)
 						{
-							EpicGames.Core.Log.Logger.LogInformation("Context {Idx}: {Line}", idx, Context[idx].TrimEnd());
+							EpicGames.Core.Log.Logger.LogDebug("Context {Idx}: {Line}", idx, Context[idx].TrimEnd());
 						}
 					}
 				}
