PR #11018: BuildGraph: Create LogPath directory before writing to directory. (prevents System.IO.DirectoryNotFoundException)

diff --git a/Engine/Source/Programs/AutomationTool/BuildGraph/TempStorage.cs b/Engine/Source/Programs/AutomationTool/BuildGraph/TempStorage.cs
--- a/Engine/Source/Programs/AutomationTool/BuildGraph/TempStorage.cs
+++ b/Engine/Source/Programs/AutomationTool/BuildGraph/TempStorage.cs
@@ -1082,6 +1082,7 @@ namespace AutomationTool
 					string CopyResult = CopyDirectory(ZipDir, OutputDir);
 					string LogPath = CommandUtils.CombinePaths(RootDir.FullName, "Engine/Programs/AutomationTool/Saved/Logs", $"Copy Files to Temp Storage.log");
 					Logger.LogInformation("Saving copy log to {LogPath}", LogPath);
+					Directory.CreateDirectory(Path.GetDirectoryName(LogPath));
 					File.WriteAllText(LogPath, string.Join(Environment.NewLine, CopyResult));
 					Parallel.ForEach(ZipFiles, (z) => z.Delete());
 					ZipFiles = new ConcurrentBag<FileInfo>(ZipFiles.Select(z => new FileInfo(CommandUtils.MakeRerootedFilePath(z.FullName, StagingDir.FullName, OutputDir.FullName))));
