[Gauntlet] Improves SystemHelpers.Delete functions by allowing force clean option that sets file attributes

diff --git a/Engine/Source/Programs/AutomationTool/Gauntlet/Framework/Gauntlet.Utils.cs b/Engine/Source/Programs/AutomationTool/Gauntlet/Framework/Gauntlet.Utils.cs
--- a/Engine/Source/Programs/AutomationTool/Gauntlet/Framework/Gauntlet.Utils.cs
+++ b/Engine/Source/Programs/AutomationTool/Gauntlet/Framework/Gauntlet.Utils.cs
@@ -1536,47 +1536,69 @@ namespace Gauntlet
 			}
 
 			/// <summary>
-			/// Wrap DirectoryInfo.Delete(Force) and retry once if exception is "A retry should be performed."
+			/// Deletes the specified directory
 			/// </summary>
-			public static void Delete(DirectoryInfo Directory, bool Force)
+			/// <param name="Directory">The directory to delete</param>
+			/// <param name="bRecursive">Whether or not subdirectories and it's contents should also be deleted</param>
+			/// <param name="bForce">If true, a second attempt will me made at file deletion after setting file attributes to normal</param>
+			public static void Delete(DirectoryInfo Directory, bool bRecursive, bool bForce = false)
 			{
 				try
 				{
-					Directory.Delete(Force);
+					Directory.Delete(bRecursive);
 				}
-				catch (Exception ex)
+				catch (Exception Ex)
 				{
-					if (ex.ToString().Contains("A retry should be performed"))
+					if (Ex.ToString().Contains("A retry should be performed"))
+					{
+						Log.Info("Retrying deletion once.");
+						Directory.Delete(bRecursive);
+					}
+					else if(bForce)
 					{
-						Log.Info("Retrying Directory.Delete(Force) once.");
-						Directory.Delete(Force);
+						Log.Info("Setting files in {Directory} to have normal attributes (no longer read-only) and retrying deletion.", Directory);
+						Directory.Attributes = FileAttributes.Normal;
+
+						foreach(FileSystemInfo Info in Directory.EnumerateFiles("*", SearchOption.AllDirectories))
+						{
+							Info.Attributes = FileAttributes.Normal;
+							Info.Delete(); // throw exceptions here because we have requested a force clean
+						}
 					}
 					else
 					{
-						throw;
+						Log.Warning("Failed to delete directory {Directory}!", Directory);
 					}
 				}
 			}
 
 			/// <summary>
-			/// Wrap FileInfo.Delete() and retry once if exception is "A retry should be performed."
+			/// Deletes the specified file
 			/// </summary>
-			public static void Delete(FileInfo File)
+			/// <param name="File">The file to delete</param>
+			/// <param name="bForce">If true, a second attempt will me made at file deletion after setting file attributes to normal</param>
+			public static void Delete(FileInfo File, bool bForce = false)
 			{
 				try
 				{
 					File.Delete();
 				}
-				catch (Exception ex)
+				catch (Exception Ex)
 				{
-					if (ex.ToString().Contains("A retry should be performed"))
+					if (Ex.ToString().Contains("A retry should be performed"))
 					{
-						Log.Info("Retrying File.Delete() once.");
+						Log.Info("Retrying deletion once.");
 						File.Delete();
 					}
+					else if (bForce)
+					{
+						Log.Info("Setting files {File} to have normal attributes (no longer read-only) and retrying deletion.", File);
+						File.Attributes = FileAttributes.Normal;
+						File.Delete(); // throw exceptions here because we have requested a force clean
+					}
 					else
 					{
-						throw;
+						Log.Warning("Failed to delete directory {Directory}!", File);
 					}
 				}
 			}
