AutomationTool: Check for dotnet instead of always relying on an installed visual studio

#rnx

diff --git a/Engine/Source/Programs/AutomationTool/AutomationUtils/WindowsHostPlatform.cs b/Engine/Source/Programs/AutomationTool/AutomationUtils/WindowsHostPlatform.cs
--- a/Engine/Source/Programs/AutomationTool/AutomationUtils/WindowsHostPlatform.cs
+++ b/Engine/Source/Programs/AutomationTool/AutomationUtils/WindowsHostPlatform.cs
@@ -10,15 +10,36 @@ using UnrealBuildTool;
 using EpicGames.Core;
 using UnrealBuildBase;
 using System.Runtime.Versioning;
+using Microsoft.Extensions.Logging;
 
 namespace AutomationTool
 {
 	class WindowsHostPlatform : HostPlatform
 	{
+		static string CachedFrameworkMsbuildTool = "";
+
 		[SupportedOSPlatform("windows")]
 		public override string GetFrameworkMsbuildExe()
 		{
-			return WindowsExports.GetMSBuildToolPath();
+			// Look for dotnet, we only support dotnet.
+			if (string.IsNullOrEmpty(CachedFrameworkMsbuildTool))
+			{
+				bool CanUseMsBuild = string.IsNullOrEmpty(CommandUtils.WhichApp("dotnet")) == false;
+
+				if (CanUseMsBuild)
+				{
+					Logger.LogInformation($"using {CommandUtils.WhichApp("dotnet")}!");
+
+					CachedFrameworkMsbuildTool = "dotnet msbuild";
+				}
+				else
+				{
+					// Look for visual studio msbuild
+					CachedFrameworkMsbuildTool = WindowsExports.GetMSBuildToolPath();
+				}
+			}
+
+			return CachedFrameworkMsbuildTool;
 		}
 
 		public override string RelativeBinariesFolder
