IAS - append the encryption key filename to upload command

diff --git a/Engine/Source/Programs/AutomationTool/Scripts/CopyBuildToStagingDirectory.Automation.cs b/Engine/Source/Programs/AutomationTool/Scripts/CopyBuildToStagingDirectory.Automation.cs
--- a/Engine/Source/Programs/AutomationTool/Scripts/CopyBuildToStagingDirectory.Automation.cs
+++ b/Engine/Source/Programs/AutomationTool/Scripts/CopyBuildToStagingDirectory.Automation.cs
@@ -3739,6 +3739,12 @@ namespace AutomationScripts
 					UploadArgs.Append(" -BuildVersion=0");
 					UploadArgs.Append(" -WriteTocToDisk");
 					UploadArgs.Append(Params.AdditionalPakOptions);
+
+					FileReference CryptoKeysFilename = FileReference.Combine(SC.MetadataDir, "Crypto.json");
+					if (FileExists_NoExceptions(CryptoKeysFilename.FullName))
+					{
+						UploadArgs.Append(string.Format(" -CryptoKeys={0}", CryptoKeysFilename.FullName));
+					}
 				}
 				else
 				{
