AutomationTool: Remove obsolete ParallelExecutor

#rnx

diff --git a/Engine/Source/Programs/AutomationTool/AutomationUtils/ParallelExecutor.cs b/Engine/Source/Programs/AutomationTool/AutomationUtils/ParallelExecutor.cs
deleted file mode 100644
diff --git a/Engine/Source/Programs/AutomationTool/AutomationUtils/UnrealBuild.cs b/Engine/Source/Programs/AutomationTool/AutomationUtils/UnrealBuild.cs
--- a/Engine/Source/Programs/AutomationTool/AutomationUtils/UnrealBuild.cs
+++ b/Engine/Source/Programs/AutomationTool/AutomationUtils/UnrealBuild.cs
@@ -521,9 +521,10 @@ namespace AutomationTool
 		/// <param name="InDeleteBuildProducts">if specified, determines if the build products will be deleted before building. If not specified -clean parameter will be used,</param>
 		/// <param name="InUpdateVersionFiles">True if the version files are to be updated </param>
 		/// <param name="InForceNoXGE">If true will force XGE off</param>
-		/// <param name="InUseParallelExecutor">If true AND XGE not present or not being used then use ParallelExecutor</param>
 		/// <param name="InAllCores">If true AND XGE not present or not being used then ensure UBT uses all available cores</param>
-		public void Build(BuildAgenda Agenda, bool? InDeleteBuildProducts = null, bool InUpdateVersionFiles = true, bool InForceNoXGE = false, bool InUseParallelExecutor = false, bool InShowProgress = false, bool InAllCores = false, int? InChangelistNumberOverride = null, Dictionary<BuildTarget, BuildManifest> InTargetToManifest = null)
+		/// <param name="InChangelistNumberOverride"></param>
+		/// <param name="InTargetToManifest"></param>
+		public void Build(BuildAgenda Agenda, bool? InDeleteBuildProducts = null, bool InUpdateVersionFiles = true, bool InForceNoXGE = false, bool InAllCores = false, int? InChangelistNumberOverride = null, Dictionary<BuildTarget, BuildManifest> InTargetToManifest = null)
 		{
 			if (!CommandUtils.CmdEnv.HasCapabilityToCompile)
 			{
@@ -572,11 +573,8 @@ namespace AutomationTool
 			bool bDisableXGE = ParseParam("NoXGE") || InForceNoXGE;
 			bool bCanUseXGE = !bDisableXGE && PlatformExports.TryGetXgConsoleExecutable(out XGEConsole);
 
-			// only run ParallelExecutor if not running XGE (and we've requested ParallelExecutor and it exists)
-			bool bCanUseParallelExecutor = InUseParallelExecutor && (HostPlatform.Current.HostEditorPlatform == UnrealTargetPlatform.Win64);
 			Logger.LogDebug("************************* UnrealBuild:");
 			Logger.LogDebug("************************* UseXGE: {bCanUseXGE}", bCanUseXGE);
-			Logger.LogDebug("************************* UseParallelExecutor: {bCanUseParallelExecutor}", bCanUseParallelExecutor);
 
 			// Clean all the targets
 			foreach (BuildTarget Target in Agenda.Targets)
diff --git a/Engine/Source/Programs/AutomationTool/BuildGraph/Tasks/CompileTask.cs b/Engine/Source/Programs/AutomationTool/BuildGraph/Tasks/CompileTask.cs
--- a/Engine/Source/Programs/AutomationTool/BuildGraph/Tasks/CompileTask.cs
+++ b/Engine/Source/Programs/AutomationTool/BuildGraph/Tasks/CompileTask.cs
@@ -56,9 +56,10 @@ namespace AutomationTool.Tasks
 		public bool AllowXGE = true;
 
 		/// <summary>
-		/// Whether to allow using the parallel executor for this compile.
+		/// No longer necessary as UnrealBuildTool is run to compile targets.
 		/// </summary>
 		[TaskParameter(Optional = true)]
+		[Obsolete]
 		public bool AllowParallelExecutor = true;
 
 		/// <summary>
@@ -100,11 +101,6 @@ namespace AutomationTool.Tasks
 		/// </summary>
 		bool bAllowXGE = true;
 
-		/// <summary>
-		/// Whether to allow using the parallel executor for this job
-		/// </summary>
-		bool bAllowParallelExecutor = true;
-
 		/// <summary>
 		/// Whether to allow using all available cores for this job, when bAllowXGE is false
 		/// </summary>
@@ -138,15 +134,10 @@ namespace AutomationTool.Tasks
 				{
 					return false;
 				}
-				if (!bAllowParallelExecutor || !CompileTask.Parameters.AllowParallelExecutor)
-				{
-					return false;
-				}
 			}
 
 			CompileTaskParameters Parameters = CompileTask.Parameters;
 			bAllowXGE &= Parameters.AllowXGE;
-			bAllowParallelExecutor &= Parameters.AllowParallelExecutor;
 			bAllowAllCores &= Parameters.AllowAllCores;
 
 			UnrealBuild.BuildTarget Target = new UnrealBuild.BuildTarget { TargetName = Parameters.Target, Platform = Parameters.Platform, Config = Parameters.Configuration, UprojectPath = CompileTask.FindProjectFile(), UBTArgs = (Parameters.Arguments ?? ""), Clean = Parameters.Clean };
@@ -176,9 +167,8 @@ namespace AutomationTool.Tasks
 			Dictionary<UnrealBuild.BuildTarget, BuildManifest> TargetToManifest = new Dictionary<UnrealBuild.BuildTarget,BuildManifest>();
 			UnrealBuild Builder = new UnrealBuild(Job.OwnerCommand);
 
-			bool bCanUseParallelExecutor = (BuildHostPlatform.Current.Platform == UnrealTargetPlatform.Win64 && bAllowParallelExecutor);	// parallel executor is only available on Windows as of 2016-09-22
 			bool bAllCores = (CommandUtils.IsBuildMachine || bAllowAllCores);	// Enable using all cores if this is a build agent or the flag was passed in to the task and XGE is disabled.
-			Builder.Build(Agenda, InDeleteBuildProducts: null, InUpdateVersionFiles: false, InForceNoXGE: !bAllowXGE, InUseParallelExecutor: bCanUseParallelExecutor, InAllCores: bAllCores, InTargetToManifest: TargetToManifest);
+			Builder.Build(Agenda, InDeleteBuildProducts: null, InUpdateVersionFiles: false, InForceNoXGE: !bAllowXGE, InAllCores: bAllCores, InTargetToManifest: TargetToManifest);
 
 			UnrealBuild.CheckBuildProducts(Builder.BuildProductFiles);
 
@@ -291,7 +281,7 @@ namespace AutomationTool.Tasks
 		{
 			//
 			// Don't do any logic here. You have to do it in the ctor or a getter
-			//  otherwise you break the ParallelExecutor pathway, which doesn't call this function!
+			//  otherwise you break the ITaskExecutor pathway, which doesn't call this function!
 			//
 			return GetExecutor().ExecuteAsync(Job, BuildProducts, TagNameToFileSet);
 		}
@@ -381,10 +371,9 @@ namespace AutomationTool.Tasks
 		/// <param name="Project">The project to compile with</param>
 		/// <param name="Arguments">Additional arguments for UnrealBuildTool</param>
 		/// <param name="AllowXGE">Whether to allow using XGE for compilation</param>
-		/// <param name="AllowParallelExecutor">Whether to allow using the parallel executor for this compile</param>
 		/// <param name="Clean">Whether to allow cleaning this target. If unspecified, targets are cleaned if the -Clean argument is passed on the command line</param>
 		/// <returns>Build products from the compile</returns>
-		public static async Task<FileSet> CompileAsync(string Target, UnrealTargetPlatform Platform, UnrealTargetConfiguration Configuration, FileReference Project = null, string Arguments = null, bool AllowXGE = true, bool AllowParallelExecutor = true, bool? Clean = null)
+		public static async Task<FileSet> CompileAsync(string Target, UnrealTargetPlatform Platform, UnrealTargetConfiguration Configuration, FileReference Project = null, string Arguments = null, bool AllowXGE = true, bool? Clean = null)
 		{
 			CompileTaskParameters Parameters = new CompileTaskParameters();
 			Parameters.Target = Target;
@@ -393,7 +382,6 @@ namespace AutomationTool.Tasks
 			Parameters.Project = Project?.FullName;
 			Parameters.Arguments = Arguments;
 			Parameters.AllowXGE = AllowXGE;
-			Parameters.AllowParallelExecutor = AllowParallelExecutor;
 			Parameters.Clean = Clean;
 			return await ExecuteAsync(new CompileTask(Parameters));
 		}
diff --git a/Engine/Source/Programs/AutomationTool/Scripts/MegaXGE.Automation.cs b/Engine/Source/Programs/AutomationTool/Scripts/MegaXGE.Automation.cs
--- a/Engine/Source/Programs/AutomationTool/Scripts/MegaXGE.Automation.cs
+++ b/Engine/Source/Programs/AutomationTool/Scripts/MegaXGE.Automation.cs
@@ -148,7 +148,7 @@ class MegaXGE : BuildCommand
 		}
 		Logger.LogInformation("*************************");
 
-		UnrealBuild.Build(Agenda, InUpdateVersionFiles: IsBuildMachine, InUseParallelExecutor: ParseParam("useparallelexecutor"), InShowProgress: ShowProgress);
+		UnrealBuild.Build(Agenda, InUpdateVersionFiles: IsBuildMachine);
 
 		// 		if (WorkingCL > 0) // only move UAT files if we intend to check in some build products
 		// 		{
