AutomationTool: Check for dotnet msbuild only when visual studio msbuild is not available on windows

#rnx

diff --git a/Engine/Source/Programs/AutomationTool/AutomationUtils/HostPlatform.cs b/Engine/Source/Programs/AutomationTool/AutomationUtils/HostPlatform.cs
--- a/Engine/Source/Programs/AutomationTool/AutomationUtils/HostPlatform.cs
+++ b/Engine/Source/Programs/AutomationTool/AutomationUtils/HostPlatform.cs
@@ -8,6 +8,7 @@ using System.Text;
 using System.Diagnostics;
 using EpicGames.Core;
 using Microsoft.Extensions.Logging;
+using System.Runtime.CompilerServices;
 
 namespace AutomationTool
 {
diff --git a/Engine/Source/Programs/AutomationTool/AutomationUtils/LinuxHostPlatform.cs b/Engine/Source/Programs/AutomationTool/AutomationUtils/LinuxHostPlatform.cs
--- a/Engine/Source/Programs/AutomationTool/AutomationUtils/LinuxHostPlatform.cs
+++ b/Engine/Source/Programs/AutomationTool/AutomationUtils/LinuxHostPlatform.cs
@@ -1,34 +1,28 @@
 // Copyright Epic Games, Inc. All Rights Reserved.
 
+using EpicGames.Core;
+using Microsoft.Extensions.Logging;
 using System;
-using System.Collections.Generic;
-using System.Linq;
-using System.Text;
 using System.Diagnostics;
 using System.IO;
 using UnrealBuildTool;
-using EpicGames.Core;
-using UnrealBuildBase;
-using Microsoft.Extensions.Logging;
 
 namespace AutomationTool
 {
 	class LinuxHostPlatform : HostPlatform
 	{
-		static string CachedFrameworkMsbuildTool = "";
+		static string CachedFrameworkMsbuildExe = string.Empty;
 
 		public override string GetFrameworkMsbuildExe()
 		{
 			// Look for dotnet, we only support dotnet.
-			if (string.IsNullOrEmpty(CachedFrameworkMsbuildTool))
+			if (string.IsNullOrEmpty(CachedFrameworkMsbuildExe))
 			{
-				bool CanUseMsBuild = string.IsNullOrEmpty(CommandUtils.WhichApp("dotnet")) == false;
-
-				if (CanUseMsBuild)
+				FileReference dotnet = FileReference.FromString(CommandUtils.WhichApp("dotnet"));
+				if (dotnet != null && FileReference.Exists(dotnet))
 				{
-					Logger.LogInformation($"using {CommandUtils.WhichApp("dotnet")}!");
-
-					CachedFrameworkMsbuildTool = "dotnet msbuild";
+					Logger.LogInformation("Using {DotNet}", dotnet.FullName);
+					CachedFrameworkMsbuildExe = "dotnet msbuild";
 				}
 				else
 				{
@@ -36,7 +30,7 @@ namespace AutomationTool
 				}
 			}
 
-			return CachedFrameworkMsbuildTool;
+			return CachedFrameworkMsbuildExe;
 		}
 
 		public override string RelativeBinariesFolder
diff --git a/Engine/Source/Programs/AutomationTool/AutomationUtils/MacHostPlatform.cs b/Engine/Source/Programs/AutomationTool/AutomationUtils/MacHostPlatform.cs
--- a/Engine/Source/Programs/AutomationTool/AutomationUtils/MacHostPlatform.cs
+++ b/Engine/Source/Programs/AutomationTool/AutomationUtils/MacHostPlatform.cs
@@ -1,34 +1,27 @@
 // Copyright Epic Games, Inc. All Rights Reserved.
 
+using EpicGames.Core;
+using Microsoft.Extensions.Logging;
 using System;
-using System.Collections.Generic;
-using System.Linq;
-using System.Text;
 using System.Diagnostics;
 using System.IO;
 using UnrealBuildTool;
-using System.Text.RegularExpressions;
-using EpicGames.Core;
-using UnrealBuildBase;
-using Microsoft.Extensions.Logging;
 
 namespace AutomationTool
 {
 	class MacHostPlatform : HostPlatform
 	{
-		static string CachedFrameworkMsbuildExe = "";
+		static string CachedFrameworkMsbuildExe = string.Empty;
 
 		public override string GetFrameworkMsbuildExe()
 		{
 			// Look for dotnet, we only support dotnet.
 			if (string.IsNullOrEmpty(CachedFrameworkMsbuildExe))
 			{
-				bool CanUseMsBuild = string.IsNullOrEmpty(CommandUtils.WhichApp("dotnet")) == false;
-
-				if (CanUseMsBuild)
+				FileReference dotnet = FileReference.FromString(CommandUtils.WhichApp("dotnet"));
+				if (dotnet != null && FileReference.Exists(dotnet))
 				{
-					Logger.LogInformation("using {DotNet}!", CommandUtils.WhichApp("dotnet"));
-
+					Logger.LogInformation("Using {DotNet}", dotnet.FullName);
 					CachedFrameworkMsbuildExe = "dotnet msbuild";
 				}
 				else
diff --git a/Engine/Source/Programs/AutomationTool/AutomationUtils/WindowsHostPlatform.cs b/Engine/Source/Programs/AutomationTool/AutomationUtils/WindowsHostPlatform.cs
--- a/Engine/Source/Programs/AutomationTool/AutomationUtils/WindowsHostPlatform.cs
+++ b/Engine/Source/Programs/AutomationTool/AutomationUtils/WindowsHostPlatform.cs
@@ -1,24 +1,52 @@
 // Copyright Epic Games, Inc. All Rights Reserved.
 
+using EpicGames.Core;
+using Microsoft.Extensions.Logging;
 using System;
-using System.Collections.Generic;
-using System.Linq;
-using System.Text;
-using System.IO;
 using System.Diagnostics;
-using UnrealBuildTool;
-using EpicGames.Core;
-using UnrealBuildBase;
+using System.IO;
 using System.Runtime.Versioning;
+using UnrealBuildTool;
 
 namespace AutomationTool
 {
 	class WindowsHostPlatform : HostPlatform
 	{
+		static string CachedFrameworkMsbuildTool = "";
+
 		[SupportedOSPlatform("windows")]
 		public override string GetFrameworkMsbuildExe()
 		{
-			return WindowsExports.GetMSBuildToolPath();
+			if (string.IsNullOrEmpty(CachedFrameworkMsbuildTool))
+			{
+				try
+				{
+					// Look for visual studio msbuild
+					FileReference msbuild = FileReference.FromString(WindowsExports.GetMSBuildToolPath());
+					if (msbuild != null && FileReference.Exists(msbuild))
+					{
+						CachedFrameworkMsbuildTool = msbuild.FullName;
+						Logger.LogInformation("Using {MsBuild}", CachedFrameworkMsbuildTool);
+						return CachedFrameworkMsbuildTool;
+					}
+				}
+				catch (BuildException)
+				{
+				}
+
+				FileReference dotnet = FileReference.FromString(CommandUtils.WhichApp("dotnet"));
+				if (dotnet != null && FileReference.Exists(dotnet))
+				{
+					Logger.LogInformation("Using {DotNet}", dotnet.FullName);
+					CachedFrameworkMsbuildTool = "dotnet msbuild";
+				}
+				else
+				{
+					throw new BuildException("Unable to find installation of MSBuild.");
+				}
+			}
+
+			return CachedFrameworkMsbuildTool;
 		}
 
 		public override string RelativeBinariesFolder
