[Gauntlet] Retain desktop builds in the device cache when using device reservation blocks
Also remove destlocalinstalldir from UnrealSession and adds an equivalent InstallRoot into TargetDeviceDesktopCommon

diff --git a/Engine/Source/Programs/AutomationTool/Gauntlet/Framework/Base/Gauntlet.TargetDeviceDesktopCommon.cs b/Engine/Source/Programs/AutomationTool/Gauntlet/Framework/Base/Gauntlet.TargetDeviceDesktopCommon.cs
--- a/Engine/Source/Programs/AutomationTool/Gauntlet/Framework/Base/Gauntlet.TargetDeviceDesktopCommon.cs
+++ b/Engine/Source/Programs/AutomationTool/Gauntlet/Framework/Base/Gauntlet.TargetDeviceDesktopCommon.cs
@@ -61,14 +61,27 @@ namespace Gauntlet
 
 		public string UserDir { get; protected set; }
 
+		/// <summary>
+		/// Optional directory staged builds are installed to if the requested build is not already on the same volume
+		/// </summary>
+		[AutoParam]
+		public string InstallRoot { get; protected set; }
+
 		// TODO - move this be part of ITargetDevice
 		protected Dictionary<EIntendedBaseCopyDirectory, string> LocalDirectoryMappings { get; set; }
 
 		public TargetDeviceDesktopCommon(string InName, string InCacheDir)
 		{
+			AutoParam.ApplyParamsAndDefaults(this, Globals.Params.AllArguments);
+
 			Name = InName;
 			LocalCachePath = InCacheDir;
-			UserDir = Path.Combine(LocalCachePath, "UserDir");
+			UserDir = Path.Combine(InCacheDir, "UserDir");
+
+			if(string.IsNullOrEmpty(InstallRoot))
+			{
+				InstallRoot = InCacheDir;
+			}
 
 			LocalDirectoryMappings = new Dictionary<EIntendedBaseCopyDirectory, string>();
 		}
diff --git a/Engine/Source/Programs/AutomationTool/Gauntlet/Framework/Devices/Gauntlet.DevicePool.cs b/Engine/Source/Programs/AutomationTool/Gauntlet/Framework/Devices/Gauntlet.DevicePool.cs
--- a/Engine/Source/Programs/AutomationTool/Gauntlet/Framework/Devices/Gauntlet.DevicePool.cs
+++ b/Engine/Source/Programs/AutomationTool/Gauntlet/Framework/Devices/Gauntlet.DevicePool.cs
@@ -12,6 +12,7 @@ using System.Text.Json;
 using System.Text.Json.Serialization;
 using EpicGames.Core;
 using System.Data;
+using Gauntlet.Utils;
 
 namespace Gauntlet
 {
@@ -745,66 +746,81 @@ namespace Gauntlet
 		}
 
 		/// <summary>
-		/// Force clean the cache path by setting all of the files in the cache path
-		/// to have "normal" attributes (i.e. not read-only) before attempting to delete.
+		/// Construct a path to hold cache files and make sure it's properly cleaned
 		/// </summary>
-		public bool ForceCleanCachePath(string ClientTempDir)
+		private string GetCleanCachePath(DeviceDefinition InDeviceDefiniton)
 		{
-			DirectoryInfo ClientTempDirInfo = new DirectoryInfo(ClientTempDir) { Attributes = FileAttributes.Normal };
-
-			Log.Info(KnownLogEvents.Gauntlet_DeviceEvent, "Setting files in {0} to have normal attributes (no longer read-only).", ClientTempDir);
-			foreach (FileSystemInfo info in ClientTempDirInfo.GetFileSystemInfos("*", SearchOption.AllDirectories))
+			// Give the desktop platform a temp folder with its name under the device cache
+			string DeviceCache = Path.Combine(LocalTempDir, "DeviceCache");
+			string PlatformCache = Path.Combine(DeviceCache, InDeviceDefiniton.Platform.ToString());
+			string ClientCache = Path.Combine(PlatformCache, InDeviceDefiniton.Name);
+
+			// On Desktops, builds are installed in the device cache.
+			// When using device reservation blocks, we don't want to fully clean the cache and lose previously installed builds.
+			// If bRetainBuilds evaluates to true, it means we are in the second step or beyond in a device reservation block.
+			// In this case we'll just delete the left over UserDir which should already have been emptied by UnrealSession.
+			bool? bForceClean = UnrealAppConfig.ForceFullClean;
+			bool? bSkipInstall = UnrealAppConfig.ForceSkipInstall;
+			bool bUsingReservationBlock = bForceClean.HasValue && bSkipInstall.HasValue;
+			bool bRetainBuilds = bUsingReservationBlock && !bForceClean.Value && bSkipInstall.Value;
+
+			if(bRetainBuilds)
 			{
-				info.Attributes = FileAttributes.Normal;
-			}
+				Log.Info("Retaining build cache for device reservation block");
 
-			try
-			{
-				Log.Info(KnownLogEvents.Gauntlet_DeviceEvent, "Clearing artifact path {0} (force)", ClientTempDir);
-				Directory.Delete(ClientTempDir, true);
+				DirectoryInfo UserDirectory = new(Path.Combine(ClientCache, "UserDir"));
+				if(UserDirectory.Exists)
+				{
+					try
+					{
+						Log.Info("Cleaning stale user directory...");
+						SystemHelpers.Delete(UserDirectory, true, true);
+					}
+					catch(Exception Ex)
+					{
+						throw new AutomationException("Failed to clean user directory {0}. This could result in improper artifact reporting. {1}", UserDirectory, Ex);
+					}
+				}
 			}
-			catch (Exception Ex)
+			else
 			{
-				Log.Warning(KnownLogEvents.Gauntlet_DeviceEvent, "Failed to force delete {File}. {Exception}", ClientTempDir, Ex.Message);
-				return false;
-			}
+				int CleanAttempts = 0;
 
-			return true;
-		}
+				while (Directory.Exists(ClientCache))
+				{
+					DirectoryInfo ClientCacheDirectory = new(ClientCache);
 
-		/// <summary>
-		/// Construct a path to hold cache files and make sure it is empty
-		/// </summary>
-		private string GetCleanCachePath(DeviceDefinition InDeviceDefiniton)
-		{
-			// Give the desktop platform a temp folder with its name under the device cache 
-			string PlatformCache = Path.Combine(LocalTempDir, "DeviceCache", InDeviceDefiniton.Platform.ToString());
-			string ClientTempDir = Path.Combine(PlatformCache, InDeviceDefiniton.Name);
+					try
+					{
+						Log.Info("Cleaning stale client device cache...");
+						SystemHelpers.Delete(ClientCacheDirectory, true, true);
+					}
+					catch (Exception Ex)
+					{
+						// If we fail to acquire the default client cache while using device reservation blocks,
+						// we can't ensure future tests will have their cache directories mapped to the correct build location
+						if(bUsingReservationBlock)
+						{
+							throw new AutomationException("Failed to clean default client device cache {0}. {1}", ClientCache, Ex);
+						}
 
-			int CleanAttempts = 1;
+						// When not using device reservation blocks, we can just create a newly indexed directory for the client cache
+						else
+						{
+							string Warning = "Failed to clean client device cache {Folder}. A newly indexed directory will be created instead. {Message}";
+							Log.Warning(KnownLogEvents.Gauntlet_DeviceEvent, Warning, ClientCache, Ex.Message);
 
-			// Make sure this is a fresh directory 
-			while (Directory.Exists(ClientTempDir))
-			{
-				try
-				{
-					Directory.Delete(ClientTempDir, true);
-				}
-				catch (Exception Ex)
-				{
-					Log.Info(KnownLogEvents.Gauntlet_DeviceEvent, "Clearing artifact path {0} failed - attempting to force clean", ClientTempDir);
-					if (!ForceCleanCachePath(ClientTempDir))
-					{
-						// warn and use a different directory
-						Log.Warning(KnownLogEvents.Gauntlet_DeviceEvent, "Failed to delete {Folder}. {Message}", ClientTempDir, Ex.Message);
-						ClientTempDir = Path.Combine(PlatformCache, string.Format("{0}_{1}", InDeviceDefiniton.Name, CleanAttempts++));
+							ClientCache = Path.Combine(PlatformCache, $"{InDeviceDefiniton.Name}_{++CleanAttempts}");
+						}
 					}
 				}
+
+				// create this path
+				Log.Info("Client device cache set to {Directory}", ClientCache);
+				Directory.CreateDirectory(ClientCache);
 			}
-			// create this path
-			Directory.CreateDirectory(ClientTempDir);
 
-			return ClientTempDir;
+			return ClientCache;
 		}
 
 		/// <summary>
diff --git a/Engine/Source/Programs/AutomationTool/Gauntlet/Platform/Linux/Gauntlet.TargetDeviceLinux.cs b/Engine/Source/Programs/AutomationTool/Gauntlet/Platform/Linux/Gauntlet.TargetDeviceLinux.cs
--- a/Engine/Source/Programs/AutomationTool/Gauntlet/Platform/Linux/Gauntlet.TargetDeviceLinux.cs
+++ b/Engine/Source/Programs/AutomationTool/Gauntlet/Platform/Linux/Gauntlet.TargetDeviceLinux.cs
@@ -118,45 +118,43 @@ namespace Gauntlet
 
 		protected override IAppInstall InstallStagedBuild(UnrealAppConfig AppConfig, StagedBuild InBuild)
 		{
-			string BuildPath = InBuild.BuildPath;
+			string BuildDir = InBuild.BuildPath;
 
-			if (Utils.SystemHelpers.IsNetworkPath(BuildPath))
+			if (Utils.SystemHelpers.IsNetworkPath(BuildDir))
 			{
 				string SubDir = string.IsNullOrEmpty(AppConfig.Sandbox) ? AppConfig.ProjectName : AppConfig.Sandbox;
-				string BasePath = string.IsNullOrEmpty(AppConfig.DestLocalInstallDir) ? this.LocalCachePath : AppConfig.DestLocalInstallDir;
-				string DestPath = Path.Combine(BasePath, SubDir, AppConfig.ProcessType.ToString());
+				string InstallDir = Path.Combine(InstallRoot, SubDir, AppConfig.ProcessType.ToString());
 
 				if (!AppConfig.SkipInstall)
 				{
-					DestPath = StagedBuild.InstallBuildParallel(AppConfig, InBuild, BuildPath, DestPath, ToString());
+					InstallDir = StagedBuild.InstallBuildParallel(AppConfig, InBuild, BuildDir, InstallDir, ToString());
 				}
 				else
 				{
-					Log.Info("Skipping install of {0} (-skipdeploy)", BuildPath);
+					Log.Info("Skipping install of {0} (-SkipInstall)", BuildDir);
 				}
 
-				Utils.SystemHelpers.MarkDirectoryForCleanup(DestPath);
-
-				BuildPath = DestPath;
+				BuildDir = InstallDir;
+				Utils.SystemHelpers.MarkDirectoryForCleanup(InstallDir);
 			}
 
 			LinuxAppInstall LinuxApp = new LinuxAppInstall(AppConfig.Name, AppConfig.ProjectName, this);
 			LinuxApp.RunOptions = RunOptions;
 
 			// Set commandline replace any InstallPath arguments with the path we use
-			LinuxApp.CommandArguments = Regex.Replace(AppConfig.CommandLine, @"\$\(InstallPath\)", BuildPath, RegexOptions.IgnoreCase);
+			LinuxApp.CommandArguments = Regex.Replace(AppConfig.CommandLine, @"\$\(InstallPath\)", BuildDir, RegexOptions.IgnoreCase);
 
 			if (string.IsNullOrEmpty(UserDir) == false)
 			{
 				LinuxApp.CommandArguments += string.Format(" -userdir=\"{0}\"", UserDir);
-				LinuxApp.ArtifactPath = Path.Combine(UserDir, @"Saved");
+				LinuxApp.ArtifactPath = Path.Combine(UserDir, "Saved");
 
 				Utils.SystemHelpers.MarkDirectoryForCleanup(UserDir);
 			}
 			else
 			{
 				// e.g d:\Unreal\GameName\Saved
-				LinuxApp.ArtifactPath = Path.Combine(BuildPath, AppConfig.ProjectName, @"Saved");
+				LinuxApp.ArtifactPath = Path.Combine(BuildDir, AppConfig.ProjectName, "Saved");
 
 			}
 
@@ -165,7 +163,7 @@ namespace Gauntlet
 
 			if (LocalDirectoryMappings.Count == 0)
 			{
-				PopulateDirectoryMappings(Path.Combine(BuildPath, AppConfig.ProjectName));
+				PopulateDirectoryMappings(Path.Combine(BuildDir, AppConfig.ProjectName));
 			}
 
 			CopyAdditionalFiles(AppConfig);
@@ -177,7 +175,7 @@ namespace Gauntlet
 			else
 			{
 				// TODO - this check should be at a higher level....
-				string BinaryPath = Path.Combine(BuildPath, InBuild.ExecutablePath);
+				string BinaryPath = Path.Combine(BuildDir, InBuild.ExecutablePath);
 
 				// check for a local newer executable
 				if (Globals.Params.ParseParam("dev") && AppConfig.ProcessType.UsesEditor() == false)
diff --git a/Engine/Source/Programs/AutomationTool/Gauntlet/Platform/Mac/Gauntlet.TargetDeviceMac.cs b/Engine/Source/Programs/AutomationTool/Gauntlet/Platform/Mac/Gauntlet.TargetDeviceMac.cs
--- a/Engine/Source/Programs/AutomationTool/Gauntlet/Platform/Mac/Gauntlet.TargetDeviceMac.cs
+++ b/Engine/Source/Programs/AutomationTool/Gauntlet/Platform/Mac/Gauntlet.TargetDeviceMac.cs
@@ -149,10 +149,7 @@ namespace Gauntlet
 			MacApp.CommandArguments += string.Format(" -userdir=\"{0}\"", UserDir);
 			MacApp.ArtifactPath = Path.Combine(UserDir, @"Saved");
 
-			if (LocalDirectoryMappings.Count == 0)
-			{
-				PopulateDirectoryMappings(BuildPath);
-			}
+			PopulateDirectoryMappings(BuildPath);
 
 			// clear artifact path
 			MacApp.CleanDeviceArtifacts();
@@ -232,36 +229,31 @@ namespace Gauntlet
 		/// <returns></returns>
 		protected string CopyBuildIfNecessary(UnrealAppConfig AppConfig, string InBuildPath)
 		{
-			bool SkipDeploy = Globals.Params.ParseParam("SkipDeploy");
+			string BuildDir = InBuildPath;
 
-			string OutBuildPath = InBuildPath;
-
-			string BuildVolume = GetVolumeName(OutBuildPath);
+			string BuildVolume = GetVolumeName(BuildDir);
 			string LocalRoot = GetVolumeName(Environment.CurrentDirectory);
 
 			// Must be on our volume to run
 			if (BuildVolume.Equals(LocalRoot, StringComparison.OrdinalIgnoreCase) == false)
 			{
 				string SubDir = string.IsNullOrEmpty(AppConfig.Sandbox) ? AppConfig.ProjectName : AppConfig.Sandbox;
-				string BasePath = string.IsNullOrEmpty(AppConfig.DestLocalInstallDir) ? this.LocalCachePath : AppConfig.DestLocalInstallDir;
-				string DestPath = Path.Combine(BasePath, SubDir, AppConfig.ProcessType.ToString());
+				string InstallDir = Path.Combine(InstallRoot, SubDir, AppConfig.ProcessType.ToString());
 
-				if (!SkipDeploy)
+				if (!AppConfig.SkipInstall)
 				{
-					Utils.SystemHelpers.CopyDirectory(InBuildPath, DestPath, Utils.SystemHelpers.CopyOptions.Mirror);
+					Utils.SystemHelpers.CopyDirectory(BuildDir, InstallDir, Utils.SystemHelpers.CopyOptions.Mirror);
 				}
-
 				else
 				{
-					Log.Info("Skipping install of {0} (-skipdeploy)", InBuildPath);
+					Log.Info("Skipping install of {0} (-SkipInstall)", BuildDir);
 				}
 
-				Utils.SystemHelpers.MarkDirectoryForCleanup(DestPath);
-
-				OutBuildPath = DestPath;
+				BuildDir = InstallDir;
+				Utils.SystemHelpers.MarkDirectoryForCleanup(InstallDir);
 			}
 
-			return OutBuildPath;
+			return BuildDir;
 		}
 	}
 
diff --git a/Engine/Source/Programs/AutomationTool/Gauntlet/Platform/Windows/Gauntlet.TargetDeviceWindows.cs b/Engine/Source/Programs/AutomationTool/Gauntlet/Platform/Windows/Gauntlet.TargetDeviceWindows.cs
--- a/Engine/Source/Programs/AutomationTool/Gauntlet/Platform/Windows/Gauntlet.TargetDeviceWindows.cs
+++ b/Engine/Source/Programs/AutomationTool/Gauntlet/Platform/Windows/Gauntlet.TargetDeviceWindows.cs
@@ -165,26 +165,24 @@ namespace Gauntlet
 
 		protected override IAppInstall InstallStagedBuild(UnrealAppConfig AppConfig, StagedBuild InBuild)
 		{
-			string BuildPath = InBuild.BuildPath;
+			string BuildDir = InBuild.BuildPath;
 
-			if (Utils.SystemHelpers.IsNetworkPath(BuildPath))
+			if (Utils.SystemHelpers.IsNetworkPath(BuildDir))
 			{
 				string SubDir = string.IsNullOrEmpty(AppConfig.Sandbox) ? AppConfig.ProjectName : AppConfig.Sandbox;
-				string BasePath = string.IsNullOrEmpty(AppConfig.DestLocalInstallDir) ? this.LocalCachePath : AppConfig.DestLocalInstallDir;
-				string DestPath = Path.Combine(BasePath, SubDir, AppConfig.ProcessType.ToString());
+				string InstallDir = Path.Combine(InstallRoot, SubDir, AppConfig.ProcessType.ToString());
 
 				if (!AppConfig.SkipInstall)
 				{
-					DestPath = StagedBuild.InstallBuildParallel(AppConfig, InBuild, BuildPath, DestPath, ToString());
+					InstallDir = StagedBuild.InstallBuildParallel(AppConfig, InBuild, BuildDir, InstallDir, ToString());
 				}
 				else
 				{
-					Log.Info("Skipping install of {0} (-skipdeploy)", BuildPath);
+					Log.Info("Skipping install of {0} (-SkipInstall)", BuildDir);
 				}
 
-				Utils.SystemHelpers.MarkDirectoryForCleanup(DestPath);
-
-				BuildPath = DestPath;
+				BuildDir = InstallDir;
+				Utils.SystemHelpers.MarkDirectoryForCleanup(InstallDir);
 			}
 
 			WindowsAppInstall WinApp = new WindowsAppInstall(AppConfig.Name, AppConfig.ProjectName, this);
@@ -192,20 +190,19 @@ namespace Gauntlet
 			WinApp.CanAlterCommandArgs = AppConfig.CanAlterCommandArgs;
 
 			// Set commandline replace any InstallPath arguments with the path we use
-			WinApp.CommandArguments = Regex.Replace(AppConfig.CommandLine, @"\$\(InstallPath\)", BuildPath, RegexOptions.IgnoreCase);
+			WinApp.CommandArguments = Regex.Replace(AppConfig.CommandLine, @"\$\(InstallPath\)", BuildDir, RegexOptions.IgnoreCase);
 
 			if (string.IsNullOrEmpty(UserDir) == false)
 			{
 				WinApp.CommandArguments += string.Format(" -userdir=\"{0}\"", UserDir);
-				WinApp.ArtifactPath = Path.Combine(UserDir, @"Saved");
+				WinApp.ArtifactPath = Path.Combine(UserDir, "Saved");
 
 				Utils.SystemHelpers.MarkDirectoryForCleanup(UserDir);
 			}
 			else
 			{
 				// e.g d:\Unreal\GameName\Saved
-				WinApp.ArtifactPath = Path.Combine(BuildPath, AppConfig.ProjectName, @"Saved");
-
+				WinApp.ArtifactPath = Path.Combine(BuildDir, AppConfig.ProjectName, "Saved");
 			}
 
 			// clear artifact path
@@ -213,7 +210,7 @@ namespace Gauntlet
 
 			if (LocalDirectoryMappings.Count == 0)
 			{
-				PopulateDirectoryMappings(Path.Combine(BuildPath, AppConfig.ProjectName));
+				PopulateDirectoryMappings(Path.Combine(BuildDir, AppConfig.ProjectName));
 			}
 
 			CopyAdditionalFiles(AppConfig);
@@ -225,7 +222,7 @@ namespace Gauntlet
 			else
 			{
 				// TODO - this check should be at a higher level....
-				string BinaryPath = Path.Combine(BuildPath, InBuild.ExecutablePath);
+				string BinaryPath = Path.Combine(BuildDir, InBuild.ExecutablePath);
 
 				// check for a local newer executable
 				if (Globals.Params.ParseParam("dev") && AppConfig.ProcessType.UsesEditor() == false)
diff --git a/Engine/Source/Programs/AutomationTool/Gauntlet/Unreal/Base/Gauntlet.UnrealAppConfig.cs b/Engine/Source/Programs/AutomationTool/Gauntlet/Unreal/Base/Gauntlet.UnrealAppConfig.cs
--- a/Engine/Source/Programs/AutomationTool/Gauntlet/Unreal/Base/Gauntlet.UnrealAppConfig.cs
+++ b/Engine/Source/Programs/AutomationTool/Gauntlet/Unreal/Base/Gauntlet.UnrealAppConfig.cs
@@ -99,12 +99,6 @@ namespace Gauntlet
 		/// </summary>
 		public string Sandbox { get; set; }
 
-		/// <summary>
-		/// Optional directory for builds that need to be run locally. This can be used to persist a locally-install build between build graph nodes.
-		/// </summary>
-		public string DestLocalInstallDir { get; set; }
-
-		// new system
 		public IBuild Build { get; set; }
 
 		// Prevents installing a build on device
diff --git a/Engine/Source/Programs/AutomationTool/Gauntlet/Unreal/Base/Gauntlet.UnrealSession.cs b/Engine/Source/Programs/AutomationTool/Gauntlet/Unreal/Base/Gauntlet.UnrealSession.cs
--- a/Engine/Source/Programs/AutomationTool/Gauntlet/Unreal/Base/Gauntlet.UnrealSession.cs
+++ b/Engine/Source/Programs/AutomationTool/Gauntlet/Unreal/Base/Gauntlet.UnrealSession.cs
@@ -589,11 +589,6 @@ namespace Gauntlet
 		/// </summary>
 		public string Sandbox { get; set; }
 
-		/// <summary>
-		/// Optional directory where local builds will be installed
-		/// </summary>
-		public string DestLocalInstallDir { get; set; }
-
 		/// <summary>
 		/// Whether or not we should retain our devices this pass
 		/// </summary>
@@ -900,12 +895,6 @@ namespace Gauntlet
 					// todo - should this be elsewhere?
 					AppConfig.Sandbox = Sandbox;
 
-					if (!string.IsNullOrEmpty(DestLocalInstallDir) && BuildHostPlatform.Current.Platform == Device.Platform)
-					{
-						AppConfig.DestLocalInstallDir = Path.Combine(DestLocalInstallDir, Device.Platform.ToString(), Device.Name );
-					}
-
-
 					Device.ClearSavedDirectory(AppConfig);
 
 					IAppInstall Install = null;
diff --git a/Engine/Source/Programs/AutomationTool/Gauntlet/Unreal/Base/Gauntlet.UnrealTestContext.cs b/Engine/Source/Programs/AutomationTool/Gauntlet/Unreal/Base/Gauntlet.UnrealTestContext.cs
--- a/Engine/Source/Programs/AutomationTool/Gauntlet/Unreal/Base/Gauntlet.UnrealTestContext.cs
+++ b/Engine/Source/Programs/AutomationTool/Gauntlet/Unreal/Base/Gauntlet.UnrealTestContext.cs
@@ -130,7 +130,6 @@ namespace Gauntlet
 		/// </summary>
 		[AutoParam("")]
 		public string Namespaces;
-		
 
 		public IEnumerable<string> SearchPaths;
 
@@ -140,12 +139,6 @@ namespace Gauntlet
 		[AutoParam("")]
 		public string Sandbox { get; set; }
 
-		/// <summary>
-		/// Optional directory where local builds will be installed
-		/// </summary>
-		[AutoParam("")]
-		public string DestLocalInstallDir;
-
 		/// <summary>
 		/// Skip any check or copying of builds. Mostly useful when debugging
 		/// </summary>
diff --git a/Engine/Source/Programs/AutomationTool/Gauntlet/Unreal/Base/Gauntlet.UnrealTestNode.cs b/Engine/Source/Programs/AutomationTool/Gauntlet/Unreal/Base/Gauntlet.UnrealTestNode.cs
--- a/Engine/Source/Programs/AutomationTool/Gauntlet/Unreal/Base/Gauntlet.UnrealTestNode.cs
+++ b/Engine/Source/Programs/AutomationTool/Gauntlet/Unreal/Base/Gauntlet.UnrealTestNode.cs
@@ -682,12 +682,7 @@ namespace Gauntlet
 				}
 			}
 
-			UnrealApp = new UnrealSession(Context.BuildInfo, SessionRoles) 
-			{ 
-				Sandbox = Context.Options.Sandbox, 
-				DestLocalInstallDir = Context.Options.DestLocalInstallDir 
-			};
-
+			UnrealApp = new UnrealSession(Context.BuildInfo, SessionRoles) { Sandbox = Context.Options.Sandbox };
 			return true;
 		}
 
