reusable singleton httpclient as per MS guidance for .net5+

updates PDB telemetry task to use new client

diff --git a/Engine/Source/Programs/AutomationTool/AutomationUtils/HttpClientSingleton.cs b/Engine/Source/Programs/AutomationTool/AutomationUtils/HttpClientSingleton.cs
new file mode 100644
--- /dev/null
+++ b/Engine/Source/Programs/AutomationTool/AutomationUtils/HttpClientSingleton.cs
@@ -0,0 +1,32 @@
+// Copyright Epic Games, Inc. All Rights Reserved.
+using Microsoft.Extensions.Http;
+using Polly;
+using Polly.Extensions.Http;
+using System;
+using System.Net.Http;
+
+namespace AutomationUtils
+{
+	// This is used for instances where DI cannot be used for the HTTPClient
+	// More info: https://learn.microsoft.com/en-us/dotnet/fundamentals/networking/http/httpclient-guidelines
+	public class HttpClientSingleton<T>
+	{
+		private static readonly Lazy<HttpClient> internalInstance =
+			new Lazy<HttpClient>(() =>
+			{
+				var retryPolicy = HttpPolicyExtensions
+				.HandleTransientHttpError()
+				.WaitAndRetryAsync(3, retryAttempt => TimeSpan.FromSeconds(Math.Pow(2, retryAttempt)));
+
+				// per MS documentation, with a lifetime specified we shouldn't have socket exhaustion issues
+				var socketHandler = new SocketsHttpHandler { PooledConnectionLifetime = TimeSpan.FromMinutes(15) };
+				var pollyHandler = new PolicyHttpMessageHandler(retryPolicy)
+				{
+					InnerHandler = socketHandler,
+				};
+
+				return new HttpClient(pollyHandler);
+			});
+		public static HttpClient Client { get => internalInstance.Value; }
+	}
+}
