[Gauntlet] Reorder DevicePool functions public -> protected -> private

diff --git a/Engine/Source/Programs/AutomationTool/Gauntlet/Framework/Devices/Gauntlet.DevicePool.cs b/Engine/Source/Programs/AutomationTool/Gauntlet/Framework/Devices/Gauntlet.DevicePool.cs
--- a/Engine/Source/Programs/AutomationTool/Gauntlet/Framework/Devices/Gauntlet.DevicePool.cs
+++ b/Engine/Source/Programs/AutomationTool/Gauntlet/Framework/Devices/Gauntlet.DevicePool.cs
@@ -16,7 +16,6 @@
 
 namespace Gauntlet
 {
-
 	/// <summary>
 	/// Device performance specification
 	/// </summary>
@@ -52,7 +51,7 @@
 
 		public string Available { get; set; }
 
-		public bool RemoveOnShutdown { get; set; }		
+		public bool RemoveOnShutdown { get; set; }
 
 		public override string ToString()
 		{
@@ -71,7 +70,7 @@
 
 		public UnrealDeviceTargetConstraint(UnrealTargetPlatform? Platform, EPerfSpec PerfSpec = EPerfSpec.Unspecified, string Model = null)
 		{
-			this.Platform = Platform;			
+			this.Platform = Platform;
 			this.PerfSpec = PerfSpec;
 			this.Model = Model == null ? string.Empty : Model;
 		}
@@ -207,6 +206,9 @@
 	/// </summary>
 	public class DeviceReservationStates
 	{
+		private Dictionary<string, ResevationState> RequiredCollection;
+		private Dictionary<string, int> ProblemCollection;
+
 		public DeviceReservationStates()
 		{
 			ResetRequirements();
@@ -319,9 +321,6 @@
 			public int NumLacking;
 		}
 
-		private Dictionary<string, ResevationState> RequiredCollection;
-		private Dictionary<string, int> ProblemCollection;
-
 		private ResevationState FindClosestState(string DeviceKey)
 		{
 			if (RequiredCollection.ContainsKey(DeviceKey))
@@ -350,33 +349,42 @@
 	/// Singleton class that's responsible for providing a list of devices and reserving them. Code should call
 	/// EnumerateDevices to build a list of desired devices, which must then be reserved by calling ReserveDevices.
 	/// Once done ReleaseDevices should be called.
-	/// 
+	///
 	/// These reservations exist at the process level and we rely on the device implementation to provide arbitrage 
 	/// between difference processes and machines
-	/// 
+	///
 	/// </summary>
 	public class DevicePool : IDisposable
 	{
+		/// <summary>
+		/// Access to our singleton
+		/// </summary>
+		public static DevicePool Instance { get; private set; } = new DevicePool();
 
 		/// <summary>
+		/// Device reservation service URL
+		/// </summary>
+		public string DeviceURL;
+
+		/// <summary>
 		/// Object used for locking access to internal data
 		/// </summary>
-		private Object LockObject = new Object();
+		private object LockObject = new object();
 
 		/// <summary>
 		/// List of all provisioned devices that are available for reservations
 		/// </summary>
-		List<ITargetDevice> AvailableDevices = new List<ITargetDevice>();
+		private List<ITargetDevice> AvailableDevices = new List<ITargetDevice>();
 
 		/// <summary>
 		/// List of all devices that have been reserved
 		/// </summary>
-		List<ITargetDevice> ReservedDevices = new List<ITargetDevice>();
+		private List<ITargetDevice> ReservedDevices = new List<ITargetDevice>();
 
 		/// <summary>
 		/// List of platforms we've had devices for
 		/// </summary>
-		HashSet<UnrealTargetPlatform?> UsedPlatforms = new HashSet<UnrealTargetPlatform?>();
+		private HashSet<UnrealTargetPlatform?> UsedPlatforms = new HashSet<UnrealTargetPlatform?>();
 
 		/// <summary>
 		/// Device reservation left to fulfill
@@ -384,54 +392,45 @@
 		private DeviceReservationStates ReservationStates;
 
 		/// <summary>
-		/// Device reservation service URL 
-		/// </summary>
-		public string DeviceURL;
-
-		/// <summary>
 		/// Device to reservation lookup
 		/// </summary>
-		Dictionary<ITargetDevice, DeviceReservationAutoRenew> ServiceReservations = new Dictionary<ITargetDevice, DeviceReservationAutoRenew>();
-
+		private Dictionary<ITargetDevice, DeviceReservationAutoRenew> ServiceReservations = new Dictionary<ITargetDevice, DeviceReservationAutoRenew>();
 
 		/// <summary>
 		/// Target device info, private for reservation use
 		/// </summary>
-		Dictionary<ITargetDevice, DeviceDefinition> ServiceDeviceInfo = new Dictionary<ITargetDevice, DeviceDefinition>();
+		private Dictionary<ITargetDevice, DeviceDefinition> ServiceDeviceInfo = new Dictionary<ITargetDevice, DeviceDefinition>();
 
 		/// <summary>
 		/// List of device definitions that can be provisioned on demand
 		/// </summary>
-		List<DeviceDefinition> UnprovisionedDevices = new List<DeviceDefinition>();
+		private List<DeviceDefinition> UnprovisionedDevices = new List<DeviceDefinition>();
 
 		/// <summary>
 		/// List of definitions that failed to provision
 		/// </summary>
-		List<DeviceDefinition> FailedProvisions = new List<DeviceDefinition>();
+		private List<DeviceDefinition> FailedProvisions = new List<DeviceDefinition>();
 
 		/// <summary>
 		/// Local connection state for any device. We use this to avoid disconnecting devices that the
 		/// user may already be connected to when running tests
 		/// </summary>
-		Dictionary<ITargetDevice, bool> InitialConnectState = new Dictionary<ITargetDevice, bool>();
+		private Dictionary<ITargetDevice, bool> InitialConnectState = new Dictionary<ITargetDevice, bool>();
 
 		/// <summary>
 		/// Device constraints for performance profiles, etc
 		/// </summary>
-		Dictionary<ITargetDevice, UnrealDeviceTargetConstraint> Constraints = new Dictionary<ITargetDevice, UnrealDeviceTargetConstraint>();
+		private Dictionary<ITargetDevice, UnrealDeviceTargetConstraint> Constraints = new Dictionary<ITargetDevice, UnrealDeviceTargetConstraint>();
 
 		/// <summary>
-		/// Local tempdir...
+		/// Directory automation artifacts are saved to
 		/// </summary>
-		string LocalTempDir;
-
-
-		bool UniqueLocalTemps;
+		private string LocalTempDir;
 
 		/// <summary>
-		/// Static instance
+		/// Whether or not Unique temporary directories should be created
 		/// </summary>
-		private static DevicePool _Instance;
+		private bool bUniqueTemps;
 
 		/// <summary>
 		/// The maximum number of problem devices to report to device backend
@@ -444,47 +443,29 @@
 		/// </summary>
 		protected DevicePool()
 		{
-			if (_Instance == null)
-			{
-				_Instance = this;
-			}
-
 			// create two local devices by default?
 			AddLocalDevices(2);
-
 			ReservationStates = new DeviceReservationStates();
+
+			if(Instance == null)
+			{
+				Instance = this;
+			}
 		}
 
+		#region IDisposable Support
 		~DevicePool()
 		{
 			Dispose(false);
 		}
 
 		/// <summary>
-		/// Access to our singleton
-		/// </summary>
-		public static DevicePool Instance
-		{
-			get
-			{
-				if (_Instance == null)
-				{
-					new DevicePool();
-				}
-				return _Instance;
-			}
-		}
-
-		/// <summary>
 		/// Shutdown the pool and release all devices.
 		/// </summary>
 		public static void Shutdown()
 		{
-			if (_Instance != null)
-			{
-				_Instance.Dispose();
-				_Instance = null;
-			}
+			Instance?.Dispose();
+			Instance = null;
 		}
 
 		public void Dispose()
@@ -548,310 +529,61 @@
 				}
 			}
 		}
-
-		static private void CleanupDevices(IEnumerable<UnrealTargetPlatform?> Platforms)
-		{
-			IEnumerable<IDeviceService> DeviceServices = Gauntlet.Utils.InterfaceHelpers.FindImplementations<IDeviceService>();
-			if (DeviceServices.Any())
-			{
-				foreach (UnrealTargetPlatform? Platform in Platforms)
-				{
-					IDeviceService DeviceService = DeviceServices.Where(D => D.CanSupportPlatform(Platform)).FirstOrDefault();
-					if (DeviceService != null)
-					{
-						DeviceService.CleanupDevices();
-					}
-				}
-			}
-		}
-
-		void RevertDeviceConfiguration(ITargetDevice Device)
-		{
-			IConfigurableDevice ConfigurableDevice = Device as IConfigurableDevice;
-			if (ConfigurableDevice != null)
-			{
-				var Snapshot = DeviceConfigurationCache.Instance.GetConfigurationSnapshot(Device.Platform, Device.Name);
-				if (Snapshot == null)
-				{
-					return;
-				}
-
-				// Connect temporarily to be able to revert the device's configuration
-				// if the device was disconnected entering here disconnect it after this is over
-				bool bNeedsDisconnect = false;
-				if (!Device.IsConnected)
-				{
-					Device.Connect();
-					bNeedsDisconnect = true;
-				}
-				
-				if (ConfigurableDevice.ApplyConfiguration(Snapshot))
-				{
-					DeviceConfigurationCache.Instance.ClearSnapshot(Snapshot);
-				}
-
-				if (bNeedsDisconnect)
-				{
-					Device.Disconnect();
-				}
-			}
-		}
+		#endregion
 
 		/// <summary>
-		/// Explicitly release all device reservations
-		/// </summary>
-		private void ReleaseReservations()
-		{
-
-			ServiceDeviceInfo.Clear();
-
-			foreach (ITargetDevice Device in ServiceReservations.Keys)
-			{
-				AvailableDevices.Remove(Device);
-				ReservedDevices.Remove(Device);
-				Device.Dispose();
-			}
-
-			foreach (DeviceReservationAutoRenew reservation in ServiceReservations.Values.Distinct())
-			{
-				reservation.Dispose();
-			}
-
-			ServiceReservations.Clear();
-
-		}
-
-		/// <summary>
-		/// Release all devices in the provided list from our reserved list
-		/// </summary>
-		/// <param name="DeviceList"></param>
-		public void ReleaseDevices(IEnumerable<ITargetDevice> DeviceList)
-		{
-			lock (LockObject)
-			{
-				// remove all these devices from our reserved list
-				ReservedDevices = ReservedDevices.Where(D => DeviceList.Contains(D) == false).ToList();
-
-				// disconnect them
-				List<ITargetDevice> KeepList = new List<ITargetDevice>();
-				foreach (ITargetDevice Device in DeviceList)
-				{
-					RevertDeviceConfiguration(Device);
-
-					if (Device.IsConnected && InitialConnectState[Device] == false)
-					{
-						Device.Disconnect();
-					}
-
-					// if device is from reservation service, unregister it
-					DeviceReservationAutoRenew Reservation;
-					if (ServiceReservations.TryGetValue(Device, out Reservation))
-					{
-						bool DisposeReservation = ServiceReservations.Count(Entry => Entry.Value == Reservation) == 1;
-
-						// remove and dispose of device
-						// @todo: add support for reservation modification on server (partial device release)
-						ServiceReservations.Remove(Device);
-						ServiceDeviceInfo.Remove(Device);
-						AvailableDevices.Remove(Device);
-
-						InitialConnectState.Remove(Device);
-						Constraints.Remove(Device);
-
-						Device.Dispose();
-
-						// if last device in reservation, dispose the reservation
-						if (DisposeReservation)
-						{
-							Reservation.Dispose();
-						}
-					}
-					else
-					{
-						KeepList.Add(Device);
-					}
-
-				}
-
-				// Put all devices back into the available list at the front - if another device is requested we want to
-				// try and use one we just put a build on
-				KeepList.ForEach(D => { AvailableDevices.Remove(D); AvailableDevices.Insert(0, D); });
-			}
-		}
-
-		/// <summary>
-		/// Report target device issue to service with given error message
-		/// </summary>
-		public void ReportDeviceError(ITargetDevice Device, string ErrorMessage)
-		{
-			DeviceDefinition Def = null;
-			if (!ServiceDeviceInfo.TryGetValue(Device, out Def))
-			{
-				return;
-			}
-
-			ReportDeviceError(Def.Name, ErrorMessage);
-		}
-
-		/// <summary>
-		/// Reporting device issue using the device name provided by service
-		/// </summary>
-		private void ReportDeviceError(string ServiceDeviceName, string ErrorMessage)
-		{
-			if (MaxDeviceErrorReports == 0)
-			{
-				Log.Warning(KnownLogEvents.Gauntlet_DeviceEvent, "Maximum device errors reported to backend, {Name} : {Message} ignored", ServiceDeviceName, ErrorMessage);
-				return;
-			}
-
-			MaxDeviceErrorReports--;
-
-			Reservation.ReportDeviceError(DeviceURL, ServiceDeviceName, ErrorMessage);
-		}
-
-		/// <summary>
 		/// Get the reservation status as an array
 		/// </summary>
 		/// <returns></returns>
-		public string[] GetReservationStatuses()
-		{
-			return ReservationStates != null? ReservationStates.GetStatuses() : null;
-		}
+		public string[] GetReservationStatuses() => ReservationStates != null ? ReservationStates.GetStatuses() : null;
 
 		/// <summary>
 		/// Get the reserved device marked with problems as an array
 		/// </summary>
 		/// <returns></returns>
-		public string[] GetDeviceProblems()
-		{
-			return ReservationStates != null ? ReservationStates.GetProblems() : null;
-		}
+		public string[] GetDeviceProblems() => ReservationStates != null ? ReservationStates.GetProblems() : null;
 
 		/// <summary>
-		/// Report on Device reservation state
+		/// Returns the number of available devices of the provided type. This includes unprovisioned devices but not reserved ones.
+		/// Note: unprovisioned devices are currently only returned when device is not constrained
 		/// </summary>
-		public void ReportDeviceReservationState()
+		public int GetAvailableDeviceCount(UnrealDeviceTargetConstraint Constraint, Func<ITargetDevice, bool> Validate = null)
 		{
-			string PoolId = string.IsNullOrEmpty(Globals.DevicePoolId) ? "Local" : Globals.DevicePoolId;
-			string[] DeviceReservationStatus = GetReservationStatuses();
-			if (DeviceReservationStatus != null && DeviceReservationStatus.Length > 0)
+			lock (LockObject)
 			{
-				Log.Warning(KnownLogEvents.Gauntlet_DeviceEvent, "Device Reservation status from {Pool} pool: {Status}", PoolId, string.Join(", ", DeviceReservationStatus));
-			}
-			string[] DeviceProblemStatus = GetDeviceProblems();
-			if (DeviceProblemStatus != null && DeviceProblemStatus.Length > 0)
-			{
-				Log.Warning(KnownLogEvents.Gauntlet_DeviceEvent, "Reserved Devices marked with problems from {Pool} pool: {Status}", PoolId, string.Join(", ", DeviceProblemStatus));
+				return AvailableDevices.Where(D => Validate == null ? Constraint.Check(D) : Validate(D)).Count() +
+					UnprovisionedDevices.Where(D => Constraint.Check(D)).Count();
 			}
 		}
 
 		/// <summary>
-		/// Construct a path to hold cache files and make sure it's properly cleaned
+		/// Returns the number of available devices of the provided type. This includes unprovisioned devices but not reserved ones
+		/// Note: unprovisioned devices are currently only returned when device is not constrained
 		/// </summary>
-		private string GetCleanCachePath(DeviceDefinition InDeviceDefiniton)
-		{
-			// Give the desktop platform a temp folder with its name under the device cache
-			string DeviceCache = Path.Combine(LocalTempDir, "DeviceCache");
-			string PlatformCache = Path.Combine(DeviceCache, InDeviceDefiniton.Platform.ToString());
-			string ClientCache = Path.Combine(PlatformCache, InDeviceDefiniton.Name);
-
-			// On Desktops, builds are installed in the device cache.
-			// When using device reservation blocks, we don't want to fully clean the cache and lose previously installed builds.
-			// If bRetainBuilds evaluates to true, it means we are in the second step or beyond in a device reservation block.
-			// In this case we'll just delete the left over UserDir which should already have been emptied by UnrealSession.
-			bool? bForceClean = UnrealAppConfig.ForceFullClean;
-			bool? bSkipInstall = UnrealAppConfig.ForceSkipInstall;
-			bool bUsingReservationBlock = bForceClean.HasValue && bSkipInstall.HasValue;
-			bool bRetainBuilds = bUsingReservationBlock && !bForceClean.Value && bSkipInstall.Value;
-
-			if(bRetainBuilds)
-			{
-				Log.Info("Retaining build cache for device reservation block");
-
-				DirectoryInfo UserDirectory = new(Path.Combine(ClientCache, "UserDir"));
-				if(UserDirectory.Exists)
-				{
-					try
-					{
-						Log.Info("Cleaning stale user directory...");
-						SystemHelpers.Delete(UserDirectory, true, true);
-					}
-					catch(Exception Ex)
-					{
-						throw new AutomationException("Failed to clean user directory {0}. This could result in improper artifact reporting. {1}", UserDirectory, Ex);
-					}
-				}
-			}
-			else
-			{
-				int CleanAttempts = 0;
-
-				while (Directory.Exists(ClientCache))
-				{
-					DirectoryInfo ClientCacheDirectory = new(ClientCache);
-
-					try
-					{
-						Log.Info("Cleaning stale client device cache...");
-						SystemHelpers.Delete(ClientCacheDirectory, true, true);
-					}
-					catch (Exception Ex)
-					{
-						// If we fail to acquire the default client cache while using device reservation blocks,
-						// we can't ensure future tests will have their cache directories mapped to the correct build location
-						if(bUsingReservationBlock)
-						{
-							throw new AutomationException("Failed to clean default client device cache {0}. {1}", ClientCache, Ex);
-						}
-
-						// When not using device reservation blocks, we can just create a newly indexed directory for the client cache
-						else
-						{
-							string Warning = "Failed to clean client device cache {Folder}. A newly indexed directory will be created instead. {Message}";
-							Log.Warning(KnownLogEvents.Gauntlet_DeviceEvent, Warning, ClientCache, Ex.Message);
-
-							ClientCache = Path.Combine(PlatformCache, $"{InDeviceDefiniton.Name}_{++CleanAttempts}");
-						}
-					}
-				}
-
-				// create this path
-				Log.Info("Client device cache set to {Directory}", ClientCache);
-				Directory.CreateDirectory(ClientCache);
-			}
-
-			return ClientCache;
-		}
-
-		/// <summary>
-		/// Reserve all devices in the provided list. Once reserved a device will not be seen by any code that
-		/// calls EnumerateDevices
-		/// </summary>
-		/// <param name="DeviceList"></param>
-		/// <returns></returns>
-		public bool ReserveDevices(IEnumerable<ITargetDevice> DeviceList)
+		public int GetTotalDeviceCount(UnrealDeviceTargetConstraint Constraint, Func<ITargetDevice, bool> Validate = null)
 		{
 			lock (LockObject)
 			{
-				// can reserve if not reserved...
-				if (ReservedDevices.Intersect(DeviceList).Count() > 0)
-				{
-					return false;
-				}
-
-				// remove these devices from the available list
-				AvailableDevices = AvailableDevices.Where(D => DeviceList.Contains(D) == false).ToList();
-
-				ReservedDevices.AddRange(DeviceList);
-				DeviceList.All(D => UsedPlatforms.Add(D.Platform));
+				return AvailableDevices.Union(ReservedDevices).Where(D => Validate == null ? Constraint.Check(D) : Validate(D)).Count() +
+					UnprovisionedDevices.Where(D => Constraint.Check(D)).Count();
 			}
-			return true;
 		}
 
+		public UnrealDeviceTargetConstraint GetConstraint(ITargetDevice Device)
+		{
+
+			if (!Constraints.ContainsKey(Device))
+			{
+				throw new AutomationException("Device pool has no contstaint for {0} (device was likely released)", Device);
+			}
+
+			return Constraints[Device];
+		}
+
 		public void SetLocalOptions(string InLocalTemp, bool InUniqueTemps, string InDeviceURL = "")
 		{
 			LocalTempDir = InLocalTemp;
-			UniqueLocalTemps = InUniqueTemps;
+			bUniqueTemps = InUniqueTemps;
 			DeviceURL = InDeviceURL;
 		}
 
@@ -890,7 +622,7 @@
 					{
 						DeviceDefinition Def = new DeviceDefinition();
 						Def.Name = string.Format("Virtual{0}{1}", DevicePlatform.ToString(), i);
-						Def.Platform = DevicePlatform??BuildHostPlatform.Current.Platform;
+						Def.Platform = DevicePlatform ?? BuildHostPlatform.Current.Platform;
 						UnprovisionedDevices.Add(Def);
 					}
 				}
@@ -898,3 +630,303 @@
 		}
 
 		/// <summary>
+		/// Created a list of device definitions from the passed in reference. Needs work....
+		/// </summary>
+		/// <param name="InputReference"></param>
+		/// <param name="InLocalTempDir"></param>
+		/// <param name="ObeyConstraints"></param>
+		public void AddDevices(UnrealTargetPlatform DefaultPlatform, string InputReference, bool ObeyConstraints = true)
+		{
+			lock (LockObject)
+			{
+				List<ITargetDevice> NewDevices = new List<ITargetDevice>();
+
+				int SlashIndex = InputReference.IndexOf("\\") >= 0 ? InputReference.IndexOf("\\") : InputReference.IndexOf("/");
+
+				bool PossibleFileName = InputReference.IndexOfAny(Path.GetInvalidPathChars()) < 0 &&
+							(InputReference.IndexOf(":") == -1 || (InputReference.IndexOf(":") == SlashIndex - 1));
+				// Did they specify a file?
+				if (PossibleFileName && File.Exists(InputReference))
+				{
+					Log.Info(KnownLogEvents.Gauntlet_DeviceEvent, "Adding devices from {Reference}", InputReference);
+					List<DeviceDefinition> DeviceDefinitions = JsonSerializer.Deserialize<List<DeviceDefinition>>(
+						File.ReadAllText(InputReference),
+						new JsonSerializerOptions { PropertyNameCaseInsensitive = true }
+					);
+
+					foreach (DeviceDefinition Def in DeviceDefinitions)
+					{
+						Log.Info(KnownLogEvents.Gauntlet_DeviceEvent, "Adding {Name}", Def);
+
+						// use Legacy field if it exists
+						if (Def.Platform == null)
+						{
+							Def.Platform = Def.Type;
+						}
+
+						// check for an availability constraint
+						if (string.IsNullOrEmpty(Def.Available) == false && ObeyConstraints)
+						{
+							// check whether disabled
+							if (String.Compare(Def.Available, "disabled", true) == 0)
+							{
+								Log.Info(KnownLogEvents.Gauntlet_DeviceEvent, "Skipping {Name} due to being disabled", Def.Name);
+								continue;
+							}
+
+							// availability is specified as a range, e.g 21:00-09:00.
+							Match M = Regex.Match(Def.Available, @"(\d{1,2}:\d\d)\s*-\s*(\d{1,2}:\d\d)");
+
+							if (M.Success)
+							{
+								DateTime From, To;
+
+								if (DateTime.TryParse(M.Groups[1].Value, out From) && DateTime.TryParse(M.Groups[2].Value, out To))
+								{
+									// these are just times so when parsed will have todays date. If the To time is less than
+									// From (22:00-06:00) time it spans midnight so move it to the next day
+									if (To < From)
+									{
+										To = To.AddDays(1);
+									}
+
+									// if From is in the future (e.g. it's 01:00 and range is 22:00-08:00) we may be in the previous days window, 
+									// so move them both back a day
+									if (From > DateTime.Now)
+									{
+										From = From.AddDays(-1);
+										To = To.AddDays(-1);
+									}
+
+									if (DateTime.Now < From || DateTime.Now > To)
+									{
+										Log.Info(KnownLogEvents.Gauntlet_DeviceEvent, "Skipping {Name} due to availability constraint {Constraint}", Def.Name, Def.Available);
+										continue;
+									}
+								}
+								else
+								{
+									Log.Warning(KnownLogEvents.Gauntlet_DeviceEvent, "Failed to parse availability {Constraint} for {Name}", Def.Available, Def.Name);
+								}
+							}
+						}
+
+						Def.RemoveOnShutdown = true;
+
+						if (Def.Platform == null)
+						{
+							Def.Platform = DefaultPlatform;
+						}
+
+						UnprovisionedDevices.Add(Def);
+					}
+
+					// randomize devices so if there's a bad device st the start so we don't always hit it (or we do if its later)
+					UnprovisionedDevices = UnprovisionedDevices.OrderBy(D => Guid.NewGuid()).ToList();
+				}
+				else
+				{
+					if (string.IsNullOrEmpty(InputReference) == false)
+					{
+						string[] DevicesList = InputReference.Split(',');
+
+						foreach (string DeviceRef in DevicesList)
+						{
+
+							// check for <platform>:<address>:<port>|<model>. We pass address:port to device constructor
+							Match M = Regex.Match(DeviceRef, @"(.+?):(.+)");
+
+							UnrealTargetPlatform DevicePlatform = DefaultPlatform;
+							string DeviceAddress = DeviceRef;
+							string Model = string.Empty;
+
+							// when using device service, skip adding local non-desktop devices to pool
+							bool IsDesktop = UnrealBuildTool.Utils.GetPlatformsInClass(UnrealPlatformClass.Desktop).Contains(DevicePlatform);
+							if (!IsDesktop && DeviceRef.Equals("default", StringComparison.OrdinalIgnoreCase) && !String.IsNullOrEmpty(DeviceURL))
+							{
+								continue;
+							}
+
+							if (M.Success)
+							{
+								if (!UnrealTargetPlatform.TryParse(M.Groups[1].ToString(), out DevicePlatform))
+								{
+									throw new AutomationException("platform {0} is not a recognized device type", M.Groups[1].ToString());
+								}
+
+								DeviceAddress = M.Groups[2].ToString();
+
+								// parse device model
+								if (DeviceAddress.Contains("|"))
+								{
+									string[] Components = DeviceAddress.Split(new char[] { '|' });
+									DeviceAddress = Components[0];
+									Model = Components[1];
+								}
+
+							}
+
+							Log.Info(KnownLogEvents.Gauntlet_DeviceEvent, "Added device {Platform}:{Address} to pool", DevicePlatform, DeviceAddress);
+							DeviceDefinition Def = new DeviceDefinition();
+							Def.Address = DeviceAddress;
+							Def.Name = DeviceAddress;
+							Def.Platform = DevicePlatform;
+							Def.Model = Model;
+							UnprovisionedDevices.Add(Def);
+						}
+					}
+				}
+			}
+		}
+
+		/// <summary>
+		/// Adds the list of devices to our internal availability list
+		/// </summary>
+		/// <param name="InDevices"></param>
+		public void RegisterDevices(IEnumerable<ITargetDevice> InDevices)
+		{
+			lock (LockObject)
+			{
+				AvailableDevices = AvailableDevices.Union(InDevices).ToList();
+			}
+		}
+
+		/// <summary>
+		/// Registers the provided device for availability
+		/// </summary>
+		/// <param name="Device"></param>
+		public void RegisterDevice(ITargetDevice Device, UnrealDeviceTargetConstraint Constraint = null)
+		{
+			lock (LockObject)
+			{
+				if (AvailableDevices.Contains(Device))
+				{
+					throw new Exception("Device already registered!");
+				}
+
+				InitialConnectState[Device] = Device.IsConnected;
+				Constraints[Device] = Constraint ?? new UnrealDeviceTargetConstraint(Device.Platform.Value);
+
+				AvailableDevices.Add(Device);
+
+				if (Log.IsVerbose)
+				{
+					Device.RunOptions = Device.RunOptions & ~CommandUtils.ERunOptions.NoLoggingOfRunCommand;
+				}
+			}
+		}
+
+		/// <summary>
+		/// Run the provided function across all our devices until it returns false. Devices are provisioned on demand (e.g turned from info into an ITargetDevice)
+		/// </summary>
+		public void EnumerateDevices(UnrealTargetPlatform Platform, Func<ITargetDevice, bool> Predicate)
+		{
+			EnumerateDevices(new UnrealDeviceTargetConstraint(Platform), Predicate);
+		}
+
+		public void EnumerateDevices(UnrealDeviceTargetConstraint Constraint, Func<ITargetDevice, bool> Predicate)
+		{
+			lock (LockObject)
+			{
+				List<ITargetDevice> Selection = new List<ITargetDevice>();
+
+				Log.Verbose(KnownLogEvents.Gauntlet_DeviceEvent, $"Enumerating devices for constraint {Constraint}");
+
+				Log.Verbose(KnownLogEvents.Gauntlet_DeviceEvent, $"   Available devices:");
+				AvailableDevices.ForEach(D => Log.Verbose(KnownLogEvents.Gauntlet_DeviceEvent, $"      {D.Platform}:{D.Name}"));
+
+				Log.Verbose(KnownLogEvents.Gauntlet_DeviceEvent, $"   Unprovisioned devices:");
+				UnprovisionedDevices.ForEach(D => Log.Verbose(KnownLogEvents.Gauntlet_DeviceEvent, $"      {D}"));
+
+				// randomize the order of all devices that are of this platform
+				var MatchingProvisionedDevices = AvailableDevices.Where(D => Constraint.Check(D)).ToList();
+				var MatchingUnprovisionedDevices = UnprovisionedDevices.Where(D => Constraint.Check(D)).ToList();
+
+				bool OutOfDevices = false;
+				bool ContinuePredicate = true;
+
+				do
+				{
+					// Go through all our provisioned devices to see if these fulfill the predicates
+					// requirements
+
+					ITargetDevice NextDevice = MatchingProvisionedDevices.FirstOrDefault();
+
+					while (NextDevice != null && ContinuePredicate)
+					{
+						Log.Verbose(KnownLogEvents.Gauntlet_DeviceEvent, "Checking {Name} against predicate", NextDevice.Name);
+						MatchingProvisionedDevices.Remove(NextDevice);
+						ContinuePredicate = Predicate(NextDevice);
+
+						NextDevice = MatchingProvisionedDevices.FirstOrDefault();
+					}
+
+					if (ContinuePredicate)
+					{
+						// add more devices if possible
+						OutOfDevices = MatchingUnprovisionedDevices.Count() == 0;
+
+						DeviceDefinition NextDeviceDef = MatchingUnprovisionedDevices.FirstOrDefault();
+
+						if (NextDeviceDef != null)
+						{
+							Log.Verbose(KnownLogEvents.Gauntlet_DeviceEvent, "Provisioning device {Name} for the pool", NextDeviceDef.Name);
+
+							// try to create a device. This can fail, but if so we'll just end up back here
+							// on the next iteration
+							ITargetDevice NewDevice = CreateAndRegisterDeviceFromDefinition(NextDeviceDef);
+
+							MatchingUnprovisionedDevices.Remove(NextDeviceDef);
+							UnprovisionedDevices.Remove(NextDeviceDef);
+
+							if (NewDevice != null)
+							{
+								MatchingProvisionedDevices.Add(NewDevice);
+								Log.Verbose(KnownLogEvents.Gauntlet_DeviceEvent, "Added device {Name} to pool", NewDevice.Name);
+							}
+							else
+							{
+								Log.Info(KnownLogEvents.Gauntlet_DeviceEvent, "Failed to provision {Name}", NextDeviceDef.Name);
+								// track this
+								if (FailedProvisions.Contains(NextDeviceDef) == false)
+								{
+									FailedProvisions.Add(NextDeviceDef);
+								}
+							}
+						}
+						else
+						{
+							Log.Info(KnownLogEvents.Gauntlet_DeviceEvent, "Pool ran out of devices of type {Constraint}!", Constraint);
+							OutOfDevices = true;
+						}
+					}
+				} while (OutOfDevices == false && ContinuePredicate);
+			}
+		}
+
+		/// <summary>
+		/// Reserve all devices in the provided list. Once reserved a device will not be seen by any code that
+		/// calls EnumerateDevices
+		/// </summary>
+		/// <param name="DeviceList"></param>
+		/// <returns></returns>
+		public bool ReserveDevices(IEnumerable<ITargetDevice> DeviceList)
+		{
+			lock (LockObject)
+			{
+				// can reserve if not reserved...
+				if (ReservedDevices.Intersect(DeviceList).Count() > 0)
+				{
+					return false;
+				}
+
+				// remove these devices from the available list
+				AvailableDevices = AvailableDevices.Where(D => DeviceList.Contains(D) == false).ToList();
+
+				ReservedDevices.AddRange(DeviceList);
+				DeviceList.All(D => UsedPlatforms.Add(D.Platform));
+			}
+			return true;
+		}
+
+		/// <summary>
@@ -924,7 +956,7 @@
 					continue;
 				}
 
-                if (!DeviceMap.ContainsKey(Entry.Key.Platform.Value))
+				if (!DeviceMap.ContainsKey(Entry.Key.Platform.Value))
 				{
 					// if an unsupported device, we can't reserve it
 					Log.Info(KnownLogEvents.Gauntlet_DeviceEvent, "Unable to reserve service device of type: {Type}", Entry.Key);
@@ -939,14 +971,13 @@
 						Log.Info(KnownLogEvents.Gauntlet_DeviceEvent, "Unable to reserve service device of model: {Model} on legacy service", Entry.Key.Model);
 						return false;
 					}
-									
 				}
 
 				for (int i = 0; i < Entry.Value; i++)
 				{
 					// @todo: if any additional reservation requirements, encode constraint into json
 					ReservationStates.AddRequired(Entry.Key.FormatWithIdentifier(), Entry.Value);
-				}				
+				}
 			}
 
 			// reserve devices
@@ -1006,7 +1037,7 @@
 					Def.Name = Device.Name;
 					Def.Platform = DeviceMap.FirstOrDefault(Entry => Entry.Value == Device.Type.Replace("-DevKit", "", StringComparison.OrdinalIgnoreCase)).Key;
 					Def.DeviceData = Device.DeviceData;
-					Def.Model = Device.Model;					
+					Def.Model = Device.Model;
 
 					EPerfSpec Out = EPerfSpec.Unspecified;
 					if (!String.IsNullOrEmpty(Device.PerfSpec) && !Enum.TryParse<EPerfSpec>(Device.PerfSpec, true, out Out))
@@ -1015,7 +1046,7 @@
 					}
 					Def.PerfSpec = Out;
 
-				    ITargetDevice TargetDevice = CreateAndRegisterDeviceFromDefinition(Def);
+					ITargetDevice TargetDevice = CreateAndRegisterDeviceFromDefinition(Def);
 
 					// If a device from service can't be added, fail reservation and cleanup devices
 					// @todo: device problem reporting, requesting additional devices
@@ -1041,7 +1072,7 @@
 						return false;
 					}
 					else
-					{						
+					{
 						ReservedDevices.Add(TargetDevice);
 						UsedPlatforms.Add(TargetDevice.Platform);
 					}
@@ -1061,194 +1092,192 @@
 		}
 
 		/// <summary>
-		/// Created a list of device definitions from the passed in reference. Needs work....
+		/// Report on Device reservation state
 		/// </summary>
-		/// <param name="InputReference"></param>
-		/// <param name="InLocalTempDir"></param>
-		/// <param name="ObeyConstraints"></param>
-		public void AddDevices(UnrealTargetPlatform DefaultPlatform, string InputReference, bool ObeyConstraints=true)
+		public void ReportDeviceReservationState()
 		{
-			lock (LockObject)
+			string PoolId = string.IsNullOrEmpty(Globals.DevicePoolId) ? "Local" : Globals.DevicePoolId;
+			string[] DeviceReservationStatus = GetReservationStatuses();
+			if (DeviceReservationStatus != null && DeviceReservationStatus.Length > 0)
 			{
-				List<ITargetDevice> NewDevices = new List<ITargetDevice>();
-
-				int SlashIndex = InputReference.IndexOf("\\") >= 0 ? InputReference.IndexOf("\\") : InputReference.IndexOf("/");
-
-				bool PossibleFileName = InputReference.IndexOfAny(Path.GetInvalidPathChars()) < 0 &&
-							(InputReference.IndexOf(":") == -1 || (InputReference.IndexOf(":") == SlashIndex - 1));
-				// Did they specify a file?
-				if (PossibleFileName && File.Exists(InputReference))
-				{
-					Log.Info(KnownLogEvents.Gauntlet_DeviceEvent, "Adding devices from {Reference}", InputReference);
-					List<DeviceDefinition> DeviceDefinitions = JsonSerializer.Deserialize<List < DeviceDefinition >>(
-						File.ReadAllText(InputReference),
-						new JsonSerializerOptions { PropertyNameCaseInsensitive = true }
-					);
-
-					foreach (DeviceDefinition Def in DeviceDefinitions)
-					{
-						Log.Info(KnownLogEvents.Gauntlet_DeviceEvent, "Adding {Name}", Def);
-
-						// use Legacy field if it exists
-						if (Def.Platform == null)
-						{
-							Def.Platform = Def.Type;
-						}
-
-						// check for an availability constraint
-						if (string.IsNullOrEmpty(Def.Available) == false && ObeyConstraints)
-						{
-							// check whether disabled
-							if (String.Compare(Def.Available, "disabled", true) == 0)
-							{
-								Log.Info(KnownLogEvents.Gauntlet_DeviceEvent, "Skipping {Name} due to being disabled", Def.Name);
-								continue;
-							}
-
-							// availability is specified as a range, e.g 21:00-09:00.
-							Match M = Regex.Match(Def.Available, @"(\d{1,2}:\d\d)\s*-\s*(\d{1,2}:\d\d)");
-
-							if (M.Success)
-							{
-								DateTime From, To;
-
-								if (DateTime.TryParse(M.Groups[1].Value, out From) && DateTime.TryParse(M.Groups[2].Value, out To))
-								{
-									// these are just times so when parsed will have todays date. If the To time is less than
-									// From (22:00-06:00) time it spans midnight so move it to the next day
-									if (To < From)
-									{
-										To = To.AddDays(1);
-									}
-
-									// if From is in the future (e.g. it's 01:00 and range is 22:00-08:00) we may be in the previous days window, 
-									// so move them both back a day
-									if (From > DateTime.Now)
-									{
-										From = From.AddDays(-1);
-										To = To.AddDays(-1);
-									}
-
-									if (DateTime.Now < From || DateTime.Now > To)
-									{
-										Log.Info(KnownLogEvents.Gauntlet_DeviceEvent, "Skipping {Name} due to availability constraint {Constraint}", Def.Name, Def.Available);
-										continue;
-									}
-								}
-								else
-								{
-									Log.Warning(KnownLogEvents.Gauntlet_DeviceEvent, "Failed to parse availability {Constraint} for {Name}", Def.Available, Def.Name);
-								}
-							}
-						}
-
-						Def.RemoveOnShutdown = true;
-
-						if (Def.Platform == null)
-						{
-							Def.Platform = DefaultPlatform;
-						}
-
-						UnprovisionedDevices.Add(Def);
-					}
-
-					// randomize devices so if there's a bad device st the start so we don't always hit it (or we do if its later)
-					UnprovisionedDevices = UnprovisionedDevices.OrderBy(D => Guid.NewGuid()).ToList();
-				}
-				else
-				{
-					if (string.IsNullOrEmpty(InputReference) == false)
-					{
-						string[] DevicesList = InputReference.Split(',');
-
-						foreach (string DeviceRef in DevicesList)
-						{							
-
-							// check for <platform>:<address>:<port>|<model>. We pass address:port to device constructor
-							Match M = Regex.Match(DeviceRef, @"(.+?):(.+)");
-
-							UnrealTargetPlatform DevicePlatform = DefaultPlatform;
-							string DeviceAddress = DeviceRef;
-							string Model = string.Empty;
-
-							// when using device service, skip adding local non-desktop devices to pool
-							bool IsDesktop = UnrealBuildTool.Utils.GetPlatformsInClass(UnrealPlatformClass.Desktop).Contains(DevicePlatform);							
-							if (!IsDesktop && DeviceRef.Equals("default", StringComparison.OrdinalIgnoreCase) && !String.IsNullOrEmpty(DeviceURL))
-							{
-								continue;
-							}
-
-							if (M.Success)
-							{
-								if (!UnrealTargetPlatform.TryParse(M.Groups[1].ToString(), out DevicePlatform))
-								{
-									throw new AutomationException("platform {0} is not a recognized device type", M.Groups[1].ToString());
-								}
-
-								DeviceAddress = M.Groups[2].ToString();
-
-								// parse device model
-								if (DeviceAddress.Contains("|"))
-								{
-									string[] Components = DeviceAddress.Split(new char[] { '|' });
-									DeviceAddress = Components[0];
-									Model = Components[1];
-								}
-
-							} 
-
-							Log.Info(KnownLogEvents.Gauntlet_DeviceEvent, "Added device {Platform}:{Address} to pool", DevicePlatform, DeviceAddress);
-							DeviceDefinition Def = new DeviceDefinition();
-							Def.Address = DeviceAddress;
-							Def.Name = DeviceAddress;
-							Def.Platform = DevicePlatform;
-							Def.Model = Model;
-							UnprovisionedDevices.Add(Def);
-						}
-					}
-				}
+				Log.Warning(KnownLogEvents.Gauntlet_DeviceEvent, "Device Reservation status from {Pool} pool: {Status}", PoolId, string.Join(", ", DeviceReservationStatus));
+			}
+			string[] DeviceProblemStatus = GetDeviceProblems();
+			if (DeviceProblemStatus != null && DeviceProblemStatus.Length > 0)
+			{
+				Log.Warning(KnownLogEvents.Gauntlet_DeviceEvent, "Reserved Devices marked with problems from {Pool} pool: {Status}", PoolId, string.Join(", ", DeviceProblemStatus));
 			}
 		}
 
 		/// <summary>
-		/// Adds the list of devices to our internal availability list
+		/// Report target device issue to service with given error message
 		/// </summary>
-		/// <param name="InDevices"></param>
-		public void RegisterDevices(IEnumerable<ITargetDevice> InDevices)
+		public void ReportDeviceError(ITargetDevice Device, string ErrorMessage)
+		{
+			DeviceDefinition Def = null;
+			if (!ServiceDeviceInfo.TryGetValue(Device, out Def))
+			{
+				return;
+			}
+
+			ReportDeviceError(Def.Name, ErrorMessage);
+		}
+
+		/// <summary>
+		/// Release all devices in the provided list from our reserved list
+		/// </summary>
+		/// <param name="DeviceList"></param>
+		public void ReleaseDevices(IEnumerable<ITargetDevice> DeviceList)
 		{
 			lock (LockObject)
 			{
-				AvailableDevices = AvailableDevices.Union(InDevices).ToList();
+				// remove all these devices from our reserved list
+				ReservedDevices = ReservedDevices.Where(D => DeviceList.Contains(D) == false).ToList();
+
+				// disconnect them
+				List<ITargetDevice> KeepList = new List<ITargetDevice>();
+				foreach (ITargetDevice Device in DeviceList)
+				{
+					RevertDeviceConfiguration(Device);
+
+					if (Device.IsConnected && InitialConnectState[Device] == false)
+					{
+						Device.Disconnect();
+					}
+
+					// if device is from reservation service, unregister it
+					DeviceReservationAutoRenew Reservation;
+					if (ServiceReservations.TryGetValue(Device, out Reservation))
+					{
+						bool DisposeReservation = ServiceReservations.Count(Entry => Entry.Value == Reservation) == 1;
+
+						// remove and dispose of device
+						// @todo: add support for reservation modification on server (partial device release)
+						ServiceReservations.Remove(Device);
+						ServiceDeviceInfo.Remove(Device);
+						AvailableDevices.Remove(Device);
+
+						InitialConnectState.Remove(Device);
+						Constraints.Remove(Device);
+
+						Device.Dispose();
+
+						// if last device in reservation, dispose the reservation
+						if (DisposeReservation)
+						{
+							Reservation.Dispose();
+						}
+					}
+					else
+					{
+						KeepList.Add(Device);
+					}
+
+				}
+
+				// Put all devices back into the available list at the front - if another device is requested we want to
+				// try and use one we just put a build on
+				KeepList.ForEach(D => { AvailableDevices.Remove(D); AvailableDevices.Insert(0, D); });
 			}
 		}
 
 		/// <summary>
-		/// Registers the provided device for availability
+		///	Checks whether device pool can accommodate requirements, optionally add service devices to meet demand
 		/// </summary>
-		/// <param name="Device"></param>
-		public void RegisterDevice(ITargetDevice Device, UnrealDeviceTargetConstraint Constraint = null)
+		public bool CheckAvailableDevices(Dictionary<UnrealDeviceTargetConstraint, int> RequiredDevices, IReadOnlyCollection<ProblemDevice> ProblemDevices = null, bool UseServiceDevices = true)
 		{
-			lock (LockObject)
+			Dictionary<UnrealDeviceTargetConstraint, int> AvailableDeviceTypes = new Dictionary<UnrealDeviceTargetConstraint, int>();
+			Dictionary<UnrealDeviceTargetConstraint, int> TotalDeviceTypes = new Dictionary<UnrealDeviceTargetConstraint, int>();
+
+			// Do these "how many available" checks every time because the DevicePool provisions on demand so while it may think it has N machines, 
+			// some of them may fail to be provisioned and we could end up with none!
+
+			// See how many of these types are in device pool (mostly to supply informative info if we can't meet these)
+
+			foreach (var PlatformRequirement in RequiredDevices)
 			{
-				if (AvailableDevices.Contains(Device))
+				UnrealDeviceTargetConstraint Constraint = PlatformRequirement.Key;
+
+				Func<ITargetDevice, bool> Validate = (ITargetDevice Device) =>
 				{
-					throw new Exception("Device already registered!");
+
+					if (!Constraint.Check(Device))
+					{
+						return false;
+					}
+
+					if (ProblemDevices == null)
+					{
+						return true;
+					}
+
+					foreach (ProblemDevice PDevice in ProblemDevices)
+					{
+						if (PDevice.Platform == Device.Platform && PDevice.Name == Device.Name)
+						{
+							return false;
+						}
+					}
+
+					return true;
+				};
+
+				AvailableDeviceTypes[Constraint] = DevicePool.Instance.GetAvailableDeviceCount(Constraint, Validate);
+				TotalDeviceTypes[Constraint] = DevicePool.Instance.GetTotalDeviceCount(Constraint, Validate);
+
+
+				Log.Verbose(KnownLogEvents.Gauntlet_DeviceEvent, "{Constraint}: {Platform} devices required. Total:{Total}, Available:{Available}",
+					Constraint, PlatformRequirement.Value,
+					TotalDeviceTypes[PlatformRequirement.Key], AvailableDeviceTypes[PlatformRequirement.Key]);
+			}
+
+			// get a list of any platforms where we don't have enough
+			var TooFewTotalDevices = RequiredDevices.Where(KP => TotalDeviceTypes[KP.Key] < RequiredDevices[KP.Key]).Select(KP => KP.Key);
+			var TooFewCurrentDevices = RequiredDevices.Where(KP => AvailableDeviceTypes[KP.Key] < RequiredDevices[KP.Key]).Select(KP => KP.Key);
+
+
+			var Devices = TooFewTotalDevices.Concat(TooFewCurrentDevices);
+
+			// Request devices from the service if we need them
+			if (UseServiceDevices && !String.IsNullOrEmpty(DeviceURL) && (TooFewTotalDevices.Count() > 0 || TooFewCurrentDevices.Count() > 0))
+			{
+
+				Dictionary<UnrealDeviceTargetConstraint, int> DeviceCounts = new Dictionary<UnrealDeviceTargetConstraint, int>();
+
+				Devices.ToList().ForEach(Platform => DeviceCounts[Platform] = RequiredDevices[Platform]);
+
+				Log.Info(KnownLogEvents.Gauntlet_DeviceEvent, "Requesting devices from service at {URL}", DeviceURL);
+
+				// Acquire necessary devices from service
+				if (!ReserveDevicesFromService(DeviceURL, DeviceCounts))
+				{
+					return false;
+				}
+			}
+			else
+			{
+				// if we can't ever run then throw an exception
+				if (TooFewTotalDevices.Count() > 0)
+				{
+					var MissingDeviceStrings = TooFewTotalDevices.Select(D => string.Format("Not enough devices of type {0} exist for test. ({1} required, {2} available)", D, RequiredDevices[D], AvailableDeviceTypes[D]));
+					Log.Error(KnownLogEvents.Gauntlet_DeviceEvent, string.Join("\n", MissingDeviceStrings));
+					throw new AutomationException("Not enough devices available");
 				}
 
-				InitialConnectState[Device] = Device.IsConnected;
-				Constraints[Device] = Constraint ?? new UnrealDeviceTargetConstraint(Device.Platform.Value);
-
-				AvailableDevices.Add(Device);
-
-				if (Log.IsVerbose)
+				// if we can't  run now then return false
+				if (TooFewCurrentDevices.Count() > 0)
 				{
-					Device.RunOptions = Device.RunOptions & ~CommandUtils.ERunOptions.NoLoggingOfRunCommand;
+					var MissingDeviceStrings = TooFewCurrentDevices.Select(D => string.Format("Not enough devices of type {0} available for test. ({1} required, {2} available)", D, RequiredDevices[D], AvailableDeviceTypes[D]));
+					Log.Verbose(KnownLogEvents.Gauntlet_DeviceEvent, string.Join("\n", MissingDeviceStrings));
+					return false;
 				}
 			}
+
+			return true;
 		}
 
 		/// <summary>
-		/// Created and registered a device from the provided definition 
+		/// Created and registered a device from the provided definition
 		/// </summary>
 		/// <param name="Def"></param>
 		/// <returns></returns>
@@ -1256,7 +1285,7 @@
 		{
 			ITargetDevice NewDevice = null;
 
-			IDeviceFactory Factory = Gauntlet.Utils.InterfaceHelpers.FindImplementations<IDeviceFactory>()
+			IDeviceFactory Factory = InterfaceHelpers.FindImplementations<IDeviceFactory>()
 					.Where(F => F.CanSupportPlatform(Def.Platform))
 					.FirstOrDefault();
 
@@ -1296,13 +1325,13 @@
 					Log.Info(KnownLogEvents.Gauntlet_DeviceEvent, "Assigned device {Name} reports unavailable. Requesting a forced disconnect", NewDevice);
 					NewDevice.Disconnect(true);
 
-					if  (NewDevice.IsAvailable == false)
+					if (NewDevice.IsAvailable == false)
 					{
 						Log.Info(KnownLogEvents.Gauntlet_DeviceEvent, "Assigned device {Name} still  unavailable. Requesting a reboot", NewDevice);
 						NewDevice.Reboot();
 					}
 				}
-					
+
 				lock (LockObject)
 				{
 					if (NewDevice != null)
@@ -1319,226 +1348,169 @@
 			return NewDevice;
 		}
 
-		public UnrealDeviceTargetConstraint GetConstraint(ITargetDevice Device)
+		/// <summary>
+		/// Reporting device issue using the device name provided by service
+		/// </summary>
+		private void ReportDeviceError(string ServiceDeviceName, string ErrorMessage)
 		{
-
-			if (!Constraints.ContainsKey(Device))
+			if (MaxDeviceErrorReports == 0)
 			{
-				throw new AutomationException("Device pool has no contstaint for {0} (device was likely released)", Device);
+				Log.Warning(KnownLogEvents.Gauntlet_DeviceEvent, "Maximum device errors reported to backend, {Name} : {Message} ignored", ServiceDeviceName, ErrorMessage);
+				return;
 			}
 
-			return Constraints[Device];
+			MaxDeviceErrorReports--;
+
+			Reservation.ReportDeviceError(DeviceURL, ServiceDeviceName, ErrorMessage);
 		}
 
 		/// <summary>
-		/// Returns the number of available devices of the provided type. This includes unprovisioned devices but not reserved ones.
-		/// Note: unprovisioned devices are currently only returned when device is not constrained
+		/// Explicitly release all device reservations
 		/// </summary>
-		public int GetAvailableDeviceCount(UnrealDeviceTargetConstraint Constraint, Func<ITargetDevice, bool> Validate = null)
+		private void ReleaseReservations()
 		{
-			lock (LockObject)
+
+			ServiceDeviceInfo.Clear();
+
+			foreach (ITargetDevice Device in ServiceReservations.Keys)
 			{
-				return AvailableDevices.Where(D => Validate == null ? Constraint.Check(D) : Validate(D)).Count() +
-					UnprovisionedDevices.Where(D => Constraint.Check(D)).Count();
+				AvailableDevices.Remove(Device);
+				ReservedDevices.Remove(Device);
+				Device.Dispose();
+			}
+
+			foreach (DeviceReservationAutoRenew reservation in ServiceReservations.Values.Distinct())
+			{
+				reservation.Dispose();
+			}
+
+			ServiceReservations.Clear();
+		}
+
+		private void CleanupDevices(IEnumerable<UnrealTargetPlatform?> Platforms)
+		{
+			IEnumerable<IDeviceService> DeviceServices = InterfaceHelpers.FindImplementations<IDeviceService>();
+			if (DeviceServices.Any())
+			{
+				foreach (UnrealTargetPlatform? Platform in Platforms)
+				{
+					IDeviceService DeviceService = DeviceServices.Where(D => D.CanSupportPlatform(Platform)).FirstOrDefault();
+					if (DeviceService != null)
+					{
+						DeviceService.CleanupDevices();
+					}
+				}
+			}
+		}
+
+		private void RevertDeviceConfiguration(ITargetDevice Device)
+		{
+			IConfigurableDevice ConfigurableDevice = Device as IConfigurableDevice;
+			if (ConfigurableDevice != null)
+			{
+				var Snapshot = DeviceConfigurationCache.Instance.GetConfigurationSnapshot(Device.Platform, Device.Name);
+				if (Snapshot == null)
+				{
+					return;
+				}
+
+				// Connect temporarily to be able to revert the device's configuration
+				// if the device was disconnected entering here disconnect it after this is over
+				bool bNeedsDisconnect = false;
+				if (!Device.IsConnected)
+				{
+					Device.Connect();
+					bNeedsDisconnect = true;
+				}
+
+				if (ConfigurableDevice.ApplyConfiguration(Snapshot))
+				{
+					DeviceConfigurationCache.Instance.ClearSnapshot(Snapshot);
+				}
+
+				if (bNeedsDisconnect)
+				{
+					Device.Disconnect();
+				}
 			}
 		}
 
 		/// <summary>
-		/// Returns the number of available devices of the provided type. This includes unprovisioned devices but not reserved ones
-		/// Note: unprovisioned devices are currently only returned when device is not constrained
+		/// Construct a path to hold cache files and make sure it's properly cleaned
 		/// </summary>
-		public int GetTotalDeviceCount(UnrealDeviceTargetConstraint Constraint, Func<ITargetDevice, bool> Validate = null)
+		private string GetCleanCachePath(DeviceDefinition InDeviceDefiniton)
 		{
-			lock (LockObject)
+			// Give the desktop platform a temp folder with its name under the device cache
+			string DeviceCache = Path.Combine(LocalTempDir, "DeviceCache");
+			string PlatformCache = Path.Combine(DeviceCache, InDeviceDefiniton.Platform.ToString());
+			string ClientCache = Path.Combine(PlatformCache, InDeviceDefiniton.Name);
+
+			// On Desktops, builds are installed in the device cache.
+			// When using device reservation blocks, we don't want to fully clean the cache and lose previously installed builds.
+			// If bRetainBuilds evaluates to true, it means we are in the second step or beyond in a device reservation block.
+			// In this case we'll just delete the left over UserDir which should already have been emptied by UnrealSession.
+			bool? bForceClean = UnrealAppConfig.ForceFullClean;
+			bool? bSkipInstall = UnrealAppConfig.ForceSkipInstall;
+			bool bUsingReservationBlock = bForceClean.HasValue && bSkipInstall.HasValue;
+			bool bRetainBuilds = bUsingReservationBlock && !bForceClean.Value && bSkipInstall.Value;
+
+			if(bRetainBuilds)
 			{
-				return AvailableDevices.Union(ReservedDevices).Where(D => Validate == null ? Constraint.Check(D) : Validate(D)).Count() +
-					UnprovisionedDevices.Where(D => Constraint.Check(D)).Count();
-			}
-		}
+				Log.Info("Retaining build cache for device reservation block");
 
-		/// <summary>
-		///	Checks whether device pool can accommodate requirements, optionally add service devices to meet demand
-		/// </summary>
-		public bool CheckAvailableDevices(Dictionary<UnrealDeviceTargetConstraint, int> RequiredDevices, IReadOnlyCollection<ProblemDevice> ProblemDevices = null, bool UseServiceDevices = true)
-		{
-			Dictionary<UnrealDeviceTargetConstraint, int> AvailableDeviceTypes = new Dictionary<UnrealDeviceTargetConstraint, int>();
-			Dictionary<UnrealDeviceTargetConstraint, int> TotalDeviceTypes = new Dictionary<UnrealDeviceTargetConstraint, int>();
-
-			// Do these "how many available" checks every time because the DevicePool provisions on demand so while it may think it has N machines, 
-			// some of them may fail to be provisioned and we could end up with none!
-
-			// See how many of these types are in device pool (mostly to supply informative info if we can't meet these)
-
-			foreach (var PlatformRequirement in RequiredDevices)
-			{
-				UnrealDeviceTargetConstraint Constraint = PlatformRequirement.Key;
-
-				Func<ITargetDevice, bool> Validate = (ITargetDevice Device) =>
+				DirectoryInfo UserDirectory = new(Path.Combine(ClientCache, "UserDir"));
+				if(UserDirectory.Exists)
 				{
-
-					if (!Constraint.Check(Device))
+					try
 					{
-						return false;
+						Log.Info("Cleaning stale user directory...");
+						SystemHelpers.Delete(UserDirectory, true, true);
 					}
-
-					if (ProblemDevices == null)
+					catch(Exception Ex)
 					{
-						return true;
+						throw new AutomationException("Failed to clean user directory {0}. This could result in improper artifact reporting. {1}", UserDirectory, Ex);
 					}
-
-					foreach (ProblemDevice PDevice in ProblemDevices)
-					{
-						if (PDevice.Platform == Device.Platform && PDevice.Name == Device.Name)
-						{
-							return false;
-						}
-					}
-
-					return true;
-				};
-
-				AvailableDeviceTypes[Constraint] = DevicePool.Instance.GetAvailableDeviceCount(Constraint, Validate);
-				TotalDeviceTypes[Constraint] = DevicePool.Instance.GetTotalDeviceCount(Constraint, Validate);
-
-
-				Log.Verbose(KnownLogEvents.Gauntlet_DeviceEvent, "{Constraint}: {Platform} devices required. Total:{Total}, Available:{Available}",
-					Constraint, PlatformRequirement.Value,
-					TotalDeviceTypes[PlatformRequirement.Key], AvailableDeviceTypes[PlatformRequirement.Key]);
-			}
-
-			// get a list of any platforms where we don't have enough
-			var TooFewTotalDevices = RequiredDevices.Where(KP => TotalDeviceTypes[KP.Key] < RequiredDevices[KP.Key]).Select(KP => KP.Key);
-			var TooFewCurrentDevices = RequiredDevices.Where(KP => AvailableDeviceTypes[KP.Key] < RequiredDevices[KP.Key]).Select(KP => KP.Key);
-
-
-			var Devices = TooFewTotalDevices.Concat(TooFewCurrentDevices);		
-
-			// Request devices from the service if we need them
-			if (UseServiceDevices && !String.IsNullOrEmpty(DeviceURL)  && (TooFewTotalDevices.Count() > 0 || TooFewCurrentDevices.Count() > 0))
-			{				
-
-				Dictionary<UnrealDeviceTargetConstraint, int> DeviceCounts = new Dictionary<UnrealDeviceTargetConstraint, int>();
-
-				Devices.ToList().ForEach(Platform => DeviceCounts[Platform] = RequiredDevices[Platform]);
-
-				Log.Info(KnownLogEvents.Gauntlet_DeviceEvent, "Requesting devices from service at {URL}", DeviceURL);
-
-				// Acquire necessary devices from service
-				if (!ReserveDevicesFromService(DeviceURL, DeviceCounts))
-				{
-					return false;
 				}
 			}
 			else
 			{
-				// if we can't ever run then throw an exception
-				if (TooFewTotalDevices.Count() > 0)
+				int CleanAttempts = 0;
+
+				while (Directory.Exists(ClientCache))
 				{
-					var MissingDeviceStrings = TooFewTotalDevices.Select(D => string.Format("Not enough devices of type {0} exist for test. ({1} required, {2} available)", D, RequiredDevices[D], AvailableDeviceTypes[D]));
-					Log.Error(KnownLogEvents.Gauntlet_DeviceEvent, string.Join("\n", MissingDeviceStrings));
-					throw new AutomationException("Not enough devices available");
-				}
+					DirectoryInfo ClientCacheDirectory = new(ClientCache);
 
-				// if we can't  run now then return false
-				if (TooFewCurrentDevices.Count() > 0)
-				{
-					var MissingDeviceStrings = TooFewCurrentDevices.Select(D => string.Format("Not enough devices of type {0} available for test. ({1} required, {2} available)", D, RequiredDevices[D], AvailableDeviceTypes[D]));
-					Log.Verbose(KnownLogEvents.Gauntlet_DeviceEvent, string.Join("\n", MissingDeviceStrings));
-					return false;
-				}
-			}
+					try
+					{
+						Log.Info("Cleaning stale client device cache...");
+						SystemHelpers.Delete(ClientCacheDirectory, true, true);
+					}
+					catch (Exception Ex)
+					{
+						// If we fail to acquire the default client cache while using device reservation blocks,
+						// we can't ensure future tests will have their cache directories mapped to the correct build location
+						if(bUsingReservationBlock)
+						{
+							throw new AutomationException("Failed to clean default client device cache {0}. {1}", ClientCache, Ex);
+						}
 
-			return true;
-
-		}
-
-		/// <summary>
-		/// Run the provided function across all our devices until it returns false. Devices are provisioned on demand (e.g turned from info into an ITargetDevice)
-		/// </summary>		
-		public void EnumerateDevices(UnrealTargetPlatform Platform, Func<ITargetDevice, bool> Predicate)
-		{
-			EnumerateDevices(new UnrealDeviceTargetConstraint(Platform), Predicate);
-		}
-
-		public void EnumerateDevices(UnrealDeviceTargetConstraint Constraint, Func<ITargetDevice, bool> Predicate)
-		{
-			lock (LockObject)
-			{
-				List<ITargetDevice> Selection = new List<ITargetDevice>();
-				
-				Log.Verbose(KnownLogEvents.Gauntlet_DeviceEvent, $"Enumerating devices for constraint {Constraint}");
-
-				Log.Verbose(KnownLogEvents.Gauntlet_DeviceEvent, $"   Available devices:");
-				AvailableDevices.ForEach(D => Log.Verbose(KnownLogEvents.Gauntlet_DeviceEvent, $"      {D.Platform}:{D.Name}"));
-
-				Log.Verbose(KnownLogEvents.Gauntlet_DeviceEvent, $"   Unprovisioned devices:");
-				UnprovisionedDevices.ForEach(D => Log.Verbose(KnownLogEvents.Gauntlet_DeviceEvent, $"      {D}"));
-				
-				// randomize the order of all devices that are of this platform
-				var MatchingProvisionedDevices = AvailableDevices.Where(D => Constraint.Check(D)).ToList();
-				var MatchingUnprovisionedDevices = UnprovisionedDevices.Where(D => Constraint.Check(D)).ToList();
-
-				bool OutOfDevices = false;
-				bool ContinuePredicate = true;
-
-				do
-				{
-					// Go through all our provisioned devices to see if these fulfill the predicates
-					// requirements
-
-					ITargetDevice NextDevice = MatchingProvisionedDevices.FirstOrDefault();
-
-					while (NextDevice != null && ContinuePredicate)
-					{
-						Log.Verbose(KnownLogEvents.Gauntlet_DeviceEvent, "Checking {Name} against predicate", NextDevice.Name);
-						MatchingProvisionedDevices.Remove(NextDevice);
-						ContinuePredicate = Predicate(NextDevice);
-
-						NextDevice = MatchingProvisionedDevices.FirstOrDefault();
-					}
-
-					if (ContinuePredicate)
-					{
-						// add more devices if possible
-						OutOfDevices = MatchingUnprovisionedDevices.Count() == 0;
-
-						DeviceDefinition NextDeviceDef = MatchingUnprovisionedDevices.FirstOrDefault();
-
-						if (NextDeviceDef != null)
-						{
-							Log.Verbose(KnownLogEvents.Gauntlet_DeviceEvent, "Provisioning device {Name} for the pool", NextDeviceDef.Name);
-
-							// try to create a device. This can fail, but if so we'll just end up back here
-							// on the next iteration
-							ITargetDevice NewDevice = CreateAndRegisterDeviceFromDefinition(NextDeviceDef);
-
-							MatchingUnprovisionedDevices.Remove(NextDeviceDef);
-							UnprovisionedDevices.Remove(NextDeviceDef);
-
-							if (NewDevice != null)
-							{
-								MatchingProvisionedDevices.Add(NewDevice);
-								Log.Verbose(KnownLogEvents.Gauntlet_DeviceEvent, "Added device {Name} to pool", NewDevice.Name);
-							}
-							else
-							{
-								Log.Info(KnownLogEvents.Gauntlet_DeviceEvent, "Failed to provision {Name}", NextDeviceDef.Name);
-								// track this
-								if (FailedProvisions.Contains(NextDeviceDef) == false)
-								{
-									FailedProvisions.Add(NextDeviceDef);
-								}
-							}
-						}
+						// When not using device reservation blocks, we can just create a newly indexed directory for the client cache
 						else
 						{
-							Log.Info(KnownLogEvents.Gauntlet_DeviceEvent, "Pool ran out of devices of type {Constraint}!", Constraint);
-							OutOfDevices = true;
+							string Warning = "Failed to clean client device cache {Folder}. A newly indexed directory will be created instead. {Message}";
+							Log.Warning(KnownLogEvents.Gauntlet_DeviceEvent, Warning, ClientCache, Ex.Message);
+
+							ClientCache = Path.Combine(PlatformCache, $"{InDeviceDefiniton.Name}_{++CleanAttempts}");
 						}
 					}
-				} while (OutOfDevices == false && ContinuePredicate);
+				}
+
+				// create this path
+				Log.Info("Client device cache set to {Directory}", ClientCache);
+				Directory.CreateDirectory(ClientCache);
 			}
+
+			return ClientCache;
 		}
 	}
 }
- 
