- Fixed Staging of Program Mac targets to keep the bogus .app copies out of the Staged dir

diff --git a/Engine/Source/Programs/AutomationTool/Apple/ApplePlatform.Automation.cs b/Engine/Source/Programs/AutomationTool/Apple/ApplePlatform.Automation.cs
--- a/Engine/Source/Programs/AutomationTool/Apple/ApplePlatform.Automation.cs
+++ b/Engine/Source/Programs/AutomationTool/Apple/ApplePlatform.Automation.cs
@@ -95,7 +95,23 @@ public abstract class ApplePlatform : Platform
 		// staging will put binaries into Staged/<game>/Binaries/<platform> and they aren't needed, and when we pull this into a .app, it 
 		// messes with the resulting .app. So, we remove the game binary now (leaving in helper .app's and raw .dylibs, etc)
 		// they come from BuildProducts, and we could maybe remove from that list, but it could cause issues with Horde/buildmachines
-		FileReference BinaryPath = FileReference.Combine(SC.StageDirectory, Params.IsCodeBasedProject ? SC.ShortProjectName : "Engine", "Binaries", SC.PlatformDir, SC.StageExecutables[0]);
+		// program binaries live under the Engine directory
+		string RootDirName = "Engine";
+		StagedFileReference RemappedLoc = DeploymentContext.ApplyDirectoryRemap(SC, DeploymentContext.MakeRelativeStagedReference(SC, SC.ProjectBinariesFolder));
+
+
+		FileReference BinaryPath = FileReference.Combine(SC.StageDirectory, RemappedLoc.Name, SC.StageExecutables[0]);
+
+
+		if (Params.IsCodeBasedProject && SC.StageTargets[0].Receipt.TargetType != TargetType.Program)
+		{
+			RootDirName = SC.ShortProjectName;
+		}
+		FileReference BinaryPath2 = FileReference.Combine(SC.StageDirectory, RootDirName, "Binaries", SC.PlatformDir, SC.StageExecutables[0]);
+
+Console.WriteLine("%%%%%%%% {0} // {1}", BinaryPath, BinaryPath2);
+BinaryPath = BinaryPath2;
+
 		DirectoryReference AppPath = new DirectoryReference(BinaryPath.FullName + ".app");
 		InternalUtils.SafeDeleteFile(BinaryPath.FullName, true);
 		InternalUtils.SafeDeleteDirectory(AppPath.FullName, true);
