- Make sure zipping a Mac app works with ZipTask which deletes empty directories (the Binaries/Mac dir may be empty after staging), by creating an empty file in Binaries/<platform>

diff --git a/Engine/Source/Programs/AutomationTool/Apple/ApplePlatform.Automation.cs b/Engine/Source/Programs/AutomationTool/Apple/ApplePlatform.Automation.cs
--- a/Engine/Source/Programs/AutomationTool/Apple/ApplePlatform.Automation.cs
+++ b/Engine/Source/Programs/AutomationTool/Apple/ApplePlatform.Automation.cs
@@ -103,9 +103,12 @@ public abstract class ApplePlatform : Platform
 			RootDirName = SC.ShortProjectName;
 		}
 		FileReference BinaryPath = FileReference.Combine(SC.StageDirectory, RootDirName, "Binaries", SC.PlatformDir, SC.StageExecutables[0]);
+		InternalUtils.SafeDeleteFile(BinaryPath.FullName, true);
+		// this may leave the binaries directory empty, but the Mac needs the Binaries/Mac dir to exist (see FMacPlatformProcess::BaseDir())
+		// so plop a file down in it's place
+		File.WriteAllText(Path.Combine(BinaryPath.Directory.FullName, ".binariesdir"), "");
 
 		DirectoryReference AppPath = new DirectoryReference(BinaryPath.FullName + ".app");
-		InternalUtils.SafeDeleteFile(BinaryPath.FullName, true);
 		InternalUtils.SafeDeleteDirectory(AppPath.FullName, true);
 
 		if (AppleExports.UseModernXcode(Params.RawProjectPath))
