Localization:
Introduced opt in localization for plugins. Plugins can be set to never have localization config files generated, auto-generated or have the localization pipeline use the user generated localization config files.
- All existing plugins with localization targets will be retrofitted to have a LocalizationConfigGenerationPolicy of Never to match current behavior.
- Introduced the LocalizationConfigGenerationPolicy enum to plugin descriptors and LocalizationTargetDescriptors. This controls whether plugin localization targets use user generated localization config files, auto-generated localization config files or if the plugin is never to have localization config files and thus not localized.
- Introduced a step in the Localize UAT command to auto-generate localization config files with default settings for plugins that opt into the auto-generation feature.
- Added a clean up step in the Localize uAT command to delete all auto-generated files and folders from a run of the command.
- Added a -PreserveAutoGeneratedResources flag for the Localize UAT command to preserve the auto-generated files and directories for debugging.
-Updated both the C# and C++ version of the LocalizationTargetDescriptor to contain the new LocalizationConfigGenerationPolicy and have them read and written from the plugin descriptor files.
- Updated the NewPluginLocalizationTarget command under the LocalizationTargetEditor UAT command to accept a LocalizationConfigGenerationPolicy as a command line argument for all plugin localization targets to be created.
If the parameter is not specified, the parameter defaults to the Auto LocalizationConfigGenerationPolicy for all the plugins specified.
#test Used the NewPluginLocalizationTarget sub-command to create plugins with an auto localization target. Used the Localize UAT command to perform a gather for such a plugin. Config files were indeed generated, the plugin was gathered and the loc data stored in PluginName/Content/Localization and then the config files were deleted. No fuss, no muss

diff --git a/Engine/Source/Programs/AutomationTool/Localization/LocalizationTargetEditor/Commands/NewPluginLocalizationTargetCommand.cs b/Engine/Source/Programs/AutomationTool/Localization/LocalizationTargetEditor/Commands/NewPluginLocalizationTargetCommand.cs
--- a/Engine/Source/Programs/AutomationTool/Localization/LocalizationTargetEditor/Commands/NewPluginLocalizationTargetCommand.cs
+++ b/Engine/Source/Programs/AutomationTool/Localization/LocalizationTargetEditor/Commands/NewPluginLocalizationTargetCommand.cs
@@ -4,19 +4,18 @@ using System.Collections.Generic;
 using System.Linq;
 using System.Text;
 using System.IO;
-using AutomationTool;
 using UnrealBuildTool;
 using EpicGames.Core;
 using Microsoft.Extensions.Logging;
 
 using static AutomationTool.CommandUtils;
-using UnrealBuildBase;
 
 namespace EpicGames.Localization
 {
 	public class NewPluginLocalizationTargetCommand : PluginLocalizationTargetCommand
 	{
 		private LocalizationTargetDescriptorLoadingPolicy _loadingPolicy;
+		private LocalizationConfigGenerationPolicy _configGenerationPolicy;
 		private string _localizationTargetNameSuffix = "";
 		private LocalizationConfigFileFormat _fileFormat = LocalizationConfigFileFormat.Latest;
 		private List<string> _filesToAdd = new List<string>();
@@ -40,7 +39,9 @@ namespace EpicGames.Localization
 			helpTextBuilder.AppendLine("UEProjectDirectory - Required relative path from UERootDirectory to the project. This should be  your project directory that contains the Plugins directory.");
 			helpTextBuilder.AppendLine("UEProjectName - An optional name of the project the plugin is under. If blank, this implies the plugin is under Engine.");
 			helpTextBuilder.AppendLine("LocalizationTargetLoadingPolicy - A required string that determines when the plugin should load the localization data for the localization target. Valid values are Game, Always, Editor, PropertyNames, ToolTips. See LocalizationTargetDescriptorLoadingPolicy");
-			helpTextBuilder.AppendLine("_localizationTargetNameSuffix - An optional suffix to give to the plugin localization target. By default, the name of the plugin localization target would be the name of the plugin. This allows multiple localization targets to be created for plugins.");
+			helpTextBuilder.AppendLine("LocalizationConfigGenerationPolicy- An optional string that specifies how localization config files should be generated during the localization gather process for the plugin. Acceptable values are Never, Auto, User. If not specified, defaults to Auto.");
+			helpTextBuilder.AppendLine("Never means no localization config files will be generated or used during the localization process. No localization data will be generated for the plugin druing a gather. Auto means temporary, default localization config files will be generated druing localization gather and used to generate localization data. User means there are user provided localization config files in the plugins's Config/Localization folder that will be used during localization gathers to generate localization data.");
+			helpTextBuilder.AppendLine("LocalizationTargetNameSuffix - An optional suffix to give to the plugin localization target. By default, the name of the plugin localization target would be the name of the plugin. This allows multiple localization targets to be created for plugins.");
 			// Include plugins 
 			helpTextBuilder.AppendLine("IncludePlugins - An optional comma separated list of plugins to create localization targets for. E.g PluginA,PluginB,PluginC");
 			helpTextBuilder.AppendLine("IncludePluginsDirectory - An optional relative directory to UEProjectDirectory. All plugins under this directory will have localization targets created if they are not excluded. E.g Plugins/PluginFolderA.");
@@ -61,8 +62,18 @@ namespace EpicGames.Localization
 			// @TODOLocalization: Support an override for the localization config file format 
 			_fileFormat = LocalizationConfigFileFormat.Latest;
 			_loadingPolicy = _commandLineHelper.ParseRequiredEnumParamEnum<LocalizationTargetDescriptorLoadingPolicy>("LocalizationTargetLoadingPolicy");
-
-			return true;
+			LocalizationConfigGenerationPolicy? nullableGenerationPolicy = _commandLineHelper.ParseOptionalEnumParam<LocalizationConfigGenerationPolicy>("LocalizationConfigGenerationPolicy");
+			if (nullableGenerationPolicy.HasValue)
+			{
+				_configGenerationPolicy = nullableGenerationPolicy.Value;
+			}
+			else
+			{
+				_configGenerationPolicy = LocalizationConfigGenerationPolicy.Auto;
+				Logger.LogInformation("LocalizationConfigGenerationPolicy is defaulting to Auto. All plugin localization targets created will have default, temporary localization config files generated during the localization pipeline. To change this behavior, please see the helpt text for the LocalizationConfigGenerationPolicy parameter.");
+			}
+			
+				return true;
 		}
 
 		public override bool Execute()
@@ -101,7 +112,15 @@ namespace EpicGames.Localization
 					pluginLocalizationTargetName += "_" + _localizationTargetNameSuffix;
 				}
 
-				GeneratePluginLocalizationConfigFiles(pluginInfo, pluginLocalizationTargetName, pluginInfo.Directory.MakeRelativeTo(ueProjectDirectoryReference));
+				// We  only generate the config files if it's set to user. Auto will autogen the config files and never implies we don't need localization config files at all.
+				if (_configGenerationPolicy == LocalizationConfigGenerationPolicy.User)
+				{
+					GeneratePluginLocalizationConfigFiles(pluginInfo, pluginLocalizationTargetName, pluginInfo.Directory.MakeRelativeTo(ueProjectDirectoryReference));
+				}
+				else
+				{
+					Logger.LogInformation($"Localization target generation policy is {_configGenerationPolicy.ToString()}. No localization config files will be generated for {pluginLocalizationTargetName}.	");
+				}
 				// Update the .uplugin file with updated localization target descriptors 
 				UpdatePluginLocalizationTargetDescriptor(pluginInfo, pluginLocalizationTargetName);
 			}
@@ -165,7 +184,7 @@ namespace EpicGames.Localization
 					}
 				}
 			}
-			LocalizationTargetDescriptor locTargetDescriptorToAdd = new LocalizationTargetDescriptor(localizationTargetName, _loadingPolicy);
+			LocalizationTargetDescriptor locTargetDescriptorToAdd = new LocalizationTargetDescriptor(localizationTargetName, _loadingPolicy, _configGenerationPolicy);
 			if (descriptor.LocalizationTargets is null)
 			{
 				descriptor.LocalizationTargets = new LocalizationTargetDescriptor[] { locTargetDescriptorToAdd };
@@ -179,7 +198,7 @@ namespace EpicGames.Localization
 				descriptor.LocalizationTargets = newLocalizationTargets;
 			}
 
-			Logger.LogInformation($"Updating localization target information for {pluginInfo.Name} plugin. Adding localization target {localizationTargetName} with loading policy {_loadingPolicy.ToString()}");
+			Logger.LogInformation($"Updating localization target information for {pluginInfo.Name} plugin. Adding localization target {localizationTargetName} with loading policy {_loadingPolicy.ToString()} and ConfigGenerationPolicy {_configGenerationPolicy}.");
 
 			// Add descriptor to processing list 
 			_filesToEdit.Add(pluginInfo.File.FullName);
diff --git a/Engine/Source/Programs/AutomationTool/Scripts/Localisation.Automation.cs b/Engine/Source/Programs/AutomationTool/Scripts/Localisation.Automation.cs
--- a/Engine/Source/Programs/AutomationTool/Scripts/Localisation.Automation.cs
+++ b/Engine/Source/Programs/AutomationTool/Scripts/Localisation.Automation.cs
@@ -89,8 +89,12 @@ class Localize : BuildCommand
 	private string AdditionalCommandletArguments = "";
 	private bool bEnableParallelGather = false;
 	private bool bIsRunningInPreview = false;
+	private bool bPreserveAutoGeneratedResources = false;
 	private int PendingChangeList = -1;
 
+	HashSet<string> AutoGeneratedFiles = new();
+	HashSet<string> AutoGeneratedDirectories = new();
+
 	public override void ExecuteBuild()
 	{
 		ParseCommandLine();
@@ -113,7 +117,12 @@ class Localize : BuildCommand
 		// Build up any additional batches needed for plugins
 		if (bShouldGatherPlugins)
 		{
-			AddPluginLocalizationBatches(LocalizationBatches);
+			HashSet<string> PluginsAdded = new();
+			AddPluginLocalizationBatches(LocalizationBatches, PluginsAdded);
+			if (PluginsAdded.Count > 0)
+			{
+				AutoGeneratePluginLocalizationFiles(PluginsAdded);
+			}
 		}
 
 		// Create a single changelist to use for all changes
@@ -180,6 +189,11 @@ class Localize : BuildCommand
 			CleanUpGeneratedPreviewFiles(LocalizationBatches);
 		}
 
+		if (!bPreserveAutoGeneratedResources)
+		{
+			CleanUpAutoGeneratedFiles();
+		}
+
 		// Upload the latest sources to our localization provider
 		if (LocalizationStepNames.Contains("Upload"))
 		{
@@ -306,6 +320,12 @@ class Localize : BuildCommand
 			Logger.LogInformation("Running in preview mode. Preview switch will be passed along to all localization commandlets to be run.");
 			AdditionalCommandletArguments += " -Preview";
 		}
+
+		bPreserveAutoGeneratedResources = ParseParam("PreserveAutoGeneratedResources");
+		if (bPreserveAutoGeneratedResources)
+		{
+			Logger.LogInformation("Preserving auto-generated content. Auto-generated files and folders will not be automatically cleaned up at the end of this command.");
+		}
 	}
 
 	private void AddStaticLocalizationBatches(List<LocalizationBatch> LocalizationBatches)
@@ -333,7 +353,7 @@ class Localize : BuildCommand
 		}
 	}
 
-	private void AddPluginLocalizationBatches(List<LocalizationBatch> LocalizationBatches)
+	private void AddPluginLocalizationBatches(List<LocalizationBatch> LocalizationBatches, HashSet<string> PluginsAdded)
 	{
 		var PluginsRootDirectory = new DirectoryReference(CombinePaths(UEProjectRoot, UEProjectDirectory));
 		IReadOnlyList<PluginInfo> AllPlugins = Plugins.ReadPluginsFromDirectory(PluginsRootDirectory, "Plugins", UEProjectName.Length == 0 ? PluginType.Engine : PluginType.Project);
@@ -348,6 +368,7 @@ class Localize : BuildCommand
 			bool bPluginHasLocalizationTarget = PluginInfo.Descriptor.LocalizationTargets != null && PluginInfo.Descriptor.LocalizationTargets.Length > 0;
 			if (bShouldIncludePlugin && bPluginHasLocalizationTarget)
 			{
+				PluginsAdded.Add(PluginInfo.Name);
 				var RootRelativePluginPath = PluginInfo.Directory.MakeRelativeTo(new DirectoryReference(UEProjectRoot));
 				RootRelativePluginPath = RootRelativePluginPath.Replace('\\', '/'); // Make sure we use / as these paths are used with P4
 
@@ -644,6 +665,83 @@ class Localize : BuildCommand
 		P4.RevertUnchanged(PendingChangeList);
 	}
 
+	private void AutoGeneratePluginLocalizationFiles(HashSet<string> PluginNames)
+	{
+		foreach (string PluginName in PluginNames)
+		{
+			PluginInfo Info = Plugins.GetPlugin(PluginName);
+			if (Info is null)
+			{
+				continue;
+			}
+
+			foreach (LocalizationTargetDescriptor Target in Info.Descriptor.LocalizationTargets)
+			{
+				if (Target.ConfigGenerationPolicy == LocalizationConfigGenerationPolicy.Auto)
+				{
+					string PluginConfigLocalizationPath = CombinePaths(Info.Directory.FullName, "Config", "Localization");
+					if (!Directory.Exists(PluginConfigLocalizationPath))
+					{
+						Logger.LogInformation($"Plugin '{Info.Name}' does not have a Config Localization folder. Creating '{PluginConfigLocalizationPath}'.");
+						Directory.CreateDirectory(PluginConfigLocalizationPath);
+						AutoGeneratedDirectories.Add(PluginConfigLocalizationPath);
+					}
+					LocalizationConfigFileGenerator Generator = LocalizationConfigFileGenerator.GetGeneratorForFileFormat(LocalizationConfigFileFormat.Latest);
+					LocalizationConfigFileGeneratorParams GeneratorParams = new();
+					GeneratorParams.LocalizationTargetName = Target.Name;
+					GeneratorParams.LocalizationTargetRootDirectory = Info.File.Directory.FullName;
+					List<LocalizationConfigFile> PluginLocalizationConfigFiles = Generator.GenerateDefaultSettingsConfigFiles(GeneratorParams);
+					foreach (LocalizationConfigFile PluginLocalizationConfigFile in PluginLocalizationConfigFiles)
+					{
+						string SaveFilePath = Path.Combine(PluginConfigLocalizationPath, PluginLocalizationConfigFile.Name);
+						if (File.Exists(SaveFilePath))
+						{
+							Logger.LogInformation($"Plugin localization target %{Target.Name} in {Info.Name} plugin already contains localization config file {SaveFilePath}. The file will not be auto-generated.");
+							continue;
+						}
+						Logger.LogInformation($"Auto generating localization config file '{SaveFilePath}' for localization target '{Target.Name}' in plugin {Info.Name}.");
+						AutoGeneratedFiles.Add(SaveFilePath);
+						if (!bIsRunningInPreview)
+						{
+							PluginLocalizationConfigFile.Write(new FileReference(SaveFilePath));
+						}
+					}
+				}
+			}
+		}
+	}
+
+	private void CleanUpAutoGeneratedFiles()
+	{
+		Logger.LogInformation("Cleaning up auto-generated files.");
+		foreach (string AutoGeneratedFile in AutoGeneratedFiles)
+		{
+			if (File.Exists(AutoGeneratedFile))
+			{
+				Logger.LogInformation($"Deleting auto-generated file '{AutoGeneratedFile}'");
+				if (!bIsRunningInPreview)
+				{
+					File.Delete(AutoGeneratedFile);
+				}
+			}
+		}
+		Logger.LogInformation("Finished cleaning up all auto-generated files.");
+
+		Logger.LogInformation("Cleaning up auto-generated directories.");
+		foreach (string AutoGeneratedDirectory in AutoGeneratedDirectories)
+		{
+			if (Directory.Exists(AutoGeneratedDirectory))
+			{
+				Logger.LogInformation($"Deleting auto-generated directory '{AutoGeneratedDirectory}'.");
+				if (!bIsRunningInPreview)
+				{
+					Directory.Delete(AutoGeneratedDirectory);
+				}
+			}
+		}
+		Logger.LogInformation("Finished cleaning up auto-generated directories.");
+	}
+
 	private ProjectInfo GenerateProjectInfo(string RootWorkingDirectory, string ProjectName, IReadOnlyList<string> LocalizationStepNames)
 	{
 		var LocalizationSteps = new List<ProjectStepInfo>();
diff --git a/Engine/Source/Programs/UnrealBuildTool/System/LocalizationDescriptor.cs b/Engine/Source/Programs/UnrealBuildTool/System/LocalizationDescriptor.cs
--- a/Engine/Source/Programs/UnrealBuildTool/System/LocalizationDescriptor.cs
+++ b/Engine/Source/Programs/UnrealBuildTool/System/LocalizationDescriptor.cs
@@ -43,7 +43,26 @@ namespace UnrealBuildTool
 	};
 
 	/// <summary>
-	/// 
+	/// How this localization target should be generated during the localization gather pipeline
+	/// </summary>
+	public enum LocalizationConfigGenerationPolicy
+	{
+		/// <summary>
+		/// This localization target should never have localization config files associated with it during the localization gather pipeline.
+		/// </summary>
+		Never,
+		/// <summary>
+		/// This localization target should only use user generated localization config files during the localization gather pipeline.
+		/// </summary>
+		User,
+		/// <summary>
+		/// Default auto-generated localization config files will be used to generate the localization target and localization content files during the localization gather pipeline
+		/// </summary>
+		Auto,
+	}
+
+	/// <summary>
+	/// Description of a localization target.
 	/// </summary>
 	[DebuggerDisplay("Name={Name}")]
 	public class LocalizationTargetDescriptor
@@ -58,6 +77,10 @@ namespace UnrealBuildTool
 		/// </summary>
 		public LocalizationTargetDescriptorLoadingPolicy LoadingPolicy;
 
+		/// <summary>
+		/// How should this localization target's localization config files be generated during a localization gather.
+		/// </summary>
+		public LocalizationConfigGenerationPolicy ConfigGenerationPolicy;
 		/// <summary>
 		/// Constructor
 		/// </summary>
@@ -67,6 +90,21 @@ namespace UnrealBuildTool
 		{
 			Name = InName;
 			LoadingPolicy = InLoadingPolicy;
+			// Older plugins and localization target descriptors won't have the generation policy. We default it to Never.
+			ConfigGenerationPolicy = LocalizationConfigGenerationPolicy.Never;
+		}
+
+		/// <summary>
+		/// Constructor
+		/// </summary>
+		/// <param name="InName">The name of the localization target.</param>
+		/// <param name="InLoadingPolicy"> When the localization data associated with this localization target should be loaded.</param>
+		/// <param name="InGenerationPolicy">How the localization config files should be generated during a localization gather to create the localization data files.</param>
+		public LocalizationTargetDescriptor(string InName, LocalizationTargetDescriptorLoadingPolicy InLoadingPolicy, LocalizationConfigGenerationPolicy InGenerationPolicy)
+		{
+			Name = InName;
+			LoadingPolicy = InLoadingPolicy;
+			ConfigGenerationPolicy = InGenerationPolicy;
 		}
 
 		/// <summary>
@@ -76,7 +114,13 @@ namespace UnrealBuildTool
 		/// <returns>The new localization target descriptor</returns>
 		public static LocalizationTargetDescriptor FromJsonObject(JsonObject InObject)
 		{
-			return new LocalizationTargetDescriptor(InObject.GetStringField("Name"), InObject.GetEnumField<LocalizationTargetDescriptorLoadingPolicy>("LoadingPolicy"));
+			LocalizationTargetDescriptor descriptor = new LocalizationTargetDescriptor(InObject.GetStringField("Name"), InObject.GetEnumField<LocalizationTargetDescriptorLoadingPolicy>("LoadingPolicy"));
+			LocalizationConfigGenerationPolicy policy;
+			if (InObject.TryGetEnumField<LocalizationConfigGenerationPolicy>("ConfigGenerationPolicy", out policy))
+			{
+				descriptor.ConfigGenerationPolicy = policy;
+			}
+			return descriptor;
 		}
 
 		/// <summary>
@@ -88,6 +132,7 @@ namespace UnrealBuildTool
 			Writer.WriteObjectStart();
 			Writer.WriteValue("Name", Name);
 			Writer.WriteValue("LoadingPolicy", LoadingPolicy.ToString());
+			Writer.WriteValue("ConfigGenerationPolicy", ConfigGenerationPolicy.ToString());
 			Writer.WriteObjectEnd();
 		}
 
@@ -96,7 +141,8 @@ namespace UnrealBuildTool
 			JsonObject localizationTargetObject= new JsonObject();
 			localizationTargetObject.AddOrSetFieldValue("Name", Name);
 			localizationTargetObject.AddOrSetFieldValue("LoadingPolicy", LoadingPolicy.ToString());
-			
+			localizationTargetObject.AddOrSetFieldValue("ConfigGenerationPolicy", ConfigGenerationPolicy.ToString());
+
 			return localizationTargetObject;
 		}
 
diff --git a/Engine/Source/Runtime/Projects/Private/LocalizationDescriptor.cpp b/Engine/Source/Runtime/Projects/Private/LocalizationDescriptor.cpp
--- a/Engine/Source/Runtime/Projects/Private/LocalizationDescriptor.cpp
+++ b/Engine/Source/Runtime/Projects/Private/LocalizationDescriptor.cpp
@@ -69,9 +69,43 @@ const TCHAR* ELocalizationTargetDescriptorLoadingPolicy::ToString(const ELocaliz
 	}
 }
 
-FLocalizationTargetDescriptor::FLocalizationTargetDescriptor(FString InName, ELocalizationTargetDescriptorLoadingPolicy::Type InLoadingPolicy)
+ELocalizationConfigGenerationPolicy::Type ELocalizationConfigGenerationPolicy::FromString(const TCHAR* String)
+{
+	ELocalizationConfigGenerationPolicy ::Type TestType = (ELocalizationConfigGenerationPolicy::Type)0;
+	for (; TestType < ELocalizationConfigGenerationPolicy::Max; TestType = (ELocalizationConfigGenerationPolicy::Type)(TestType + 1))
+	{
+		const TCHAR* TestString = ToString(TestType);
+		if (FCString::Stricmp(String, TestString) == 0)
+		{
+			break;
+		}
+	}
+	return TestType;
+}
+
+const TCHAR* ELocalizationConfigGenerationPolicy::ToString(const ELocalizationConfigGenerationPolicy::Type Value)
+{
+	switch (Value)
+	{
+	case Never:
+		return TEXT("Never");
+
+	case Auto:
+		return TEXT("Auto");
+
+	case User:
+		return TEXT("User");
+
+	default:
+		ensureMsgf(false, TEXT("ELocalizationTargetDescriptorGenerationPolicy ::ToString - Unrecognized ELocalizationTargetDescriptorGenerationPolicy value: %i"), Value);
+		return nullptr;
+	}
+}
+
+FLocalizationTargetDescriptor::FLocalizationTargetDescriptor(FString InName, ELocalizationTargetDescriptorLoadingPolicy::Type InLoadingPolicy, ELocalizationConfigGenerationPolicy::Type InGenerationPolicy)
 	: Name(MoveTemp(InName))
 	, LoadingPolicy(InLoadingPolicy)
+	, ConfigGenerationPolicy(InGenerationPolicy)
 {
 }
 
@@ -104,6 +138,22 @@ bool FLocalizationTargetDescriptor::Read(const FJsonObject& InObject, FText* Out
 		}
 	}
 
+		// Read the target generation policy
+	TSharedPtr<FJsonValue> ConfigGenerationPolicyValue = InObject.TryGetField(TEXT("ConfigGenerationPolicy"));
+	if (ConfigGenerationPolicyValue.IsValid() && ConfigGenerationPolicyValue->Type == EJson::String)
+	{
+		ConfigGenerationPolicy = ELocalizationConfigGenerationPolicy::FromString(*ConfigGenerationPolicyValue->AsString());
+		if (ConfigGenerationPolicy == ELocalizationConfigGenerationPolicy::Max)
+		{
+			if (OutFailReason)
+			{
+				*OutFailReason = FText::Format(LOCTEXT("TargetWithInvalidGenerationPolicy", "Localization Target entry '{0}' specified an unrecognized target GenerationPolicy'{1}'"), FText::FromString(Name), FText::FromString(ConfigGenerationPolicyValue->AsString()));
+			}
+			return false;
+		}
+	}
+	// If we can't find the value, GenerationPolicy already defaults to Never so it's ok 
+
 	return true;
 }
 
@@ -165,6 +215,7 @@ void FLocalizationTargetDescriptor::UpdateJson(FJsonObject& JsonObject) const
 {
 	JsonObject.SetStringField(TEXT("Name"), Name);
 	JsonObject.SetStringField(TEXT("LoadingPolicy"), FString(ELocalizationTargetDescriptorLoadingPolicy::ToString(LoadingPolicy)));
+	JsonObject.SetStringField(TEXT("ConfigGenerationPolicy"), FString(ELocalizationConfigGenerationPolicy::ToString(ConfigGenerationPolicy)));
 }
 
 void FLocalizationTargetDescriptor::WriteArray(TJsonWriter<>& Writer, const TCHAR* ArrayName, const TArray<FLocalizationTargetDescriptor>& Descriptors)
diff --git a/Engine/Source/Runtime/Projects/Public/LocalizationDescriptor.h b/Engine/Source/Runtime/Projects/Public/LocalizationDescriptor.h
--- a/Engine/Source/Runtime/Projects/Public/LocalizationDescriptor.h
+++ b/Engine/Source/Runtime/Projects/Public/LocalizationDescriptor.h
@@ -51,6 +51,39 @@ namespace ELocalizationTargetDescriptorLoadingPolicy
 	PROJECTS_API const TCHAR* ToString(const ELocalizationTargetDescriptorLoadingPolicy::Type Value);
 };
 
+/** How the localization target's localization config files are generated during the localization gather pipeline.*/
+namespace ELocalizationConfigGenerationPolicy
+{
+	enum Type
+	{
+		/** This localization target does not have localization config files associated with it and no localization content files will be generated for it during the localizaiton pipeline.*/
+		Never,
+		/** The user has provided localization config files for this localization target and they will be used to generate the localization content files for the localization target.*/
+		User,
+		/** Temporary localization config files will be generated for the localization target during the localization pipeline to generate the localization content files. After the localization pipeline is complete, the temporary files will be deleted.*/
+		Auto,
+		/** NOTE: If you add a new value, make sure to update the ToString() method below!. */
+		Max,
+	};
+
+	/**
+	 * Converts a string to a ELocalizationConfigGenerationPolicy::Type value
+	 *
+	 * @param	The string to convert to a value
+	 * @return	The corresponding value, or 'Max' if the string is not valid.
+	 */
+	PROJECTS_API ELocalizationConfigGenerationPolicy::Type FromString(const TCHAR* Text);
+
+	/**
+	 * Returns the name of a ELocalizationConfigGenerationPolicy::Type
+	 *
+	 * @param	The value to convert to a string
+	 * @return	The string representation of this enum value
+	 */
+	PROJECTS_API const TCHAR* ToString(const ELocalizationConfigGenerationPolicy::Type Value);
+};
+
+
 /**
  * Description of a localization target.
  */
@@ -62,8 +95,11 @@ struct FLocalizationTargetDescriptor
 	/** When should the localization data associated with a target should be loaded? */
 	ELocalizationTargetDescriptorLoadingPolicy::Type LoadingPolicy;
 
+	/** How the localizationc config files associated with the localization target are generated */
+	ELocalizationConfigGenerationPolicy::Type ConfigGenerationPolicy;
+
 	/** Normal constructor */
-	PROJECTS_API FLocalizationTargetDescriptor(FString InName = FString(), ELocalizationTargetDescriptorLoadingPolicy::Type InLoadingPolicy = ELocalizationTargetDescriptorLoadingPolicy::Never);
+	PROJECTS_API FLocalizationTargetDescriptor(FString InName = FString(), ELocalizationTargetDescriptorLoadingPolicy::Type InLoadingPolicy = ELocalizationTargetDescriptorLoadingPolicy::Never, ELocalizationConfigGenerationPolicy::Type InGenerationPolicy = ELocalizationConfigGenerationPolicy::Never);
 
 	/** Reads a descriptor from the given JSON object */
 	PROJECTS_API bool Read(const FJsonObject& InObject, FText* OutFailReason = nullptr);
