- Made the .binariesdir file only on Mac

diff --git a/Engine/Source/Programs/AutomationTool/Apple/ApplePlatform.Automation.cs b/Engine/Source/Programs/AutomationTool/Apple/ApplePlatform.Automation.cs
--- a/Engine/Source/Programs/AutomationTool/Apple/ApplePlatform.Automation.cs
+++ b/Engine/Source/Programs/AutomationTool/Apple/ApplePlatform.Automation.cs
@@ -106,7 +106,10 @@ public abstract class ApplePlatform : Platform
 		InternalUtils.SafeDeleteFile(BinaryPath.FullName, true);
 		// this may leave the binaries directory empty, but the Mac needs the Binaries/Mac dir to exist (see FMacPlatformProcess::BaseDir())
 		// so plop a file down in it's place
-		File.WriteAllText(Path.Combine(BinaryPath.Directory.FullName, ".binariesdir"), "");
+		if (SC.StageTargetPlatform.PlatformType == UnrealTargetPlatform.Mac)
+		{
+			File.WriteAllText(Path.Combine(BinaryPath.Directory.FullName, ".binariesdir"), "");
+		}
 
 		DirectoryReference AppPath = new DirectoryReference(BinaryPath.FullName + ".app");
 		InternalUtils.SafeDeleteDirectory(AppPath.FullName, true);
