Modifies buildtasks that were using HTTPClients so that they use a singleton http client.

Follow up from CL 30122407

diff --git a/Engine/Source/Programs/AutomationTool/BuildGraph/Tasks/DeployToolTask.cs b/Engine/Source/Programs/AutomationTool/BuildGraph/Tasks/DeployToolTask.cs
--- a/Engine/Source/Programs/AutomationTool/BuildGraph/Tasks/DeployToolTask.cs
+++ b/Engine/Source/Programs/AutomationTool/BuildGraph/Tasks/DeployToolTask.cs
@@ -1,6 +1,7 @@
 // Copyright Epic Games, Inc. All Rights Reserved.
 
 using AutomationTool;
+using AutomationUtils;
 using EpicGames.Core;
 using System;
 using System.Collections.Generic;
@@ -141,7 +142,7 @@ namespace AutomationTool.Tasks
 
 			HttpClient CreateHttpClient()
 			{
-				HttpClient httpClient = new HttpClient();
+				HttpClient httpClient = HttpClientSingleton<DeployToolTask>.Client;
 				httpClient.BaseAddress = serverUri;
 				if (settings?.Token != null)
 				{
diff --git a/Engine/Source/Programs/AutomationTool/BuildGraph/Tasks/NuGetLicenseCheckTask.cs b/Engine/Source/Programs/AutomationTool/BuildGraph/Tasks/NuGetLicenseCheckTask.cs
--- a/Engine/Source/Programs/AutomationTool/BuildGraph/Tasks/NuGetLicenseCheckTask.cs
+++ b/Engine/Source/Programs/AutomationTool/BuildGraph/Tasks/NuGetLicenseCheckTask.cs
@@ -1,6 +1,7 @@
 // Copyright Epic Games, Inc. All Rights Reserved.
 
 using AutomationTool;
+using AutomationUtils;
 using EpicGames.Core;
 using Microsoft.Extensions.Logging;
 using System;
@@ -274,7 +275,7 @@ namespace AutomationTool.Tasks
 
 							if (!LicenseUrlToInfo.TryGetValue(LicenseUrl, out Info.License))
 							{
-								using (HttpClient Client = new HttpClient())
+								using (HttpClient Client = HttpClientSingleton<NuGetLicenseCheckTask>.Client)
 								{
 									using HttpResponseMessage Response = await Client.GetAsync(LicenseUrl);
 									if (!Response.IsSuccessStatusCode)
