AutomationTool changes to support game projects overriding how to copy files during staging. InternalUtils.SafeCopyFile now takes an optional delegate to call when copying files. A CustomStageCopyHandler can be implemented to supply the copy operation. This is useful for transforming files during staging.

diff --git a/Engine/Source/Developer/DeveloperToolSettings/Classes/Settings/ProjectPackagingSettings.h b/Engine/Source/Developer/DeveloperToolSettings/Classes/Settings/ProjectPackagingSettings.h
--- a/Engine/Source/Developer/DeveloperToolSettings/Classes/Settings/ProjectPackagingSettings.h
+++ b/Engine/Source/Developer/DeveloperToolSettings/Classes/Settings/ProjectPackagingSettings.h
@@ -640,6 +640,12 @@ public:
 	UPROPERTY(config)
 	TArray<FProjectBuildSettings> EngineCustomBuilds;
 
+	/**
+	* The type name of a CustomStageCopyHandler subclass to instanciate during the copy build to staging directory step. See SetupCustomStageCopyHandler() in CopyBuildToStagingDirectory.Automation.cs and CustomStageCopyHandler.cs
+	*/
+	UPROPERTY(config)
+	FString CustomStageCopyHandler;
+
 private:
 	/** Helper array used to mirror Blueprint asset selections across edits */
 	TArray<FFilePath> CachedNativizeBlueprintAssets;
diff --git a/Engine/Source/Programs/AutomationTool/AutomationUtils/CommandUtils.cs b/Engine/Source/Programs/AutomationTool/AutomationUtils/CommandUtils.cs
--- a/Engine/Source/Programs/AutomationTool/AutomationUtils/CommandUtils.cs
+++ b/Engine/Source/Programs/AutomationTool/AutomationUtils/CommandUtils.cs
@@ -69,6 +69,12 @@ namespace AutomationTool
 		Local
 	}
 
+	/// <summary>
+	/// Delegate to override the copy operation.
+	/// <returns>Return true if the copy was handled. Otherwise false to fallback on the built in copy operation</returns>
+	/// </summary>
+	public delegate bool OverrideCopyDelegate(ILogger Logger, string SourceName, string TargetName);
+
 	/// <summary>
 	/// Base utility function for script commands.
 	/// </summary>
@@ -1415,7 +1421,7 @@ namespace AutomationTool
 		/// <param name="Dest">The full path to the destination file</param>
 		/// <param name="bAllowDifferingTimestamps">If true, will always skip a file if the destination exists, even if timestamp differs; defaults to false</param>
 		/// <returns>True if the operation was successful, false otherwise.</returns>
-		public static void CopyFileIncremental(FileReference Source, FileReference Dest, bool bAllowDifferingTimestamps = false, List<string> IniKeyDenyList = null, List<string> IniSectionDenyList = null)
+		public static void CopyFileIncremental(FileReference Source, FileReference Dest, OverrideCopyDelegate OverrideCopyHandler = null, bool bAllowDifferingTimestamps = false, List<string> IniKeyDenyList = null, List<string> IniSectionDenyList = null)
 		{
 			if (InternalUtils.SafeFileExists(Dest.FullName, true))
 			{
@@ -1443,7 +1449,7 @@ namespace AutomationTool
 			{
 				throw new AutomationException("Failed to delete {0} for copy", Dest);
 			}
-			if (!InternalUtils.SafeCopyFile(Source.FullName, Dest.FullName, IniKeyDenyList: IniKeyDenyList, IniSectionDenyList: IniSectionDenyList))
+			if (!InternalUtils.SafeCopyFile(Source.FullName, Dest.FullName, OverrideCopyHandler: OverrideCopyHandler, IniKeyDenyList: IniKeyDenyList, IniSectionDenyList: IniSectionDenyList))
 			{
 				throw new AutomationException("Failed to copy {0} to {1}", Source, Dest);
 			}
diff --git a/Engine/Source/Programs/AutomationTool/AutomationUtils/CustomStageCopyHandler.cs b/Engine/Source/Programs/AutomationTool/AutomationUtils/CustomStageCopyHandler.cs
new file mode 100644
--- /dev/null
+++ b/Engine/Source/Programs/AutomationTool/AutomationUtils/CustomStageCopyHandler.cs
@@ -0,0 +1,62 @@
+// Copyright Epic Games, Inc. All Rights Reserved.
+
+using AutomationTool;
+using System;
+using System.Collections.Generic;
+using System.Linq;
+using System.Reflection;
+using UnrealBuildTool;
+using Microsoft.Extensions.Logging;
+
+public abstract class CustomStageCopyHandler
+{
+	public virtual bool StageFile(ILogger Logger, string SourceName, string TargetName) { return false; }
+
+	public static CustomStageCopyHandler Create(string Name)
+	{
+		// one-time static init. NB. Can't use the partial struct convention because automation classes are all in different assemblies
+		if (HandlerRegistry == null)
+		{
+			HandlerRegistry = new();
+
+			foreach (Assembly LoadedAssembly in ScriptManager.AllScriptAssemblies)
+			{
+				IEnumerable<Type> HandlerTypes = LoadedAssembly.GetTypes().Where(X => !X.IsAbstract && X.IsSubclassOf(typeof(CustomStageCopyHandler)) && X.GetCustomAttribute(typeof(CustomStageCopyHandlerAttribute)) != null);
+				foreach (Type HandlerType in HandlerTypes)
+				{
+					string HandlerName = HandlerType.GetCustomAttribute<CustomStageCopyHandlerAttribute>().Name;
+					HandlerRegistry.Add(HandlerName, HandlerType);
+				}
+			}
+		}
+
+		Type SelectedType;
+		if (!HandlerRegistry.TryGetValue(Name, out SelectedType) || SelectedType == null)
+		{
+			throw new BuildException($"Unknown custom stage copy handler {Name}");
+		}
+
+		CustomStageCopyHandler Handler = (CustomStageCopyHandler)Activator.CreateInstance(SelectedType);
+		if (Handler == null)
+		{
+			throw new BuildException($"Could not instantiate the custom stage copy handler {Name}");
+		}
+
+		return Handler;
+	}
+
+	private static Dictionary<string, Type> HandlerRegistry = null;
+}
+
+/// <summary>
+/// CustomStageCopyHandler class declarations must be tagged with this, used to find the handler by name
+/// </summary>
+[System.AttributeUsage(System.AttributeTargets.Class)]
+public class CustomStageCopyHandlerAttribute : System.Attribute
+{
+	public string Name;
+	public CustomStageCopyHandlerAttribute(string Name)
+	{
+		this.Name = Name;
+	}
+}
diff --git a/Engine/Source/Programs/AutomationTool/AutomationUtils/DeploymentContext.cs b/Engine/Source/Programs/AutomationTool/AutomationUtils/DeploymentContext.cs
--- a/Engine/Source/Programs/AutomationTool/AutomationUtils/DeploymentContext.cs
+++ b/Engine/Source/Programs/AutomationTool/AutomationUtils/DeploymentContext.cs
@@ -313,6 +313,11 @@ public class DeploymentContext //: ProjectParams
 	/// </summary>
 	public HashSet<StagedFileReference> ExtraFilesAllowList = new HashSet<StagedFileReference>();
 
+	/// <summary>
+	/// Optional copy handler that during CopyOrWriteManifestFilesToStageDir will be passed to SafeCopy to handle
+	/// the copy operation of files.
+	/// </summary>
+	public OverrideCopyDelegate OverrideCopyHandler = null;
 
 	/// <summary>
 	/// List of ini keys to strip when staging
diff --git a/Engine/Source/Programs/AutomationTool/AutomationUtils/Utils.cs b/Engine/Source/Programs/AutomationTool/AutomationUtils/Utils.cs
--- a/Engine/Source/Programs/AutomationTool/AutomationUtils/Utils.cs
+++ b/Engine/Source/Programs/AutomationTool/AutomationUtils/Utils.cs
@@ -437,7 +437,7 @@ namespace AutomationTool
 		/// <param name="SourceName">Source name</param>
 		/// <param name="TargetName">Target name</param>
 		/// <returns>True if the operation was successful, false otherwise.</returns>
-		public static bool SafeCopyFile(string SourceName, string TargetName, bool bQuiet = false, List<string> IniKeyDenyList = null, List<string> IniSectionDenyList = null, List<string> IniSectionAllowList = null, bool bSafeCreateDirectory = false)
+		public static bool SafeCopyFile(string SourceName, string TargetName, bool bQuiet = false, OverrideCopyDelegate OverrideCopyHandler = null, List<string> IniKeyDenyList = null, List<string> IniSectionDenyList = null, List<string> IniSectionAllowList = null, bool bSafeCreateDirectory = false)
 		{
 			if (!bQuiet)
 			{
@@ -468,6 +468,11 @@ namespace AutomationTool
 						// ini files may change size, don't check
 						bSkipSizeCheck = true;
 					}
+					else if (OverrideCopyHandler != null && OverrideCopyHandler.Invoke(Logger, SourceName, TargetName))
+					{
+						// handlers may change the file sizes, don't check
+						bSkipSizeCheck = true;
+					}
 					else
 					{
 						if(File.Exists(SourceName))
diff --git a/Engine/Source/Programs/AutomationTool/Scripts/CopyBuildToStagingDirectory.Automation.cs b/Engine/Source/Programs/AutomationTool/Scripts/CopyBuildToStagingDirectory.Automation.cs
--- a/Engine/Source/Programs/AutomationTool/Scripts/CopyBuildToStagingDirectory.Automation.cs
+++ b/Engine/Source/Programs/AutomationTool/Scripts/CopyBuildToStagingDirectory.Automation.cs
@@ -2145,7 +2145,7 @@ namespace AutomationScripts
 					FileReference Dest = FileReference.Combine(StageDirectory, Pair.Key.Name);
 					if (Src != Dest)  // special case for things created in the staging directory, like the pak file
 					{
-						CopyFileIncremental(Src, Dest, IniKeyDenyList: SC.IniKeyDenyList, IniSectionDenyList: SC.IniSectionDenyList);
+						CopyFileIncremental(Src, Dest, SC.OverrideCopyHandler, IniKeyDenyList: SC.IniKeyDenyList, IniSectionDenyList: SC.IniSectionDenyList);
 					}
 				}
 			}
@@ -5223,10 +5223,32 @@ namespace AutomationScripts
 				{
 					DeploymentContexts.Insert(0, SC);
 				}
+
+				SetupCustomStageCopyHandler(Params, SC);
 			}
 
 			return DeploymentContexts;
 		}
+		private static void SetupCustomStageCopyHandler(ProjectParams Params, DeploymentContext SC)
+		{
+			string CustomStageCopyHandlerName = null;
+			ConfigHierarchy GameIni = ConfigCache.ReadHierarchy(ConfigHierarchyType.Game, Params.RawProjectPath.Directory, SC.StageTargetPlatform.PlatformType, SC.CustomConfig);
+			GameIni.GetString("/Script/UnrealEd.ProjectPackagingSettings", "CustomStageCopyHandler", out CustomStageCopyHandlerName);
+
+			if (!string.IsNullOrEmpty(CustomStageCopyHandlerName))
+			{
+				CustomStageCopyHandler CustomStageCopyHandler = CustomStageCopyHandler.Create(CustomStageCopyHandlerName);
+				if (CustomStageCopyHandler != null)
+				{
+					Logger.LogInformation("CopyBuildToStagingDirectory using CustomStageCopyHandler {0}", CustomStageCopyHandlerName);
+
+					SC.OverrideCopyHandler = (ILogger Logger, string SourceName, string TargetName) =>
+					{
+						return CustomStageCopyHandler.StageFile(Logger, SourceName, TargetName);
+					};
+				}
+			}
+		}
 
 		public static void CopyBuildToStagingDirectory(ProjectParams Params)
 		{
