Gauntlet LLT: update regex that parses console report to include potentially skipped tests
#rnx

diff --git a/Engine/Source/Programs/AutomationTool/LowLevelTests/Utility/LowLevelTests.LogParser.cs b/Engine/Source/Programs/AutomationTool/LowLevelTests/Utility/LowLevelTests.LogParser.cs
--- a/Engine/Source/Programs/AutomationTool/LowLevelTests/Utility/LowLevelTests.LogParser.cs
+++ b/Engine/Source/Programs/AutomationTool/LowLevelTests/Utility/LowLevelTests.LogParser.cs
@@ -1,6 +1,8 @@
 // Copyright Epic Games, Inc. All Rights Reserved.
 
 using Gauntlet;
+using System;
+using System.Linq;
 using System.Text.RegularExpressions;
 
 namespace LowLevelTests
@@ -25,12 +27,14 @@ namespace LowLevelTests
 			public string TotalCases;
 			public string Passed;
 			public string Failed;
+			public string Skipped;
 
 			public CatchTestCaseResult()
 			{
 				TotalCases = "";
 				Passed = "";
 				Failed = "";
+				Skipped = "";
 			}
 		}
 		public class CatchAssertionResult
@@ -51,7 +55,11 @@ namespace LowLevelTests
 		{
 			internal const string AllCatchTestPassed = @"All tests passed\s\((?<assertions>[\w\d]+)\sassertion[s]?\sin\s(?<testcases>[\w\d]+)\stest\scase[s]?.*\)";
 
-			internal const string CatchTestCaseResults = @"test\scase[s]?\:\s+(?<totalcases>[\d]+)\s+\|\s+(?<passed>[\d]+)\s+passed\s+\|\s+(?<failed>[\d]+)\s+failed";
+			internal const string CatchTestCaseResults = @"test\scase[s]?\:\s+(?<totalcases>[\d]+)\s+";
+			internal const string CatchPassedCases = @"(?<passed>[\d]+)\s+passed";
+			internal const string CatchFailedCases = @"(?<failed>[\d]+)\s+failed";
+			internal const string CatchSkippedCases = @"(?<skipped>[\d]+)\s+skipped";
+
 
 			internal const string CatchTestAssertionResults = @"assertion[s]?\:\s+(?<totalcases>[\d]+)\s+\|\s+(?<passed>[\d]+)\s+passed\s+\|\s+(?<failed>[\d]+)\s+failed";
 
@@ -94,22 +102,36 @@ namespace LowLevelTests
 				}
 			}
 
-			// If we did not find any matching test passed string from catch, check to see test results of pass/fail were logged.
-			MatchCollection rCollection = Regex.Matches(Content, ParsePatterns.CatchTestCaseResults);
+			// If we did not find any matching test passed string from catch, check to see test results of pass/fail/skip were logged.
+			Match testResultsMatch = Regex.Matches(Content, ParsePatterns.CatchTestCaseResults).FirstOrDefault();
 
-			foreach (Match aMatch in rCollection)
+			if (testResultsMatch != null && testResultsMatch.Success)
 			{
-				if (aMatch.Success)
+				CatchTestCaseResult TestCaseResult = new CatchTestCaseResult()
 				{
-					TestResults.Passed = false;
-					TestResults.TestCases = new CatchTestCaseResult()
-					{
+					TotalCases = testResultsMatch.Groups["totalcases"].ToString(),
+				};
+				
+				Match testPassedMatch = Regex.Matches(Content, ParsePatterns.CatchPassedCases).FirstOrDefault();
+				if (testPassedMatch != null && testPassedMatch.Success)
+				{
+					TestCaseResult.Passed = testPassedMatch.Groups["passed"].ToString();
+				}
 
-						TotalCases = aMatch.Groups["totalcases"].ToString(),
-						Passed = aMatch.Groups["passed"].ToString(),
-						Failed = aMatch.Groups["failed"].ToString()
-					};
+				Match testFailedMatch = Regex.Matches(Content, ParsePatterns.CatchFailedCases).FirstOrDefault();
+				if (testFailedMatch != null && testFailedMatch.Success)
+				{
+					TestCaseResult.Failed = testFailedMatch.Groups["failed"].ToString();
 				}
+
+				Match testSkippedMatch = Regex.Matches(Content, ParsePatterns.CatchSkippedCases).FirstOrDefault();
+				if (testSkippedMatch != null && testSkippedMatch.Success)
+				{
+					TestCaseResult.Skipped = testSkippedMatch.Groups["skipped"].ToString();
+				}
+
+				TestResults.Passed = string.IsNullOrEmpty(TestCaseResult.Failed) || Convert.ToInt32(TestCaseResult.Failed) == 0;
+				TestResults.TestCases = TestCaseResult;
 			}
 
 			MatchCollection aCollection = Regex.Matches(Content, ParsePatterns.CatchTestAssertionResults);
