Remove some unacceptable words

#rnx

diff --git a/Engine/Source/Programs/AutomationTool/Scripts/BenchmarkBuild/BenchmarkBuild.Automation.cs b/Engine/Source/Programs/AutomationTool/Scripts/BenchmarkBuild/BenchmarkBuild.Automation.cs
--- a/Engine/Source/Programs/AutomationTool/Scripts/BenchmarkBuild/BenchmarkBuild.Automation.cs
+++ b/Engine/Source/Programs/AutomationTool/Scripts/BenchmarkBuild/BenchmarkBuild.Automation.cs
@@ -325,7 +325,7 @@ namespace AutomationTool.Benchmark
 
 				// parse map args
 				{
-					// master arg that sets all three
+					// primary arg that sets all three
 					var EditorMaps = ParseMapList("editor-maps");
 
 					if (EditorMaps.Any())
diff --git a/Engine/Source/Programs/AutomationTool/Scripts/Localisation.Automation.cs b/Engine/Source/Programs/AutomationTool/Scripts/Localisation.Automation.cs
--- a/Engine/Source/Programs/AutomationTool/Scripts/Localisation.Automation.cs
+++ b/Engine/Source/Programs/AutomationTool/Scripts/Localisation.Automation.cs
@@ -306,7 +306,7 @@ class Localize : BuildCommand
 			AdditionalCommandletArguments = AdditionalCommandletArguments.Trim();
 			if (AdditionalCommandletArguments.StartsWith("\"") && AdditionalCommandletArguments.EndsWith("\""))
 			{
-				// We subtract 2 to nuke the last " character
+				// We subtract 2 to remove the last " character
 				AdditionalCommandletArguments = AdditionalCommandletArguments[1..^1];
 			}
 		}
