Localization:
Optimizationof Smartling download and upload.
- Previously, downloading of .po files from Smartling are done in series. Now, each language is downloaded in its own Task. Parallelizing downloads have improved download times anywhere from 2x to 10x.
- Logs for downloads, uploads and authentication token retrieval has been changed to be thread safe. The logs for a single Task will be batched and only printed after the task has completed.
- API for GetAuthenticationToken(), RequestAuthenticationToken() and RetrieveAuthenticationToken() has been changed to facilitate multithreaded logging
- Uploading of .po files has also been changed to be threaded to conform to the changes in APIs for the authentication token retrieval functions.
- Switched from Console.WriteLine to the preferred Logger.LogInformation

#test: Performed the localization downloads locally. Checked the logs for the download step and compared timings of localization downloads. All tested localization files seems to have been downloaded successfully.

diff --git a/Engine/Source/Programs/AutomationTool/SmartlingLocalization/SmartlingLocalizationProvider.cs b/Engine/Source/Programs/AutomationTool/SmartlingLocalization/SmartlingLocalizationProvider.cs
--- a/Engine/Source/Programs/AutomationTool/SmartlingLocalization/SmartlingLocalizationProvider.cs
+++ b/Engine/Source/Programs/AutomationTool/SmartlingLocalization/SmartlingLocalizationProvider.cs
@@ -16,6 +16,8 @@ using EpicGames.Localization;
 using System.Threading;
 using Microsoft.Extensions.Logging;
 
+using static AutomationTool.CommandUtils;
+using System.Linq;
 
 #pragma warning disable SYSLIB0014
 
@@ -72,7 +74,7 @@ namespace EpicGames.SmartlingLocalization
 		{
 			if (String.IsNullOrEmpty(LocalizationBranchName))
 			{
-				Console.WriteLine("LocalizationBranchName is null or empty. The branch name is used to create the correct file URIs for Smartling. Wrong files may be downloaded from Smartling and files may be clobbered in Smartling. Please append a value to LocalizationName as a command line argument to fix this.");
+				Logger.LogWarning("LocalizationBranchName is null or empty. The branch name is used to create the correct file URIs for Smartling. Wrong files may be downloaded from Smartling and files may be clobbered in Smartling. Please append a value to LocalizationName as a command line argument to fix this.");
 			}
 			Config = new SmartlingConfig();
 			Client = new HttpClient();
@@ -89,8 +91,9 @@ namespace EpicGames.SmartlingLocalization
 
 		public async override Task DownloadProjectFromLocalizationProvider(string ProjectName, ProjectImportExportInfo ProjectImportInfo)
 		{
-			Console.WriteLine($"Starting Smartling download for {ProjectName} project files.");
+			Logger.LogInformation($"Starting Smartling download for {ProjectName} project files.");
 			Stopwatch Watch = Stopwatch.StartNew();
+			List<Task> DownloadTasks = new List<Task>();
 			// Get the latest files for each culture.
 			foreach (var Culture in ProjectImportInfo.CulturesToGenerate)
 			{
@@ -100,20 +103,33 @@ namespace EpicGames.SmartlingLocalization
 					continue;
 				}
 
-				await DownloadLatestPOFile(Culture, null, ProjectImportInfo);
+				DownloadTasks.Add(DownloadLatestPOFileAndLog(Culture, null, ProjectImportInfo));
 				foreach (var Platform in ProjectImportInfo.SplitPlatformNames)
 				{
-					await DownloadLatestPOFile(Culture, Platform, ProjectImportInfo);
+					DownloadTasks.Add(DownloadLatestPOFileAndLog(Culture, Platform, ProjectImportInfo));
 				}
 			}
+			await Task.WhenAll(DownloadTasks);
 			Watch.Stop();
-			Console.WriteLine($"Completed Smartling download for {ProjectName} project files in {Watch.ElapsedMilliseconds / 1000} seconds.");
+			Logger.LogInformation($"Completed Smartling download for {ProjectName} project files in {Watch.ElapsedMilliseconds / 1000} seconds.");
 		}
 
-		private async Task DownloadLatestPOFile(string EpicLocale, string Platform, ProjectImportExportInfo ProjectImportInfo)
+		private async Task DownloadLatestPOFileAndLog(string EpicLocale, string Platform, ProjectImportExportInfo ProjectImportInfo)
 		{
-			await GetAuthenticationToken();
-			
+			List<string> Logs = await DownloadLatestPOFile(EpicLocale, Platform, ProjectImportInfo);
+			lock (LoggerLock)
+			{
+				foreach (string Log in Logs)
+				{
+					Logger.LogInformation(Log);
+				}
+			}
+		}
+
+		private async Task<List<string>> DownloadLatestPOFile(string EpicLocale, string Platform, ProjectImportExportInfo ProjectImportInfo)
+		{
+			List<string> Logs = new List<string>();
+			await GetAuthenticationToken(Logs);
 
 			var DestinationDirectory = String.IsNullOrEmpty(Platform)
 			? new DirectoryInfo(CommandUtils.CombinePaths(RootWorkingDirectory, ProjectImportInfo.DestinationPath))
@@ -124,7 +140,7 @@ namespace EpicGames.SmartlingLocalization
 			// note that the Smartling Filename and Uri may be different 
 			var SmartlingFilename = GetSmartlingFilename(ProjectImportInfo.PortableObjectName, Platform);
 
-			Console.WriteLine($"Exporting: '{SmartlingFilename}' as '{ExportFile.FullName}' ({EpicLocale})");
+			Logs.Add($"Downloading: '{SmartlingFilename}' as '{ExportFile.FullName}' ({EpicLocale})");
 
 			string SmartlingLocale = ConvertEpicLocaleToSmartlingLocale(EpicLocale);
 			string DownloadEndpoint = $"https://api.smartling.com/files-api/v2/projects/{Config.ProjectId}/locales/{SmartlingLocale}/file";
@@ -148,7 +164,7 @@ namespace EpicGames.SmartlingLocalization
 			while (true)
 			{
 				CurrentTimeOut = CurrentTries * InitialTimeOut;
-				Console.WriteLine($"Current time out for download response {CurrentTimeOut}s.");
+				Logs.Add($"Current time out for download response {CurrentTimeOut}s.");
 				TimeSpan Timeout = TimeSpan.FromSeconds(CurrentTimeOut);
 				using var CancellationToken = new CancellationTokenSource(Timeout);
 				try
@@ -161,11 +177,11 @@ namespace EpicGames.SmartlingLocalization
 						if (CurrentTries > MaxTries)
 						{
 							DownloadStopWatch.Stop();
-							Console.WriteLine($"[FAILED] Exporting: '{ExportFile.FullName}' ({EpicLocale}) exhausted all retries after {DownloadStopWatch.ElapsedMilliseconds / 1000}.");
-							return;
+							Logs.Add($"[FAILED] Downloading: '{ExportFile.FullName}' ({EpicLocale}) exhausted all retries after {DownloadStopWatch.ElapsedMilliseconds / 1000}.");
+							return Logs;
 						}
-						Console.WriteLine($"Encountered HTTP Status Code 401. Authentication most likely expired. Retrying {CurrentTries}/{MaxTries} times with refreshed authentication token.");
-						await GetAuthenticationToken();
+						Logs.Add($"Encountered HTTP Status Code 401. Authentication most likely expired. Retrying {CurrentTries}/{MaxTries} times with refreshed authentication token.");
+						await GetAuthenticationToken(Logs);
 						continue;
 					}
 					break;
@@ -176,12 +192,12 @@ namespace EpicGames.SmartlingLocalization
 					if (CurrentTries > MaxTries)
 					{
 						DownloadStopWatch.Stop();
-						Console.WriteLine($"[FAILED] Exporting: '{ExportFile.FullName}' ({EpicLocale}) exhausted all retries in {DownloadStopWatch.ElapsedMilliseconds / 1000} seconds. - {Ex}");
-						return;
+						Logs.Add($"[FAILED] Downloading: '{ExportFile.FullName}' ({EpicLocale}) exhausted all retries in {DownloadStopWatch.ElapsedMilliseconds / 1000} seconds. - {Ex}");
+						return Logs;
 					}
-					Console.WriteLine($"Failed to get download response. Retrying {CurrentTries}/{MaxTries} times.");
+					Logs.Add($"Failed to get download response. Retrying {CurrentTries}/{MaxTries} times.");
 					// We need to retreive the authentication token again because after each timeout, we may exceed the validity of the authentication token 
-					await GetAuthenticationToken();
+					await GetAuthenticationToken(Logs);
 				}
 			}
 
@@ -207,7 +223,7 @@ namespace EpicGames.SmartlingLocalization
 					}
 				}
 				DownloadStopWatch.Stop();
-				Console.WriteLine($"[SUCCESS] Exporting: '{SmartlingFileUri}' as '{ExportFile.FullName}' ({EpicLocale}) in {DownloadStopWatch.ElapsedMilliseconds / 1000} seconds.");
+				Logs.Add($"[SUCCESS] Downloading: '{SmartlingFileUri}' as '{ExportFile.FullName}' ({EpicLocale}) in {DownloadStopWatch.ElapsedMilliseconds / 1000} seconds.");
 				// Reset the write status of the file
 				if (ExportFileWasReadOnly)
 				{
@@ -221,7 +237,7 @@ namespace EpicGames.SmartlingLocalization
 				if (bCreateBackupCopy)
 				{
 					string ExportFileCopyPath = Path.Combine(ExportFile.DirectoryName, $"{Path.GetFileNameWithoutExtension(ExportFile.Name)}_FromSmartling{ExportFile.Extension}");
-					Console.WriteLine($"Updating Smartling copy '{ExportFileCopyPath}'");
+					Logs.Add($"Updating Smartling copy '{ExportFileCopyPath}'");
 					var ExportFileCopy = new FileInfo(ExportFileCopyPath);
 
 					var ExportFileCopyWasReadOnly = false;
@@ -245,38 +261,60 @@ namespace EpicGames.SmartlingLocalization
 			{
 				// The file may not currently exist in Smartling and will need to be uploaded first via the Upload step later on. 
 				DownloadStopWatch.Stop();
-				Console.WriteLine($"[FAILED] Exporting: '{ExportFile.FullName}' ({EpicLocale} in {DownloadStopWatch.ElapsedMilliseconds / 1000} seconds. The file may need to be uploaded first.)");
-				await PrintRequestErrors(DownloadResponse);
+				Logs.Add($"[FAILED] Downloading: '{ExportFile.FullName}' ({EpicLocale} in {DownloadStopWatch.ElapsedMilliseconds / 1000} seconds. The file may need to be uploaded first.)");
+				await AppendRequestErrorsToLog(DownloadResponse, Logs);
 			}
+			return Logs;
 		}
-		
 
-		private async Task PrintRequestErrors(HttpResponseMessage Response)
+		private async Task AppendRequestErrorsToLog(HttpResponseMessage Response, List<string> Logs)
 		{
 			string ResponseString = await Response.Content.ReadAsStringAsync();
 			var ResponseEnvelope = JsonSerializer.Deserialize<SmartlingResponseEnvelope<SmartlingRequestErrorsEnvelope>>(ResponseString, JsonOptions);
 			var RequestErrors = ResponseEnvelope.Response.Errors;
 			foreach (SmartlingRequestError RequestError in RequestErrors)
 			{
-				Console.WriteLine($"Smartling Warning:\n {RequestError.ToString()}");
+				Logs.Add($"Smartling Warning:\n {RequestError.ToString()}");
 			}
+			
 		}
 
 		public async override Task UploadProjectToLocalizationProvider(string ProjectName, ProjectImportExportInfo ProjectExportInfo)
 		{
+			Logger.LogInformation($"Starting Smartling upload for {ProjectName} project files.");
+			Stopwatch Watch = Stopwatch.StartNew();
+			List<Task> UploadTasks = new List<Task>();
 			// Upload the .po file for the native culture first
-			await UploadLatestPOFile(ProjectExportInfo.NativeCulture, null, ProjectExportInfo);
+			UploadTasks.Add(UploadLatestPOFileAndLog(ProjectExportInfo.NativeCulture, null, ProjectExportInfo));
+
 			foreach (var Platform in ProjectExportInfo.SplitPlatformNames)
 			{
-				await UploadLatestPOFile(ProjectExportInfo.NativeCulture, Platform, ProjectExportInfo);
+				UploadTasks.Add(UploadLatestPOFileAndLog(ProjectExportInfo.NativeCulture, Platform, ProjectExportInfo));
 			}
+			await Task.WhenAll(UploadTasks);
 
 			// Uploading all cultures is a legacy behavior from Onesky that isn't currently needed within our workflow. We can support this if there is a need. 
+			Watch.Stop();
+			Logger.LogInformation($"Completed Smartling upload for {ProjectName} project files in {Watch.ElapsedMilliseconds / 1000} seconds.");
+		}
+
+		private async Task UploadLatestPOFileAndLog(string EpicLocale, string Platform, ProjectImportExportInfo ProjectExportInfo)
+		{
+			List<string> Logs = await UploadLatestPOFile(EpicLocale, Platform, ProjectExportInfo);
+			lock (LoggerLock)
+			{
+				foreach (string Log in Logs)
+				{
+					Logger.LogInformation(Log);
+				}
+			}
+
 		}
 
-		private async Task UploadLatestPOFile(string EpicLocale, string Platform, ProjectImportExportInfo ProjectExportInfo)
+		private async Task<List<string>> UploadLatestPOFile(string EpicLocale, string Platform, ProjectImportExportInfo ProjectExportInfo)
 		{
-			await GetAuthenticationToken();
+			List<string> Logs = new List<string>();
+			await GetAuthenticationToken(Logs);
 
 			var SourceDirectory = String.IsNullOrEmpty(Platform)
 				? new DirectoryInfo(CommandUtils.CombinePaths(RootWorkingDirectory, ProjectExportInfo.DestinationPath))
@@ -286,14 +324,14 @@ namespace EpicGames.SmartlingLocalization
 			var FileToUpload = new FileInfo(FileToUploadPath);
 			if (!FileToUpload.Exists)
 			{
-				Console.WriteLine($"Unable to upload '{FileToUploadPath}'. File does not exist.");
-				return;
+				Logs.Add($"Unable to upload '{FileToUploadPath}'. File does not exist.");
+				return Logs;
 			}
 			string SmartlingFilename = GetSmartlingFilename(ProjectExportInfo.PortableObjectName, Platform);
 			string SmartlingFileUri = GetSmartlingFileUri(SmartlingFilename);
 			bool bIsNative = EpicLocale == ProjectExportInfo.NativeCulture;
 
-			Console.WriteLine($"Uploading: '{FileToUpload.FullName}' as '{SmartlingFileUri}' ({EpicLocale})");
+			Logs.Add($"Uploading: '{FileToUpload.FullName}' as '{SmartlingFileUri}' ({EpicLocale})");
 			// For now we only upload the file in the native locale
 			// It is dissuaded to import translations to Smartling. Perform all translations in the Smartling dashboard instead.
 
@@ -334,30 +372,31 @@ namespace EpicGames.SmartlingLocalization
 						++CurrentTries;
 						if (CurrentTries > MaxTries)
 						{
-							Console.WriteLine($"[FAILED] Uploading: '{FileToUpload.FullName}' ({EpicLocale}). Exhausted all retries.");
-							return;
+							Logs.Add($"[FAILED] Uploading: '{FileToUpload.FullName}' ({EpicLocale}). Exhausted all retries.");
+							return Logs;
 						}
-						Console.WriteLine("Encountered HTTP Status Code 401. Authentication token most likely expired. Retrying {CurrentTries}/{MaxTries} times with refreshed authentication token.");
+						Logs.Add($"Encountered HTTP Status Code 401. Authentication token most likely expired. Retrying {CurrentTries}/{MaxTries} times with refreshed authentication token.");
 						continue;
 					}
 					break;
 				}
 				catch (Exception Ex)
 				{
-					Console.WriteLine($"[FAILED] Uploading: '{FileToUpload.FullName}' ({EpicLocale}) - {Ex}");
-					return;
+					Logs.Add($"[FAILED] Uploading: '{FileToUpload.FullName}' ({EpicLocale}) - {Ex}");
+					return Logs;
 				}
 			}
 
 			if (UploadResponse.IsSuccessStatusCode)
 			{
-				Console.WriteLine($"[SUCCESS] Uploading: '{FileToUpload.FullName}' ({EpicLocale})");
+				Logs.Add($"[SUCCESS] Uploading: '{FileToUpload.FullName}' ({EpicLocale})");
 			}
 			else
 			{
-				Console.WriteLine($"[FAILED] Uploading: '{FileToUpload.FullName}' ({EpicLocale})");
-				await PrintRequestErrors(UploadResponse);
+				Logs.Add($"[FAILED] Uploading: '{FileToUpload.FullName}' ({EpicLocale})");
+				await AppendRequestErrorsToLog(UploadResponse, Logs);
 			}
+			return Logs;
 		}
 
 		// Override in child classes if there is a different set of languages that you need translations for 
@@ -417,37 +456,54 @@ namespace EpicGames.SmartlingLocalization
 			return SmartlingFilename;
 		}
 
-		private async Task GetAuthenticationToken()
+		private async Task GetAuthenticationToken(List<string> Logs)
 		{
 			// If the token is still valid, we just return. Otherwise we either refresh or authenticate again.
 
-			// If the last successful update + time to expire is still less than our current time accounting for some slack, we don't need to do anything to the token.
-			// We account for slack as it makes little sense to try authenticating with only 1s of validity left on a busy network
-			int Slack = 30; 
-			if (AuthenticationToken != null && AuthenticationToken.LastSuccessfulUpdateTime.AddSeconds(AuthenticationToken.ExpiresIn - Slack) > DateTime.UtcNow)
+		
+			// Any retreival or modification of the authentication token must be protected by this semaphore
+			// Calls to RequestAuthenticationToken() and RefreshAuthenticationToken() are also protected by this semaphore.
+			//  Do not call RequestAuthenticationToken() or RefreshAuthenticationToken() by themselves as they are not thread safe.
+			await AuthenticationTokenSemaphore.WaitAsync();
+			try
 			{
-				Console.WriteLine("Authentication token still valid. Using current authentication token.");
-				return;
-			}
+				// If the last successful update + time to expire is still less than our current time accounting for some slack, we don't need to do anything to the token.
+				// We account for slack as it makes little sense to try authenticating with only 1s of validity left on a busy network
+				int Slack = 30;
+				if (AuthenticationToken != null && AuthenticationToken.LastSuccessfulUpdateTime.AddSeconds(AuthenticationToken.ExpiresIn - Slack) > DateTime.UtcNow)
+				{
+					Logs.Add("Authentication token still valid. Using current authentication token.");
+					return;
+				}
 
-			if (AuthenticationToken != null && AuthenticationToken.LastSuccessfulUpdateTime.AddSeconds(AuthenticationToken.RefreshExpiresIn) > DateTime.UtcNow)
-			{
-				try
+				if (AuthenticationToken != null && AuthenticationToken.LastSuccessfulUpdateTime.AddSeconds(AuthenticationToken.RefreshExpiresIn) > DateTime.UtcNow)
 				{
-					await RefreshAuthenticationToken();
+					try
+					{
+						await RefreshAuthenticationToken(Logs);
+					}
+					catch (Exception)
+					{
+						await RequestAuthenticationToken(Logs);
+					}
 				}
-				catch (Exception)
+				else
 				{
-					await RequestAuthenticationToken();
+					await RequestAuthenticationToken(Logs);
 				}
 			}
-			else
+			finally
 			{
-				await RequestAuthenticationToken();
+				AuthenticationTokenSemaphore.Release();
 			}
 		}
 
-		private async Task RequestAuthenticationToken()
+		/// <summary>
+		/// Requests a new authentication token from the Smartling endpoints. This call is not threadsafe and should ONLY be used from GetAuthenticationToken() where the AuthenticationToken instane is protected by syncrhonization primitives. Do not call this function directly.
+		/// </summary>
+		/// <param name="Logs"></param>
+		/// <returns></returns>
+		private async Task RequestAuthenticationToken(List<string> Logs)
 		{
 			string AuthenticateEndpoint = "https://api.smartling.com/auth-api/v2/authenticate";
 			// serialize from Dictionary as MultipartFormDataContent not supported as a content-type 
@@ -468,21 +524,26 @@ namespace EpicGames.SmartlingLocalization
 					AuthenticationToken = AuthenticateResponseEnvelope.Response.Data;
 					AuthenticationToken.LastSuccessfulUpdateTime = DateTime.UtcNow;
 					Client.DefaultRequestHeaders.Authorization = new AuthenticationHeaderValue(AuthenticationToken.TokenType, AuthenticationToken.AccessToken);
-					Console.WriteLine("Successfully retrieved authentication token!");
+					Logs.Add("Successfully retrieved authentication token!");
 				}
 				else
 				{
-					Console.WriteLine("Failed to retrieve authentication token.");
-					await PrintRequestErrors(AuthenticateResponse);
+					Logs.Add("Failed to retrieve authentication token.");
+					await AppendRequestErrorsToLog(AuthenticateResponse, Logs);
 				}
 			}
 			catch (Exception Ex)
 			{
-				Console.WriteLine($"Failed to retrieve authentication token. {Ex}");
+				Logs.Add($"Failed to retrieve authentication token. {Ex}");
 			}
 		}
 
-		private async Task RefreshAuthenticationToken()
+		/// <summary>
+		/// Refreshes the authentication token using the Smarglin endpoint. This function is not thread safe and should ONLY be called by GetAuthenticationToken() where synchronization is used to protect the AuthenticationToken instance. Do not call this function directly, use GetAuthenticationToken() isntead.
+		/// </summary>
+		/// <param name="Logs"></param>
+		/// <returns></returns>
+		private async Task RefreshAuthenticationToken(List<string> Logs)
 		{
 			string RefreshEndpoint = "https://api.smartling.com/auth-api/v2/authenticate/refresh"; 
 			// serialize from Dictionary as MultipartFormDataContent not supported as a content-type 
@@ -502,17 +563,17 @@ namespace EpicGames.SmartlingLocalization
 					AuthenticationToken = RefreshResponseEnvelope.Response.Data;
 					AuthenticationToken.LastSuccessfulUpdateTime = DateTime.UtcNow;
 					Client.DefaultRequestHeaders.Authorization = new AuthenticationHeaderValue(AuthenticationToken.TokenType, AuthenticationToken.AccessToken);
-					Console.WriteLine("Successfully refreshed authentication token!");
+					Logger.LogInformation("Successfully refreshed authentication token!");
 				}
 				else
 				{
-					Console.WriteLine("Failed to refresh authentication token.");
-					await PrintRequestErrors(RefreshResponse);
+					Logs.Add("Failed to refresh authentication token.");
+					await AppendRequestErrorsToLog(RefreshResponse, Logs);
 				}
 			}
 			catch (Exception Ex)
 			{
-				Console.WriteLine($"Failed to refresh authentication token. {Ex}");
+				Logs.Add($"Failed to refresh authentication token. {Ex}");
 			}
 		}
 
@@ -520,5 +581,7 @@ namespace EpicGames.SmartlingLocalization
 		protected HttpClient Client;
 		private JsonSerializerOptions JsonOptions;
 		private SmartlingAuthenticationToken AuthenticationToken;
+		private readonly SemaphoreSlim AuthenticationTokenSemaphore= new SemaphoreSlim(1, 1);
+		private readonly Object LoggerLock = new object();
 	}
 }
