Add InstanceName to UnrealTestResults

#rnx

diff --git a/Engine/Source/Programs/AutomationTool/Gauntlet/Unreal/Utils/Gauntlet.UnrealTestResults.cs b/Engine/Source/Programs/AutomationTool/Gauntlet/Unreal/Utils/Gauntlet.UnrealTestResults.cs
--- a/Engine/Source/Programs/AutomationTool/Gauntlet/Unreal/Utils/Gauntlet.UnrealTestResults.cs
+++ b/Engine/Source/Programs/AutomationTool/Gauntlet/Unreal/Utils/Gauntlet.UnrealTestResults.cs
@@ -16,6 +16,8 @@ namespace Gauntlet
 		public string DeviceName { get; set; }
 		[JsonPropertyName("instance")] 
 		public string Instance { get; set; }
+		[JsonPropertyName("instanceName")]
+		public string InstanceName { get; set; }
 		[JsonPropertyName("platform")]
 		public string Platform { get; set; }
 		[JsonPropertyName("oSversion")]
