Update System.Drawing.Common to 4.7.3

#rnx

diff --git a/Engine/Platforms/VisionOS/Source/Programs/AutomationTool/VisionOS.Automation.csproj b/Engine/Platforms/VisionOS/Source/Programs/AutomationTool/VisionOS.Automation.csproj
--- a/Engine/Platforms/VisionOS/Source/Programs/AutomationTool/VisionOS.Automation.csproj
+++ b/Engine/Platforms/VisionOS/Source/Programs/AutomationTool/VisionOS.Automation.csproj
@@ -52,6 +52,6 @@
     <PackageReference Include="Microsoft.CSharp" Version="4.7.0" />
     <PackageReference Include="System.Data.DataSetExtensions" Version="4.5.0" />
     <PackageReference Include="System.Net.Http" Version="4.3.4" />
-    <PackageReference Include="System.Drawing.Common" Version="4.7.2" />
+    <PackageReference Include="System.Drawing.Common" Version="4.7.3" />
   </ItemGroup>
 </Project>
\ No newline at end of file
diff --git a/Engine/Source/Programs/AutomationTool/AutomationUtils/AutomationUtils.Automation.csproj b/Engine/Source/Programs/AutomationTool/AutomationUtils/AutomationUtils.Automation.csproj
--- a/Engine/Source/Programs/AutomationTool/AutomationUtils/AutomationUtils.Automation.csproj
+++ b/Engine/Source/Programs/AutomationTool/AutomationUtils/AutomationUtils.Automation.csproj
@@ -60,7 +60,7 @@
   <ItemGroup>
     <PackageReference Include="System.Data.DataSetExtensions" Version="4.5.0" />
     <PackageReference Include="System.Net.Http" Version="4.3.4" />
-    <PackageReference Include="System.Drawing.Common" Version="4.7.2" />
+    <PackageReference Include="System.Drawing.Common" Version="4.7.3" />
 
     <!--Required by Ionic -->
     <PackageReference Include="System.Security.Permissions" Version="4.7.0" />
diff --git a/Engine/Source/Programs/AutomationTool/BuildGraph/BuildGraph.Automation.csproj b/Engine/Source/Programs/AutomationTool/BuildGraph/BuildGraph.Automation.csproj
--- a/Engine/Source/Programs/AutomationTool/BuildGraph/BuildGraph.Automation.csproj
+++ b/Engine/Source/Programs/AutomationTool/BuildGraph/BuildGraph.Automation.csproj
@@ -62,6 +62,6 @@
     <PackageReference Include="Microsoft.CSharp" Version="4.7.0" />
     <PackageReference Include="System.Data.DataSetExtensions" Version="4.5.0" />
     <PackageReference Include="System.Net.Http" Version="4.3.4" />
-    <PackageReference Include="System.Drawing.Common" Version="4.7.2" />
+    <PackageReference Include="System.Drawing.Common" Version="4.7.3" />
   </ItemGroup>
 </Project>
\ No newline at end of file
diff --git a/Engine/Source/Programs/AutomationTool/IOS/IOS.Automation.csproj b/Engine/Source/Programs/AutomationTool/IOS/IOS.Automation.csproj
--- a/Engine/Source/Programs/AutomationTool/IOS/IOS.Automation.csproj
+++ b/Engine/Source/Programs/AutomationTool/IOS/IOS.Automation.csproj
@@ -56,7 +56,7 @@
     <PackageReference Include="Microsoft.CSharp" Version="4.7.0" PrivateAssets="all" />
     <PackageReference Include="System.Data.DataSetExtensions" Version="4.5.0" PrivateAssets="all" />
     <PackageReference Include="System.Net.Http" Version="4.3.4" PrivateAssets="all" />
-    <PackageReference Include="System.Drawing.Common" Version="4.7.2" PrivateAssets="all" />
+    <PackageReference Include="System.Drawing.Common" Version="4.7.3" PrivateAssets="all" />
     <PackageReference Include="System.Resources.Extensions" Version="4.7.1" PrivateAssets="all" />
     <PackageReference Include="Portable.BouncyCastle" Version="1.8.6.7" PrivateAssets="all" />
   </ItemGroup>
diff --git a/Engine/Source/Programs/AutomationTool/Linux/Linux.Automation.csproj b/Engine/Source/Programs/AutomationTool/Linux/Linux.Automation.csproj
--- a/Engine/Source/Programs/AutomationTool/Linux/Linux.Automation.csproj
+++ b/Engine/Source/Programs/AutomationTool/Linux/Linux.Automation.csproj
@@ -45,6 +45,6 @@
     <PackageReference Include="Microsoft.CSharp" Version="4.7.0" />
     <PackageReference Include="System.Data.DataSetExtensions" Version="4.5.0" />
     <PackageReference Include="System.Net.Http" Version="4.3.4" />
-    <PackageReference Include="System.Drawing.Common" Version="4.7.2" />
+    <PackageReference Include="System.Drawing.Common" Version="4.7.3" />
   </ItemGroup>
 </Project>
\ No newline at end of file
diff --git a/Engine/Source/Programs/AutomationTool/Mac/Mac.Automation.csproj b/Engine/Source/Programs/AutomationTool/Mac/Mac.Automation.csproj
--- a/Engine/Source/Programs/AutomationTool/Mac/Mac.Automation.csproj
+++ b/Engine/Source/Programs/AutomationTool/Mac/Mac.Automation.csproj
@@ -44,6 +44,6 @@
     <PackageReference Include="Microsoft.CSharp" Version="4.7.0" />
     <PackageReference Include="System.Data.DataSetExtensions" Version="4.5.0" />
     <PackageReference Include="System.Net.Http" Version="4.3.4" />
-    <PackageReference Include="System.Drawing.Common" Version="4.7.2" />
+    <PackageReference Include="System.Drawing.Common" Version="4.7.3" />
   </ItemGroup>
 </Project>
\ No newline at end of file
diff --git a/Engine/Source/Programs/AutomationTool/Scripts/AutomationScripts.Automation.csproj b/Engine/Source/Programs/AutomationTool/Scripts/AutomationScripts.Automation.csproj
--- a/Engine/Source/Programs/AutomationTool/Scripts/AutomationScripts.Automation.csproj
+++ b/Engine/Source/Programs/AutomationTool/Scripts/AutomationScripts.Automation.csproj
@@ -57,7 +57,7 @@
     <PackageReference Include="Microsoft.CSharp" Version="4.7.0" />
     <PackageReference Include="System.Data.DataSetExtensions" Version="4.5.0" />
     <PackageReference Include="System.Net.Http" Version="4.3.4" />
-    <PackageReference Include="System.Drawing.Common" Version="4.7.2" />
+    <PackageReference Include="System.Drawing.Common" Version="4.7.3" />
     <PackageReference Include="AWSSDK.S3" Version="3.3.113.2" PrivateAssets="all" />
   </ItemGroup>
 </Project>
\ No newline at end of file
diff --git a/Engine/Source/Programs/AutomationTool/SteamDeck/SteamDeck.Automation.csproj b/Engine/Source/Programs/AutomationTool/SteamDeck/SteamDeck.Automation.csproj
--- a/Engine/Source/Programs/AutomationTool/SteamDeck/SteamDeck.Automation.csproj
+++ b/Engine/Source/Programs/AutomationTool/SteamDeck/SteamDeck.Automation.csproj
@@ -44,6 +44,6 @@
     <PackageReference Include="Microsoft.CSharp" Version="4.7.0" />
     <PackageReference Include="System.Data.DataSetExtensions" Version="4.5.0" />
     <PackageReference Include="System.Net.Http" Version="4.3.4" />
-    <PackageReference Include="System.Drawing.Common" Version="4.7.2" />
+    <PackageReference Include="System.Drawing.Common" Version="4.7.3" />
   </ItemGroup>
 </Project>
\ No newline at end of file
diff --git a/Engine/Source/Programs/AutomationTool/TVOS/TVOS.Automation.csproj b/Engine/Source/Programs/AutomationTool/TVOS/TVOS.Automation.csproj
--- a/Engine/Source/Programs/AutomationTool/TVOS/TVOS.Automation.csproj
+++ b/Engine/Source/Programs/AutomationTool/TVOS/TVOS.Automation.csproj
@@ -51,6 +51,6 @@
     <PackageReference Include="Microsoft.CSharp" Version="4.7.0" />
     <PackageReference Include="System.Data.DataSetExtensions" Version="4.5.0" />
     <PackageReference Include="System.Net.Http" Version="4.3.4" />
-    <PackageReference Include="System.Drawing.Common" Version="4.7.2" />
+    <PackageReference Include="System.Drawing.Common" Version="4.7.3" />
   </ItemGroup>
 </Project>
\ No newline at end of file
diff --git a/Engine/Source/Programs/AutomationTool/Turnkey/Turnkey.Automation.csproj b/Engine/Source/Programs/AutomationTool/Turnkey/Turnkey.Automation.csproj
--- a/Engine/Source/Programs/AutomationTool/Turnkey/Turnkey.Automation.csproj
+++ b/Engine/Source/Programs/AutomationTool/Turnkey/Turnkey.Automation.csproj
@@ -58,7 +58,7 @@
     <PackageReference Include="Microsoft.CSharp" Version="4.7.0" />
     <PackageReference Include="System.Data.DataSetExtensions" Version="4.5.0" />
     <PackageReference Include="System.Net.Http" Version="4.3.4" />
-    <PackageReference Include="System.Drawing.Common" Version="4.7.2" />
+    <PackageReference Include="System.Drawing.Common" Version="4.7.3" />
     <PackageReference Include="System.IO.FileSystem.AccessControl" Version="4.7.0" />
     <PackageReference Include="Google.Apis.Drive.v3" Version="1.33.1.1225" PrivateAssets="all" />
   </ItemGroup>
diff --git a/Engine/Source/Programs/AutomationTool/Win/Win.Automation.csproj b/Engine/Source/Programs/AutomationTool/Win/Win.Automation.csproj
--- a/Engine/Source/Programs/AutomationTool/Win/Win.Automation.csproj
+++ b/Engine/Source/Programs/AutomationTool/Win/Win.Automation.csproj
@@ -52,6 +52,6 @@
     <PackageReference Include="Microsoft.CSharp" Version="4.7.0" />
     <PackageReference Include="System.Data.DataSetExtensions" Version="4.5.0" />
     <PackageReference Include="System.Net.Http" Version="4.3.4" />
-    <PackageReference Include="System.Drawing.Common" Version="4.7.2" />
+    <PackageReference Include="System.Drawing.Common" Version="4.7.3" />
   </ItemGroup>
 </Project>
\ No newline at end of file
