Fixup unacceptable words

```
[/Script/UnrealEd.ProjectPackagingSettings]
+MovieBlacklist=...
...change to...
+MovieDenyList=....
```

diff --git a/Engine/Source/Programs/AutomationTool/Scripts/CopyBuildToStagingDirectory.Automation.cs b/Engine/Source/Programs/AutomationTool/Scripts/CopyBuildToStagingDirectory.Automation.cs
--- a/Engine/Source/Programs/AutomationTool/Scripts/CopyBuildToStagingDirectory.Automation.cs
+++ b/Engine/Source/Programs/AutomationTool/Scripts/CopyBuildToStagingDirectory.Automation.cs
@@ -1362,7 +1362,7 @@ namespace AutomationScripts
 					if (!bSkipMovies && !SC.DedicatedServer)
 					{
 						List<string> MovieDenyList;
-						PlatformGameConfig.GetArray("/Script/UnrealEd.ProjectPackagingSettings", "MovieBlacklist", out MovieDenyList);
+						PlatformGameConfig.GetArray("/Script/UnrealEd.ProjectPackagingSettings", "MovieDenyList", out MovieDenyList);
 						if (MovieDenyList == null)
 						{
 							// Make an empty list to avoid having to null check below
