BuildGraph: Add a task to upload artifacts to Horde.

diff --git a/Engine/Build/Graph/Examples/BuildEditorAndTools.xml b/Engine/Build/Graph/Examples/BuildEditorAndTools.xml
--- a/Engine/Build/Graph/Examples/BuildEditorAndTools.xml
+++ b/Engine/Build/Graph/Examples/BuildEditorAndTools.xml
@@ -51,10 +51,13 @@
 	<Option Name="Promoted" Restrict="true|false" DefaultValue="true" Description="Whether to set the promoted flag in version files, allowing the editor to run without UBT."/>
 	<Option Name="ArchiveStream" Restrict="(?://.*)?" DefaultValue="" Description="Stream path that will contain zip file of the compiled binaries, for use by UnrealGameSync"/>
 	<Option Name="ForceSubmit" Restrict="true|false" DefaultValue="false" Description="Forces the submit to happen even if another change has been submitted (resolves in favor of local files)"/>
-	<Option Name="PreferredAgent" DefaultValue="CompileWin64" Description ="Agent type to use for job"/>
+	<Option Name="PreferredAgent" DefaultValue="CompileWin64;Win64" Description ="Agent type to use for job"/>
 	<Option Name="ArchiveName" Restrict="[^ ]*" DefaultValue="$(EditorTarget)" Description="Names of archive to generate instead of using the GameTarget, when the GameTarget may be UnrealEditor and result in a clash"/>
 	<Option Name="SymbolStorePath" DefaultValue="" Description ="Output directory for symbol storage"/>
 
+	<Property Name="ArchiveDir" Value="$(RootDir)\LocalBuilds\ArchiveForUGS"/>
+	<Property Name="ArchiveStagingDir" Value="$(ArchiveDir)\Staging"/>
+
 	<EnvVar Name="COMPUTERNAME"/>
 	<EnvVar Name="P4CLIENT"/>
 	
@@ -132,10 +135,9 @@
 			<Submit Description="[CL $(CodeChange)] Updated binaries" Files="#OutputFiles" FileType="binary+FS32"/>
 		</Node>	
 		
-		<!-- Create a zip archive and submit that to Perforce for use by UGS -->
-		<Node Name="Submit To Perforce For UGS" Requires="$(SubmissionDependencies)">
-			<!-- Clear out the archive directory -->
-			<Property Name="ArchiveDir" Value="$(RootDir)\LocalBuilds\ArchiveForUGS"/>
+		<!-- Prepare a build for UGS -->
+		<Node Name="Stage for UGS" Requires="$(SubmissionDependencies)">
+			<!-- Clear out the staging directory -->
 			<Delete Files="$(ArchiveDir)\..."/>
 
 			<!-- Tag AutomationTool build folder recursively -->
@@ -155,13 +157,13 @@
 			<Log Message="Archive symbols:" Files="#ArchiveSymbols"/>
 
 			<!-- Stage all the files to be archived -->
-			<Property Name="ArchiveStagingDir" Value="$(ArchiveDir)\Staging"/>
 			<Copy Files="#ArchiveBinaries" From="$(RootDir)" To="$(ArchiveStagingDir)"/>
 			<Strip Files="#ArchiveSymbols" BaseDir="$(RootDir)" OutputDir="$(ArchiveStagingDir)" Platform="Win64"/>
-
-
-			<!-- Create the zip file and submit it to Perforce -->
-			<Property Name="ArchivePerforceDir" Value="$(ArchiveDir)\Perforce"/>
+		</Node>
+		
+		<!-- Create a zip file and submit it to Perforce -->
+		<Node Name="Submit To Perforce For UGS" Requires="Stage for UGS">
+			<Property Name="ArchivePerforceDir" Value="$(RootDir)\LocalBuilds\ArchiveForUGS-Perforce"/>
 			<Property Name="ArchiveFile" Value="$(ArchivePerforceDir)\$(EscapedBranch)-$(ArchiveName).zip"/>
 			<Zip FromDir="$(ArchiveStagingDir)" ZipFile="$(ArchiveFile)"/>
 			<Warning Message="The 'ArchiveStream' argument is not set. Binaries will not be submitted." If="'$(ArchiveStream)' == ''"/>
@@ -169,6 +171,11 @@
 			<Property Name="SubmitClient" Value="$(P4CLIENT)_ArchiveForUGS" If="'$(P4CLIENT)' != ''"/>
 			<Submit Description="[CL $(CodeChange)] Updated binaries" Files="$(ArchiveFile)" FileType="binary+FS32" Workspace="$(SubmitClient)" Stream="$(ArchiveStream)" RootDir="$(ArchivePerforceDir)" Force="$(ForceSubmit)" If="'$(ArchiveStream)' != ''"/>
 		</Node>
+
+		<!-- Upload binaries to Horde -->
+		<Node Name="Submit To Horde For UGS" Requires="Stage for UGS">
+			<CreateArtifact Name="editor" Type="ugs-pcb" BaseDir="$(ArchiveStagingDir)"/>
+		</Node>
 		
 	</Agent>
 
diff --git a/Engine/Source/Programs/AutomationTool/BuildGraph/Tasks/CreateArtifactTask.cs b/Engine/Source/Programs/AutomationTool/BuildGraph/Tasks/CreateArtifactTask.cs
new file mode 100644
--- /dev/null
+++ b/Engine/Source/Programs/AutomationTool/BuildGraph/Tasks/CreateArtifactTask.cs
@@ -0,0 +1,164 @@
+// Copyright Epic Games, Inc. All Rights Reserved.
+
+using AutomationTool;
+using EpicGames.Core;
+using System;
+using System.Collections.Generic;
+using System.Text.Json;
+using System.IO;
+using System.Net.Http;
+using System.Net.Http.Headers;
+using System.Text;
+using System.Threading.Tasks;
+using System.Xml;
+using UnrealBuildBase;
+using Microsoft.Extensions.Logging;
+using EpicGames.Horde.Storage.Bundles;
+using EpicGames.Horde.Storage.Clients;
+using EpicGames.Horde.Storage;
+using EpicGames.Horde.Storage.Nodes;
+using System.Threading;
+using System.Data;
+using EpicGames.Horde.Storage.Backends;
+using Microsoft.Extensions.DependencyInjection;
+using EpicGames.Horde;
+using EpicGames.Horde.Artifacts;
+using System.Linq;
+using System.Diagnostics;
+
+#nullable enable
+
+namespace AutomationTool.Tasks
+{
+	/// <summary>
+	/// Parameters for a CreateArtifact task
+	/// </summary>
+	public class CreateArtifactTaskParameters
+	{
+		/// <summary>
+		/// Name for the artifact
+		/// </summary>
+		[TaskParameter]
+		public string Name = String.Empty;
+
+		/// <summary>
+		/// Type of the artifact
+		/// </summary>
+		[TaskParameter(Optional = true)]
+		public string Type = "unknown";
+
+		/// <summary>
+		/// Base directory to resolve relative paths for input files.
+		/// </summary>
+		[TaskParameter]
+		public string? BaseDir;
+
+		/// <summary>
+		/// Files to be uploaded.
+		/// </summary>
+		[TaskParameter(Optional = true, ValidationType = TaskParameterValidationType.FileSpec)]
+		public string Files = "...";
+	}
+
+	/// <summary>
+	/// Deploys a tool update through Horde
+	/// </summary>
+	[TaskElement("CreateArtifact", typeof(CreateArtifactTaskParameters))]
+	public class CreateArtifactTask : SpawnTaskBase
+	{
+		class LoggerProviderAdapter : ILoggerProvider
+		{
+			readonly ILogger _logger;
+
+			public LoggerProviderAdapter(ILogger logger) => _logger = logger;
+			public ILogger CreateLogger(string categoryName) => _logger;
+			public void Dispose() { }
+		}
+
+		/// <summary>
+		/// Parameters for this task
+		/// </summary>
+		CreateArtifactTaskParameters Parameters;
+
+		/// <summary>
+		/// Construct a Helm task
+		/// </summary>
+		/// <param name="InParameters">Parameters for the task</param>
+		public CreateArtifactTask(CreateArtifactTaskParameters InParameters)
+		{
+			Parameters = InParameters;
+		}
+
+		/// <summary>
+		/// Execute the task.
+		/// </summary>
+		/// <param name="Job">Information about the current job</param>
+		/// <param name="BuildProducts">Set of build products produced by this node.</param>
+		/// <param name="TagNameToFileSet">Mapping from tag names to the set of files they include</param>
+		public override async Task ExecuteAsync(JobContext Job, HashSet<FileReference> BuildProducts, Dictionary<string, HashSet<FileReference>> TagNameToFileSet)
+		{
+			// Create a DI container for building the graph
+			ServiceCollection serviceCollection = new ServiceCollection();
+			serviceCollection.AddLogging(builder => builder.AddProvider(new LoggerProviderAdapter(Log.Logger)));
+			serviceCollection.AddHorde(new HordeOptions { AllowAuthPrompt = !Automation.IsBuildMachine });
+			serviceCollection.Configure<LoggerFilterOptions>(options => options.AddFilter(typeof(HttpClient).FullName, LogLevel.Warning));
+
+			await using ServiceProvider serviceProvider = serviceCollection.BuildServiceProvider();
+
+			ArtifactName artifactName = new ArtifactName(Parameters.Name);
+			ArtifactType artifactType = new ArtifactType(Parameters.Type);
+
+			HordeHttpClient hordeHttpClient = serviceProvider.GetRequiredService<HordeHttpClient>();
+			CreateArtifactResponse response = await hordeHttpClient.CreateArtifactAsync(artifactName, artifactType);
+			Logger.LogInformation("Creating artifact {ArtifactId} '{ArtifactName}' ({ArtifactType}) (ns: {NamespaceId}, ref: {RefName})", response.ArtifactId, artifactName, artifactType, response.NamespaceId, response.RefName);
+
+			Stopwatch timer = Stopwatch.StartNew();
+
+			HttpStorageClientFactory httpStorageClientFactory = serviceProvider.GetRequiredService<HttpStorageClientFactory>();
+			using (IStorageClient client = httpStorageClientFactory.CreateClient(response.NamespaceId, response.Token))
+			{
+				await using (IStorageWriter writer = client.CreateWriter(response.RefName))
+				{
+					DirectoryReference baseDir = ResolveDirectory(Parameters.BaseDir);
+					List<FileInfo> files = ResolveFilespec(baseDir, Parameters.Files, TagNameToFileSet).Select(x => x.ToFileInfo()).ToList();
+
+					int totalCount = files.Count;
+					long totalSize = files.Sum(x => x.Length);
+
+					HashedNodeRef<DirectoryNode> outputNodeRef = await writer.WriteFilesAsync(baseDir.ToDirectoryInfo(), files, progress: new CopyStatsLoggerWithTotals(totalCount, totalSize, Logger));
+					await writer.FlushAsync();
+
+					await client.WriteRefAsync(response.RefName, outputNodeRef.Handle);
+				}
+			}
+
+			Logger.LogInformation("Completed in {Time:n1}s", timer.Elapsed.TotalSeconds);
+		}
+
+		/// <summary>
+		/// Output this task out to an XML writer.
+		/// </summary>
+		public override void Write(XmlWriter Writer)
+		{
+			Write(Writer, Parameters);
+		}
+
+		/// <summary>
+		/// Find all the tags which are used as inputs to this task
+		/// </summary>
+		/// <returns>The tag names which are read by this task</returns>
+		public override IEnumerable<string> FindConsumedTagNames()
+		{
+			return FindTagNamesFromList(Parameters.Files);
+		}
+
+		/// <summary>
+		/// Find all the tags which are modified by this task
+		/// </summary>
+		/// <returns>The tag names which are modified by this task</returns>
+		public override IEnumerable<string> FindProducedTagNames()
+		{
+			yield break;
+		}
+	}
+}
diff --git a/Engine/Source/Programs/Shared/EpicGames.Horde/HordeExtensions.cs b/Engine/Source/Programs/Shared/EpicGames.Horde/HordeExtensions.cs
new file mode 100644
--- /dev/null
+++ b/Engine/Source/Programs/Shared/EpicGames.Horde/HordeExtensions.cs
@@ -0,0 +1,65 @@
+// Copyright Epic Games, Inc. All Rights Reserved.
+
+using System;
+using System.Net.Http;
+using System.Net.Http.Headers;
+using EpicGames.Horde.Storage;
+using EpicGames.Horde.Storage.Clients;
+using Microsoft.Extensions.DependencyInjection;
+
+namespace EpicGames.Horde
+{
+	/// <summary>
+	/// Extension methods for Horde
+	/// </summary>
+	public static class HordeExtensions
+	{
+		/// <summary>
+		/// Adds Horde-related services with the default settings
+		/// </summary>
+		/// <param name="serviceCollection">Collection to register services with</param>
+		public static void AddHorde(this IServiceCollection serviceCollection)
+		{
+			AddHorde(serviceCollection, new HordeOptions());
+		}
+
+		/// <summary>
+		/// Adds Horde-related services with the default settings
+		/// </summary>
+		/// <param name="serviceCollection">Collection to register services with</param>
+		/// <param name="configure">Callback to configure options</param>
+		public static void AddHorde(this IServiceCollection serviceCollection, Action<HordeOptions> configure)
+		{
+			HordeOptions options = new HordeOptions();
+			configure(options);
+			AddHorde(serviceCollection, options);
+		}
+
+		/// <summary>
+		/// Adds Horde-related services with the default settings
+		/// </summary>
+		/// <param name="serviceCollection">Collection to register services with</param>
+		/// <param name="options">Options for the collection</param>
+		public static void AddHorde(this IServiceCollection serviceCollection, HordeOptions options)
+		{
+			void ConfigureHttpClient(HttpClient httpClient)
+			{
+				if (options.ServerUrl != null)
+				{
+					httpClient.BaseAddress = options.ServerUrl;
+				}
+				if (options.AccessToken != null)
+				{
+					httpClient.DefaultRequestHeaders.Authorization = new AuthenticationHeaderValue("Bearer", options.AccessToken);
+				}
+
+				options.ConfigureHttpClient?.Invoke(httpClient);
+			}
+
+			serviceCollection.AddHordeHttpClient(ConfigureHttpClient, options.AllowAuthPrompt);
+			serviceCollection.AddSingleton<BundleCache>(sp => new BundleCache(options.BundleCache));
+			serviceCollection.AddSingleton<StorageBackendCache>();
+			serviceCollection.AddSingleton<HttpStorageClientFactory>();
+		}
+	}
+}
diff --git a/Engine/Source/Programs/Shared/EpicGames.Horde/HordeHttpClient.cs b/Engine/Source/Programs/Shared/EpicGames.Horde/HordeHttpClient.cs
--- a/Engine/Source/Programs/Shared/EpicGames.Horde/HordeHttpClient.cs
+++ b/Engine/Source/Programs/Shared/EpicGames.Horde/HordeHttpClient.cs
@@ -19,6 +19,7 @@ using EpicGames.Horde.Secrets;
 using EpicGames.Horde.Server;
 using EpicGames.Horde.Storage;
 using EpicGames.Horde.Storage.Clients;
+using EpicGames.Horde.Streams;
 using EpicGames.Horde.Tools;
 using Microsoft.Extensions.DependencyInjection;
 using Microsoft.Extensions.Logging;
@@ -104,6 +105,20 @@ namespace EpicGames.Horde
 
 		#region Artifacts
 
+		/// <summary>
+		/// Creates a new artifact
+		/// </summary>
+		/// <param name="name">Name of the artifact</param>
+		/// <param name="type">Additional search keys tagged on the artifact</param>
+		/// <param name="streamId">Stream to create the artifact for</param>
+		/// <param name="change">Change number for the artifact</param>
+		/// <param name="keys">Keys used to identify the artifact</param>
+		/// <param name="cancellationToken">Cancellation token for the operation</param>
+		public Task<CreateArtifactResponse> CreateArtifactAsync(ArtifactName name, ArtifactType type, StreamId? streamId = null, int? change = null, List<string>? keys = null, CancellationToken cancellationToken = default)
+		{
+			return PostAsync<CreateArtifactResponse, CreateArtifactRequest>(_httpClient, $"api/v2/artifacts", new CreateArtifactRequest(name, type, streamId, change, keys ?? new List<string>()), cancellationToken);
+		}
+
 		/// <summary>
 		/// Gets metadata about an artifact object
 		/// </summary>
@@ -490,6 +505,12 @@ namespace EpicGames.Horde
 						}
 					}
 				}
+
+				// Make sure we have a base URL set
+				if (httpClient.BaseAddress == null)
+				{
+					throw new Exception("No Horde server is configured, or can be detected from the environment. Consider specifying a URL when calling AddHordeHttpClient().");
+				}
 			}
 
 			IHttpClientBuilder builder = services.AddHttpClient<HordeHttpClient>(HordeHttpClient.HttpClientName, ConfigureClientFromEnvironment)
diff --git a/Engine/Source/Programs/Shared/EpicGames.Horde/HordeOptions.cs b/Engine/Source/Programs/Shared/EpicGames.Horde/HordeOptions.cs
new file mode 100644
--- /dev/null
+++ b/Engine/Source/Programs/Shared/EpicGames.Horde/HordeOptions.cs
@@ -0,0 +1,45 @@
+// Copyright Epic Games, Inc. All Rights Reserved.
+
+using System;
+using System.Net.Http;
+using EpicGames.Horde.Storage;
+using EpicGames.Horde.Storage.Bundles;
+
+namespace EpicGames.Horde
+{
+	/// <summary>
+	/// Options for configuring the Horde connection
+	/// </summary>
+	public class HordeOptions
+	{
+		/// <summary>
+		/// Address of the Horde server
+		/// </summary>
+		public Uri? ServerUrl { get; set; }
+
+		/// <summary>
+		/// Access token to use for connecting to the server
+		/// </summary>
+		public string? AccessToken { get; set; }
+
+		/// <summary>
+		/// Whether to allow opening a browser window to prompt for authentication
+		/// </summary>
+		public bool AllowAuthPrompt { get; set; }
+
+		/// <summary>
+		/// Callback to allow configuring any HTTP client created for Horde
+		/// </summary>
+		public Action<HttpClient>? ConfigureHttpClient { get; set; }
+
+		/// <summary>
+		/// Options for creating new bundles
+		/// </summary>
+		public BundleOptions Bundle { get; } = new BundleOptions();
+
+		/// <summary>
+		/// Options for caching bundles 
+		/// </summary>
+		public BundleCacheOptions BundleCache { get; } = new BundleCacheOptions();
+	}
+}
