Fixup unacceptable words

```
[MemoryMappedFiles]
MasterEnable=true
...change to...
Enable=true
```

diff --git a/Engine/Config/IOS/IOSEngine.ini b/Engine/Config/IOS/IOSEngine.ini
--- a/Engine/Config/IOS/IOSEngine.ini
+++ b/Engine/Config/IOS/IOSEngine.ini
@@ -50,7 +50,7 @@ r.HairStrands.Binding=0
 PoolSizeVRAMPercentage=70
 
 [MemoryMappedFiles]
-MasterEnable=true
+Enable=true
 Alignment=16384
 
 [BackgroundHttp.iOSSettings]
diff --git a/Engine/Source/Programs/AutomationTool/Scripts/CopyBuildToStagingDirectory.Automation.cs b/Engine/Source/Programs/AutomationTool/Scripts/CopyBuildToStagingDirectory.Automation.cs
--- a/Engine/Source/Programs/AutomationTool/Scripts/CopyBuildToStagingDirectory.Automation.cs
+++ b/Engine/Source/Programs/AutomationTool/Scripts/CopyBuildToStagingDirectory.Automation.cs
@@ -3115,7 +3115,7 @@ namespace AutomationScripts
 			if (bHasPlatformEngineConfig)
 			{
 				bool bEnabled = false;
-				PlatformEngineConfig.GetBool("MemoryMappedFiles", "MasterEnable", out bEnabled);
+				PlatformEngineConfig.GetBool("MemoryMappedFiles", "Enable", out bEnabled);
 				if (bEnabled)
 				{
 					int Value = 0;
