[Gauntlet] ITargetDevice cleanup/install functions

This CL stubs out the following ITargetDevice functions which will be implemented for each platform one by one
- FullClean
- CleanArtifacts
- InstallApplication (new signature)
- CreateAppInstall
- CopyAdditionalFiles

Once stubs have been implemented, UnrealSession.LaunchSession will optionally be able to execute a new flow using these modularized pieces
After further validations are made, the legacy implementation of LaunchSession will be deprecated and the new flow will become the standard

diff --git a/Engine/Source/Programs/AutomationTool/Gauntlet/Framework/Base/Gauntlet.TargetDevice.cs b/Engine/Source/Programs/AutomationTool/Gauntlet/Framework/Base/Gauntlet.TargetDevice.cs
--- a/Engine/Source/Programs/AutomationTool/Gauntlet/Framework/Base/Gauntlet.TargetDevice.cs
+++ b/Engine/Source/Programs/AutomationTool/Gauntlet/Framework/Base/Gauntlet.TargetDevice.cs
@@ -85,17 +85,46 @@ namespace Gauntlet
 		/// <returns></returns>
 		Dictionary<EIntendedBaseCopyDirectory, string> GetPlatformDirectoryMappings();
 
-		void ClearSavedDirectory(UnrealAppConfig AppConfiguration);
-
 		IAppInstall InstallApplication(UnrealAppConfig AppConfiguration);
 
-		void CopyAppConfigurationFiles(UnrealAppConfig AppConfiguration);
-
 		IAppInstance Run(IAppInstall App);
 
-		bool CheckRequiredSettings(List<string> RequiredSettingsList) { return true; }
-			
-		string GetPackagedExecutableLocation() { return null; }
+		/// Begin new flow ///
+
+		/// <summary>
+		/// Fully cleans the device by deleting
+		///	 - Artifacts and other loose files associated with UE processes
+		///	 - Staged/Packaged builds
+		/// </summary>
+		void FullClean();
+
+		/// <summary>
+		/// Deletes artifacts and other loose files associated with UE processes
+		/// </summary>
+		void CleanArtifacts();
+
+		/// <summary>
+		/// Installs a build to the device
+		/// </summary>
+		/// <param name="Build">A reference to the build to install</param>
+		/// <param name="ProjectName">Name of the UE project being installed</param>
+		/// <param name="Sandbox">Which sandbox to install to</param>
+		void InstallApplication(IBuild Build, string ProjectName, string Sandbox);
+
+		/// <summary>
+		/// Create an IAppInstall that is configured by the provided AppConfiguration
+		/// </summary>
+		/// <param name="AppConfiguration">The configuration used to create the IAppInstall</param>
+		/// <returns>An AppInstall handle which can be used to run the process</returns>
+		IAppInstall CreateAppInstall(UnrealAppConfig AppConfiguration);
+
+		/// <summary>
+		/// Copies any additional files to the device
+		/// </summary>
+		/// <param name="FilesToCopy">The collection of files to copy</param>
+		void CopyAdditionalFiles(IEnumerable<UnrealFileToCopy> FilesToCopy);
+
+		/// End new flow ///
 
 		/// <summary>
 		/// Path to the crash dumps on a device
@@ -113,7 +142,6 @@ namespace Gauntlet
 		/// Returns true if there were any crash dumps for the run, and false otherwise
 		/// </summary>
 		bool CopyCrashDumps() { return false; }
-
 	};
 
 	/// <summary>
diff --git a/Engine/Source/Programs/AutomationTool/Gauntlet/Framework/Base/Gauntlet.TargetDeviceDesktopCommon.cs b/Engine/Source/Programs/AutomationTool/Gauntlet/Framework/Base/Gauntlet.TargetDeviceDesktopCommon.cs
--- a/Engine/Source/Programs/AutomationTool/Gauntlet/Framework/Base/Gauntlet.TargetDeviceDesktopCommon.cs
+++ b/Engine/Source/Programs/AutomationTool/Gauntlet/Framework/Base/Gauntlet.TargetDeviceDesktopCommon.cs
@@ -6,6 +6,7 @@ using System;
 using System.Collections.Generic;
 using System.IO;
 using System.Threading;
+using System.Linq;
 using UnrealBuildTool;
 
 namespace Gauntlet
@@ -104,58 +105,76 @@ namespace Gauntlet
 			}
 		}
 
-		public virtual void ClearSavedDirectory(UnrealAppConfig AppConfiguration) { }
+		public void FullClean()
+		{
+			
+		}
 
-		public virtual void CopyAppConfigurationFiles(UnrealAppConfig AppConfiguration) { }
+		public void CleanArtifacts()
+		{
 
-		public abstract IAppInstance Run(IAppInstall Install);
+		}
 
-		public virtual void PopulateDirectoryMappings(string BaseDirectory)
+		public void InstallApplication(IBuild Build, string ProjectName, string Sandbox)
 		{
-			LocalDirectoryMappings.Clear();
 
-			LocalDirectoryMappings.Add(EIntendedBaseCopyDirectory.Build, Path.Combine(BaseDirectory, "Build"));
-			LocalDirectoryMappings.Add(EIntendedBaseCopyDirectory.Binaries, Path.Combine(BaseDirectory, "Binaries"));
-			LocalDirectoryMappings.Add(EIntendedBaseCopyDirectory.Config, Path.Combine(BaseDirectory, "Saved", "Config"));
-            LocalDirectoryMappings.Add(EIntendedBaseCopyDirectory.Content, Path.Combine(BaseDirectory, "Content"));
-            LocalDirectoryMappings.Add(EIntendedBaseCopyDirectory.Demos, Path.Combine(BaseDirectory, "Saved", "Demos"));
-			LocalDirectoryMappings.Add(EIntendedBaseCopyDirectory.PersistentDownloadDir, Path.Combine(BaseDirectory, "Saved", "PersistentDownloadDir"));
-			LocalDirectoryMappings.Add(EIntendedBaseCopyDirectory.Profiling, Path.Combine(BaseDirectory, "Saved", "Profiling"));
-            LocalDirectoryMappings.Add(EIntendedBaseCopyDirectory.Saved, Path.Combine(BaseDirectory, "Saved"));
 		}
 
-		protected void CopyAdditionalFiles(UnrealAppConfig AppConfig)
+		public IAppInstall CreateAppInstall(UnrealAppConfig AppConfig)
 		{
-			if (AppConfig.FilesToCopy != null)
+			return null;
+		}
+
+		public void CopyAdditionalFiles(IEnumerable<UnrealFileToCopy> FilesToCopy)
+		{
+			if (FilesToCopy == null || FilesToCopy.Any())
+			{
+				return;
+			}
+
+			foreach (UnrealFileToCopy FileToCopy in FilesToCopy)
 			{
-				foreach (UnrealFileToCopy FileToCopy in AppConfig.FilesToCopy)
+				string PathToCopyTo = Path.Combine(LocalDirectoryMappings[FileToCopy.TargetBaseDirectory], FileToCopy.TargetRelativeLocation);
+				if (File.Exists(FileToCopy.SourceFileLocation))
 				{
-					string PathToCopyTo = Path.Combine(LocalDirectoryMappings[FileToCopy.TargetBaseDirectory], FileToCopy.TargetRelativeLocation);
-					if (File.Exists(FileToCopy.SourceFileLocation))
+					FileInfo SrcInfo = new FileInfo(FileToCopy.SourceFileLocation);
+					SrcInfo.IsReadOnly = false;
+					string DirectoryToCopyTo = Path.GetDirectoryName(PathToCopyTo);
+					if (!Directory.Exists(DirectoryToCopyTo))
 					{
-						FileInfo SrcInfo = new FileInfo(FileToCopy.SourceFileLocation);
-						SrcInfo.IsReadOnly = false;
-						string DirectoryToCopyTo = Path.GetDirectoryName(PathToCopyTo);
-						if (!Directory.Exists(DirectoryToCopyTo))
-						{
-							Directory.CreateDirectory(DirectoryToCopyTo);
-						}
-						if (File.Exists(PathToCopyTo))
-						{
-							FileInfo ExistingFile = new FileInfo(PathToCopyTo);
-							ExistingFile.IsReadOnly = false;
-						}
-						SrcInfo.CopyTo(PathToCopyTo, true);
-						Log.Info("Copying {0} to {1}", FileToCopy.SourceFileLocation, PathToCopyTo);
+						Directory.CreateDirectory(DirectoryToCopyTo);
 					}
-					else
+					if (File.Exists(PathToCopyTo))
 					{
-						Log.Warning(KnownLogEvents.Gauntlet_DeviceEvent, "File to copy {File} not found", FileToCopy);
+						FileInfo ExistingFile = new FileInfo(PathToCopyTo);
+						ExistingFile.IsReadOnly = false;
 					}
+					SrcInfo.CopyTo(PathToCopyTo, true);
+					Log.Info("Copying {0} to {1}", FileToCopy.SourceFileLocation, PathToCopyTo);
+				}
+				else
+				{
+					Log.Warning(KnownLogEvents.Gauntlet_DeviceEvent, "File to copy {File} not found", FileToCopy);
 				}
 			}
 		}
 
+		public abstract IAppInstance Run(IAppInstall Install);
+
+		public virtual void PopulateDirectoryMappings(string BaseDirectory)
+		{
+			LocalDirectoryMappings.Clear();
+
+			LocalDirectoryMappings.Add(EIntendedBaseCopyDirectory.Build, Path.Combine(BaseDirectory, "Build"));
+			LocalDirectoryMappings.Add(EIntendedBaseCopyDirectory.Binaries, Path.Combine(BaseDirectory, "Binaries"));
+			LocalDirectoryMappings.Add(EIntendedBaseCopyDirectory.Config, Path.Combine(BaseDirectory, "Saved", "Config"));
+            LocalDirectoryMappings.Add(EIntendedBaseCopyDirectory.Content, Path.Combine(BaseDirectory, "Content"));
+            LocalDirectoryMappings.Add(EIntendedBaseCopyDirectory.Demos, Path.Combine(BaseDirectory, "Saved", "Demos"));
+			LocalDirectoryMappings.Add(EIntendedBaseCopyDirectory.PersistentDownloadDir, Path.Combine(BaseDirectory, "Saved", "PersistentDownloadDir"));
+			LocalDirectoryMappings.Add(EIntendedBaseCopyDirectory.Profiling, Path.Combine(BaseDirectory, "Saved", "Profiling"));
+            LocalDirectoryMappings.Add(EIntendedBaseCopyDirectory.Saved, Path.Combine(BaseDirectory, "Saved"));
+		}
+
 		// TODO - b.lienau: implement these at desktop level and remove implementations from each desktop
 		protected abstract IAppInstall InstallNativeStagedBuild(UnrealAppConfig AppConfig, NativeStagedBuild Build);
 		protected abstract IAppInstall InstallStagedBuild(UnrealAppConfig AppConfig, StagedBuild Build);
diff --git a/Engine/Source/Programs/AutomationTool/Gauntlet/Framework/Devices/Gauntlet.DevicePool.cs b/Engine/Source/Programs/AutomationTool/Gauntlet/Framework/Devices/Gauntlet.DevicePool.cs
--- a/Engine/Source/Programs/AutomationTool/Gauntlet/Framework/Devices/Gauntlet.DevicePool.cs
+++ b/Engine/Source/Programs/AutomationTool/Gauntlet/Framework/Devices/Gauntlet.DevicePool.cs
@@ -1309,17 +1309,6 @@ namespace Gauntlet
 					NewDevice = Factory.CreateDevice(Def.Address, ClientTempDir, Def.DeviceData);
 				}
 
-				List<string> RequiredSettings = Globals.Params.ParseValues("RequiredSettings");
-				if (RequiredSettings.Count > 0)
-				{
-					Log.Info("Checking required settings for device {0}...", NewDevice.Name);
-					if (!NewDevice.CheckRequiredSettings(RequiredSettings))
-					{
-						Log.Warning("Device settings doesn't meet the requirements, skipping device.");
-						return null;
-					}
-					Log.Info("Device has all required settings, keeping device.");
-				}
 				if (NewDevice.IsAvailable == false)
 				{
 					Log.Info(KnownLogEvents.Gauntlet_DeviceEvent, "Assigned device {Name} reports unavailable. Requesting a forced disconnect", NewDevice);
diff --git a/Engine/Source/Programs/AutomationTool/Gauntlet/Platform/Android/Gauntlet.TargetDeviceAndroid.cs b/Engine/Source/Programs/AutomationTool/Gauntlet/Platform/Android/Gauntlet.TargetDeviceAndroid.cs
--- a/Engine/Source/Programs/AutomationTool/Gauntlet/Platform/Android/Gauntlet.TargetDeviceAndroid.cs
+++ b/Engine/Source/Programs/AutomationTool/Gauntlet/Platform/Android/Gauntlet.TargetDeviceAndroid.cs
@@ -873,11 +873,6 @@ namespace Gauntlet
 			return true;
 		}
 
-		public void ClearSavedDirectory(UnrealAppConfig AppConfiguration)
-		{
-
-		}
-
 		public IAppInstall InstallApplication(UnrealAppConfig AppConfig)
 		{
 			// todo - pass this through
@@ -1101,12 +1096,27 @@ namespace Gauntlet
 			return AppInstall;
 		}
 
-		public void CopyAppConfigurationFiles(UnrealAppConfig AppConfiguration)
+		public void FullClean()
 		{
 
 		}
-		
-		public void FullClean()
+
+		public void CleanArtifacts()
+		{
+
+		}
+
+		public void InstallApplication(IBuild Build, string ProjectName, string Sandbox)
+		{
+
+		}
+
+		public IAppInstall CreateAppInstall(UnrealAppConfig AppConfig)
+		{
+			return null;
+		}
+
+		public void CopyAdditionalFiles(IEnumerable<UnrealFileToCopy> FilesToCopy)
 		{
 
 		}
diff --git a/Engine/Source/Programs/AutomationTool/Gauntlet/Platform/IOS/Gauntlet.TargetDeviceIOS.cs b/Engine/Source/Programs/AutomationTool/Gauntlet/Platform/IOS/Gauntlet.TargetDeviceIOS.cs
--- a/Engine/Source/Programs/AutomationTool/Gauntlet/Platform/IOS/Gauntlet.TargetDeviceIOS.cs
+++ b/Engine/Source/Programs/AutomationTool/Gauntlet/Platform/IOS/Gauntlet.TargetDeviceIOS.cs
@@ -464,11 +464,6 @@ namespace Gauntlet
 			}
 		}
 
-		public void ClearSavedDirectory(UnrealAppConfig AppConfiguration)
-		{
-
-		}
-
 		// We need to lock around setting up the IPA
 		static object IPALock = new object();
 		public IAppInstall InstallApplication(UnrealAppConfig AppConfig)
@@ -588,16 +583,31 @@ namespace Gauntlet
 			return IOSApp;
 		}
 
-		public void CopyAppConfigurationFiles(UnrealAppConfig AppConfiguration)
+		public void FullClean()
 		{
 
 		}
 
-		public void FullClean(UnrealAppConfig AppConfiguration)
+		public void CleanArtifacts()
 		{
 
 		}
-		
+
+		public void InstallApplication(IBuild Build, string ProjectName, string Sandbox)
+		{
+
+		}
+
+		public IAppInstall CreateAppInstall(UnrealAppConfig AppConfig)
+		{
+			return null;
+		}
+
+		public void CopyAdditionalFiles(IEnumerable<UnrealFileToCopy> FilesToCopy)
+		{
+
+		}
+
 		public void PopulateDirectoryMappings(string ProjectDir)
 		{
 			LocalDirectoryMappings.Add(EIntendedBaseCopyDirectory.Build, Path.Combine(ProjectDir, "Build"));
diff --git a/Engine/Source/Programs/AutomationTool/Gauntlet/Platform/Linux/Gauntlet.TargetDeviceLinux.cs b/Engine/Source/Programs/AutomationTool/Gauntlet/Platform/Linux/Gauntlet.TargetDeviceLinux.cs
--- a/Engine/Source/Programs/AutomationTool/Gauntlet/Platform/Linux/Gauntlet.TargetDeviceLinux.cs
+++ b/Engine/Source/Programs/AutomationTool/Gauntlet/Platform/Linux/Gauntlet.TargetDeviceLinux.cs
@@ -108,7 +108,7 @@ namespace Gauntlet
 			LinuxApp.ArtifactPath = Path.Combine(InBuild.BuildPath, AppConfig.ProjectName, @"Saved");
 			LinuxApp.CleanDeviceArtifacts();
 
-			CopyAdditionalFiles(AppConfig);
+			CopyAdditionalFiles(AppConfig.FilesToCopy);
 
 			LinuxApp.ExecutablePath = Path.Combine(InBuild.BuildPath, InBuild.ExecutablePath);
 			LinuxApp.WorkingDirectory = InBuild.BuildPath;
@@ -166,7 +166,7 @@ namespace Gauntlet
 				PopulateDirectoryMappings(Path.Combine(BuildDir, AppConfig.ProjectName));
 			}
 
-			CopyAdditionalFiles(AppConfig);
+			CopyAdditionalFiles(AppConfig.FilesToCopy);
 
 			if (Path.IsPathRooted(InBuild.ExecutablePath))
 			{
@@ -220,7 +220,7 @@ namespace Gauntlet
 				PopulateDirectoryMappings(AppConfig.ProjectFile.Directory.FullName);
 			}
 
-			CopyAdditionalFiles(AppConfig);
+			CopyAdditionalFiles(AppConfig.FilesToCopy);
 
 			return LinuxApp;
 		}
diff --git a/Engine/Source/Programs/AutomationTool/Gauntlet/Platform/Null/Gauntlet.TargetDeviceNull.cs b/Engine/Source/Programs/AutomationTool/Gauntlet/Platform/Null/Gauntlet.TargetDeviceNull.cs
--- a/Engine/Source/Programs/AutomationTool/Gauntlet/Platform/Null/Gauntlet.TargetDeviceNull.cs
+++ b/Engine/Source/Programs/AutomationTool/Gauntlet/Platform/Null/Gauntlet.TargetDeviceNull.cs
@@ -151,17 +151,32 @@ namespace Gauntlet
 			return new NullAppInstance(this);
 		}
 
-		public void ClearSavedDirectory(UnrealAppConfig AppConfiguration)
+		public IAppInstall InstallApplication(UnrealAppConfig AppConfig)
+		{
+			return new NullAppInstall(AppConfig.Name, this, AppConfig.CommandLine);
+		}
+
+		public void FullClean()
 		{
 
 		}
 
-		public IAppInstall InstallApplication(UnrealAppConfig AppConfig)
+		public void CleanArtifacts()
 		{
-			return new NullAppInstall(AppConfig.Name, this, AppConfig.CommandLine);
+
+		}
+
+		public void InstallApplication(IBuild Build, string ProjectName, string Sandbox)
+		{
+
+		}
+
+		public IAppInstall CreateAppInstall(UnrealAppConfig AppConfig)
+		{
+			return null;
 		}
 
-		public void CopyAppConfigurationFiles(UnrealAppConfig AppConfiguration)
+		public void CopyAdditionalFiles(IEnumerable<UnrealFileToCopy> FilesToCopy)
 		{
 
 		}
diff --git a/Engine/Source/Programs/AutomationTool/Gauntlet/Platform/Windows/Gauntlet.TargetDeviceWindows.cs b/Engine/Source/Programs/AutomationTool/Gauntlet/Platform/Windows/Gauntlet.TargetDeviceWindows.cs
--- a/Engine/Source/Programs/AutomationTool/Gauntlet/Platform/Windows/Gauntlet.TargetDeviceWindows.cs
+++ b/Engine/Source/Programs/AutomationTool/Gauntlet/Platform/Windows/Gauntlet.TargetDeviceWindows.cs
@@ -158,7 +158,7 @@ namespace Gauntlet
 			WinApp.ArtifactPath = Path.Combine(InBuild.BuildPath, AppConfig.ProjectName, @"Saved");
 			WinApp.CleanDeviceArtifacts();
 
-			CopyAdditionalFiles(AppConfig);
+			CopyAdditionalFiles(AppConfig.FilesToCopy);
 
 			return WinApp;
 		}
@@ -213,7 +213,7 @@ namespace Gauntlet
 				PopulateDirectoryMappings(Path.Combine(BuildDir, AppConfig.ProjectName));
 			}
 
-			CopyAdditionalFiles(AppConfig);
+			CopyAdditionalFiles(AppConfig.FilesToCopy);
 
 			if (Path.IsPathRooted(InBuild.ExecutablePath))
 			{
@@ -267,7 +267,7 @@ namespace Gauntlet
 				PopulateDirectoryMappings(AppConfig.ProjectFile.Directory.FullName);
 			}
 
-			CopyAdditionalFiles(AppConfig);
+			CopyAdditionalFiles(AppConfig.FilesToCopy);
 
 			return WinApp;
 		}
@@ -298,7 +298,7 @@ namespace Gauntlet
 				PopulateDirectoryMappings(Path.Combine(BasePath, AppConfig.ProjectName));
 			}
 
-			CopyAdditionalFiles(AppConfig);
+			CopyAdditionalFiles(AppConfig.FilesToCopy);
 			return WinApp;
 		}
 	}
diff --git a/Engine/Source/Programs/AutomationTool/Gauntlet/SelfTest/Unreal/Gauntlet.SelfTest.TestUnrealInstallAndRunDesktop.cs b/Engine/Source/Programs/AutomationTool/Gauntlet/SelfTest/Unreal/Gauntlet.SelfTest.TestUnrealInstallAndRunDesktop.cs
--- a/Engine/Source/Programs/AutomationTool/Gauntlet/SelfTest/Unreal/Gauntlet.SelfTest.TestUnrealInstallAndRunDesktop.cs
+++ b/Engine/Source/Programs/AutomationTool/Gauntlet/SelfTest/Unreal/Gauntlet.SelfTest.TestUnrealInstallAndRunDesktop.cs
@@ -44,7 +44,7 @@ namespace Gauntlet.SelfTest
 			}
 
 			// Now delete the file using clear cache
-			TargetDevice.ClearSavedDirectory(AppConfig);
+			TargetDevice.CleanArtifacts();
 
 			// We should see these files and their parent directories are now gone.
 			bool bSavedWasDeleted = !SavedFile.Directory.Exists;
@@ -67,7 +67,7 @@ namespace Gauntlet.SelfTest
 			UnrealFileToCopy File = CreateDummyUnrealFileToCopy();
 			AppConfig.FilesToCopy.Add(File);
 
-			TargetDevice.CopyAppConfigurationFiles(AppConfig);
+			TargetDevice.CopyAdditionalFiles(AppConfig.FilesToCopy);
 
 			string CopyDirectory = Path.Combine(TargetDevice.GetPlatformDirectoryMappings()[EIntendedBaseCopyDirectory.Saved], File.TargetRelativeLocation);
 			FileInfo CopiedFile = new FileInfo(CopyDirectory);
diff --git a/Engine/Source/Programs/AutomationTool/Gauntlet/Unreal/Base/Gauntlet.UnrealSession.cs b/Engine/Source/Programs/AutomationTool/Gauntlet/Unreal/Base/Gauntlet.UnrealSession.cs
--- a/Engine/Source/Programs/AutomationTool/Gauntlet/Unreal/Base/Gauntlet.UnrealSession.cs
+++ b/Engine/Source/Programs/AutomationTool/Gauntlet/Unreal/Base/Gauntlet.UnrealSession.cs
@@ -895,8 +895,6 @@ namespace Gauntlet
 					// todo - should this be elsewhere?
 					AppConfig.Sandbox = Sandbox;
 
-					Device.ClearSavedDirectory(AppConfig);
-
 					IAppInstall Install = null;
 					bool bReinstallPerPass = Globals.Params.ParseParam("ReinstallPerPass");
 					if (RolesToInstalls == null || !RolesToInstalls.ContainsKey(Role) || bReinstallPerPass)
@@ -974,8 +972,6 @@ namespace Gauntlet
 						InstallsToRoles[Install] = Role;
 						Log.Info("Using previous install of {0} on {1}", Install.Name, Install.Device.Name);
 					}
-
-					Device.CopyAppConfigurationFiles(AppConfig);
 				}
 
 				if (InstallSuccess == false)
