[Gauntlet] Adds IDeviceValidator to separate concerns around validating whether a device is

IDeviceValidator provides the ability to verify a TargetDevice matches a set of requirements.
They are particularily useful for ensuring development kits are in a desired state prior to a test.
This can include, but is not limited to:
    - Proper firmware versions
    - Having an account signed in
    - Pushing a list of settings onto the device
To add new DeviceValidators, simply create a concrete type that inherits from IDeviceValidator.
Implement ValidateDevice() and determine whether this validator should be enabled within the constructor
 It can be especially helpful to toggle whether a validator is enabled via commandline arguments, for ease of opt-in

diff --git a/Engine/Source/Programs/AutomationTool/Gauntlet/Framework/Devices/Gauntlet.DevicePool.cs b/Engine/Source/Programs/AutomationTool/Gauntlet/Framework/Devices/Gauntlet.DevicePool.cs
--- a/Engine/Source/Programs/AutomationTool/Gauntlet/Framework/Devices/Gauntlet.DevicePool.cs
+++ b/Engine/Source/Programs/AutomationTool/Gauntlet/Framework/Devices/Gauntlet.DevicePool.cs
@@ -1321,6 +1321,12 @@ namespace Gauntlet
 					}
 				}
 
+				if(!TryValidateDeviceRequirements(NewDevice))
+				{
+					Log.Info("\nSkipping device.");
+					return null;
+				}
+
 				lock (LockObject)
 				{
 					if (NewDevice != null)
@@ -1337,6 +1343,49 @@ namespace Gauntlet
 			return NewDevice;
 		}
 
+		/// <summary>
+		/// Verifies if the provided TargetDevice meets requirements such as firmware, login status, settings, etc.
+		/// </summary>
+		/// <param name="Device">The device to validate</param>
+		/// <returns>True if the device matches the required specifications</returns>
+		protected bool TryValidateDeviceRequirements(ITargetDevice Device)
+		{
+			IEnumerable<IDeviceValidator> Validators = InterfaceHelpers.FindImplementations<IDeviceValidator>(true).Where(Validator => Validator.bEnabled);
+			if(!Validators.Any())
+			{
+				return true;
+			}
+
+			Log.Info("\nValidating requirements for {Device}...", Device);
+			bool bSucceeded = true;
+
+			foreach (IDeviceValidator Validator in Validators)
+			{
+				Log.Info("\nStarting validation for {Validator}", Validator);
+
+				if(!Validator.TryValidateDevice(Device))
+				{
+					Log.Warning("Failure! See above for details");
+					bSucceeded = false;
+				}
+				else
+				{
+					Log.Info("Success!");
+				}
+			}
+
+			if(!bSucceeded)
+			{
+				Log.Warning("\nFailed to validate requirements.");
+			}
+			else
+			{
+				Log.Info("\nAll validators passed, selecting device {Device}\n", Device);
+			}
+
+			return bSucceeded;
+		}
+
 		/// <summary>
 		/// Reporting device issue using the device name provided by service
 		/// </summary>
diff --git a/Engine/Source/Programs/AutomationTool/Gauntlet/Framework/Devices/Gauntlet.DeviceValidator.cs b/Engine/Source/Programs/AutomationTool/Gauntlet/Framework/Devices/Gauntlet.DeviceValidator.cs
new file mode 100644
--- /dev/null
+++ b/Engine/Source/Programs/AutomationTool/Gauntlet/Framework/Devices/Gauntlet.DeviceValidator.cs
@@ -0,0 +1,73 @@
+// Copyright Epic Games, Inc. All Rights Reserved.
+
+namespace Gauntlet
+{
+	/// <summary>
+	/// IDeviceValidator provides the ability to verify a TargetDevice matches a set of requirements.
+	/// They are particularily useful for ensuring development kits are in a desired state prior to a test.
+	/// This can include, but is not limited to:
+	///		- Proper firmware versions
+	///		- Having an account signed in
+	///		- Pushing a list of settings onto the device
+	///	To add new DeviceValidators, simply create a concrete type that inherits from IDeviceValidator.
+	/// Implement ValidateDevice() and determine whether this validator should be enabled within the constructor
+	/// It can be especially helpful to toggle whether a validator is enabled via commandline arguments, for ease of opt-in
+	/// </summary>
+	public interface IDeviceValidator
+	{
+		/// <summary>
+		/// If false, this validator will not be used when reserving devices for automated tests
+		/// </summary>
+		bool bEnabled { get; }
+
+		/// <summary>
+		/// Verifies if the device is valid in the context of the running program
+		/// </summary>
+		/// <param name="Device">The device to validate</param>
+		/// <returns>True if the device succeeded validation</returns>
+		bool TryValidateDevice(ITargetDevice Device);
+
+		/// <summary>
+		/// Helper function that can be used to ensure a device is powered on/connected
+		/// before attempting to run any validation checks
+		/// </summary>
+		/// <param name="Device"></param>
+		/// <returns></returns>
+		protected static bool EnsureConnection(ITargetDevice Device)
+		{
+			if (Device == null)
+			{
+				Log.Warning("Cannot ensure connection for null device");
+				return false;
+			}
+
+			if (!Device.IsOn)
+			{
+				if (!Device.PowerOn())
+				{
+					Log.Warning("Failed to power on device");
+					return false;
+				}
+				else
+				{
+					Log.VeryVerbose("Powered on device while ensuring connection");
+				}
+			}
+
+			if (!Device.IsConnected)
+			{
+				if (!Device.Connect())
+				{
+					Log.Warning("Failed to connect to device");
+					return false;
+				}
+				else
+				{
+					Log.VeryVerbose("Connected to device while ensuring connection");
+				}
+			}
+
+			return Device.IsAvailable;
+		}
+	}
+}
