[Gauntlet] Move VerifyLogin from UnrealSession into DeviceLoginValidator

diff --git a/Engine/Source/Programs/AutomationTool/Gauntlet/Framework/Devices/Gauntlet.DeviceLoginValidator.cs b/Engine/Source/Programs/AutomationTool/Gauntlet/Framework/Devices/Gauntlet.DeviceLoginValidator.cs
new file mode 100644
--- /dev/null
+++ b/Engine/Source/Programs/AutomationTool/Gauntlet/Framework/Devices/Gauntlet.DeviceLoginValidator.cs
@@ -0,0 +1,53 @@
+// Copyright Epic Games, Inc. All Rights Reserved.
+
+using System;
+using System.Collections.Generic;
+using System.Linq;
+using System.Text;
+using System.Threading.Tasks;
+
+namespace Gauntlet
+{
+	/// <summary>
+	/// DeviceLoginValidator works for any ITargetDevice that implements IOnlineServiceLogin (typically consoles)
+	/// This validator ensures the device have an account logged into that platforms online service.
+	/// </summary>
+	public class DeviceLoginValidator : IDeviceValidator
+	{
+		[AutoParamWithNames(false, "VerifyLogin")]
+		public bool bEnabled { get; set; }
+
+		public DeviceLoginValidator()
+		{
+			AutoParam.ApplyParamsAndDefaults(this, Globals.Params.AllArguments);
+		}
+
+		public bool TryValidateDevice(ITargetDevice Device)
+		{
+			if(!IDeviceValidator.EnsureConnection(Device))
+			{
+				return false;
+			}
+
+			if (Device is IOnlineServiceLogin DeviceLogin)
+			{
+				Log.Info("Verifying device login...");
+				if(DeviceLogin.VerifyLogin())
+				{
+					Log.Info("User signed-in");
+					return true;
+				}
+				else
+				{
+					Log.Warning("Unable to secure login to an online platform account!");
+					return false;
+				}
+			}
+			else
+			{
+				Log.Verbose("{Platform} does not implement IOnlineServiceLogin, skipping login validation.", Device.Platform.Value);
+				return true;
+			}
+		}
+	}
+}
diff --git a/Engine/Source/Programs/AutomationTool/Gauntlet/Unreal/Base/Gauntlet.UnrealSession.cs b/Engine/Source/Programs/AutomationTool/Gauntlet/Unreal/Base/Gauntlet.UnrealSession.cs
--- a/Engine/Source/Programs/AutomationTool/Gauntlet/Unreal/Base/Gauntlet.UnrealSession.cs
+++ b/Engine/Source/Programs/AutomationTool/Gauntlet/Unreal/Base/Gauntlet.UnrealSession.cs
@@ -909,20 +909,10 @@ namespace Gauntlet
 						IDeviceUsageReporter.RecordStart(Device.Name, Device.Platform, IDeviceUsageReporter.EventType.Install, IDeviceUsageReporter.EventState.Success, BuildSource.BuildName);
 						try
 						{
-							if ((Role.Options as UnrealTestConfiguration).VerifyLogin && Device is IOnlineServiceLogin)
-							{
-								Log.Info("\nVerifying device login...");
-								if (!(Device as IOnlineServiceLogin).VerifyLogin())
-								{
-									throw new AutomationException("Unable to secure login to an online platform account!");
-								}
-								Log.Info("Success! User signed-in.\n");
-							}
-
 							Install = Device.InstallApplication(AppConfig);
 							IDeviceUsageReporter.RecordEnd(Device.Name, Device.Platform, IDeviceUsageReporter.EventType.Install, IDeviceUsageReporter.EventState.Success);
 						}
-						catch (System.Exception Ex)
+						catch (Exception Ex)
 						{
 							// Warn, ignore the device, and do not continue
 							string ErrorMessage = string.Format("Encountered error setting up device {0} for role {1}. {2}. Will retry with new device", Device, Role, Ex);
diff --git a/Engine/Source/Programs/AutomationTool/Gauntlet/Unreal/Base/Gauntlet.UnrealTestConfiguration.cs b/Engine/Source/Programs/AutomationTool/Gauntlet/Unreal/Base/Gauntlet.UnrealTestConfiguration.cs
--- a/Engine/Source/Programs/AutomationTool/Gauntlet/Unreal/Base/Gauntlet.UnrealTestConfiguration.cs
+++ b/Engine/Source/Programs/AutomationTool/Gauntlet/Unreal/Base/Gauntlet.UnrealTestConfiguration.cs
@@ -690,11 +690,6 @@ namespace Gauntlet
 		/// </summary>
 		public bool AllRolesExit { get; set; }
 
-		/// <summary>
-		/// Whether a given role should attempt login to a device platform service
-		/// </summary>
-		public bool VerifyLogin { get; set; }
-
 		/// <summary>
 		/// The collection of options which define heartbeat behavior
 		/// </summary>
@@ -1005,11 +1000,6 @@ namespace Gauntlet
 				}
 			}
 
-			if (ConfigRole.RoleType.IsClient())
-			{
-				VerifyLogin = Globals.Params.ParseParam("VerifyLogin");
-			}
-
 			if (CommandUtils.IsBuildMachine)
 			{
 				AppConfig.CommandLineParams.AddUnique("BUILDMACHINE");
diff --git a/Engine/Source/Programs/AutomationTool/Gauntlet/Unreal/Game/EpicGame.TestConfig.cs b/Engine/Source/Programs/AutomationTool/Gauntlet/Unreal/Game/EpicGame.TestConfig.cs
--- a/Engine/Source/Programs/AutomationTool/Gauntlet/Unreal/Game/EpicGame.TestConfig.cs
+++ b/Engine/Source/Programs/AutomationTool/Gauntlet/Unreal/Game/EpicGame.TestConfig.cs
@@ -237,18 +237,12 @@ namespace EpicGame
 				// select an account
 				if (!NoMCP)
 				{
-					VerifyLogin = Globals.Params.ParseParam("VerifyLogin");
-
 					if (!bNoAccountOverride && PreAssignAccount)
 					{
 						Account UserAccount = AccountPool.Instance.ReserveAccount();
 						UserAccount.ApplyToConfig(AppConfig);
 					}
 				}
-				else
-				{
-					VerifyLogin = false;
-				}
 			}
 
 			if (ConfigRole.RoleType.IsClient())
