[Gauntlet] Move TryUpdateFirmware from UnrealSession into DeviceFirmwareValidator

diff --git a/Engine/Source/Programs/AutomationTool/Gauntlet/Framework/Devices/Gauntlet.DeviceFirmwareValidator.cs b/Engine/Source/Programs/AutomationTool/Gauntlet/Framework/Devices/Gauntlet.DeviceFirmwareValidator.cs
new file mode 100644
--- /dev/null
+++ b/Engine/Source/Programs/AutomationTool/Gauntlet/Framework/Devices/Gauntlet.DeviceFirmwareValidator.cs
@@ -0,0 +1,195 @@
+// Copyright Epic Games, Inc. All Rights Reserved.
+
+using System;
+using System.Linq;
+using UnrealBuildTool;
+
+namespace Gauntlet
+{
+	public interface IPlatformFirmwareHandler
+	{
+		bool CanSupportPlatform(UnrealTargetPlatform Platform);
+
+		bool CanSupportProject(string ProjectName);
+
+		bool GetDesiredVersion(UnrealTargetPlatform Platform, string ProjectName, out string Version);
+
+		bool GetCurrentVersion(ITargetDevice Device, out string Version);
+
+		bool UpdateDeviceFirmware(ITargetDevice Device, string Version);
+	};
+
+	public class FirmwareValidator : IDeviceValidator
+	{
+		[AutoParam("")]
+		public string FirmwareVersion { get; set; }
+
+		[AutoParamWithNames(Default: "", "Project", "ProjectName")]
+		public string ProjectName { get; set; }
+
+		[AutoParamWithNames(false, "UpdateFirmware", "TryUpdateFirmware")]
+		public bool bEnabled { get; set; }
+
+		public FirmwareValidator()
+		{
+			AutoParam.ApplyParamsAndDefaults(this, Globals.Params.AllArguments);
+
+			// Only enable this validator if a firmware version or a project name is specified
+			bEnabled = bEnabled && (!string.IsNullOrEmpty(FirmwareVersion) || !string.IsNullOrEmpty(ProjectName));
+		}
+
+		public bool TryValidateDevice(ITargetDevice Device)
+		{
+			if(!string.IsNullOrEmpty(FirmwareVersion))
+			{
+				// Requesting an explicit firmware version
+				return TryUpdateFirmware(Device, FirmwareVersion);
+			}
+			else if(!string.IsNullOrEmpty(ProjectName))
+			{
+				// Requesting the firmware version a project requires
+				return TryUpdateFirmwareForProject(Device, ProjectName);
+			}
+
+			return true;
+		}
+
+		/// <summary>
+		/// Attempts to update the provided Device's firmware to the requested version
+		/// </summary>
+		/// <param name="Device">The Device to update</param>
+		/// <param name="DesiredVersion">The firmware version to update to</param>
+		/// <returns></returns>
+		public bool TryUpdateFirmware(ITargetDevice Device, string DesiredVersion)
+		{
+			try
+			{
+				if (!IDeviceValidator.EnsureConnection(Device))
+				{
+					return false;
+				}
+
+				if (Device.Platform == null)
+				{
+					Log.Warning("Cannot update firmware for device {Device} because the platform value is null. Was the device mocked?");
+					return false;
+				}
+
+				if(string.IsNullOrEmpty(DesiredVersion))
+				{
+					Log.Warning("Cannot apply a null or empty firmware version to the device");
+					return false;
+				}
+
+				UnrealTargetPlatform Platform = Device.Platform.Value;
+
+				IPlatformFirmwareHandler Handler = GetRequiredPlatformFirmwareHandler(Platform);
+				if (Handler == null)
+				{
+					Log.Info("Cannot find an IPlatformFirmwareHandler that supports {Platform}. Has one been implemented? Skipping firmware update", Platform);
+					return true;
+				}
+
+				string CurrentVersion;
+				if (!Handler.GetCurrentVersion(Device, out CurrentVersion))
+				{
+					Log.Warning("Failed to retreive current firmware version from {Device}", Device);
+					return false;
+				}
+
+				if (!CurrentVersion.Equals(DesiredVersion, StringComparison.OrdinalIgnoreCase))
+				{
+					if (!Handler.UpdateDeviceFirmware(Device, DesiredVersion))
+					{
+						Log.Warning("Failed to update {Device}'s firmware to version {DesiredVersion}", Device, DesiredVersion);
+						return false;
+					}
+					else
+					{
+						Log.Info("Successfully updated {Device} to new firmware version {DesiredVersion}", Device, DesiredVersion);
+					}
+				}
+				else
+				{
+					Log.Info("Current firmware version {CurrentVersion} matches desired version {DesiredVersion}. Skipping firmware update", CurrentVersion, DesiredVersion);
+				}
+
+				return true;
+			}
+			catch(Exception Ex)
+			{
+				Log.Warning("Encountered an {ExceptionType} while attempting to update firmware. {Exception}", Ex.GetType(), Ex);
+				return false;
+			}
+		}
+
+		/// <summary>
+		/// Attempts to update the provided Device's firmware version to the version requested by the specified project
+		/// </summary>
+		/// <param name="Device">The Device to update</param>
+		/// <param name="ProjectName">Name of the project to get firmware version of</param>
+		/// <returns></returns>
+		public bool TryUpdateFirmwareForProject(ITargetDevice Device, string ProjectName)
+		{
+			try
+			{
+				if (!IDeviceValidator.EnsureConnection(Device))
+				{
+					return false;
+				}
+
+				if (Device.Platform == null)
+				{
+					Log.Warning("Cannot update firmware for device {Device} because the platform value is null. Was the device mocked?");
+					return false;
+				}
+
+				if (string.IsNullOrEmpty(ProjectName))
+				{
+					Log.Warning("Cannot use a null or empty project to determine desired firmware version");
+					return false;
+				}
+
+				UnrealTargetPlatform Platform = Device.Platform.Value;
+
+				IPlatformFirmwareHandler Handler = GetRequiredPlatformFirmwareHandler(Platform, ProjectName);
+
+				if (Handler == null)
+				{
+					Log.Info("Cannot find an IPlatformFirmwareHandler that supports {Platform} and {Project}. Has one been implemented? Skipping firmware update", Platform, ProjectName);
+					return true;
+				}
+
+				string DesiredVersion;
+				if (!Handler.GetDesiredVersion(Platform, ProjectName, out DesiredVersion))
+				{
+					Log.Warning("{Handler} failed to find a desired firmware version for {Platform} and {Project}", Handler.GetType(), Platform, ProjectName);
+					return false;
+				}
+
+				return TryUpdateFirmware(Device, DesiredVersion);
+			}
+			catch (Exception Ex)
+			{
+				Log.Warning("Encountered an {ExceptionType} while attempting to update firmware. {Exception}", Ex.GetType(), Ex);
+				return false;
+			}
+		}
+
+		// Gets the required platform handler that supports the provided platform
+		private IPlatformFirmwareHandler GetRequiredPlatformFirmwareHandler(UnrealTargetPlatform Platform)
+		{
+			return Utils.InterfaceHelpers.FindImplementations<IPlatformFirmwareHandler>(true)
+				.Where(Handler => Handler.CanSupportPlatform(Platform))
+				.FirstOrDefault();
+		}
+
+		// Gets the required platform handler that supports the provided platform and project
+		private IPlatformFirmwareHandler GetRequiredPlatformFirmwareHandler(UnrealTargetPlatform Platform, string ProjectName)
+		{
+			return Utils.InterfaceHelpers.FindImplementations<IPlatformFirmwareHandler>(true)
+				.Where(Handler => Handler.CanSupportPlatform(Platform) && Handler.CanSupportProject(ProjectName))
+				.FirstOrDefault();
+		}
+	}
+};
\ No newline at end of file
diff --git a/Engine/Source/Programs/AutomationTool/Gauntlet/Framework/Devices/Gauntlet.PlatformFirmwareHandler.cs b/Engine/Source/Programs/AutomationTool/Gauntlet/Framework/Devices/Gauntlet.PlatformFirmwareHandler.cs
deleted file mode 100644
diff --git a/Engine/Source/Programs/AutomationTool/Gauntlet/Unreal/Base/Gauntlet.UnrealSession.cs b/Engine/Source/Programs/AutomationTool/Gauntlet/Unreal/Base/Gauntlet.UnrealSession.cs
--- a/Engine/Source/Programs/AutomationTool/Gauntlet/Unreal/Base/Gauntlet.UnrealSession.cs
+++ b/Engine/Source/Programs/AutomationTool/Gauntlet/Unreal/Base/Gauntlet.UnrealSession.cs
@@ -832,66 +832,6 @@ namespace Gauntlet
 						}
 					}
 
-					//Verify the device's OS version, and update if necessary
-					if (Globals.Params.ParseParam("TryFirmwareUpdate") && AppConfig.Platform != null)
-					{
-						List<IPlatformFirmwareHandler> PlatformFirmwareHandlers = Gauntlet.Utils.InterfaceHelpers.FindImplementations<IPlatformFirmwareHandler>(true).ToList();
-
-						PlatformFirmwareHandlers = PlatformFirmwareHandlers.Where(FC => FC.CanSupportPlatform((UnrealTargetPlatform)AppConfig.Platform)
-							&& FC.CanSupportProject(AppConfig.ProjectName)).ToList();
-						if (PlatformFirmwareHandlers.Count > 0)
-						{
-							IPlatformFirmwareHandler SelectedFirmwareHandler = PlatformFirmwareHandlers.First();
-							Log.Verbose("Found IPlatformFirmwareHandler {0}", SelectedFirmwareHandler.GetType().Name);
-							string DesiredFirmware = string.Empty;
-							if (!SelectedFirmwareHandler.GetDesiredVersion((UnrealTargetPlatform)AppConfig.Platform, AppConfig.ProjectName, out DesiredFirmware))
-							{
-								Log.Info("Failed to get desired os version for project {0} for platform {1}, skipping firmware check",
-									AppConfig.ProjectName, AppConfig.Platform);
-							}
-							else
-							{
-								Log.Info("Desired Firmware for project {0} and platform {1}: {2}", AppConfig.ProjectName, AppConfig.Platform, DesiredFirmware);
-
-								string CurrentFirmware = string.Empty;
-								if (!SelectedFirmwareHandler.GetCurrentVersion(Device, out CurrentFirmware))
-								{
-									Log.Info("Failed to get current os version for device {0} for role {1}, skipping firmware check", Device, Role);
-								}
-								else
-								{
-									Log.Info("Current Firmware for device {0}: {1}", Device, CurrentFirmware);
-
-									if(CurrentFirmware.Equals(DesiredFirmware))
-									{
-										Log.Info("Device {0} os version match!  No need to update", Device);
-									}
-									else
-									{
-										Log.Info("Device {0} os version out of date!  Updating to version {1}", Device, DesiredFirmware);
-
-										if(SelectedFirmwareHandler.UpdateDeviceFirmware(Device, DesiredFirmware))
-										{
-											Log.Info("Successfully updated device {0} to os version {1}", Device, DesiredFirmware);
-										}
-										else
-										{
-											Log.Info("Failed to update os of device {0} for role {1}.  Will retry with new device", Device, Role);
-											UnrealDeviceReservation.MarkProblemDevice(Device);
-											InstallSuccess = false;
-											break;
-										}
-									}
-								}
-							}
-						}
-						else
-						{
-							Log.Info("Unable to locate any IPlatformFirmwareCheckers that support Project {0} and platform {1}, skipping firmware check",
-								AppConfig.ProjectName, AppConfig.Platform);
-						}
-					}
-
 					// todo - should this be elsewhere?
 					AppConfig.Sandbox = Sandbox;
 
