[Gauntlet] Move Configurable device application into DeviceProfileValidator

diff --git a/Engine/Source/Programs/AutomationTool/Gauntlet/Framework/Base/Gauntlet.DeviceConfigurationProfile.cs b/Engine/Source/Programs/AutomationTool/Gauntlet/Framework/Base/Gauntlet.DeviceConfigurationProfile.cs
deleted file mode 100644
diff --git a/Engine/Source/Programs/AutomationTool/Gauntlet/Framework/Base/Gauntlet.TargetDevice.cs b/Engine/Source/Programs/AutomationTool/Gauntlet/Framework/Base/Gauntlet.TargetDevice.cs
--- a/Engine/Source/Programs/AutomationTool/Gauntlet/Framework/Base/Gauntlet.TargetDevice.cs
+++ b/Engine/Source/Programs/AutomationTool/Gauntlet/Framework/Base/Gauntlet.TargetDevice.cs
@@ -237,23 +237,4 @@ namespace Gauntlet
 
 		public virtual bool NeedBuildDeployed() => true;
 	}
-
-	/// <summary>
-	/// Represents a device that is configurable
-	/// </summary>
-	public interface IConfigurableDevice
-	{
-		/// <summary>
-		/// Returns a configuration profile with the current device settings
-		/// </summary>
-		/// <returns></returns>
-		PlatformConfigurationBase GetCurrentConfigurationSnapshot();
-
-		/// <summary>
-		/// Applies a configuration profile to the device. If required it will reboot the device
-		/// </summary>
-		/// <param name="Configuration"></param>
-		/// <returns>false if applying the config profile fails</returns>
-		bool ApplyConfiguration(PlatformConfigurationBase Configuration);
-	}
 }
\ No newline at end of file
diff --git a/Engine/Source/Programs/AutomationTool/Gauntlet/Framework/Devices/Gauntlet.ConfigurableDevice.cs b/Engine/Source/Programs/AutomationTool/Gauntlet/Framework/Devices/Gauntlet.ConfigurableDevice.cs
new file mode 100644
--- /dev/null
+++ b/Engine/Source/Programs/AutomationTool/Gauntlet/Framework/Devices/Gauntlet.ConfigurableDevice.cs
@@ -0,0 +1,262 @@
+// Copyright Epic Games, Inc. All Rights Reserved.
+
+using EpicGames.Core;
+using System;
+using System.Collections.Generic;
+using System.IO;
+using System.Linq;
+using UnrealBuildTool;
+
+namespace Gauntlet
+{
+	/// <summary>
+	/// ITargetDevice extension for platforms that support configurable settings
+	/// </summary>
+	public interface IConfigurableDevice
+	{
+		/// <summary>
+		/// Returns a configuration profile with the current device settings
+		/// </summary>
+		/// <returns></returns>
+		PlatformConfigurationBase GetCurrentConfigurationSnapshot();
+
+		/// <summary>
+		/// Applies a configuration profile to the device. If required it will reboot the device
+		/// </summary>
+		/// <param name="Configuration"></param>
+		/// <returns>false if applying the config profile fails</returns>
+		bool ApplyConfiguration(PlatformConfigurationBase Configuration);
+	}
+
+	/// <summary>
+	/// Base interface for a platform specific configuration reader
+	/// </summary>
+	public interface IPlatformConfigurationReader
+	{
+		bool SupportsPlatform(UnrealTargetPlatform? Platform);
+
+		/// <summary>
+		/// Returns the supported platform config file extension
+		/// </summary>
+		/// <returns></returns>
+		string ConfigFileExtension();
+
+		/// <summary>
+		/// Reads the configuration from the passed in file location
+		/// </summary>
+		/// <param name="Location">An absolute path the to a file with platform configuration</param>
+		/// <returns></returns>
+		PlatformConfigurationBase ReadConfiguration(FileReference Location);
+	}
+
+	/// <summary>
+	/// Base class for a platform configuration profile.
+	/// This class needs to be implemented for each configurable platform
+	/// </summary>
+	public abstract class PlatformConfigurationBase
+	{
+		/// <summary>
+		/// Name of the config profile.
+		/// This name doesn't need to be unique across namespaces.
+		/// </summary>
+		public string ProfileName { get; set; }
+
+		/// <summary>
+		/// Namespace of the profile. Profiles have to have a unique name inside a namespace.
+		/// </summary>
+		public string Namespace { get; set; }
+
+		/// <summary>
+		/// Platform this configuration profile is for
+		/// </summary>
+		public UnrealTargetPlatform Platform { get; set; }
+	}
+
+	/// <summary>
+	/// A singleton class that encapsulates a cache of device config profiles
+	/// </summary>
+	public class DeviceConfigurationCache
+	{
+		public static DeviceConfigurationCache Instance { get; private set; } = new();
+
+		private Object LockObject = new Object();
+		private Dictionary<ConfigurationCacheKey, PlatformConfigurationBase> ConfigurationCache = new();
+
+		protected DeviceConfigurationCache()
+		{
+			Instance = this;
+		}
+
+		/// <summary>
+		/// Key by which a configuration profile is identified in the cache.
+		/// </summary>
+		class ConfigurationCacheKey
+		{
+			public UnrealTargetPlatform Platform;
+			public string Namespace;
+			public string ProfileName;
+
+			public ConfigurationCacheKey(UnrealTargetPlatform Platform, string ProjectName, string ProfileName)
+			{
+				this.Platform = Platform;
+				this.Namespace = ProjectName;
+				this.ProfileName = ProfileName;
+			}
+
+			public ConfigurationCacheKey(PlatformConfigurationBase Configuration)
+			{
+				this.Platform = Configuration.Platform;
+				this.Namespace = Configuration.Namespace;
+				this.ProfileName = Configuration.ProfileName;
+			}
+
+			public override bool Equals(object Other)
+			{
+				ConfigurationCacheKey OtherKey = Other as ConfigurationCacheKey;
+				return OtherKey != null &&
+					OtherKey.Platform == Platform &&
+					OtherKey.ProfileName == ProfileName &&
+					OtherKey.Namespace == Namespace;
+
+			}
+
+			public override int GetHashCode()
+			{
+				return (Platform, Namespace, ProfileName).GetHashCode();
+			}
+		}
+
+		/// <summary>
+		/// Scans the SettingDir for configuration profiles
+		/// </summary>
+		/// <param name="Platform">Platform for which to look for config files</param>
+		/// <param name="Namespace">The namespace in which to put the configuration profiles</param>
+		/// <param name="SettingDir"></param>
+		/// <returns></returns>
+		public bool DiscoverConfigurationProfiles(UnrealTargetPlatform? Platform, string Namespace, string SettingDir)
+		{
+			if (Platform == null || !Directory.Exists(SettingDir))
+			{
+				return false;
+			}
+
+			IPlatformConfigurationReader ConfigReader = Utils.InterfaceHelpers.FindImplementations<IPlatformConfigurationReader>(true)
+				.Where(D => D.SupportsPlatform(Platform))
+				.FirstOrDefault();
+
+			if (ConfigReader == null)
+			{
+				Log.Info("Couldn't find a configuration reader for {0}", Platform);
+				return false;
+			}
+
+			foreach (string File in Directory.EnumerateFiles(SettingDir, ConfigReader.ConfigFileExtension(), SearchOption.AllDirectories))
+			{
+				FileReference FileRef = new FileReference(File);
+				PlatformConfigurationBase DeviceProfile = ConfigReader.ReadConfiguration(FileRef);
+				if (DeviceProfile != null)
+				{
+					DeviceProfile.Platform = Platform.Value;
+					DeviceProfile.Namespace = Namespace;
+					DeviceProfile.ProfileName = FileRef.GetFileNameWithoutAnyExtensions();
+
+					lock (LockObject)
+					{
+						ConfigurationCacheKey Key = new ConfigurationCacheKey(Platform.Value, Namespace, DeviceProfile.ProfileName);
+						if (ConfigurationCache.ContainsKey(Key))
+						{
+							Log.Verbose("Device configuration profile {0} already exists", FileRef.FullName);
+							continue;
+						}
+
+						ConfigurationCache.Add(Key, DeviceProfile);
+					}
+				}
+			}
+
+			return true;
+		}
+
+		public void CacheConfigurationSnapshot(PlatformConfigurationBase Configuration, bool Overwrite = false)
+		{
+			lock (LockObject)
+			{
+				ConfigurationCacheKey Key = new ConfigurationCacheKey(Configuration.Platform, "Snapshot", Configuration.ProfileName);
+				if (ConfigurationCache.ContainsKey(Key))
+				{
+					if(Overwrite)
+					{
+						ConfigurationCache[Key] = Configuration;
+					}
+				}
+				else
+				{
+					ConfigurationCache.Add(Key, Configuration);
+				}
+			}
+		}
+
+		public PlatformConfigurationBase GetConfigurationSnapshot(UnrealTargetPlatform? Platform, string DeviceName)
+		{
+			return GetConfiguration(Platform, "Snapshot", DeviceName);
+		}
+
+		public void ClearSnapshot(PlatformConfigurationBase Snapshot)
+		{
+			lock (LockObject)
+			{
+				ConfigurationCacheKey Key = new ConfigurationCacheKey(Snapshot);
+				ConfigurationCache.Remove(Key);
+			}
+		}
+
+		public void RevertDeviceConfiguration(ITargetDevice Device)
+		{
+			if (Device is IConfigurableDevice ConfigurableDevice)
+			{
+				var Snapshot = GetConfigurationSnapshot(Device.Platform, Device.Name);
+				if (Snapshot == null)
+				{
+					return;
+				}
+
+				// Connect temporarily to be able to revert the device's configuration
+				// if the device was disconnected entering here disconnect it after this is over
+				bool bNeedsDisconnect = false;
+				if (!Device.IsConnected)
+				{
+					Device.Connect();
+					bNeedsDisconnect = true;
+				}
+
+				if (ConfigurableDevice.ApplyConfiguration(Snapshot))
+				{
+					ClearSnapshot(Snapshot);
+				}
+
+				if (bNeedsDisconnect)
+				{
+					Device.Disconnect();
+				}
+			}
+		}
+
+		public PlatformConfigurationBase GetConfiguration(UnrealTargetPlatform? Platform, string Namespace, string ProfileName)
+		{
+			if (Platform == null)
+			{
+				return null;
+			}
+
+			PlatformConfigurationBase Res = null;
+
+			lock (LockObject)
+			{
+				ConfigurationCacheKey Key = new ConfigurationCacheKey(Platform.Value, Namespace, ProfileName);
+				ConfigurationCache.TryGetValue(Key, out Res);
+			}
+
+			return Res;
+		}
+	}
+}
diff --git a/Engine/Source/Programs/AutomationTool/Gauntlet/Framework/Devices/Gauntlet.DevicePool.cs b/Engine/Source/Programs/AutomationTool/Gauntlet/Framework/Devices/Gauntlet.DevicePool.cs
--- a/Engine/Source/Programs/AutomationTool/Gauntlet/Framework/Devices/Gauntlet.DevicePool.cs
+++ b/Engine/Source/Programs/AutomationTool/Gauntlet/Framework/Devices/Gauntlet.DevicePool.cs
@@ -1052,7 +1052,6 @@ namespace Gauntlet
 					// @todo: device problem reporting, requesting additional devices
 					if (TargetDevice == null)
 					{
-
 						ReportDeviceError(Device.Name, "CreateDeviceError");
 						ReservationStates.MarkProblem(Device);
 
@@ -1138,7 +1137,7 @@ namespace Gauntlet
 				List<ITargetDevice> KeepList = new List<ITargetDevice>();
 				foreach (ITargetDevice Device in DeviceList)
 				{
-					RevertDeviceConfiguration(Device);
+					DeviceConfigurationCache.Instance.RevertDeviceConfiguration(Device);
 
 					if (Device.IsConnected && InitialConnectState[Device] == false)
 					{
@@ -1321,7 +1320,8 @@ namespace Gauntlet
 					}
 				}
 
-				if(!TryValidateDeviceRequirements(NewDevice))
+				// Now validate if the kit meets the necessary requirements
+				if (!TryValidateDeviceRequirements(NewDevice))
 				{
 					Log.Info("\nSkipping device.");
 					return null;
@@ -1351,11 +1351,13 @@ namespace Gauntlet
 		protected bool TryValidateDeviceRequirements(ITargetDevice Device)
 		{
 			IEnumerable<IDeviceValidator> Validators = InterfaceHelpers.FindImplementations<IDeviceValidator>(true).Where(Validator => Validator.bEnabled);
-			if(!Validators.Any())
+			if (!Validators.Any())
 			{
 				return true;
 			}
 
+			bool bInitiallyConnected = Device.IsConnected;
+
 			Log.Info("\nValidating requirements for {Device}...", Device);
 			bool bSucceeded = true;
 
@@ -1383,6 +1385,13 @@ namespace Gauntlet
 				Log.Info("\nAll validators passed, selecting device {Device}\n", Device);
 			}
 
+			// Most validators require establishing a connection to the device.
+			// If we weren't originally connected, disconnect so the initial connection state can be cached during device reservation
+			if(!bInitiallyConnected && Device.IsConnected)
+			{
+				Device.Disconnect();
+			}
+
 			return bSucceeded;
 		}
 
@@ -1407,7 +1416,6 @@ namespace Gauntlet
 		/// </summary>
 		private void ReleaseReservations()
 		{
-
 			ServiceDeviceInfo.Clear();
 
 			foreach (ITargetDevice Device in ServiceReservations.Keys)
@@ -1441,38 +1449,6 @@ namespace Gauntlet
 			}
 		}
 
-		private void RevertDeviceConfiguration(ITargetDevice Device)
-		{
-			IConfigurableDevice ConfigurableDevice = Device as IConfigurableDevice;
-			if (ConfigurableDevice != null)
-			{
-				var Snapshot = DeviceConfigurationCache.Instance.GetConfigurationSnapshot(Device.Platform, Device.Name);
-				if (Snapshot == null)
-				{
-					return;
-				}
-
-				// Connect temporarily to be able to revert the device's configuration
-				// if the device was disconnected entering here disconnect it after this is over
-				bool bNeedsDisconnect = false;
-				if (!Device.IsConnected)
-				{
-					Device.Connect();
-					bNeedsDisconnect = true;
-				}
-
-				if (ConfigurableDevice.ApplyConfiguration(Snapshot))
-				{
-					DeviceConfigurationCache.Instance.ClearSnapshot(Snapshot);
-				}
-
-				if (bNeedsDisconnect)
-				{
-					Device.Disconnect();
-				}
-			}
-		}
-
 		/// <summary>
 		/// Construct a path to hold cache files and make sure it's properly cleaned
 		/// </summary>
diff --git a/Engine/Source/Programs/AutomationTool/Gauntlet/Framework/Devices/Gauntlet.DeviceProfileValidator.cs b/Engine/Source/Programs/AutomationTool/Gauntlet/Framework/Devices/Gauntlet.DeviceProfileValidator.cs
new file mode 100644
--- /dev/null
+++ b/Engine/Source/Programs/AutomationTool/Gauntlet/Framework/Devices/Gauntlet.DeviceProfileValidator.cs
@@ -0,0 +1,83 @@
+// Copyright Epic Games, Inc. All Rights Reserved.
+
+using System;
+
+namespace Gauntlet
+{
+	public class DeviceProfileValidator : IDeviceValidator
+	{
+		public bool bEnabled { get; }
+
+		[AutoParam("")]
+		public string DeviceProfile { get; set; }
+
+		private const string DefaultNamespace = "Engine";
+
+		public DeviceProfileValidator()
+		{
+			AutoParam.ApplyParamsAndDefaults(this, Globals.Params.AllArguments);
+			bEnabled = !string.IsNullOrEmpty(DeviceProfile);
+		}
+
+		public bool TryValidateDevice(ITargetDevice Device)
+		{
+			if(!IDeviceValidator.EnsureConnection(Device))
+			{
+				return false;
+			}
+
+			// Device profile comes in the format <namespace>.<profilename>. If no namespace is given, assume DefaultNamespace
+			string Namespace;
+			string Profile;
+			if (DeviceProfile.Contains('.'))
+			{
+				int DotIndex = DeviceProfile.IndexOf('.');
+				Namespace = DeviceProfile.Substring(0, DotIndex);
+				Profile = DeviceProfile.Substring(DotIndex + 1);
+			}
+			else
+			{
+				Namespace = DefaultNamespace;
+				Profile = DeviceProfile;
+			}
+
+			if(Device is IConfigurableDevice ConfigurableDevice)
+			{
+				try
+				{
+					// Create a configuration based on the device profile we requested
+					PlatformConfigurationBase DesiredConfiguration = DeviceConfigurationCache.Instance.GetConfiguration(Device.Platform, Namespace, Profile);
+					if (DesiredConfiguration == null)
+					{
+						Log.Warning("Failed to locate device profile '{Namespace}.{Profile}' for {Platform} doesn't exist", Namespace, Profile, Device.Platform);
+						return false;
+					}
+
+					// Cache the existing configuration on the device
+					PlatformConfigurationBase CurrentConfiguration = ConfigurableDevice.GetCurrentConfigurationSnapshot();
+					if (CurrentConfiguration == null)
+					{
+						Log.Warning("Failed to snapshot existing device configuration");
+						return false;
+					}
+
+					// Cache the existing configuration, we'll return to this state when we have finished with the device
+					DeviceConfigurationCache.Instance.CacheConfigurationSnapshot(CurrentConfiguration);
+
+					// Now apply the requested configuration
+					return ConfigurableDevice.ApplyConfiguration(DesiredConfiguration);
+				}
+				catch(Exception Ex)
+				{
+					Log.Warning("Encountered an exception while trying to configure device.\n{ExceptionType}: {Message}\n\t{Stack}", Ex.GetType(), Ex.Message, Ex.StackTrace);
+					return false;
+				}
+			}
+			else
+			{
+				Log.Verbose("{Platform} does not implement IConfigurableDevice, skipping device configuration.", Device.Platform.Value);
+				return true;
+			}
+		}
+	}
+}
diff --git a/Engine/Source/Programs/AutomationTool/Gauntlet/Unreal/Base/Gauntlet.UnrealSession.cs b/Engine/Source/Programs/AutomationTool/Gauntlet/Unreal/Base/Gauntlet.UnrealSession.cs
--- a/Engine/Source/Programs/AutomationTool/Gauntlet/Unreal/Base/Gauntlet.UnrealSession.cs
+++ b/Engine/Source/Programs/AutomationTool/Gauntlet/Unreal/Base/Gauntlet.UnrealSession.cs
@@ -232,7 +232,7 @@ namespace Gauntlet
 	/// <summary>
 	/// Represents an instance of a running an Unreal session. Basically an aggregate of all processes for
 	/// all roles (clients, server, etc
-	/// 
+	///
 	/// TODO - combine this into UnrealSession
 	/// </summary>
 	public class UnrealSessionInstance : IDisposable
@@ -272,7 +272,7 @@ namespace Gauntlet
 		/// All roles
 		/// </summary>
 		public RoleInstance[] AllRoles { get; protected set; }
-		
+
 		/// <summary>
 		/// All running roles
 		/// </summary>
@@ -691,28 +691,6 @@ namespace Gauntlet
 			}
 		}
 
-		public bool TryConfigureDevice(ITargetDevice Device, string ProjectName, string ProfileName)
-		{
-			IConfigurableDevice ConfigurableDevice = Device as IConfigurableDevice;
-			if (ConfigurableDevice == null)
-			{
-				Log.Info("Trying to set a config profile on a device that isn't configurable: {0}", Device.Platform);
-				return false;
-			}
-
-			var Configuration = DeviceConfigurationCache.Instance.GetConfiguration(Device.Platform, ProjectName, ProfileName);
-			if (Configuration == null)
-			{
-				Log.Info("Device profile {0}.{1} for {2} doesn't exist", ProjectName, ProfileName, Device.Platform);
-				return false;
-			}
-
-			var ConfigSnapshot = ConfigurableDevice.GetCurrentConfigurationSnapshot();
-			DeviceConfigurationCache.Instance.CacheConfigurationSnapshot(ConfigSnapshot);
-			
-			return ConfigurableDevice.ApplyConfiguration(Configuration);
-		}
-
 		/// <summary>
 		/// Check that all the current roles can be performed by our build source
 		/// </summary>
@@ -814,24 +792,6 @@ namespace Gauntlet
 					// create a config from the build source (this also applies the role options)
 					UnrealAppConfig AppConfig = BuildSource.CreateConfiguration(Role, OtherRoles);
 
-					// device profile comes in the format <namespace>.<profilename>. If no namespace is given, assume Engine
-					string ProfileName = Globals.Params.ParseValue("deviceprofile", "");
-					string Namespace = "Engine";
-					int DotIndex = ProfileName.IndexOf('.');
-					if (DotIndex != -1)
-					{
-						Namespace = ProfileName.Substring(0, DotIndex);
-						ProfileName = ProfileName.Substring(DotIndex + 1);
-					}
-					
-					if (!string.IsNullOrWhiteSpace(ProfileName))
-					{
-						if (!TryConfigureDevice(Device, Namespace, ProfileName))
-						{
-							Log.Info("Failed to apply {0} configuration profile to device {1}", ProfileName, Device.Name);
-						}
-					}
-
 					// todo - should this be elsewhere?
 					AppConfig.Sandbox = Sandbox;
 
diff --git a/Engine/Source/Programs/AutomationTool/Gauntlet/Unreal/RunUnreal.cs b/Engine/Source/Programs/AutomationTool/Gauntlet/Unreal/RunUnreal.cs
--- a/Engine/Source/Programs/AutomationTool/Gauntlet/Unreal/RunUnreal.cs
+++ b/Engine/Source/Programs/AutomationTool/Gauntlet/Unreal/RunUnreal.cs
@@ -458,9 +458,6 @@ namespace Gauntlet
 			DevicePool.Instance.AddLocalDevices(Options.MaxLocalDevices);
 			DevicePool.Instance.AddVirtualDevices(Options.MaxVirtualDevices);
 
-			string EngineDeviceConfigDir = Path.Combine(Globals.UnrealRootDir, "Engine", "Build", "DeviceConfigProfiles");
-			string ProjectDeviceConfigDir = Path.Combine(Options.ProjectPath.Directory.FullName, "Build", "DeviceConfigProfiles");
-
 			foreach (var DeviceWithParams in Options.DeviceList)
 			{
 				UnrealTargetPlatform Platform = DefaultPlatform;
