Fix Platforms, Restricted, and Shaders not being included for consideration
Fix numeric separators being mistaken for single characters (which were being treated as strings) and consuming text it shouldn't have
Add detection for missing #undef LOCTEXT_NAMESPACE if one has been defined in the file

diff --git a/Engine/Source/Programs/AutomationTool/Scripts/CheckBalancedMacros.cs b/Engine/Source/Programs/AutomationTool/Scripts/CheckBalancedMacros.cs
--- a/Engine/Source/Programs/AutomationTool/Scripts/CheckBalancedMacros.cs
+++ b/Engine/Source/Programs/AutomationTool/Scripts/CheckBalancedMacros.cs
@@ -2,9 +2,11 @@
 
 using System;
 using System.Collections.Generic;
+using System.Diagnostics;
 using System.IO;
 using System.Linq;
 using System.Text;
+using System.Text.RegularExpressions;
 using System.Threading.Tasks;
 using EpicGames.Core;
 using UnrealBuildBase;
@@ -20,6 +22,11 @@ namespace AutomationTool
 	[Help("Ignore=<Name>", "File name (without path) to exclude from testing")]
 	class CheckBalancedMacros : BuildCommand
 	{
+		/// <summary>
+		/// List of directories relative to the root that can contain source files
+		/// </summary>
+		public static readonly string[] SourceDirectories = { "Platforms", "Plugins", "Restricted", "Shaders", "Source" };
+
 		/// <summary>
 		/// List of macros that should be paired up
 		/// </summary>
@@ -75,6 +82,13 @@ namespace AutomationTool
 			},
 		};
 
+		/// <summary>
+		/// Regexes for LOCTEXT_NAMESPACE preprocessor identification
+		/// This could be generalized like the above if we had other pairings we wanted to manage
+		/// </summary>
+		static readonly Regex OpenLoctextNamespaceRegex = new Regex(@"\G#define\s+LOCTEXT_NAMESPACE");
+		static readonly Regex CloseLoctextNamespaceRegex = new Regex(@"\G#undef\s+LOCTEXT_NAMESPACE");
+
 		/// <summary>
 		/// List of files to ignore for balanced macros. Additional filenames may be specified on the command line via -Ignore=...
 		/// </summary>
@@ -202,16 +216,13 @@ namespace AutomationTool
 				{
 					foreach(DirectoryReference RootDir in RootDirs)
 					{
-						DirectoryInfo PluginsDir = new DirectoryInfo(Path.Combine(RootDir.FullName, "Plugins"));
-						if(PluginsDir.Exists)
-						{
-							Queue.Enqueue(() => FindSourceFiles(PluginsDir, SourceFiles, Queue));
-						}
-
-						DirectoryInfo SourceDir = new DirectoryInfo(Path.Combine(RootDir.FullName, "Source"));
-						if(SourceDir.Exists)
+						foreach (String Directory in SourceDirectories)
 						{
-							Queue.Enqueue(() => FindSourceFiles(SourceDir, SourceFiles, Queue));
+							DirectoryInfo SourceDir = new DirectoryInfo(Path.Combine(RootDir.FullName, Directory));
+							if(SourceDir.Exists)
+							{
+								Queue.Enqueue(() => FindSourceFiles(SourceDir, SourceFiles, Queue));
+							}
 						}
 					}
 					Queue.Wait();
@@ -281,6 +292,7 @@ namespace AutomationTool
 
 			// Scan through the file token by token. Each bit in the Flags array indicates an index into the MacroPairs array that is currently active.
 			int Flags = 0;
+			bool LoctextNamespaceOpen = false;
 			for(int Idx = 0; Idx < Text.Length; )
 			{
 				int StartIdx = Idx++;
@@ -361,12 +373,12 @@ namespace AutomationTool
 						}
 					}
 				}
-				else if(Text[StartIdx] == '"' || Text[StartIdx] == '\'')
+				else if(Text[StartIdx] == '"')
 				{
 					// String
 					for(; Idx < Text.Length; Idx++)
 					{
-						if(Text[Idx] == Text[StartIdx])
+						if(Text[Idx] == '"')
 						{
 							Idx++;
 							break;
@@ -377,8 +389,52 @@ namespace AutomationTool
 						}
 					}
 				}
+				else if(Text[StartIdx] == '\'')
+				{
+					// Escaped character (e.g. \n, \', \xAB, \xFFFF)
+					if(Text[StartIdx + 1] == '\\')
+					{
+						Idx += 3;
+						for (; Idx < Text.Length; Idx++)
+						{
+							if (Text[Idx] == '\'')
+							{
+								Idx++;
+								break;
+							}
+						}
+					}
+					// Standard single character
+					else if(Text[StartIdx + 2] == '\'')
+					{
+						Idx += 3;
+					}
+					// Otherwise this is probably a numeric separator and we're going to ignore it
+				}
 				else if(Text[StartIdx] == '#')
 				{
+					// Do detection of LOCTEXT_NAMESPACE directives being properly closed
+					// It is theoretically valid to redefine LOCTEXT_NAMESPACE, so this is simply
+					// ensuring there is a closing #undef 
+
+					// Peek ahead to try and reduce the number of regex comparisons we make
+					if(!LoctextNamespaceOpen && Text[StartIdx + 1] == 'd') 
+					{
+						Match m = OpenLoctextNamespaceRegex.Match(Text, StartIdx);
+						if (m.Success && m.Index == StartIdx)
+						{
+							LoctextNamespaceOpen = true;
+						}
+					}
+					else if(LoctextNamespaceOpen && Text[StartIdx + 1] == 'u')
+					{
+						Match m = CloseLoctextNamespaceRegex.Match(Text, StartIdx);
+						if (m.Success && m.Index == StartIdx)
+						{
+							LoctextNamespaceOpen = false;
+						}
+					}
+
 					// Preprocessor directive (eg. #define)
 					for(; Idx < Text.Length && Text[Idx] != '\n'; Idx++)
 					{
@@ -401,6 +457,10 @@ namespace AutomationTool
 					}
 				}
 			}
+			if(LoctextNamespaceOpen)
+			{
+				EpicGames.Core.Log.TraceWarningTask(SourceFile, "#define NAMESPACE_LOCTEXT preprocessor directive is missing matching #undef NAMESPACE_LOCTEXT directive");
+			}
 		}
 
 		/// <summary>
