Fix single character overconsuming

diff --git a/Engine/Source/Programs/AutomationTool/Scripts/CheckBalancedMacros.cs b/Engine/Source/Programs/AutomationTool/Scripts/CheckBalancedMacros.cs
--- a/Engine/Source/Programs/AutomationTool/Scripts/CheckBalancedMacros.cs
+++ b/Engine/Source/Programs/AutomationTool/Scripts/CheckBalancedMacros.cs
@@ -394,7 +394,7 @@ namespace AutomationTool
 					// Escaped character (e.g. \n, \', \xAB, \xFFFF)
 					if(Text[StartIdx + 1] == '\\')
 					{
-						Idx += 3;
+						Idx += 2;
 						for (; Idx < Text.Length; Idx++)
 						{
 							if (Text[Idx] == '\'')
@@ -407,7 +407,7 @@ namespace AutomationTool
 					// Standard single character
 					else if(Text[StartIdx + 2] == '\'')
 					{
-						Idx += 3;
+						Idx += 2;
 					}
 					// Otherwise this is probably a numeric separator and we're going to ignore it
 				}
