Fix the use of 'AdditionalPakOptions' when '-Upload=LocalZen' is set during staging

#rnx

- We need to add a whitespace before appending 'AdditionalPakOptions' aotherwise they will be merged with the previous '-WriteTocToDisk' argument.
- We should only do this if we actually have additional pak options to avoid adding a pointless additional ' ' if there are no additional options.

diff --git a/Engine/Source/Programs/AutomationTool/Scripts/CopyBuildToStagingDirectory.Automation.cs b/Engine/Source/Programs/AutomationTool/Scripts/CopyBuildToStagingDirectory.Automation.cs
--- a/Engine/Source/Programs/AutomationTool/Scripts/CopyBuildToStagingDirectory.Automation.cs
+++ b/Engine/Source/Programs/AutomationTool/Scripts/CopyBuildToStagingDirectory.Automation.cs
@@ -3738,7 +3738,10 @@ namespace AutomationScripts
 					UploadArgs.Append(" -KeepPakFiles");
 					UploadArgs.Append(" -BuildVersion=0");
 					UploadArgs.Append(" -WriteTocToDisk");
-					UploadArgs.Append(Params.AdditionalPakOptions);
+					if (Params.AdditionalPakOptions.Length > 0)
+					{
+						UploadArgs.Append(" " + Params.AdditionalPakOptions);
+					}
 
 					FileReference CryptoKeysFilename = FileReference.Combine(SC.MetadataDir, "Crypto.json");
 					if (FileExists_NoExceptions(CryptoKeysFilename.FullName))
