GenerateStreamingInstallManifest - Add AllowPerChunkCompressionWildcard to allow pak files to opt into or out of being compressed through CompressedChunkWildcard.
CopyBuildToStagingDirectory - Use bAllowPerChunkCompression to enable bCompressed to override global compression settings.

This will correctly allow chunks to opt in and out of being compressed.
NOTE: If DDPI specifies a hardware compression setting of 'None', this won't work as expected because there will be no global compression settings to opt in to. In this case,
set bForceUseProjectCompressionFormatIgnoreHardwareOverride=true and bCompressed=False in [/Script/UnrealEd.ProjectPackagingSettings], then setup whatever project compression settings you would like chunks to
be able to opt in to.

diff --git a/Engine/Source/Editor/UnrealEd/Private/Commandlets/AssetRegistryGenerator.cpp b/Engine/Source/Editor/UnrealEd/Private/Commandlets/AssetRegistryGenerator.cpp
--- a/Engine/Source/Editor/UnrealEd/Private/Commandlets/AssetRegistryGenerator.cpp
+++ b/Engine/Source/Editor/UnrealEd/Private/Commandlets/AssetRegistryGenerator.cpp
@@ -365,23 +365,8 @@ bool FAssetRegistryGenerator::GenerateStreamingInstallManifest(int64 InOverrideC
 		return false;
 	}
 
-	TArray<FString> CompressedChunkWildcards;
-	TOptional<FConfigFile> PlatformIniFile;
-	auto ConditionalLoadPlatformIniFile = [&PlatformIniFile, this]()
-	{
-		if (!PlatformIniFile)
-		{
-			PlatformIniFile.Emplace();
-			FConfigCacheIni::LoadLocalIniFile(*PlatformIniFile, TEXT("Game"), true, *TargetPlatform->IniPlatformName());
-		}
-	};
-	if (!TargetPlatform->IsServerOnly())
-	{
-		// Load the list of wildcards to specify which pakfiles should be compressed, if the targetplatform supports it.
-		// This is only used in client platforms.
-		ConditionalLoadPlatformIniFile();
-		PlatformIniFile->GetArray(TEXT("/Script/UnrealEd.ProjectPackagingSettings"), TEXT("CompressedChunkWildcard"), CompressedChunkWildcards);
-	}
+	FConfigFile PlatformIniFile;
+	FConfigCacheIni::LoadLocalIniFile(PlatformIniFile, TEXT("Game"), true, *TargetPlatform->IniPlatformName());
 
 	// Update manifests for any encryption groups that contain non-asset files
 	if (!TargetPlatform->HasSecurePackageFormat())
@@ -429,10 +414,9 @@ bool FAssetRegistryGenerator::GenerateStreamingInstallManifest(int64 InOverrideC
 	bool bUseSecondaryOpenOrder = false;
 	TArray<FString> OrderFileSpecStrings;
 	{
-		ConditionalLoadPlatformIniFile();
-		PlatformIniFile->GetBool(TEXT("/Script/UnrealEd.ProjectPackagingSettings"), TEXT("bEnableAssetRegistryGameOpenOrderSort"), bEnableGameOpenOrderSort);
-		PlatformIniFile->GetBool(TEXT("/Script/UnrealEd.ProjectPackagingSettings"), TEXT("bPakUsesSecondaryOrder"), bUseSecondaryOpenOrder);
-		PlatformIniFile->GetArray(TEXT("/Script/UnrealEd.ProjectPackagingSettings"), TEXT("PakOrderFileSpecs"), OrderFileSpecStrings);
+		PlatformIniFile.GetBool(TEXT("/Script/UnrealEd.ProjectPackagingSettings"), TEXT("bEnableAssetRegistryGameOpenOrderSort"), bEnableGameOpenOrderSort);
+		PlatformIniFile.GetBool(TEXT("/Script/UnrealEd.ProjectPackagingSettings"), TEXT("bPakUsesSecondaryOrder"), bUseSecondaryOpenOrder);
+		PlatformIniFile.GetArray(TEXT("/Script/UnrealEd.ProjectPackagingSettings"), TEXT("PakOrderFileSpecs"), OrderFileSpecStrings);
 	}
 
 	// if a game open order can be found then use that to sort the filenames
@@ -552,6 +536,22 @@ bool FAssetRegistryGenerator::GenerateStreamingInstallManifest(int64 InOverrideC
 		}
 	}
 
+	TArray<FString> CompressedChunkWildcards;
+	if (!TargetPlatform->IsServerOnly())
+	{
+		// Load the list of wildcards to specify which pakfiles should be compressed, if the targetplatform supports it.
+		// This is only used in client platforms.
+		PlatformIniFile.GetArray(TEXT("/Script/UnrealEd.ProjectPackagingSettings"), TEXT("CompressedChunkWildcard"), CompressedChunkWildcards);
+	}
+
+	// Load the list of wildcards to specify which pakfiles have per-chunk compression. These are allowed to 
+	// use individual compression settings even if the platform package doesn't want compression.
+	// NOTE: If DDPI specifies a hardware compression setting of 'None', this won't work as expected because there will be no global compression settings to opt in to. In this case, 
+	// set bForceUseProjectCompressionFormatIgnoreHardwareOverride=true and bCompressed=False in[/Script/UnrealEd.ProjectPackagingSettings], 
+	// then setup whatever project compression settings you would like chunks to be able to opt in to.
+	TArray<FString> AllowPerChunkCompressionWildcards;
+	PlatformIniFile.GetArray(TEXT("/Script/UnrealEd.ProjectPackagingSettings"), TEXT("AllowPerChunkCompressionWildcard"), AllowPerChunkCompressionWildcards);
+
 	// generate per-chunk pak list files
 	FDefaultPakFileRules DefaultPakFileRules;
 	bool bSucceeded = true;
@@ -632,6 +632,15 @@ bool FAssetRegistryGenerator::GenerateStreamingInstallManifest(int64 InOverrideC
 				}
 			}
 
+			for (const FString& AllowPerChunkCompressionWildcard : AllowPerChunkCompressionWildcards)
+			{
+				if (PakChunkFilename.MatchesWildcard(AllowPerChunkCompressionWildcard))
+				{
+					PakChunkOptions += " AllowPerChunkCompression";
+					break;
+				}
+			}
+
 			// For encryption chunks, PakchunkIndex equals ChunkID
 			FGuid Guid = UAssetManager::Get().GetChunkEncryptionKeyGuid(PakchunkIndex);
 			if (Guid.IsValid())
diff --git a/Engine/Source/Programs/AutomationTool/Scripts/CopyBuildToStagingDirectory.Automation.cs b/Engine/Source/Programs/AutomationTool/Scripts/CopyBuildToStagingDirectory.Automation.cs
--- a/Engine/Source/Programs/AutomationTool/Scripts/CopyBuildToStagingDirectory.Automation.cs
+++ b/Engine/Source/Programs/AutomationTool/Scripts/CopyBuildToStagingDirectory.Automation.cs
@@ -2765,14 +2765,19 @@ namespace AutomationScripts
 			/// Whether to rehydrate the assets when creating the pak file or not
 			/// </summary>
 			public bool bRehydrateAssets;
-			
+
+			/// <summary>
+			/// Allow the per-chunk bCompressed flag to be used instead of the global parameter
+			/// </summary>
+			public bool bAllowPerChunkCompression;
+
 			/// <summary>
-			// Whether this pak file is for content-on-demand content delivery or not
+			/// Whether this pak file is for content-on-demand content delivery or not
 			/// </summary>
 			public bool bOnDemand;
 
 			/// <summary>
-			// Whether this pak file is to copy the source files as loose files
+			/// Whether this pak file is to copy the source files as loose files
 			/// </summary>
 			public bool bStageLoose;
 
@@ -2787,7 +2792,15 @@ namespace AutomationScripts
 			/// <param name="PakName">Path to the base output file for this pak file</param>
 			/// <param name="UnrealPakResponseFile">Map of files within the pak file to their source file on disk</param>
 			/// <param name="bCompressed">Whether to enable compression</param>
-			public CreatePakParams(string PakName, Dictionary<string, string> UnrealPakResponseFile, bool bCompressed, bool RehydrateAssets, string EncryptionKeyGuid, bool bOnDemand = false, bool bStageLoose = false)
+			public CreatePakParams(
+				string PakName, 
+				Dictionary<string, string> UnrealPakResponseFile, 
+				bool bCompressed, 
+				bool RehydrateAssets, 
+				string EncryptionKeyGuid, 
+				bool bOnDemand = false, 
+				bool bStageLoose = false,
+				bool bAllowPerChunkCompression = false)
 			{
 				this.PakName = PakName;
 				this.UnrealPakResponseFile = UnrealPakResponseFile;
@@ -2796,6 +2809,7 @@ namespace AutomationScripts
 				this.EncryptionKeyGuid = EncryptionKeyGuid;
 				this.bOnDemand = bOnDemand;
 				this.bStageLoose = bStageLoose;
+				this.bAllowPerChunkCompression = bAllowPerChunkCompression;
 			}
 		}
 
@@ -2971,14 +2985,19 @@ namespace AutomationScripts
 			// Params.Compressed is done being changed; push it to the PakParams list :
 			// note this is a change of behavior; PakParams.bCompressed could have been true from a Chunk Manifest
 			//	even though compression was otherwise turned off
-			// we now overwrite that with the overall compression setting
+			// we now overwrite that with the overall compression setting unless specifically flagged for per-chunk compression
 			//	this should always be an improvement in behavior
+			bool bAnyCompressed = Params.Compressed;
 			foreach (CreatePakParams PakParams in PakParamsList)
 			{
-				PakParams.bCompressed = Params.Compressed;
+				if (!PakParams.bAllowPerChunkCompression) // Use values from chunk manifest only if allowed
+				{
+					PakParams.bCompressed = Params.Compressed;
+				}
+				bAnyCompressed = bAnyCompressed || PakParams.bCompressed;
 			}
 
-			if (Params.Compressed && !string.IsNullOrWhiteSpace(CompressionFormats))
+			if (bAnyCompressed && !string.IsNullOrWhiteSpace(CompressionFormats))
 			{
 				CompressionFormats = " -compressionformats=" + CompressionFormats;
 			}
@@ -3010,7 +3029,7 @@ namespace AutomationScripts
 			}
 
 			string AdditionalCompressionOptionsOnCommandLine = "";
-			if (Params.Compressed)
+			if (bAnyCompressed)
 			{
 				// the game may want to control compression settings, but since it may be in a plugin that checks the commandline for the settings, we need to pass
 				// the settings directly on the UnrealPak commandline, and not put it into the batch file lines (plugins can't get the unrealpak command list, and
@@ -3513,21 +3532,21 @@ namespace AutomationScripts
 
 						if (!PakParams.bStageLoose)
 						{
-						string PakFileSpecificAdditionalArgs = "";
-						if (bShouldGeneratePatch && ShouldSkipGeneratingPatch(PlatformGameConfig, PakParams.PakName))
-						{
-							PakFileSpecificAdditionalArgs = " -SkipPatch";
-						}
-						string PakEncryptionKeyGuid = Params.SkipEncryption ? "" : PakParams.EncryptionKeyGuid;
-						PakCommands.Add(GetPakFileSpecificUnrealPakArguments(
-							UnrealPakResponseFile,
-							OutputLocation,
-							PakFileSpecificAdditionalArgs,
-							PakParams.bCompressed,
-							PakParams.bRehydrateAssets,
-							Params.SkipEncryption ? null : CryptoSettings,
-							PatchSourceContentPath,
-							Params.SkipEncryption ? "" : PakParams.EncryptionKeyGuid));
+							string PakFileSpecificAdditionalArgs = "";
+							if (bShouldGeneratePatch && ShouldSkipGeneratingPatch(PlatformGameConfig, PakParams.PakName))
+							{
+								PakFileSpecificAdditionalArgs = " -SkipPatch";
+							}
+							string PakEncryptionKeyGuid = Params.SkipEncryption ? "" : PakParams.EncryptionKeyGuid;
+							PakCommands.Add(GetPakFileSpecificUnrealPakArguments(
+								UnrealPakResponseFile,
+								OutputLocation,
+								PakFileSpecificAdditionalArgs,
+								PakParams.bCompressed,
+								PakParams.bRehydrateAssets,
+								Params.SkipEncryption ? null : CryptoSettings,
+								PatchSourceContentPath,
+								Params.SkipEncryption ? "" : PakParams.EncryptionKeyGuid));
 						}
 						LogNames.Add(OutputLocation.GetFileNameWithoutExtension());
 					}
@@ -4130,6 +4149,7 @@ namespace AutomationScripts
 				ResponseFile = new ConcurrentDictionary<string, string>();
 				Manifest = null;
 				bCompressed = false;
+				bAllowPerChunkCompression = false;
 				bOnDemand = false;
 				bStageLoose = false;
 			}
@@ -4148,6 +4168,9 @@ namespace AutomationScripts
 			public string EncryptionKeyGuid;
 			public string RequestedEncryptionKeyGuid;
 
+			// Allow the per-chunk bCompressed flag to be used instead of the global parameter
+			public bool bAllowPerChunkCompression;
+
 			// Whether the chunk is used for content-on-demand content delivery or not
 			public bool bOnDemand;
 
@@ -4216,15 +4239,19 @@ namespace AutomationScripts
 					// Set chunk name to string like "pakchunk0"
 					var ChunkManifestFilename = CombinePaths(TmpPackagingPath, ChunkOptions[0]);
 					ChunkDefinition CD = new ChunkDefinition(Path.GetFileNameWithoutExtension(ChunkOptions[0]));
-					for (int I = 1; I < ChunkOptions.Length; ++I)
+					for (int IOption = 1; IOption < ChunkOptions.Length; ++IOption)
 					{
-						if (string.Compare(ChunkOptions[I], "compressed", true) == 0)
+						if (string.Compare(ChunkOptions[IOption], "compressed", true) == 0)
 						{
 							CD.bCompressed = true;
 						}
-						else if (ChunkOptions[I].StartsWith("encryptionkeyguid="))
+						else if(string.Compare(ChunkOptions[IOption], "AllowPerChunkCompression", true) == 0)
+						{
+							CD.bAllowPerChunkCompression = true;
+						}
+						else if (ChunkOptions[IOption].StartsWith("encryptionkeyguid="))
 						{
-							CD.RequestedEncryptionKeyGuid = ChunkOptions[I].Substring(ChunkOptions[I].IndexOf('=') + 1); ;
+							CD.RequestedEncryptionKeyGuid = ChunkOptions[IOption].Substring(ChunkOptions[IOption].IndexOf('=') + 1);
 
 							if (PakCryptoSettings.SecondaryEncryptionKeys != null)
 							{
@@ -4535,7 +4562,8 @@ namespace AutomationScripts
 					string EncryptionKeyToUse = Params.SkipEncryption ? "" : Chunk.EncryptionKeyGuid;
 					PakInputs.Add(new CreatePakParams(Chunk.ChunkName,
 						Chunk.ResponseFile.ToDictionary(entry => entry.Key, entry => entry.Value),
-						Params.Compressed || Chunk.bCompressed, Params.RehydrateAssets, EncryptionKeyToUse, Chunk.bOnDemand, Chunk.bStageLoose));
+						Chunk.bCompressed, 
+						Params.RehydrateAssets, EncryptionKeyToUse, Chunk.bOnDemand, Chunk.bStageLoose, Chunk.bAllowPerChunkCompression));
 				}
 			}
 
