[Gauntlet] Update InstallApplication function signature for new ITargetDevice flow by renaming it to InstallBuild

diff --git a/Engine/Source/Programs/AutomationTool/Gauntlet/Framework/Base/Gauntlet.TargetDevice.cs b/Engine/Source/Programs/AutomationTool/Gauntlet/Framework/Base/Gauntlet.TargetDevice.cs
--- a/Engine/Source/Programs/AutomationTool/Gauntlet/Framework/Base/Gauntlet.TargetDevice.cs
+++ b/Engine/Source/Programs/AutomationTool/Gauntlet/Framework/Base/Gauntlet.TargetDevice.cs
@@ -106,10 +106,8 @@ namespace Gauntlet
 		/// <summary>
 		/// Installs a build to the device
 		/// </summary>
-		/// <param name="Build">A reference to the build to install</param>
-		/// <param name="ProjectName">Name of the UE project being installed</param>
-		/// <param name="Sandbox">Which sandbox to install to</param>
-		void InstallApplication(IBuild Build, string ProjectName, string Sandbox);
+		/// <param name="AppConfiguration">The configuration containing the build to install</param>
+		void InstallBuild(UnrealAppConfig AppConfiguration);
 
 		/// <summary>
 		/// Create an IAppInstall that is configured by the provided AppConfiguration
diff --git a/Engine/Source/Programs/AutomationTool/Gauntlet/Framework/Base/Gauntlet.TargetDeviceDesktopCommon.cs b/Engine/Source/Programs/AutomationTool/Gauntlet/Framework/Base/Gauntlet.TargetDeviceDesktopCommon.cs
--- a/Engine/Source/Programs/AutomationTool/Gauntlet/Framework/Base/Gauntlet.TargetDeviceDesktopCommon.cs
+++ b/Engine/Source/Programs/AutomationTool/Gauntlet/Framework/Base/Gauntlet.TargetDeviceDesktopCommon.cs
@@ -107,7 +107,7 @@ namespace Gauntlet
 
 		public void FullClean()
 		{
-			
+
 		}
 
 		public void CleanArtifacts()
@@ -115,7 +115,7 @@ namespace Gauntlet
 
 		}
 
-		public void InstallApplication(IBuild Build, string ProjectName, string Sandbox)
+		public void InstallBuild(UnrealAppConfig AppConfiguration)
 		{
 
 		}
diff --git a/Engine/Source/Programs/AutomationTool/Gauntlet/Platform/Android/Gauntlet.TargetDeviceAndroid.cs b/Engine/Source/Programs/AutomationTool/Gauntlet/Platform/Android/Gauntlet.TargetDeviceAndroid.cs
--- a/Engine/Source/Programs/AutomationTool/Gauntlet/Platform/Android/Gauntlet.TargetDeviceAndroid.cs
+++ b/Engine/Source/Programs/AutomationTool/Gauntlet/Platform/Android/Gauntlet.TargetDeviceAndroid.cs
@@ -1106,7 +1106,7 @@ namespace Gauntlet
 
 		}
 
-		public void InstallApplication(IBuild Build, string ProjectName, string Sandbox)
+		public void InstallBuild(UnrealAppConfig AppConfiguration)
 		{
 
 		}
diff --git a/Engine/Source/Programs/AutomationTool/Gauntlet/Platform/IOS/Gauntlet.TargetDeviceIOS.cs b/Engine/Source/Programs/AutomationTool/Gauntlet/Platform/IOS/Gauntlet.TargetDeviceIOS.cs
--- a/Engine/Source/Programs/AutomationTool/Gauntlet/Platform/IOS/Gauntlet.TargetDeviceIOS.cs
+++ b/Engine/Source/Programs/AutomationTool/Gauntlet/Platform/IOS/Gauntlet.TargetDeviceIOS.cs
@@ -593,7 +593,7 @@ namespace Gauntlet
 
 		}
 
-		public void InstallApplication(IBuild Build, string ProjectName, string Sandbox)
+		public void InstallBuild(UnrealAppConfig AppConfiguration)
 		{
 
 		}
diff --git a/Engine/Source/Programs/AutomationTool/Gauntlet/Platform/Null/Gauntlet.TargetDeviceNull.cs b/Engine/Source/Programs/AutomationTool/Gauntlet/Platform/Null/Gauntlet.TargetDeviceNull.cs
--- a/Engine/Source/Programs/AutomationTool/Gauntlet/Platform/Null/Gauntlet.TargetDeviceNull.cs
+++ b/Engine/Source/Programs/AutomationTool/Gauntlet/Platform/Null/Gauntlet.TargetDeviceNull.cs
@@ -166,7 +166,7 @@ namespace Gauntlet
 
 		}
 
-		public void InstallApplication(IBuild Build, string ProjectName, string Sandbox)
+		public void InstallBuild(UnrealAppConfig AppConfiguration)
 		{
 
 		}
