Horde: Fix license check job step for Horde server after NET 8 upgrade.

diff --git a/Engine/Source/Programs/AutomationTool/BuildGraph/Tasks/NuGetLicenseCheckTask.cs b/Engine/Source/Programs/AutomationTool/BuildGraph/Tasks/NuGetLicenseCheckTask.cs
--- a/Engine/Source/Programs/AutomationTool/BuildGraph/Tasks/NuGetLicenseCheckTask.cs
+++ b/Engine/Source/Programs/AutomationTool/BuildGraph/Tasks/NuGetLicenseCheckTask.cs
@@ -43,6 +43,12 @@ namespace AutomationTool.Tasks
 		/// </summary>
 		[TaskParameter(Optional = true)]
 		public DirectoryReference LicenseDir;
+
+		/// <summary>
+		/// Override path to dotnet executable
+		/// </summary>
+		[TaskParameter(Optional = true)]
+		public FileReference DotNetPath;
 	}
 
 	/// <summary>
@@ -132,7 +138,9 @@ namespace AutomationTool.Tasks
 		/// <param name="TagNameToFileSet">Mapping from tag names to the set of files they include</param>
 		public override async Task ExecuteAsync(JobContext Job, HashSet<FileReference> BuildProducts, Dictionary<string, HashSet<FileReference>> TagNameToFileSet)
 		{
-			IProcessResult NuGetOutput = await ExecuteAsync(Unreal.DotnetPath.FullName, $"nuget locals global-packages --list", LogOutput: false);
+			FileReference DotNetPath = Parameters.DotNetPath ?? Unreal.DotnetPath;
+
+			IProcessResult NuGetOutput = await ExecuteAsync(DotNetPath.FullName, $"nuget locals global-packages --list", LogOutput: false);
 			if (NuGetOutput.ExitCode != 0)
 			{
 				throw new AutomationException("DotNet terminated with an exit code indicating an error ({0})", NuGetOutput.ExitCode);
@@ -152,7 +160,7 @@ namespace AutomationTool.Tasks
 
 			const string UnknownPrefix = "Unknown-";
 
-			IProcessResult PackageListOutput = await ExecuteAsync(Unreal.DotnetPath.FullName, "list package --include-transitive", WorkingDir: Parameters.BaseDir, LogOutput: false);
+			IProcessResult PackageListOutput = await ExecuteAsync(DotNetPath.FullName, "list package --include-transitive", WorkingDir: Parameters.BaseDir, LogOutput: false);
 			if (PackageListOutput.ExitCode != 0)
 			{
 				throw new AutomationException("DotNet terminated with an exit code indicating an error ({0})", PackageListOutput.ExitCode);
diff --git a/Engine/Source/Programs/Horde/BuildHorde.xml b/Engine/Source/Programs/Horde/BuildHorde.xml
--- a/Engine/Source/Programs/Horde/BuildHorde.xml
+++ b/Engine/Source/Programs/Horde/BuildHorde.xml
@@ -43,7 +43,7 @@
 				[REDACTED_EMAIL]
 				[REDACTED_EMAIL]
 			</Property>
-			<NuGet-LicenseCheck BaseDir="Engine/Source/Programs/Horde/Horde.Server" IgnorePackages="$(IgnorePackages)" LicenseDir="Engine/Source/Programs/Horde/ThirdParty/Licenses"/>
+			<NuGet-LicenseCheck BaseDir="Engine/Source/Programs/Horde/Horde.Server" IgnorePackages="$(IgnorePackages)" LicenseDir="Engine/Source/Programs/Horde/ThirdParty/Licenses" DotNetPath="$(DotNetPath)"/>
 		</Node>
 
 		<Node Name="Build HordeServer" Requires="Check HordeServer Licenses">
