Add an option to output a CSV containing all packages and license files from the NuGet-LicenseCheck task.

diff --git a/Engine/Source/Programs/AutomationTool/BuildGraph/Tasks/NuGetLicenseCheckTask.cs b/Engine/Source/Programs/AutomationTool/BuildGraph/Tasks/NuGetLicenseCheckTask.cs
--- a/Engine/Source/Programs/AutomationTool/BuildGraph/Tasks/NuGetLicenseCheckTask.cs
+++ b/Engine/Source/Programs/AutomationTool/BuildGraph/Tasks/NuGetLicenseCheckTask.cs
@@ -44,6 +44,12 @@ namespace AutomationTool.Tasks
 		[TaskParameter(Optional = true)]
 		public DirectoryReference LicenseDir;
 
+		/// <summary>
+		/// Path to a csv file to write with list of packages and their licenses
+		/// </summary>
+		[TaskParameter(Optional = true)]
+		public FileReference CsvFile;
+
 		/// <summary>
 		/// Override path to dotnet executable
 		/// </summary>
@@ -89,6 +95,7 @@ namespace AutomationTool.Tasks
 		{
 			public string Name;
 			public string Version;
+			public string ProjectUrl;
 			public LicenseInfo License;
 			public string LicenseSource;
 			public PackageState State;
@@ -102,6 +109,7 @@ namespace AutomationTool.Tasks
 			public string NormalizedText;
 			public string Extension;
 			public bool Approved;
+			public FileReference File;
 		}
 
 		LicenseInfo FindOrAddLicense(Dictionary<IoHash, LicenseInfo> Licenses, string Text, string Extension)
@@ -207,6 +215,7 @@ namespace AutomationTool.Tasks
 							{
 								string Text = await FileReference.ReadAllTextAsync(File);
 								LicenseInfo License = FindOrAddLicense(Licenses, Text, File.GetFileNameWithoutExtension());
+								License.File = File;
 								License.Approved = true;
 							}
 						}
@@ -255,6 +264,9 @@ namespace AutomationTool.Tasks
 					XmlDocument XmlDocument = new XmlDocument();
 					XmlDocument.Load(XmlReader);
 
+					XmlNode ProjectUrlNode = XmlDocument.SelectSingleNode("/package/metadata/projectUrl");
+					Info.ProjectUrl = ProjectUrlNode?.InnerText;
+
 					if (Info.License == null)
 					{
 						XmlNode LicenseNode = XmlDocument.SelectSingleNode("/package/metadata/license");
@@ -379,6 +391,32 @@ namespace AutomationTool.Tasks
 					}
 				}
 			}
+
+			if (Parameters.CsvFile != null)
+			{
+				Logger.LogInformation("Writing {File}", Parameters.CsvFile);
+				DirectoryReference.CreateDirectory(Parameters.CsvFile.Directory);
+				using (StreamWriter writer = new StreamWriter(Parameters.CsvFile.FullName))
+				{
+					await writer.WriteLineAsync($"Package,Version,Project Url,License Url,License Hash,License File");
+					foreach (PackageInfo PackageInfo in Packages.Values)
+					{
+						string RelativeLicensePath = "";
+						if (PackageInfo.License?.File != null)
+						{
+							RelativeLicensePath = PackageInfo.License.File.MakeRelativeTo(Parameters.CsvFile.Directory);
+						}
+
+						string LicenseUrl = "";
+						if (PackageInfo.LicenseSource != null && PackageInfo.LicenseSource.StartsWith("http", StringComparison.OrdinalIgnoreCase))
+						{
+							LicenseUrl = PackageInfo.LicenseSource;
+						}
+
+						await writer.WriteLineAsync($"\"{PackageInfo.Name}\",\"{PackageInfo.Version}\",{PackageInfo.ProjectUrl},{LicenseUrl},{PackageInfo.License?.Hash},{RelativeLicensePath}");
+					}
+				}
+			}
 		}
 
 		/// <summary>
