Fxied utrace default channels for Editor Perf Tests so we enable Insghts Regions in the trace data

diff --git a/Engine/Source/Programs/AutomationTool/Gauntlet/Unreal/Editor/EditorTests.cs b/Engine/Source/Programs/AutomationTool/Gauntlet/Unreal/Editor/EditorTests.cs
--- a/Engine/Source/Programs/AutomationTool/Gauntlet/Unreal/Editor/EditorTests.cs
+++ b/Engine/Source/Programs/AutomationTool/Gauntlet/Unreal/Editor/EditorTests.cs
@@ -91,7 +91,7 @@ namespace UnrealEditor
 			{
 				AppConfig.CommandLineParams.Add(string.Format("tracefile={0}", TraceFile));
 				AppConfig.CommandLineParams.Add("tracefiletrunc"); // replace existing
-				AppConfig.CommandLineParams.Add("trace=cpu,counters,stats,gpu,frame,bookmark,log,loadtime,savetime,assetloadtime");
+				AppConfig.CommandLineParams.Add("trace=default,counters,stats,loadtime,savetime,assetloadtime");
 				AppConfig.CommandLineParams.Add("statnamedevents");
 			}
 
