Replication tests compilation errors on Linux and other platforms

diff --git a/Engine/Plugins/Runtime/ReplicationSystemTestPlugin/Source/Private/Tests/Serialization/TestRotatorNetSerializer.cpp b/Engine/Plugins/Runtime/ReplicationSystemTestPlugin/Source/Private/Tests/Serialization/TestRotatorNetSerializer.cpp
--- a/Engine/Plugins/Runtime/ReplicationSystemTestPlugin/Source/Private/Tests/Serialization/TestRotatorNetSerializer.cpp
+++ b/Engine/Plugins/Runtime/ReplicationSystemTestPlugin/Source/Private/Tests/Serialization/TestRotatorNetSerializer.cpp
@@ -87,7 +87,7 @@ const RotatorType FTestRotatorNetSerializerBase<RotatorType>::Values[] =
 	RotatorType(0.0f, 0.0f, 1.0f),
 	RotatorType(-0.0f), // 0.0f == -0.0f so we expect this to be handled gracefully
 	RotatorType(359.999f, 129.130f, 45.0f),
-	RotatorType(1.26698556E-23),
+	RotatorType(1.26698556E-23f),
 	RotatorType(47.11f),
 };
 template<typename RotatorType>
diff --git a/Engine/Source/Programs/AutomationTool/Gauntlet/Platform/Linux/Gauntlet.TargetDeviceLinux.cs b/Engine/Source/Programs/AutomationTool/Gauntlet/Platform/Linux/Gauntlet.TargetDeviceLinux.cs
--- a/Engine/Source/Programs/AutomationTool/Gauntlet/Platform/Linux/Gauntlet.TargetDeviceLinux.cs
+++ b/Engine/Source/Programs/AutomationTool/Gauntlet/Platform/Linux/Gauntlet.TargetDeviceLinux.cs
@@ -27,6 +27,14 @@ namespace Gauntlet
 			RunOptions = CommandUtils.ERunOptions.NoWaitForExit | CommandUtils.ERunOptions.NoLoggingOfRunCommand;
 		}
 
+		public TargetDeviceLinux(string InName, string InCacheDir, bool InIsArm64)
+			: base(InName, InCacheDir)
+		{
+			IsArm64 = InIsArm64;
+			Platform = InIsArm64 ? UnrealTargetPlatform.LinuxArm64 : UnrealTargetPlatform.Linux;
+			RunOptions = CommandUtils.ERunOptions.NoWaitForExit | CommandUtils.ERunOptions.NoLoggingOfRunCommand;
+		}
+
 		public override IAppInstance Run(IAppInstall App)
 		{
 			LinuxAppInstall LinuxApp = App as LinuxAppInstall;
@@ -276,7 +284,7 @@ namespace Gauntlet
 
 		public ITargetDevice CreateDevice(string InRef, string InCachePath, string InParam = null)
 		{
-			return new TargetDeviceLinux(InRef, InCachePath) { IsArm64 = true };
+			return new TargetDeviceLinux(InRef, InCachePath, true);
 		}
 	}
 }
\ No newline at end of file
