Add command line option to control stdout/stderr read timeout in child processes of AutomationTool. Needed for workaround in Wine where stderr is not always properly closed, causing a minute delay on some operations. Patching Wine is cumbersome at this point, and this command line option introduces a workaround while maintaining current behaviour by default.

Tested by running AutomationTool under Wine using Linux and noticing a one minute delay as process exit is waiting for stderr to close. Then re-running the test with -WaitForStdStreams=100 and verifying the exit happens within 100ms and no output is lost.

diff --git a/Engine/Source/Programs/AutomationTool/AutomationUtils/Automation.cs b/Engine/Source/Programs/AutomationTool/AutomationUtils/Automation.cs
--- a/Engine/Source/Programs/AutomationTool/AutomationUtils/Automation.cs
+++ b/Engine/Source/Programs/AutomationTool/AutomationUtils/Automation.cs
@@ -409,6 +409,12 @@ AutomationTool.exe [-verbose] [-compileonly] [-p4] Command0 [-Arg0 -Arg1 -Arg2 .
 			NoSubmit = AutomationToolCommandLine.IsSetGlobal("-NoSubmit");
 			P4 = AutomationToolCommandLine.IsSetGlobal("-P4");
 			NoP4 = AutomationToolCommandLine.IsSetGlobal("-NoP4");
+
+			int WaitTimeMs;
+			if (int.TryParse((string)AutomationToolCommandLine.GetValueUnchecked("-WaitForStdStreams"), out WaitTimeMs))
+			{
+				WaitForStdStreams = WaitTimeMs;
+			}
 		}
 
 		// Using Nullable bools here to ensure that Initialize() has been called before the members are accesed.
@@ -511,5 +517,11 @@ AutomationTool.exe [-verbose] [-compileonly] [-p4] Command0 [-Arg0 -Arg1 -Arg2 .
 			get; 
 			private set;
 		}
+
+		public static int WaitForStdStreams
+		{
+			get;
+			private set;
+		} = -1;
 	}
 }
diff --git a/Engine/Source/Programs/AutomationTool/AutomationUtils/ProcessUtils.cs b/Engine/Source/Programs/AutomationTool/AutomationUtils/ProcessUtils.cs
--- a/Engine/Source/Programs/AutomationTool/AutomationUtils/ProcessUtils.cs
+++ b/Engine/Source/Programs/AutomationTool/AutomationUtils/ProcessUtils.cs
@@ -486,51 +486,65 @@ namespace AutomationTool
 				bProcTerminated = (Proc == null);
 			}
 			// Keep checking if we got all output messages until the process terminates.
-			if (!bProcTerminated)
+			Stopwatch Watch = Stopwatch.StartNew();
+			int MaxWaitUntilMessagesReceived = 60 * 1000;
+			int WaitTimeout = 500;
+			if (GlobalCommandLine.WaitForStdStreams >= 0)
 			{
-				// Check messages
-				int MaxWaitUntilMessagesReceived = 120;
-				while (MaxWaitUntilMessagesReceived > 0 && !(bStdOutSignalReceived && bStdErrSignalReceived))
+				MaxWaitUntilMessagesReceived = GlobalCommandLine.WaitForStdStreams;
+			}
+			if (MaxWaitUntilMessagesReceived > WaitTimeout)
+			{
+				WaitTimeout = 1 + (MaxWaitUntilMessagesReceived / 10);
+			}
+			while (!(bStdOutSignalReceived && bStdErrSignalReceived))
+			{
+				if (!bStdOutSignalReceived)
 				{
-					if (!bStdOutSignalReceived)
-					{
-						bStdOutSignalReceived = OutputWaitHandle.WaitOne(500);
-					}
-					if (!bStdErrSignalReceived)
-					{
-						bStdErrSignalReceived = ErrorWaitHandle.WaitOne(500);
-					}
-					// Check if the process terminated
-					lock (ProcSyncObject)
-					{
-						bProcTerminated = (Proc == null) || Proc.HasExited;
-					}
-					if (bProcTerminated)
-					{
-						// Process terminated but make sure we got all messages, don't wait forever though
-						MaxWaitUntilMessagesReceived--;
-					}
+					bStdOutSignalReceived = OutputWaitHandle.WaitOne(WaitTimeout);
 				}
-				if (!(bStdOutSignalReceived && bStdErrSignalReceived))
+				if (!bStdErrSignalReceived)
 				{
-					Logger.LogDebug("Waited for a long time for output of {AppName}, some output may be missing; we gave up.", AppName);
+					bStdErrSignalReceived = ErrorWaitHandle.WaitOne(WaitTimeout);
 				}
-
-				// Double-check if the process terminated
+				// Check if the process terminated
 				lock (ProcSyncObject)
 				{
 					bProcTerminated = (Proc == null) || Proc.HasExited;
-
-					if (Proc != null)
+				}
+				if (!bProcTerminated)
+				{
+					// Timeout starts when process has terminated
+					Watch.Restart();
+				}
+				else
+				{
+					if (Watch.ElapsedMilliseconds > MaxWaitUntilMessagesReceived)
 					{
-						if (!bProcTerminated)
-						{
-							// The process did not terminate yet but we've read all output messages, wait until the process terminates
-							Proc.WaitForExit();
-						}
+						// Timeout passed, do not wait any longer
+						break;
+					}
+				}
+			}
+			if (!(bStdOutSignalReceived && bStdErrSignalReceived))
+			{
+				Logger.LogInformation("Waited for {0:n2}s for output of {AppName}, some output may be missing; we gave up.", Watch.Elapsed.TotalSeconds, AppName);
+			}
+
+			// Double-check if the process terminated
+			lock (ProcSyncObject)
+			{
+				bProcTerminated = (Proc == null) || Proc.HasExited;
 
-						ExitCode = Proc.ExitCode;
+				if (Proc != null)
+				{
+					if (!bProcTerminated)
+					{
+						// The process did not terminate yet but we've read all output messages, wait until the process terminates
+						Proc.WaitForExit();
 					}
+
+					ExitCode = Proc.ExitCode;
 				}
 			}
 		}
diff --git a/Engine/Source/Programs/AutomationTool/Program.cs b/Engine/Source/Programs/AutomationTool/Program.cs
--- a/Engine/Source/Programs/AutomationTool/Program.cs
+++ b/Engine/Source/Programs/AutomationTool/Program.cs
@@ -42,6 +42,7 @@ namespace AutomationToolDriver
 			string Option_ScriptsForProject = "-ScriptsForProject";
 			string Option_ScriptDir = "-ScriptDir";
 			string Option_Telemetry = "-Telemetry";
+			string Option_WaitForStdStreams = "-WaitForStdStreams";
 
 			// The parameter was not found in the list of global parameters, continue looking...
 			if (CurrentParam.StartsWith(Option_ScriptsForProject + "=", StringComparison.InvariantCultureIgnoreCase))
@@ -103,7 +104,12 @@ namespace AutomationToolDriver
             {
 				string TelemetryPath = CurrentParam.Substring(CurrentParam.IndexOf('=') + 1);
 				AutomationToolCommandLine.SetUnchecked(Option_Telemetry, TelemetryPath);
-            }
+			}
+			else if (CurrentParam.StartsWith(Option_WaitForStdStreams + "=", StringComparison.InvariantCultureIgnoreCase))
+			{
+				string WaitTime = CurrentParam.Substring(CurrentParam.IndexOf('=') + 1);
+				AutomationToolCommandLine.SetUnchecked(Option_WaitForStdStreams, WaitTime);
+			}
 			else if (CurrentParam.StartsWith("-"))
 			{
 				if (CurrentCommand != null)
@@ -286,7 +292,8 @@ namespace AutomationToolDriver
 					{"-UseLocalBuildStorage", @"Allows you to use local storage for your root build storage dir {default of P:\Builds {on PC} is changed to Engine\Saved\LocalBuilds}. Used for local testing."},
 					{"-WaitForDebugger", "Waits for a debugger to be attached, and breaks once debugger successfully attached."},
 					{"-BuildMachine", "" },
-					{"-WaitForUATMutex", "" }
+					{"-WaitForUATMutex", "" },
+					{"-WaitForStdStreams", "Time in milliseconds to wait for std streams to close in child processes." }
 				},
 				new HashSet<string>(StringComparer.InvariantCultureIgnoreCase) { "-msbuild-verbose", "-NoCompileUAT" }
 			);
