- Added support for UAT to cook with a specific project's editor .exe for projects that use a Unique BuildEnvironment (this will cook with the game's editor, so it will use the game's uniquely compiled engine modules (useful for per-project SDKs that have different targetplatform dlls)
  - The target receipt already listed the proper exe to run, taking unique/shared into account, so this change delayed-looks-up the path from the receipt on first use (like when running the editor to cook)

diff --git a/Engine/Source/Programs/AutomationTool/AutomationUtils/ProjectParams.cs b/Engine/Source/Programs/AutomationTool/AutomationUtils/ProjectParams.cs
--- a/Engine/Source/Programs/AutomationTool/AutomationUtils/ProjectParams.cs
+++ b/Engine/Source/Programs/AutomationTool/AutomationUtils/ProjectParams.cs
@@ -372,7 +372,7 @@ namespace AutomationTool
 			this.Manifests = InParams.Manifests;
             this.CreateChunkInstall = InParams.CreateChunkInstall;
 			this.SkipEncryption = InParams.SkipEncryption;
-			this.UnrealExe = InParams.UnrealExe;
+			this.SpecifiedUnrealExe = InParams.SpecifiedUnrealExe;
 			this.NoDebugInfo = InParams.NoDebugInfo;
 			this.SeparateDebugInfo = InParams.SeparateDebugInfo;
 			this.MapFile = InParams.MapFile;
@@ -982,7 +982,7 @@ namespace AutomationTool
             this.RunAutomationTest = ParseParamValueIfNotSpecified(Command, RunAutomationTest, "RunAutomationTest");
             this.RunAutomationTests = this.RunAutomationTest != "" || GetParamValueIfNotSpecified(Command, RunAutomationTests, this.RunAutomationTests, "RunAutomationTests");
             this.SkipServer = GetParamValueIfNotSpecified(Command, SkipServer, this.SkipServer, "skipserver");
-			this.UnrealExe = ParseParamValueIfNotSpecified(Command, UnrealExe, "unrealexe", "UnrealEditor-Cmd.exe", ObsoleteSpecifiedValue: UE4Exe, ObsoleteParamName: "ue4exe");
+			this.SpecifiedUnrealExe = ParseParamValueIfNotSpecified(Command, UnrealExe, "unrealexe", null, ObsoleteSpecifiedValue: UE4Exe, ObsoleteParamName: "ue4exe");
 			this.Unattended = GetParamValueIfNotSpecified(Command, Unattended, this.Unattended, "unattended");
 			this.DeviceUsername = ParseParamValueIfNotSpecified(Command, DeviceUsername, "deviceuser", String.Empty);
 			this.DevicePassword = ParseParamValueIfNotSpecified(Command, DevicePassword, "devicepass", String.Empty);
@@ -1587,7 +1587,30 @@ namespace AutomationTool
 		public string CookerSupportFilesSubdirectory;
 		
 		[Help("unrealexe=ExecutableName", "Name of the Unreal Editor executable, i.e. -unrealexe=UnrealEditor.exe")]
-		public string UnrealExe;
+		private string SpecifiedUnrealExe = null;
+
+		public string UnrealExe
+		{
+			get
+			{
+				if (SpecifiedUnrealExe == null)
+				{
+					SpecifiedUnrealExe = "UnrealEditor-Cmd.exe";
+					if (CodeBasedUprojectPath != null)
+					{
+						FileReference ReceiptLocation = TargetReceipt.GetDefaultPath(CodeBasedUprojectPath.Directory, EditorTargets[0], HostPlatform.Platform, UnrealTargetConfiguration.Development, null);
+						TargetReceipt Receipt;
+						if (!TargetReceipt.TryRead(ReceiptLocation, out Receipt))
+						{
+							throw new AutomationException($"Missing {ReceiptLocation} receipt. Editor needs to be built first.");
+						}
+						SpecifiedUnrealExe = Receipt.LaunchCmd.FullName;
+					}
+				}
+
+				return SpecifiedUnrealExe;
+			}
+		}
 
 		/// <summary>
 		/// Shared: true if this build is archived, command line: -archive
