- Allow for code to set ProjectParams.UnrealExe (not used in Main, but this restores functionality that was removed)

diff --git a/Engine/Source/Programs/AutomationTool/AutomationUtils/ProjectParams.cs b/Engine/Source/Programs/AutomationTool/AutomationUtils/ProjectParams.cs
--- a/Engine/Source/Programs/AutomationTool/AutomationUtils/ProjectParams.cs
+++ b/Engine/Source/Programs/AutomationTool/AutomationUtils/ProjectParams.cs
@@ -1610,6 +1610,11 @@ namespace AutomationTool
 
 				return SpecifiedUnrealExe;
 			}
+			set 
+			{
+				// allow code override
+				SpecifiedUnrealExe = value; 
+			}
 		}
 
 		/// <summary>
