UAT - Extending ICustomStageCopyHandler to support overriding of the plugin manifest file creation. Allows game projects to have control over how plugin descriptors are added to the manifest

diff --git a/Engine/Source/Programs/AutomationTool/AutomationUtils/CustomStageCopyHandler.cs b/Engine/Source/Programs/AutomationTool/AutomationUtils/CustomStageCopyHandler.cs
--- a/Engine/Source/Programs/AutomationTool/AutomationUtils/CustomStageCopyHandler.cs
+++ b/Engine/Source/Programs/AutomationTool/AutomationUtils/CustomStageCopyHandler.cs
@@ -1,6 +1,7 @@
 // Copyright Epic Games, Inc. All Rights Reserved.
 
 using AutomationTool;
+using EpicGames.Core;
 using System;
 using System.Collections.Generic;
 using System.Linq;
@@ -10,7 +11,23 @@ using Microsoft.Extensions.Logging;
 
 public abstract class CustomStageCopyHandler
 {
-	public virtual bool StageFile(ILogger Logger, string SourceName, string TargetName) { return false; }
+	/// <summary>
+	/// Called when copying files to the staging directory.
+	/// </summary>
+	/// <returns>true if the file was handled. false to fallback on the default stage copy implementation</returns>
+	public virtual bool StageFile(ILogger Logger, string SourceName, string TargetName)
+	{ 
+		return false;
+	}
+
+	/// <summary>
+	/// Called when writing plugin descriptor files to the plugin manifest file.
+	/// </summary>
+	/// <returns>The PluginDescriptor to write into the manifest</returns>
+	public virtual PluginDescriptor CreateDescriptorForPluginManifest(ILogger Logger, FileReference PluginFile)
+	{ 
+		return PluginDescriptor.FromFile(PluginFile);
+	}
 
 	public static CustomStageCopyHandler Create(string Name)
 	{
diff --git a/Engine/Source/Programs/AutomationTool/AutomationUtils/DeploymentContext.cs b/Engine/Source/Programs/AutomationTool/AutomationUtils/DeploymentContext.cs
--- a/Engine/Source/Programs/AutomationTool/AutomationUtils/DeploymentContext.cs
+++ b/Engine/Source/Programs/AutomationTool/AutomationUtils/DeploymentContext.cs
@@ -314,10 +314,10 @@ public class DeploymentContext //: ProjectParams
 	public HashSet<StagedFileReference> ExtraFilesAllowList = new HashSet<StagedFileReference>();
 
 	/// <summary>
-	/// Optional copy handler that during CopyOrWriteManifestFilesToStageDir will be passed to SafeCopy to handle
-	/// the copy operation of files.
+	/// Optional stage handler that during CopyOrWriteManifestFilesToStageDir will handle the copy operation of files and creation
+	/// of the plugin manifest file.
 	/// </summary>
-	public OverrideCopyDelegate OverrideCopyHandler = null;
+	public CustomStageCopyHandler CustomStageCopyHandler = null;
 
 	/// <summary>
 	/// List of ini keys to strip when staging
diff --git a/Engine/Source/Programs/AutomationTool/Scripts/CopyBuildToStagingDirectory.Automation.cs b/Engine/Source/Programs/AutomationTool/Scripts/CopyBuildToStagingDirectory.Automation.cs
--- a/Engine/Source/Programs/AutomationTool/Scripts/CopyBuildToStagingDirectory.Automation.cs
+++ b/Engine/Source/Programs/AutomationTool/Scripts/CopyBuildToStagingDirectory.Automation.cs
@@ -2050,6 +2050,18 @@ namespace AutomationScripts
 			return InputPaths;
 		}
 
+		static PluginDescriptor CreatePluginDescriptorForManifest(DeploymentContext SC, FileReference File)
+		{
+			if (SC.CustomStageCopyHandler != null)
+			{
+				return SC.CustomStageCopyHandler.CreateDescriptorForPluginManifest(Logger, File);
+			}
+			else
+			{
+				return PluginDescriptor.FromFile(File);
+			}
+		}
+
 		static void CreatePluginManifest(DeploymentContext SC, Dictionary<StagedFileReference, FileReference> FileMapping, StagedFileType FileType, string ManifestName)
 		{
 			// Get the path to the project's mods directory. We wont include anything under here in the manifest
@@ -2061,7 +2073,7 @@ namespace AutomationScripts
 			{
 				if (!File.Value.IsUnderDirectory(ModsDir) && File.Value.HasExtension(".uplugin"))
 				{
-					PluginDescriptor Descriptor = PluginDescriptor.FromFile(File.Value);
+					PluginDescriptor Descriptor = CreatePluginDescriptorForManifest(SC, File.Value);
 					StagedPlugins[File.Key] = Descriptor;
 				}
 			}
@@ -2137,6 +2149,8 @@ namespace AutomationScripts
 		{
 			if (bPerformCopy)
 			{
+				OverrideCopyDelegate OverrideCopyHandler = (SC.CustomStageCopyHandler != null) ? SC.CustomStageCopyHandler.StageFile : null;
+
 				var StageDirectory = ManifestName == "DebugFiles" ? SC.DebugStageDirectory : SC.StageDirectory;
 				Logger.LogInformation("Copying {ManifestName} to staging directory: {StageDirectory}", ManifestName, StageDirectory);
 				foreach (KeyValuePair<StagedFileReference, FileReference> Pair in Mapping)
@@ -2145,7 +2159,7 @@ namespace AutomationScripts
 					FileReference Dest = FileReference.Combine(StageDirectory, Pair.Key.Name);
 					if (Src != Dest)  // special case for things created in the staging directory, like the pak file
 					{
-						CopyFileIncremental(Src, Dest, SC.OverrideCopyHandler, IniKeyDenyList: SC.IniKeyDenyList, IniSectionDenyList: SC.IniSectionDenyList);
+						CopyFileIncremental(Src, Dest, OverrideCopyHandler, IniKeyDenyList: SC.IniKeyDenyList, IniSectionDenyList: SC.IniSectionDenyList);
 					}
 				}
 			}
@@ -5268,16 +5282,8 @@ namespace AutomationScripts
 
 			if (!string.IsNullOrEmpty(CustomStageCopyHandlerName))
 			{
-				CustomStageCopyHandler CustomStageCopyHandler = CustomStageCopyHandler.Create(CustomStageCopyHandlerName);
-				if (CustomStageCopyHandler != null)
-				{
-					Logger.LogInformation("CopyBuildToStagingDirectory using CustomStageCopyHandler {0}", CustomStageCopyHandlerName);
-
-					SC.OverrideCopyHandler = (ILogger Logger, string SourceName, string TargetName) =>
-					{
-						return CustomStageCopyHandler.StageFile(Logger, SourceName, TargetName);
-					};
-				}
+				Logger.LogInformation("CopyBuildToStagingDirectory using CustomStageCopyHandler {0}", CustomStageCopyHandlerName);
+				SC.CustomStageCopyHandler = CustomStageCopyHandler.Create(CustomStageCopyHandlerName);
 			}
 		}
 
