Localization:
Allow multiple localization batches to run with a single invocation of the Editor.
- Previously, when plugins are added each plugin will be gathered by a new Editor process. This change collapses all the batches to be gathered by a single Editor.
- This consolidation is done by writing all of the config files needed for the gather into a .txt file and passed to the GatherText commandlet
- GatherTextCommandlet now supports the -ConfigList parameter which will load a .txt file with 1 config file per line to be executed.
- Introduced the ConsolidateConfigFiles parameter to the Localize UAT command. This consolidates all of the config files across the various localization batches into a single .txt file. This file is then passed to the GatherText commandlet to be run.

#test: Ran this locally to verify that localizing 2 plugins in batches and in consolidated form yield the same result

diff --git a/Engine/Source/Editor/UnrealEd/Private/Commandlets/GatherTextCommandlet.cpp b/Engine/Source/Editor/UnrealEd/Private/Commandlets/GatherTextCommandlet.cpp
--- a/Engine/Source/Editor/UnrealEd/Private/Commandlets/GatherTextCommandlet.cpp
+++ b/Engine/Source/Editor/UnrealEd/Private/Commandlets/GatherTextCommandlet.cpp
@@ -52,21 +52,45 @@ int32 UGatherTextCommandlet::Main( const FString& Params )
 	if (const FString* ConfigParamPtr = ParamVals.Find(UGatherTextCommandletBase::ConfigParam))
 	{
 		ConfigParamPtr->ParseIntoArray(GatherTextConfigPaths, TEXT(";"));
+	}
 
-		const FString& ProjectBasePath = UGatherTextCommandletBase::GetProjectBasePath();
-
-		for (FString& GatherTextConfigPath : GatherTextConfigPaths)
+	if (const FString* ConfigListFileParamPtr = ParamVals.Find(TEXT("ConfigList")))
+	{
+		if (FPaths::FileExists(*ConfigListFileParamPtr))
 		{
-			if (FPaths::IsRelative(GatherTextConfigPath))
+			TArray<FString> ConfigFiles;
+			FFileHelper::LoadFileToStringArray(ConfigFiles, **ConfigListFileParamPtr);
+			if (ConfigFiles.Num() > 0)
 			{
-				GatherTextConfigPath = FPaths::Combine(*ProjectBasePath, *GatherTextConfigPath);
+				GatherTextConfigPaths.Append(MoveTemp(ConfigFiles));
 			}
+			else
+			{
+				UE_LOG(LogGatherTextCommandlet, Warning, TEXT("There are no config file paths in specified config ,list '%s'. Please check to see the is correctly populated."), **ConfigListFileParamPtr);
+			}
+		}
+		else
+		{
+			UE_LOG(LogGatherTextCommandlet, Warning, TEXT("Specified config list file '%s' does not exist. No additional config files from -ConfigList can be added."), **ConfigListFileParamPtr);
+		}
+	}
+
+	// @TODOLocalization: Handle the case where -config and -ConfigList both specify the same files.
+	// Currently that would just mean that the config files will be launched with the relevant commandlets multiple times. The results should be correct, but it's wasted work.
+	
+	// Turn all relative paths into absolute paths 
+	const FString& ProjectBasePath = UGatherTextCommandletBase::GetProjectBasePath();
+	for (FString& GatherTextConfigPath : GatherTextConfigPaths)
+	{
+		if (FPaths::IsRelative(GatherTextConfigPath))
+		{
+			GatherTextConfigPath = FPaths::Combine(*ProjectBasePath, *GatherTextConfigPath);
 		}
 	}
 
 	if (GatherTextConfigPaths.Num() == 0)
 	{
-		UE_LOG(LogGatherTextCommandlet, Error, TEXT("-config not specified.\n%s"), *UsageText);
+		UE_LOG(LogGatherTextCommandlet, Error, TEXT("-config or -ConfigList not specified. If -ConfigList was specified, please check that the file path is valid.\n%s"), *UsageText);
 		return -1;
 	}
 
diff --git a/Engine/Source/Programs/AutomationTool/Localization/LocalizationProvider.cs b/Engine/Source/Programs/AutomationTool/Localization/LocalizationProvider.cs
--- a/Engine/Source/Programs/AutomationTool/Localization/LocalizationProvider.cs
+++ b/Engine/Source/Programs/AutomationTool/Localization/LocalizationProvider.cs
@@ -98,6 +98,19 @@ namespace EpicGames.Localization
 			ExportInfo = InExportInfo;
 		}
 
+		public List<string> GetConfigFilesToRun(List<string> RequestedLocalizationStepNames)
+		{
+			List<string> ConfigFiles = new();
+			foreach (var LocalizationStep in LocalizationSteps)
+			{
+				if (RequestedLocalizationStepNames.Contains(LocalizationStep.Name))
+				{
+					ConfigFiles.Add(LocalizationStep.LocalizationConfigFile);
+				}
+			}
+			return ConfigFiles;
+		}
+
 		/** The name of this project */
 		public string ProjectName { get; private set; }
 
diff --git a/Engine/Source/Programs/AutomationTool/Scripts/Localisation.Automation.cs b/Engine/Source/Programs/AutomationTool/Scripts/Localisation.Automation.cs
--- a/Engine/Source/Programs/AutomationTool/Scripts/Localisation.Automation.cs
+++ b/Engine/Source/Programs/AutomationTool/Scripts/Localisation.Automation.cs
@@ -12,6 +12,9 @@ using EpicGames.Core;
 using System.Threading.Tasks;
 using Microsoft.Extensions.Logging;
 
+
+
+
 [Help("Updates the external localization data using the arguments provided.")]
 [Help("UEProjectRoot", "Optional root-path to the project we're gathering for (defaults to CmdEnv.LocalRoot if unset).")]
 [Help("UEProjectDirectory", "Sub-path to the project we're gathering for (relative to UEProjectRoot).")]
@@ -75,6 +78,133 @@ class Localize : BuildCommand
 		public List<IProcessResult> GatherProcessResults = new List<IProcessResult>();
 	};
 
+	private abstract class IGatherTextCommandletLauncherStrategy
+	{
+		public class Args
+		{
+			public string AbsoluteEditorExePath { get; set; } = "";
+			public ERunOptions CommandletRunOptions { get; set; } = ERunOptions.Default;
+			public string AbsoluteUEProjectDirectoryPath { get; set; } = "";
+			public string UEProjectName { get; set; } = "";
+			public string EditorArgs { get; set; } = "";
+			public List<string> LocalizationStepNames { get; set; } = new();
+
+			public string GetAbsoluteUEProjectPath()
+			{
+				// This implies that we're gathering for the Engine 
+				if (String.IsNullOrEmpty(UEProjectName))
+				{
+					return "";
+				}
+				return Path.Combine(AbsoluteUEProjectDirectoryPath, $"{UEProjectName}.uproject");
+			}
+		}
+
+		public IGatherTextCommandletLauncherStrategy(IGatherTextCommandletLauncherStrategy.Args InArgs)
+		{
+			LauncherArgs = InArgs;
+		}
+		public abstract void LaunchCommandlet(List<LocalizationTask> LocalizationTasks);
+
+		protected readonly IGatherTextCommandletLauncherStrategy.Args LauncherArgs;
+	}
+
+	private class BatchedGatherTextCommandletLauncherStrategy : IGatherTextCommandletLauncherStrategy
+	{
+		public BatchedGatherTextCommandletLauncherStrategy(IGatherTextCommandletLauncherStrategy.Args InArgs) : base(InArgs)
+		{
+
+		}
+		
+		public override void LaunchCommandlet(List<LocalizationTask> LocalizationTasks)
+		{
+			foreach (var LocalizationTask in LocalizationTasks)
+			{
+				foreach (var ProjectInfo in LocalizationTask.ProjectInfos)
+				{
+					List<string> LocalizationConfigFiles = ProjectInfo.GetConfigFilesToRun(LauncherArgs.LocalizationStepNames);
+					if (LocalizationConfigFiles.Count > 0)
+					{
+						string ConcatenatedConfigFiles = String.Join(";", LocalizationConfigFiles);
+						string CommandLine = $"\"{LauncherArgs.GetAbsoluteUEProjectPath()}\" -run=GatherText -config=\"{ConcatenatedConfigFiles}\" {LauncherArgs.EditorArgs}";
+						Logger.LogInformation("Running localization commandlet for '{Arg0}': {Arguments}", ProjectInfo.ProjectName, CommandLine);
+						LocalizationTask.GatherProcessResults.Add(Run(LauncherArgs.AbsoluteEditorExePath, CommandLine, null,	LauncherArgs.CommandletRunOptions));
+					}
+					else
+					{
+						Logger.LogInformation($"Localization target '{ProjectInfo.ProjectName}' has no valid config files associated with the specified localization steps. The localization target will not be gathered.");
+						LocalizationTask.GatherProcessResults.Add(null);
+					}
+				}
+			}
+		}
+	}
+
+	private class ConsolidatedGatherTextCommandletLauncherStrategy : IGatherTextCommandletLauncherStrategy
+	{
+		public ConsolidatedGatherTextCommandletLauncherStrategy(IGatherTextCommandletLauncherStrategy.Args InArgs) : base(InArgs)
+		{
+
+		}
+		public override void LaunchCommandlet(List<LocalizationTask> LocalizationTasks)
+		{
+			List<string> ConsolidatedConfigFiles = new List<string>();
+			List<string> ProjectsToGather = new List<string>();
+			foreach (var LocalizationTask in LocalizationTasks)
+			{
+				foreach (var ProjectInfo in LocalizationTask.ProjectInfos)
+				{
+					List<string> LocalizationConfigFiles = ProjectInfo.GetConfigFilesToRun(LauncherArgs.LocalizationStepNames);
+					if (LocalizationConfigFiles.Count > 0)
+					{
+						ConsolidatedConfigFiles.AddRange(ProjectInfo.GetConfigFilesToRun(LauncherArgs.LocalizationStepNames));
+						ProjectsToGather.Add(ProjectInfo.ProjectName);
+					}
+					else
+					{
+						Logger.LogInformation($"Localization target {ProjectInfo.ProjectName} will not be gathered because it does not have localization config files that can be run.");
+					}
+					
+				}
+			}
+
+			// Now that we've consolidated all the config files, we write them out to the Saved Directory 
+			string SaveDirectory = Path.Combine(LauncherArgs.AbsoluteUEProjectDirectoryPath, "Saved", "Localization");
+			if (!Directory.Exists(SaveDirectory))
+			{
+				Directory.CreateDirectory(SaveDirectory);
+			}
+			// we append a GUID to the file name to ensure that parallel gathers can still work without stomping on the file and multiple runs of the Localize command can still result in unique files for debugging 
+			string SaveFile = Path.Combine(SaveDirectory, $"ConfigList_{Guid.NewGuid()}.txt");
+			File.WriteAllLines(SaveFile, ConsolidatedConfigFiles);
+
+			// now construct the command line and run 
+			string CommandLine = $"\"{LauncherArgs.GetAbsoluteUEProjectPath()}\" -run=GatherText -ConfigList=\"{SaveFile}\" {LauncherArgs.EditorArgs}";
+			string ConcatenatedProjectsString = String.Join(",", ProjectsToGather);
+			Logger.LogInformation($"Consolidating localization gather for following localization targets: \"{ConcatenatedProjectsString}\"");
+			Logger.LogInformation($"Running consolidated gather text commandlet - {CommandLine}");
+			IProcessResult ConsolidatedProcessResult = Run(LauncherArgs.AbsoluteEditorExePath, CommandLine, null, LauncherArgs.CommandletRunOptions);
+
+			// Go through all tasks and set the gather process result  accordingly for the various projects
+			// We need to do this otherwise projects won't be uploaded successfully 
+			foreach (var LocalizationTask in LocalizationTasks)
+			{
+				foreach (var ProjectInfo in LocalizationTask.ProjectInfos)
+				{
+					List<string> LocalizationConfigFiles = ProjectInfo.GetConfigFilesToRun(LauncherArgs.LocalizationStepNames);
+					if (LocalizationConfigFiles.Count > 0)
+					{
+						LocalizationTask.GatherProcessResults.Add(ConsolidatedProcessResult);
+					}
+					else
+					{
+						LocalizationTask.GatherProcessResults.Add(null);
+					}
+				}
+			}
+		}
+	}
+
 	private string UEProjectRoot = CmdEnv.LocalRoot;
 	private string UEProjectDirectory = "";
 	private string UEProjectName = "";
@@ -90,6 +220,7 @@ class Localize : BuildCommand
 	private bool bEnableParallelGather = false;
 	private bool bIsRunningInPreview = false;
 	private bool bPreserveAutoGeneratedResources = false;
+	private bool bConsolidateConfigFiles = false;
 	private int PendingChangeList = -1;
 
 	HashSet<string> AutoGeneratedFiles = new();
@@ -183,7 +314,7 @@ class Localize : BuildCommand
 		// This runs even for non-parallel execution to log the exit state of the process.
 		WaitForCommandletResults(LocalizationTasks);
 
-		// If we are running in preview, we can go ahead and delete all generated preview files after the gather step is complete 
+		// If we are running in preview, we can go ahead and delete all generated preview files after the gather step is complete
 		if (bIsRunningInPreview)
         {
 			CleanUpGeneratedPreviewFiles(LocalizationBatches);
@@ -312,6 +443,10 @@ class Localize : BuildCommand
 		}
 
 		bEnableParallelGather = ParseParam("ParallelGather");
+		if (bEnableParallelGather)
+		{
+			Logger.LogInformation("Parallel gather enabled. Multiple instances of the editor will be used to gather each individual localization batch.");
+		}
 
 		bIsRunningInPreview = ParseParam("Preview");
 		// We pass the preview switch along to have the gather text commandlets exhibit different behaviors. See UGatherTextCommandlet
@@ -326,6 +461,12 @@ class Localize : BuildCommand
 		{
 			Logger.LogInformation("Preserving auto-generated content. Auto-generated files and folders will not be automatically cleaned up at the end of this command.");
 		}
+
+		bConsolidateConfigFiles = ParseParam("ConsolidateConfigFiles");
+		if (bConsolidateConfigFiles)
+		{
+			Logger.LogInformation("Consolidating config files for various localization targets. A single instance of the Editor will be run to perform all localization gather steps.");
+		}
 	}
 
 	private void AddStaticLocalizationBatches(List<LocalizationBatch> LocalizationBatches)
@@ -531,7 +672,7 @@ class Localize : BuildCommand
 		}
 	}
 
-	private void StartGatherCommands(List<LocalizationTask> LocalizationTasks)
+	private string GetEditorExePath()
 	{
 		var EditorExe = CombinePaths(CmdEnv.LocalRoot, @"Engine/Binaries/Win64/UnrealEditor-Cmd.exe");
 		if (!File.Exists(EditorExe))
@@ -540,43 +681,44 @@ class Localize : BuildCommand
 			EditorExe = CombinePaths(CmdEnv.LocalRoot, @"Engine/Binaries/Win64/UnrealEditor-Win64-Debug-Cmd.exe");
 		}
 
-		// Set the common basic editor arguments
-		string EditorArguments = BuildEditorArguments();
+		return EditorExe;
+	}
 
-		// Set the common process run options
+	private ERunOptions GetCommandletRunOptions()
+	{
 		var CommandletRunOptions = ERunOptions.Default | ERunOptions.NoLoggingOfRunCommand; // Disable logging of the run command as it will print the exit code which GUBP can pick up as an error (we do that ourselves later)
 		if (bEnableParallelGather)
 		{
 			CommandletRunOptions |= ERunOptions.NoWaitForExit;
 		}
+		return CommandletRunOptions;
+	}
 
-		foreach (var LocalizationTask in LocalizationTasks)
-		{
-			var ProjectArgument = String.IsNullOrEmpty(UEProjectName) ? "" : String.Format("\"{0}\"", Path.Combine(LocalizationTask.RootWorkingDirectory, String.Format("{0}.uproject", UEProjectName)));
+	private void StartGatherCommands(List<LocalizationTask> LocalizationTasks)
+	{
+		IGatherTextCommandletLauncherStrategy.Args Args = new();
+		Args.AbsoluteEditorExePath = GetEditorExePath();
 
-			foreach (var ProjectInfo in LocalizationTask.ProjectInfos)
-			{
-				var LocalizationConfigFiles = new List<string>();
-				foreach (var LocalizationStep in ProjectInfo.LocalizationSteps)
-				{
-					if (LocalizationStepNames.Contains(LocalizationStep.Name))
-					{
-						LocalizationConfigFiles.Add(LocalizationStep.LocalizationConfigFile);
-					}
-				}
+		// Set the common basic editor arguments
+		Args.EditorArgs = BuildEditorArguments();
 
-				if (LocalizationConfigFiles.Count > 0)
-				{
-					var Arguments = String.Format("{0} -run=GatherText -config=\"{1}\" {2}", ProjectArgument, String.Join(";", LocalizationConfigFiles), EditorArguments);
-					Logger.LogInformation("Running localization commandlet for '{Arg0}': {Arguments}", ProjectInfo.ProjectName, Arguments);
-					LocalizationTask.GatherProcessResults.Add(Run(EditorExe, Arguments, null, CommandletRunOptions));
-				}
-				else
-				{
-					LocalizationTask.GatherProcessResults.Add(null);
-				}
-			}
+		// Set the common process run options
+		Args.CommandletRunOptions = GetCommandletRunOptions();
+
+		Args.UEProjectName = UEProjectName;
+		Args.AbsoluteUEProjectDirectoryPath = CombinePaths(UEProjectRoot, UEProjectDirectory);
+		Args.LocalizationStepNames = LocalizationStepNames;
+
+		IGatherTextCommandletLauncherStrategy Launcher;
+		if (bConsolidateConfigFiles)
+		{
+			Launcher = new ConsolidatedGatherTextCommandletLauncherStrategy(Args);
+		}
+		else
+		{
+			Launcher = new BatchedGatherTextCommandletLauncherStrategy(Args);
 		}
+		Launcher.LaunchCommandlet(LocalizationTasks);
 	}
 
 	private void WaitForCommandletResults(List<LocalizationTask> LocalizationTasks)
