Fix log output not being human readable
+ Handle setting up the default command line arguments through TargetDeviceDesktopCommon
+ Align log settings across configs

#rnx

diff --git a/Engine/Source/Programs/AutomationTool/Gauntlet/Framework/Base/Gauntlet.TargetDeviceDesktopCommon.cs b/Engine/Source/Programs/AutomationTool/Gauntlet/Framework/Base/Gauntlet.TargetDeviceDesktopCommon.cs
--- a/Engine/Source/Programs/AutomationTool/Gauntlet/Framework/Base/Gauntlet.TargetDeviceDesktopCommon.cs
+++ b/Engine/Source/Programs/AutomationTool/Gauntlet/Framework/Base/Gauntlet.TargetDeviceDesktopCommon.cs
@@ -5,9 +5,12 @@ using EpicGames.Core;
 using System;
 using System.Collections.Generic;
 using System.IO;
+using System.Text;
 using System.Threading;
 using System.Linq;
 using UnrealBuildTool;
+using System.Text.RegularExpressions;
+using static AutomationTool.ProcessResult;
 
 namespace Gauntlet
 {
@@ -197,7 +200,32 @@ namespace Gauntlet
 
 		public string WorkingDirectory;
 
-		public virtual string CommandArguments { get; set; }
+		public string LogFile;
+
+		public bool CanAlterCommandArgs;
+
+		public string CommandArguments
+		{
+			get { return CommandArgumentsPrivate; }
+			set
+			{
+				if (CanAlterCommandArgs || string.IsNullOrEmpty(CommandArgumentsPrivate))
+				{
+					CommandArgumentsPrivate = value;
+				}
+				else
+				{
+					Log.Info("Skipped setting command AppInstall line when CanAlterCommandArgs = false");
+				}
+			}
+		}
+
+		private string CommandArgumentsPrivate;
+
+		public void AppendCommandline(string AdditionalCommandline)
+		{
+			CommandArguments += AdditionalCommandline;
+		}
 
 		public CommandUtils.ERunOptions RunOptions { get; set; }
 
@@ -208,6 +236,7 @@ namespace Gauntlet
 			DesktopDevice = InDevice;
 			CommandArguments = string.Empty;
 			RunOptions = CommandUtils.ERunOptions.NoWaitForExit;
+			CanAlterCommandArgs = true;
 		}
 
 		public virtual IAppInstance Run()
@@ -217,6 +246,19 @@ namespace Gauntlet
 
 		public virtual void CleanDeviceArtifacts()
 		{
+			// log file
+			try
+			{
+				if (LogFile != null && File.Exists(LogFile))
+				{
+					EpicGames.Core.FileUtils.ForceDeleteFile(LogFile);
+				}
+			}
+			catch (Exception Ex)
+			{
+				Log.Warning(KnownLogEvents.Gauntlet_DeviceEvent, "Unable to delete existing log file {File}. {Exception}", LogFile, Ex.Message);
+			}
+			// all other artifacts
 			if (!string.IsNullOrEmpty(ArtifactPath) && Directory.Exists(ArtifactPath))
 			{
 				try
@@ -255,6 +297,69 @@ namespace Gauntlet
 			}
 			return true;
 		}
+
+		public virtual void SetDefaultCommandLineArguments(UnrealAppConfig AppConfig, CommandUtils.ERunOptions InRunOptions, string BuildDir)
+		{
+			RunOptions = InRunOptions;
+			if (Log.IsVeryVerbose)
+			{
+				RunOptions |= CommandUtils.ERunOptions.AllowSpew;
+			}
+
+			string UserDir = DesktopDevice.UserDir;
+			// Set commandline replace any InstallPath arguments with the path we use
+			CommandArguments = Regex.Replace(AppConfig.CommandLine, @"\$\(InstallPath\)", BuildDir, RegexOptions.IgnoreCase);
+			CanAlterCommandArgs = AppConfig.CanAlterCommandArgs;
+
+			if (CanAlterCommandArgs)
+			{
+				// Unreal userdir
+				if (string.IsNullOrEmpty(UserDir) == false)
+				{
+					CommandArguments += $" -userdir=\"{UserDir}\"";
+					ArtifactPath = Path.Combine(UserDir, "Saved");
+
+					Utils.SystemHelpers.MarkDirectoryForCleanup(UserDir);
+				}
+				else
+				{
+					// e.g d:\Unreal\GameName\Saved
+					ArtifactPath = Path.Combine(BuildDir, ProjectName, "Saved");
+				}
+
+				// Unreal abslog
+				// Look in app parameters if abslog is specified, if so use it
+				Regex LogRegex = new Regex(@"-abslog[:\s=](?:""([^""]*)""|([^-][^\s]*))?");
+				Match M = LogRegex.Match(CommandArguments);
+				if (M.Success)
+				{
+					LogFile = M.Groups[2].Value;
+				}
+
+				// Explicitly set log file when not already defined if not build machine
+				// -abslog makes sure Unreal dynamically update the log window when using -log
+				if (!CommandUtils.IsBuildMachine && (!string.IsNullOrEmpty(LogFile) || AppConfig.CommandLine.Contains("-log")))
+				{
+					string LogFolder = string.IsNullOrEmpty(LogFile) ? Path.Combine(ArtifactPath, "Logs") : Path.GetDirectoryName(LogFile);
+
+					if (!Directory.Exists(LogFolder))
+					{
+						Directory.CreateDirectory(LogFolder);
+					}
+
+					if (string.IsNullOrEmpty(LogFile))
+					{
+						LogFile = Path.Combine(LogFolder, $"{ProjectName}.log");
+						CommandArguments += string.Format(" -abslog=\"{0}\"", LogFile);
+					}
+
+					RunOptions |= CommandUtils.ERunOptions.NoStdOutRedirect;
+				}
+
+				// clear artifact path
+				CleanDeviceArtifacts();
+			}
+		}
 	}
 
 	public abstract class DesktopCommonAppInstance<DesktopAppInstall, DesktopTargetDevice> : LocalAppProcess
@@ -267,8 +372,8 @@ namespace Gauntlet
 
 		protected DesktopAppInstall Install;
 
-		public DesktopCommonAppInstance(DesktopAppInstall InInstall, IProcessResult InProcess, string ProcessLogFile = null)
-			: base(InProcess, InInstall.CommandArguments, ProcessLogFile)
+		public DesktopCommonAppInstance(DesktopAppInstall InInstall, IProcessResult InProcess, string InProcessLogFile = null)
+			: base(InProcess, InInstall.CommandArguments, InProcessLogFile)
 		{
 			Install = InInstall;
 		}
diff --git a/Engine/Source/Programs/AutomationTool/Gauntlet/Platform/Linux/Gauntlet.TargetDeviceLinux.cs b/Engine/Source/Programs/AutomationTool/Gauntlet/Platform/Linux/Gauntlet.TargetDeviceLinux.cs
--- a/Engine/Source/Programs/AutomationTool/Gauntlet/Platform/Linux/Gauntlet.TargetDeviceLinux.cs
+++ b/Engine/Source/Programs/AutomationTool/Gauntlet/Platform/Linux/Gauntlet.TargetDeviceLinux.cs
@@ -90,7 +90,7 @@ namespace Gauntlet
 				Environment.CurrentDirectory = OldWD;
 			}
 
-			return new LinuxAppInstance(LinuxApp, Result);
+			return new LinuxAppInstance(LinuxApp, Result, LinuxApp.LogFile);
 		}
 
 		protected override IAppInstall InstallNativeStagedBuild(UnrealAppConfig AppConfig, NativeStagedBuild InBuild)
@@ -104,17 +104,7 @@ namespace Gauntlet
 			{
 				LinuxApp = new LinuxAppInstall(AppConfig.Name, AppConfig.ProjectName, this);
 			}
-
-			LinuxApp.RunOptions = RunOptions;
-			if (Log.IsVeryVerbose)
-			{
-				LinuxApp.RunOptions |= CommandUtils.ERunOptions.AllowSpew;
-			}
-
-			LinuxApp.CommandArguments = AppConfig.CommandLine;
-
-			LinuxApp.ArtifactPath = Path.Combine(InBuild.BuildPath, AppConfig.ProjectName, @"Saved");
-			LinuxApp.CleanDeviceArtifacts();
+			LinuxApp.SetDefaultCommandLineArguments(AppConfig, RunOptions, InBuild.BuildPath);
 
 			CopyAdditionalFiles(AppConfig.FilesToCopy);
 
@@ -147,27 +137,7 @@ namespace Gauntlet
 			}
 
 			LinuxAppInstall LinuxApp = new LinuxAppInstall(AppConfig.Name, AppConfig.ProjectName, this);
-			LinuxApp.RunOptions = RunOptions;
-
-			// Set commandline replace any InstallPath arguments with the path we use
-			LinuxApp.CommandArguments = Regex.Replace(AppConfig.CommandLine, @"\$\(InstallPath\)", BuildDir, RegexOptions.IgnoreCase);
-
-			if (string.IsNullOrEmpty(UserDir) == false)
-			{
-				LinuxApp.CommandArguments += string.Format(" -userdir=\"{0}\"", UserDir);
-				LinuxApp.ArtifactPath = Path.Combine(UserDir, "Saved");
-
-				Utils.SystemHelpers.MarkDirectoryForCleanup(UserDir);
-			}
-			else
-			{
-				// e.g d:\Unreal\GameName\Saved
-				LinuxApp.ArtifactPath = Path.Combine(BuildDir, AppConfig.ProjectName, "Saved");
-
-			}
-
-			// clear artifact path
-			LinuxApp.CleanDeviceArtifacts();
+			LinuxApp.SetDefaultCommandLineArguments(AppConfig, RunOptions, BuildDir);
 
 			if (LocalDirectoryMappings.Count == 0)
 			{
@@ -216,11 +186,7 @@ namespace Gauntlet
 			LinuxAppInstall LinuxApp = new LinuxAppInstall(AppConfig.Name, AppConfig.ProjectName, this);
 
 			LinuxApp.WorkingDirectory = Path.GetDirectoryName(Build.ExecutablePath);
-			LinuxApp.RunOptions = RunOptions;
-
-			// Force this to stop logs and other artifacts going to different places
-			LinuxApp.CommandArguments = AppConfig.CommandLine + string.Format(" -userdir=\"{0}\"", UserDir);
-			LinuxApp.ArtifactPath = Path.Combine(UserDir, @"Saved");
+			LinuxApp.SetDefaultCommandLineArguments(AppConfig, RunOptions, LinuxApp.WorkingDirectory);
 			LinuxApp.ExecutablePath = Build.ExecutablePath;
 
 			if (LocalDirectoryMappings.Count == 0)
@@ -257,8 +223,8 @@ namespace Gauntlet
 
 	public class LinuxAppInstance : DesktopCommonAppInstance<LinuxAppInstall, TargetDeviceLinux>
 	{
-		public LinuxAppInstance(LinuxAppInstall InInstall, IProcessResult InProcess, string ProcessLogFile = null)
-			: base(InInstall, InProcess, ProcessLogFile)
+		public LinuxAppInstance(LinuxAppInstall InInstall, IProcessResult InProcess, string InProcessLogFile = null)
+			: base(InInstall, InProcess, InProcessLogFile)
 		{ }
 	}
 
diff --git a/Engine/Source/Programs/AutomationTool/Gauntlet/Platform/Mac/Gauntlet.TargetDeviceMac.cs b/Engine/Source/Programs/AutomationTool/Gauntlet/Platform/Mac/Gauntlet.TargetDeviceMac.cs
--- a/Engine/Source/Programs/AutomationTool/Gauntlet/Platform/Mac/Gauntlet.TargetDeviceMac.cs
+++ b/Engine/Source/Programs/AutomationTool/Gauntlet/Platform/Mac/Gauntlet.TargetDeviceMac.cs
@@ -72,50 +72,35 @@ namespace Gauntlet
 				Environment.CurrentDirectory = OldWD;
 			}
 
-			return new MacAppInstance(MacInstall, Result);
+			return new MacAppInstance(MacInstall, Result, MacInstall.LogFile);
 		}
 
 		protected override IAppInstall InstallNativeStagedBuild(UnrealAppConfig AppConfig, NativeStagedBuild InBuild)
 		{
-			MacAppInstall MacApp = new MacAppInstall(AppConfig.Name, this);
-
-			MacApp.RunOptions = RunOptions;
-			if (Log.IsVeryVerbose)
-			{
-				MacApp.RunOptions |= CommandUtils.ERunOptions.AllowSpew;
-			}
-
+			MacAppInstall MacApp = new MacAppInstall(AppConfig.Name, AppConfig.ProjectName, this);
+			MacApp.SetDefaultCommandLineArguments(AppConfig, RunOptions, InBuild.BuildPath);
 			MacApp.WorkingDirectory = InBuild.BuildPath;
 			MacApp.ExecutablePath = Path.Combine(InBuild.BuildPath, InBuild.ExecutablePath);
-			MacApp.CommandArguments = AppConfig.CommandLine;
 
-			MacApp.ArtifactPath = Path.Combine(InBuild.BuildPath, AppConfig.ProjectName, @"Saved");
-			MacApp.CleanDeviceArtifacts();
+			CopyAdditionalFiles(AppConfig.FilesToCopy);
 
 			return MacApp;
 		}
 
 		protected override IAppInstall InstallEditorBuild(UnrealAppConfig AppConfig, EditorBuild Build)
 		{
-			MacAppInstall MacApp = new MacAppInstall(AppConfig.Name, this);
-
+			MacAppInstall MacApp = new MacAppInstall(AppConfig.Name, AppConfig.ProjectName, this);
+			MacApp.SetDefaultCommandLineArguments(AppConfig, RunOptions, Build.ExecutablePath);
 			MacApp.WorkingDirectory = Path.GetFullPath(Build.ExecutablePath);
-			MacApp.CommandArguments = AppConfig.CommandLine;
-			MacApp.RunOptions = RunOptions;
-
-			// Mac always forces this to stop logs and other artifacts going to different places
-			MacApp.CommandArguments += string.Format(" -userdir=\"{0}\"", UserDir);
-			MacApp.ArtifactPath = Path.Combine(UserDir, @"Saved");
-
-			// now turn the Foo.app into Foo/Content/MacOS/Foo
-			string AppPath = Path.GetDirectoryName(Build.ExecutablePath);
-
 			MacApp.ExecutablePath = GetExecutableIfBundle(Build.ExecutablePath);
 
 			if (LocalDirectoryMappings.Count == 0)
 			{
-				PopulateDirectoryMappings(AppPath);
+				PopulateDirectoryMappings(AppConfig.ProjectFile.Directory.FullName);
 			}
+
+			CopyAdditionalFiles(AppConfig.FilesToCopy);
+
 			return MacApp;
 		}
 
@@ -134,25 +119,14 @@ namespace Gauntlet
 
 		protected IAppInstall InstallBuild(UnrealAppConfig AppConfig, string BuildPath, string BundlePath)
 		{
-			MacAppInstall MacApp = new MacAppInstall(AppConfig.Name, this);
-
+			MacAppInstall MacApp = new MacAppInstall(AppConfig.Name, AppConfig.ProjectName, this);
+			MacApp.SetDefaultCommandLineArguments(AppConfig, RunOptions, BuildPath);
 			MacApp.WorkingDirectory = BuildPath;
-			MacApp.RunOptions = RunOptions;
-
 			MacApp.ExecutablePath = GetExecutableIfBundle(BundlePath);
 
-			// Set commandline replace any InstallPath arguments with the path we use
-			MacApp.CommandArguments = Regex.Replace(AppConfig.CommandLine, @"\$\(InstallPath\)", BuildPath, RegexOptions.IgnoreCase);
-
-			// Mac always forces this to stop logs and other artifacts going to different places
-			// Mac always forces this to stop logs and other artifacts going to different places
-			MacApp.CommandArguments += string.Format(" -userdir=\"{0}\"", UserDir);
-			MacApp.ArtifactPath = Path.Combine(UserDir, @"Saved");
-
 			PopulateDirectoryMappings(BuildPath);
 
-			// clear artifact path
-			MacApp.CleanDeviceArtifacts();
+			CopyAdditionalFiles(AppConfig.FilesToCopy);
 
 			// check for a local newer executable
 			if (Globals.Params.ParseParam("dev")
@@ -266,8 +240,8 @@ namespace Gauntlet
 			set { } // intentional nop
 		}
 
-		public MacAppInstall(string InName, TargetDeviceMac InDevice)
-			: base(InName, string.Empty, InDevice)
+		public MacAppInstall(string InName, string InProjectName, TargetDeviceMac InDevice)
+			: base(InName, InProjectName, InDevice)
 		{ }
 	}
 
diff --git a/Engine/Source/Programs/AutomationTool/Gauntlet/Platform/Windows/Gauntlet.TargetDeviceWindows.cs b/Engine/Source/Programs/AutomationTool/Gauntlet/Platform/Windows/Gauntlet.TargetDeviceWindows.cs
--- a/Engine/Source/Programs/AutomationTool/Gauntlet/Platform/Windows/Gauntlet.TargetDeviceWindows.cs
+++ b/Engine/Source/Programs/AutomationTool/Gauntlet/Platform/Windows/Gauntlet.TargetDeviceWindows.cs
@@ -4,7 +4,6 @@ using System;
 using System.IO;
 using AutomationTool;
 using UnrealBuildTool;
-using System.Text.RegularExpressions;
 using EpicGames.Core;
 using static AutomationTool.ProcessResult;
 
@@ -58,7 +57,6 @@ namespace Gauntlet
 			}
 
 			IProcessResult Result = null;
-			string ProcessLogFile = null;
 
 			lock (Globals.MainLock)
 			{
@@ -70,63 +68,15 @@ namespace Gauntlet
 				Log.Info("Launching {0} on {1}", App.Name, ToString());
 
 				string CmdLine = WinApp.CommandArguments;
-
-				if (WinApp.CanAlterCommandArgs)
-				{
-					// Look in app parameters if abslog is specified, if so use it
-					Regex CLRegex = new Regex(@"(--?[a-zA-Z]+)[:\s=]?([A-Z]:(?:\\[\w\s-]+)+\\?(?=\s-)|\""[^\""]*\""|[^-][^\s]*)?");
-					foreach (Match M in CLRegex.Matches(CmdLine))
-					{
-						if (M.Groups.Count == 3 && M.Groups[1].Value == "-abslog")
-						{
-							ProcessLogFile = M.Groups[2].Value;
-						}
-					}
-				}
-
-				// Explicitly set log file when not already defined if not build machine
-				// -abslog makes sure Unreal dymanicaly update the log window when using -log
-				if (WinApp.CanAlterCommandArgs && !AutomationTool.Automation.IsBuildMachine && string.IsNullOrEmpty(ProcessLogFile))
-				{
-					string LogFolder = string.Format(@"{0}\Logs", WinApp.ArtifactPath);
-
-					if (!Directory.Exists(LogFolder))
-					{
-						Directory.CreateDirectory(LogFolder);
-					}
-
-					ProcessLogFile = string.Format("{0}\\{1}.log", LogFolder, WinApp.ProjectName);
-					CmdLine = string.Format("{0} -abslog=\"{1}\"", CmdLine, ProcessLogFile);
-				}
-
-				// cleanup any existing log file
-				try
-				{
-					if (ProcessLogFile != null && File.Exists(ProcessLogFile))
-					{
-						EpicGames.Core.FileUtils.ForceDeleteFile(ProcessLogFile);
-					}
-				}
-				catch (Exception Ex)
-				{
-					//throw new AutomationException("Unable to delete existing log file {0} {1}", ProcessLogFile, Ex.Message);
-					Log.Warning(KnownLogEvents.Gauntlet_DeviceEvent, "Unable to delete existing log file {File}. {Exception}", ProcessLogFile, Ex.Message);
-				}
-
-				Log.Verbose("\t{0}", CmdLine);
-
 				CommandUtils.ERunOptions FinalRunOptions = WinApp.RunOptions;
-				if (!WinApp.CanAlterCommandArgs)
-				{
-					CmdLine = WinApp.CommandArguments;
-				}
+				bool bAllowSpew = WinApp.RunOptions.HasFlag(CommandUtils.ERunOptions.AllowSpew);
 
-				FinalRunOptions = WinApp.RunOptions | (ProcessLogFile != null ? CommandUtils.ERunOptions.NoStdOutRedirect : 0);
+				Log.Verbose("\t{0}", CmdLine);
 
 				Result = CommandUtils.Run(WinApp.ExecutablePath,
 					CmdLine,
 					Options: FinalRunOptions,
-					SpewFilterCallback: new SpewFilterCallbackType(M => { return ProcessLogFile == null ? M : null; }) /* make sure stderr does not spew in the stdout */,
+					SpewFilterCallback: new SpewFilterCallbackType(delegate (string M) { return bAllowSpew ? M : null; }) /* make sure stderr does not spew in the stdout */,
 					WorkingDir: WinApp.WorkingDirectory);
 
 				if (Result.HasExited && Result.ExitCode != 0)
@@ -137,26 +87,15 @@ namespace Gauntlet
 				Environment.CurrentDirectory = OldWD;
 			}
 
-			return new WindowsAppInstance(WinApp, Result, ProcessLogFile);
+			return new WindowsAppInstance(WinApp, Result, WinApp.LogFile);
 		}
 
 		protected override IAppInstall InstallNativeStagedBuild(UnrealAppConfig AppConfig, NativeStagedBuild InBuild)
 		{
 			WindowsAppInstall WinApp = new WindowsAppInstall(AppConfig.Name, AppConfig.ProjectName, this);
-			WinApp.CanAlterCommandArgs = AppConfig.CanAlterCommandArgs;
-
-			WinApp.RunOptions = RunOptions;
-			if (Log.IsVeryVerbose)
-			{
-				WinApp.RunOptions |= CommandUtils.ERunOptions.AllowSpew;
-			}
-
 			WinApp.WorkingDirectory = InBuild.BuildPath;
+			WinApp.SetDefaultCommandLineArguments(AppConfig, RunOptions, InBuild.BuildPath);
 			WinApp.ExecutablePath = Path.Combine(InBuild.BuildPath, InBuild.ExecutablePath);
-			WinApp.CommandArguments = AppConfig.CommandLine;
-
-			WinApp.ArtifactPath = Path.Combine(InBuild.BuildPath, AppConfig.ProjectName, @"Saved");
-			WinApp.CleanDeviceArtifacts();
 
 			CopyAdditionalFiles(AppConfig.FilesToCopy);
 
@@ -186,27 +125,7 @@ namespace Gauntlet
 			}
 
 			WindowsAppInstall WinApp = new WindowsAppInstall(AppConfig.Name, AppConfig.ProjectName, this);
-			WinApp.RunOptions = RunOptions;
-			WinApp.CanAlterCommandArgs = AppConfig.CanAlterCommandArgs;
-
-			// Set commandline replace any InstallPath arguments with the path we use
-			WinApp.CommandArguments = Regex.Replace(AppConfig.CommandLine, @"\$\(InstallPath\)", BuildDir, RegexOptions.IgnoreCase);
-
-			if (string.IsNullOrEmpty(UserDir) == false)
-			{
-				WinApp.CommandArguments += string.Format(" -userdir=\"{0}\"", UserDir);
-				WinApp.ArtifactPath = Path.Combine(UserDir, "Saved");
-
-				Utils.SystemHelpers.MarkDirectoryForCleanup(UserDir);
-			}
-			else
-			{
-				// e.g d:\Unreal\GameName\Saved
-				WinApp.ArtifactPath = Path.Combine(BuildDir, AppConfig.ProjectName, "Saved");
-			}
-
-			// clear artifact path
-			WinApp.CleanDeviceArtifacts();
+			WinApp.SetDefaultCommandLineArguments(AppConfig, RunOptions, BuildDir);
 
 			if (LocalDirectoryMappings.Count == 0)
 			{
@@ -255,11 +174,7 @@ namespace Gauntlet
 			WindowsAppInstall WinApp = new WindowsAppInstall(AppConfig.Name, AppConfig.ProjectName, this);
 
 			WinApp.WorkingDirectory = Path.GetDirectoryName(Build.ExecutablePath);
-			WinApp.RunOptions = RunOptions;
-
-			// Force this to stop logs and other artifacts going to different places
-			WinApp.CommandArguments = AppConfig.CommandLine + string.Format(" -userdir=\"{0}\"", UserDir);
-			WinApp.ArtifactPath = Path.Combine(UserDir, @"Saved");
+			WinApp.SetDefaultCommandLineArguments(AppConfig, RunOptions, WinApp.WorkingDirectory);
 			WinApp.ExecutablePath = Build.ExecutablePath;
 
 			if (LocalDirectoryMappings.Count == 0)
@@ -275,23 +190,7 @@ namespace Gauntlet
 		protected IAppInstall InstallSelfInstallingBuild(UnrealAppConfig AppConfig, IWindowsSelfInstallingBuild Build)
 		{
 			WindowsAppInstall WinApp = Build.Install(this, AppConfig, out string BasePath);
-
-			if (Log.IsVeryVerbose)
-			{
-				WinApp.RunOptions |= CommandUtils.ERunOptions.AllowSpew;
-			}
-
-			if (string.IsNullOrEmpty(UserDir) == false)
-			{
-				WinApp.CommandArguments += string.Format(" -userdir=\"{0}\"", UserDir);
-				WinApp.ArtifactPath = Path.Combine(UserDir, @"Saved");
-				Utils.SystemHelpers.MarkDirectoryForCleanup(UserDir);
-			}
-			else
-			{
-				WinApp.ArtifactPath = Path.Combine(BasePath, AppConfig.ProjectName, @"Saved");
-			}
-			WinApp.CleanDeviceArtifacts();
+			WinApp.SetDefaultCommandLineArguments(AppConfig, RunOptions, BasePath);
 
 			if (LocalDirectoryMappings.Count == 0)
 			{
@@ -305,45 +204,18 @@ namespace Gauntlet
 
 	public class WindowsAppInstall : DesktopCommonAppInstall<TargetDeviceWindows>, IAppInstall.IDynamicCommandLine
 	{
-		public bool CanAlterCommandArgs;
-
 		[Obsolete("Will be removed in a future release. Use 'DesktopDevice' instead.")]
 		public TargetDeviceWindows WinDevice => DesktopDevice;
 
-		public override string CommandArguments
-		{
-			get { return CommandArgumentsPrivate; }
-			set
-			{
-				if (CanAlterCommandArgs || string.IsNullOrEmpty(CommandArgumentsPrivate))
-				{
-					CommandArgumentsPrivate = value;
-				}
-				else
-				{
-					Log.Info("Skipped setting command AppInstall line when CanAlterCommandArgs = false");
-				}
-			}
-		}
-
-		private string CommandArgumentsPrivate;
-
 		public WindowsAppInstall(string InName, string InProjectName, TargetDeviceWindows InDevice)
 			: base(InName, InProjectName, InDevice)
-		{
-			CanAlterCommandArgs = true;
-		}
-
-		public void AppendCommandline(string AdditionalCommandline)
-		{
-			CommandArguments += AdditionalCommandline;
-		}
+		{ }
 	}
 
 	public class WindowsAppInstance : DesktopCommonAppInstance<WindowsAppInstall, TargetDeviceWindows>
 	{
-		public WindowsAppInstance(WindowsAppInstall InInstall, IProcessResult InProcess, string ProcessLogFile = null)
-			: base(InInstall, InProcess, ProcessLogFile)
+		public WindowsAppInstance(WindowsAppInstall InInstall, IProcessResult InProcess, string InProcessLogFile = null)
+			: base(InInstall, InProcess, InProcessLogFile)
 		{ }
 	}
 
diff --git a/Engine/Source/Programs/AutomationTool/Gauntlet/Unreal/Automation/UE.Automation.cs b/Engine/Source/Programs/AutomationTool/Gauntlet/Unreal/Automation/UE.Automation.cs
--- a/Engine/Source/Programs/AutomationTool/Gauntlet/Unreal/Automation/UE.Automation.cs
+++ b/Engine/Source/Programs/AutomationTool/Gauntlet/Unreal/Automation/UE.Automation.cs
@@ -244,16 +244,16 @@ namespace UE
 
 				if (ConfigRole.RoleType.IsClient())
 				{
-					// Have the client list the tests it knows about. Useful for troubleshooting discrepencies
+					// Have the client list the tests it knows about. Useful for troubleshooting discrepancies
 					string ClientAutomationTestArgument = "List;";
 
-					// Make sure the stereo test setting propogates to the client
+					// Make sure the stereo test setting propagates to the client
 					if (EnableStereoTestVariants)
 					{
 						ClientAutomationTestArgument += "EnableStereoTests;";
 					}
 
-					AppConfig.CommandLine += string.Format(" -sessionid={0} -messaging -log -TcpMessagingConnect={1}:6666 -ExecCmds=\"Automation {2}\"", SessionID, HostIP, ClientAutomationTestArgument);
+					AppConfig.CommandLine += string.Format(" -sessionid={0} -messaging -TcpMessagingConnect={1}:6666 -ExecCmds=\"Automation {2}\"", SessionID, HostIP, ClientAutomationTestArgument);
 				}
 				else if (ConfigRole.RoleType.IsEditor())
 				{
@@ -368,7 +368,7 @@ namespace UE
 
 			// Tests in the editor only require a single role
 			UnrealTestRole EditorRole = Config.RequireRole(Config.CookedEditor ? UnrealTargetRole.CookedEditor : UnrealTargetRole.Editor);
-			EditorRole.CommandLineParams.AddRawCommandline("-NoWatchdog -stdout -FORCELOGFLUSH -CrashForUAT -log");
+			EditorRole.CommandLineParams.AddRawCommandline("-NoWatchdog -stdout -FORCELOGFLUSH -CrashForUAT");
 
 			return Config;
 		}
@@ -400,7 +400,7 @@ namespace UE
 
 			// Target tests require an editor which hosts the process
 			UnrealTestRole EditorRole = Config.RequireRole(UnrealTargetRole.Editor);
-			EditorRole.CommandLineParams.AddRawCommandline("-NoWatchdog -stdout -FORCELOGFLUSH -CrashForUAT -log");
+			EditorRole.CommandLineParams.AddRawCommandline("-NoWatchdog -stdout -FORCELOGFLUSH -CrashForUAT");
 
 			if (Config.Attended == false)
 			{
diff --git a/Engine/Source/Programs/AutomationTool/Gauntlet/Unreal/Base/Gauntlet.UnrealSession.cs b/Engine/Source/Programs/AutomationTool/Gauntlet/Unreal/Base/Gauntlet.UnrealSession.cs
--- a/Engine/Source/Programs/AutomationTool/Gauntlet/Unreal/Base/Gauntlet.UnrealSession.cs
+++ b/Engine/Source/Programs/AutomationTool/Gauntlet/Unreal/Base/Gauntlet.UnrealSession.cs
@@ -546,13 +546,12 @@ namespace Gauntlet
 		/// <param name="InAppInstance"></param>
 		/// <param name="InArtifactPath"></param>
 		/// <param name="InLogSummary"></param>
-		public UnrealRoleArtifacts(UnrealSessionRole InSessionRole, IAppInstance InAppInstance, string InArtifactPath, string InLogPath, UnrealLogParser InLog)
+		public UnrealRoleArtifacts(UnrealSessionRole InSessionRole, IAppInstance InAppInstance, string InArtifactPath, string InLogPath)
 		{
 			SessionRole = InSessionRole;
 			AppInstance = InAppInstance;
 			ArtifactPath = InArtifactPath;
 			LogPath = InLogPath;
-			//LogParser = InLog;
 		}
 	}
 
@@ -1196,12 +1195,12 @@ namespace Gauntlet
 						Writer.WriteLine(string.Format("Role: {0}\r\n", InRunningRole.Role));
 						Writer.WriteLine(string.Format("Automation Command: {0}\r\n", Environment.CommandLine));
 						Writer.WriteLine("---------------------------");
-						Writer.Write(InRunningRole.AppInstance.StdOut);
+						Writer.Write(UnrealLogParser.SanitizeLogText(InRunningRole.AppInstance.StdOut));
 					}
 					Log.Info($"Wrote {RoleName} Log to {ArtifactLogFilePath}");
 
-					// On build machines, copy all role logs to Horde.
-					if (IsBuildMachine)
+					// On build machines, copy all role logs to Horde. Except for test that already generate a dedicated report for it.
+					if (IsBuildMachine && !InContext.TestParams.ParseParam("WriteTestResultsForHorde"))
 					{
 						string HordeLogFilePath = Path.Combine(CommandUtils.CmdEnv.LogFolder, RoleName + "Output.log");
 						File.Copy(ArtifactLogFilePath, HordeLogFilePath, true);
@@ -1262,7 +1261,7 @@ namespace Gauntlet
 							{
 								Log.Info("Downsizing and gifying session images at {0}", ScreenshotSubdirectory.FullName);
 
-								// Downsize first so gif-step is quicker and takes less resoruces.
+								// Downsize first so gif-step is quicker and takes less resources.
 								Utils.Image.ConvertImages(ScreenshotSubdirectory.FullName, ScreenshotSubdirectory.FullName, "jpg", true);
 
 								string GifPath = GenerateNotTakenFilePath(Path.Combine(DestinationDirectory.FullName, RoleName + "Test.gif"));
@@ -1285,8 +1284,7 @@ namespace Gauntlet
 			// END REMOVEME
 
 			// Save the Artifact filepath
-			UnrealLogParser LogParser = new UnrealLogParser(InRunningRole.AppInstance.StdOut);
-			return new UnrealRoleArtifacts(InRunningRole.Role, InRunningRole.AppInstance, DestinationDirectory.FullName, ArtifactLogFilePath, LogParser);
+			return new UnrealRoleArtifacts(InRunningRole.Role, InRunningRole.AppInstance, DestinationDirectory.FullName, ArtifactLogFilePath);
 		}
 
 		/// <summary>
diff --git a/Engine/Source/Programs/AutomationTool/Gauntlet/Unreal/Base/Gauntlet.UnrealTestConfiguration.cs b/Engine/Source/Programs/AutomationTool/Gauntlet/Unreal/Base/Gauntlet.UnrealTestConfiguration.cs
--- a/Engine/Source/Programs/AutomationTool/Gauntlet/Unreal/Base/Gauntlet.UnrealTestConfiguration.cs
+++ b/Engine/Source/Programs/AutomationTool/Gauntlet/Unreal/Base/Gauntlet.UnrealTestConfiguration.cs
@@ -959,8 +959,8 @@ namespace Gauntlet
 				AppConfig.CommandLine += " -RemoveInvalidKeys";
 			}
 
-			// use -log on servers so we get a window..
-			if (AppConfig.ProcessType.IsServer())
+			// use -log on user machine so we get a window..
+			if (!AutomationTool.Automation.IsBuildMachine)
 			{
 				AppConfig.CommandLine += " -log";
 			}
diff --git a/Engine/Source/Programs/AutomationTool/Gauntlet/Unreal/Utils/Gauntlet.UnrealLogParser.cs b/Engine/Source/Programs/AutomationTool/Gauntlet/Unreal/Utils/Gauntlet.UnrealLogParser.cs
--- a/Engine/Source/Programs/AutomationTool/Gauntlet/Unreal/Utils/Gauntlet.UnrealLogParser.cs
+++ b/Engine/Source/Programs/AutomationTool/Gauntlet/Unreal/Utils/Gauntlet.UnrealLogParser.cs
@@ -275,7 +275,7 @@ namespace Gauntlet
 			LogEntries = ParsedEntries;
 		}
 
-		static string SanitizeLogText(string InContent)
+		public static string SanitizeLogText(string InContent)
 		{
 			StringBuilder ContentBuilder = new StringBuilder();
 
