Avoid waiting for std streams timeout in consecutive calls  if WaitForExit is called multiple times on the process result object

diff --git a/Engine/Source/Programs/AutomationTool/AutomationUtils/ProcessUtils.cs b/Engine/Source/Programs/AutomationTool/AutomationUtils/ProcessUtils.cs
--- a/Engine/Source/Programs/AutomationTool/AutomationUtils/ProcessUtils.cs
+++ b/Engine/Source/Programs/AutomationTool/AutomationUtils/ProcessUtils.cs
@@ -247,6 +247,8 @@ namespace AutomationTool
 		private Process Proc = null;
 		private AutoResetEvent OutputWaitHandle = new AutoResetEvent(false);
 		private AutoResetEvent ErrorWaitHandle = new AutoResetEvent(false);
+		private bool bStdOutSignalReceived = false;
+		private bool bStdErrSignalReceived = false;
 		private object ProcSyncObject;
 
 		public ProcessResult(string InAppName, Process InProc, bool bAllowSpew, bool bCaptureSpew = true, LogEventType SpewVerbosity = LogEventType.Console, SpewFilterCallbackType InSpewFilterCallback = null)
@@ -478,12 +480,10 @@ namespace AutomationTool
 		public void WaitForExit()
 		{
 			bool bProcTerminated = false;
-			bool bStdOutSignalReceived = false;
-			bool bStdErrSignalReceived = false;
 			// Make sure the process objeect is valid.
 			lock (ProcSyncObject)
 			{
-				bProcTerminated = (Proc == null);
+				bProcTerminated = (Proc == null) || Proc.HasExited;
 			}
 			// Keep checking if we got all output messages until the process terminates.
 			Stopwatch Watch = Stopwatch.StartNew();
