Revert change in conditional copy of log file in Gauntlet UnrealSession
#rnx

diff --git a/Engine/Source/Programs/AutomationTool/Gauntlet/Unreal/Base/Gauntlet.UnrealSession.cs b/Engine/Source/Programs/AutomationTool/Gauntlet/Unreal/Base/Gauntlet.UnrealSession.cs
--- a/Engine/Source/Programs/AutomationTool/Gauntlet/Unreal/Base/Gauntlet.UnrealSession.cs
+++ b/Engine/Source/Programs/AutomationTool/Gauntlet/Unreal/Base/Gauntlet.UnrealSession.cs
@@ -1199,8 +1199,8 @@ namespace Gauntlet
 					}
 					Log.Info($"Wrote {RoleName} Log to {ArtifactLogFilePath}");
 
-					// On build machines, copy all role logs to Horde. Except for test that already generate a dedicated report for it.
-					if (IsBuildMachine && !InContext.TestParams.ParseParam("WriteTestResultsForHorde"))
+					// On build machines, copy all role logs to Horde.
+					if (IsBuildMachine)
 					{
 						string HordeLogFilePath = Path.Combine(CommandUtils.CmdEnv.LogFolder, RoleName + "Output.log");
 						File.Copy(ArtifactLogFilePath, HordeLogFilePath, true);
