Horde: Query the server API version before uploading tool data, and fall back to using older storage formats if the server requires it.

diff --git a/Engine/Source/Programs/AutomationTool/BuildGraph/Tasks/DeployToolTask.cs b/Engine/Source/Programs/AutomationTool/BuildGraph/Tasks/DeployToolTask.cs
--- a/Engine/Source/Programs/AutomationTool/BuildGraph/Tasks/DeployToolTask.cs
+++ b/Engine/Source/Programs/AutomationTool/BuildGraph/Tasks/DeployToolTask.cs
@@ -20,6 +20,10 @@ using EpicGames.Horde.Storage.Nodes;
 using System.Threading;
 using System.Data;
 using EpicGames.Horde.Storage.Backends;
+using Microsoft.Extensions.DependencyInjection;
+using EpicGames.Horde;
+using EpicGames.Horde.Tools;
+using EpicGames.Horde.Server;
 
 #nullable enable
 
@@ -137,77 +141,66 @@ namespace AutomationTool.Tasks
 				throw new AutomationException($"Missing 'server' key from {settingsFile}");
 			}
 
-			Uri serverUri = new Uri(settings.Server);
-
-			HttpClient CreateHttpClient()
+			ServiceCollection serviceCollection = new ServiceCollection();
+			serviceCollection.Configure<HordeOptions>(options =>
 			{
-				HttpClient httpClient = new HttpClient();
-				httpClient.BaseAddress = serverUri;
-				if (settings?.Token != null)
-				{
-					httpClient.DefaultRequestHeaders.Authorization = new AuthenticationHeaderValue("Bearer", settings.Token);
-				}
-				return httpClient;
-			}
+				options.ServerUrl = new Uri(settings.Server);
+				options.AccessToken = settings.Token;
+			});
+			serviceCollection.AddHttpClient();
+			serviceCollection.AddHorde();
+
+			await using ServiceProvider serviceProvider = serviceCollection.BuildServiceProvider();
+			IHordeClient hordeClient = serviceProvider.GetRequiredService<IHordeClient>();
+
+			using HordeHttpClient hordeHttpClient = hordeClient.CreateHttpClient();
+
+			GetServerInfoResponse infoResponse = await hordeHttpClient.GetServerInfoAsync();
+			Logger.LogInformation("Uploading tool to {ServerUrl} (Version: {Version}, API v{ApiVersion})...", settings.Server, infoResponse.ServerVersion, (int)infoResponse.ApiVersion);
+
+			BlobSerializerOptions serializerOptions = BlobSerializerOptions.Create(infoResponse.ApiVersion);
 
 			IBlobHandle handle;
 
-			string basePath = $"api/v1/tools/{Parameters.Id}";
-			using HttpStorageBackend httpStorageBackend = new HttpStorageBackend(basePath, CreateHttpClient, Logger);
-			using BundleStorageClient storageClient = new BundleStorageClient(httpStorageBackend, BundleCache.None, Logger);
+			ToolId toolId = new ToolId(Parameters.Id);
 
+			using IStorageClient storageClient = hordeClient.CreateStorageClient(toolId);
 			await using (IBlobWriter treeWriter = storageClient.CreateBlobWriter())
 			{
 				DirectoryNode sandbox = new DirectoryNode();
 				if (Parameters.File != null)
 				{
 					using FileStream stream = FileReference.Open(ResolveFile(Parameters.File), FileMode.Open, FileAccess.Read);
-					await sandbox.CopyFromZipStreamAsync(stream, treeWriter, new ChunkingOptions());
+					await sandbox.CopyFromZipStreamAsync(stream, treeWriter, new ChunkingOptions(), serializerOptions);
 				}
 				else if (Parameters.Directory != null)
 				{
 					DirectoryInfo directoryInfo = ResolveDirectory(Parameters.Directory).ToDirectoryInfo();
-					await sandbox.AddFilesAsync(directoryInfo, treeWriter);
+					await sandbox.AddFilesAsync(directoryInfo, treeWriter, serializerOptions: serializerOptions);
 				}
 				else
 				{
 					throw new AutomationException("Either File=... or Directory=... must be specified");
 				}
-				handle = await treeWriter.WriteBlobAsync(sandbox);
+				handle = await treeWriter.WriteBlobAsync(sandbox, serializerOptions);
 				await treeWriter.FlushAsync();
 			}
 
-			CreateDeploymentRequest request = new CreateDeploymentRequest();
-			request.Version = Parameters.Version;
+			double? duration = null;
 			if (Parameters.Duration != 0)
 			{
-				request.Duration = Parameters.Duration;
+				duration = Parameters.Duration;
 			}
+
+			bool? createPaused = null;
 			if (Parameters.Paused)
 			{
-				request.CreatePaused = true;
+				createPaused = true;
 			}
-			request.Node = handle.GetLocator().ToString();
 
-			using (HttpClient httpClient = CreateHttpClient())
-			{
-				using (HttpResponseMessage response = await httpClient.PostAsync<CreateDeploymentRequest>(new Uri(serverUri, $"api/v2/tools/{Parameters.Id}/deployments"), request, CancellationToken.None))
-				{
-					if (!response.IsSuccessStatusCode)
-					{
-						string? responseContent;
-						try
-						{
-							responseContent = await response.Content.ReadAsStringAsync();
-						}
-						catch
-						{
-							responseContent = "(No message)";
-						}
-						throw new AutomationException($"Upload failed ({response.StatusCode}): {responseContent}");
-					}
-				}
-			}
+			BlobLocator locator = handle.GetLocator();
+			ToolDeploymentId deploymentId = await hordeHttpClient.CreateToolDeploymentAsync(toolId, Parameters.Version, duration, createPaused, locator);
+			Logger.LogInformation("Created {ToolId} deployment {DeploymentId}", toolId, deploymentId);
 		}
 
 		/// <summary>
diff --git a/Engine/Source/Programs/Horde/Horde.Server/Jobs/Artifacts/ArtifactCollectionV1.cs b/Engine/Source/Programs/Horde/Horde.Server/Jobs/Artifacts/ArtifactCollectionV1.cs
--- a/Engine/Source/Programs/Horde/Horde.Server/Jobs/Artifacts/ArtifactCollectionV1.cs
+++ b/Engine/Source/Programs/Horde/Horde.Server/Jobs/Artifacts/ArtifactCollectionV1.cs
@@ -11,7 +11,6 @@ using MongoDB.Bson.Serialization.Attributes;
 using MongoDB.Driver;
 using EpicGames.Horde.Storage;
 using EpicGames.Horde.Jobs;
-using EpicGames.Horde;
 using EpicGames.Core;
 
 namespace Horde.Server.Jobs.Artifacts
diff --git a/Engine/Source/Programs/Shared/EpicGames.Horde/HordeApiVersion.cs b/Engine/Source/Programs/Shared/EpicGames.Horde/HordeApiVersion.cs
--- a/Engine/Source/Programs/Shared/EpicGames.Horde/HordeApiVersion.cs
+++ b/Engine/Source/Programs/Shared/EpicGames.Horde/HordeApiVersion.cs
@@ -19,12 +19,17 @@ namespace EpicGames.Horde
 		/// <summary>
 		/// Unknown version
 		/// </summary>
-		Unknown,
+		Unknown = 0,
 
 		/// <summary>
 		/// Initial version
 		/// </summary>
-		Initial,
+		Initial = 1,
+
+		/// <summary>
+		/// Interior nodes in chunked data now include the length of chunked data to allow seeking.
+		/// </summary>
+		AddLengthsToInteriorNodes = 2,
 
 		/// <summary>
 		/// One past the latest known version number. Add new version numbers above this point.
diff --git a/Engine/Source/Programs/Shared/EpicGames.Horde/HordeClient.cs b/Engine/Source/Programs/Shared/EpicGames.Horde/HordeClient.cs
--- a/Engine/Source/Programs/Shared/EpicGames.Horde/HordeClient.cs
+++ b/Engine/Source/Programs/Shared/EpicGames.Horde/HordeClient.cs
@@ -1,6 +1,11 @@
 ﻿// Copyright Epic Games, Inc. All Rights Reserved.
 
+using System;
 using System.Net.Http;
+using EpicGames.Horde.Storage;
+using EpicGames.Horde.Storage.Backends;
+using EpicGames.Horde.Storage.Bundles;
+using Microsoft.Extensions.Logging;
 
 namespace EpicGames.Horde
 {
@@ -10,17 +15,29 @@ namespace EpicGames.Horde
 	class HordeClient : IHordeClient
 	{
 		readonly IHttpClientFactory _httpClientFactory;
+		readonly BundleCache _bundleCache;
+		readonly ILoggerFactory _loggerFactory;
 
 		/// <summary>
 		/// Constructor
 		/// </summary>
-		public HordeClient(IHttpClientFactory httpClientFactory)
+		public HordeClient(IHttpClientFactory httpClientFactory, BundleCache bundleCache, ILoggerFactory loggerFactory)
 		{
 			_httpClientFactory = httpClientFactory;
+			_bundleCache = bundleCache;
+			_loggerFactory = loggerFactory;
 		}
 
 		/// <inheritdoc/>
 		public HordeHttpClient CreateHttpClient()
 			=> _httpClientFactory.CreateHordeHttpClient();
+
+		/// <inheritdoc/>
+		public IStorageClient CreateStorageClient(string basePath)
+		{
+			Func<HttpClient> createClient = () => _httpClientFactory.CreateClient(HordeHttpClient.HttpClientName);
+			using HttpStorageBackend httpStorageBackend = new HttpStorageBackend(basePath, createClient, _loggerFactory.CreateLogger<HttpStorageBackend>());
+			return new BundleStorageClient(httpStorageBackend, _bundleCache, _loggerFactory.CreateLogger<BundleStorageClient>());
+		}
 	}
 }
diff --git a/Engine/Source/Programs/Shared/EpicGames.Horde/IHordeClient.cs b/Engine/Source/Programs/Shared/EpicGames.Horde/IHordeClient.cs
--- a/Engine/Source/Programs/Shared/EpicGames.Horde/IHordeClient.cs
+++ b/Engine/Source/Programs/Shared/EpicGames.Horde/IHordeClient.cs
@@ -1,5 +1,8 @@
 // Copyright Epic Games, Inc. All Rights Reserved.
 
+using EpicGames.Horde.Storage;
+using EpicGames.Horde.Tools;
+
 namespace EpicGames.Horde
 {
 	/// <summary>
@@ -12,5 +15,28 @@ namespace EpicGames.Horde
 		/// </summary>
 		/// <returns></returns>
 		HordeHttpClient CreateHttpClient();
+
+		/// <summary>
+		/// Creates a storage client for the given base path
+		/// </summary>
+		IStorageClient CreateStorageClient(string relativePath);
+	}
+
+	/// <summary>
+	/// Extension methods for <see cref="IHordeClient"/>
+	/// </summary>
+	public static class HordeClientExtensions
+	{
+		/// <summary>
+		/// Creates a storage client for a particular namespace
+		/// </summary>
+		public static IStorageClient CreateStorageClient(this IHordeClient hordeClient, NamespaceId namespaceId)
+			=> hordeClient.CreateStorageClient($"api/v1/storage/{namespaceId}");
+
+		/// <summary>
+		/// Creates a storage client for a particular tool
+		/// </summary>
+		public static IStorageClient CreateStorageClient(this IHordeClient hordeClient, ToolId toolId)
+			=> hordeClient.CreateStorageClient($"api/v1/tools/{toolId}");
 	}
 }
diff --git a/Engine/Source/Programs/Shared/EpicGames.Horde/Storage/BlobSerializer.cs b/Engine/Source/Programs/Shared/EpicGames.Horde/Storage/BlobSerializer.cs
--- a/Engine/Source/Programs/Shared/EpicGames.Horde/Storage/BlobSerializer.cs
+++ b/Engine/Source/Programs/Shared/EpicGames.Horde/Storage/BlobSerializer.cs
@@ -6,6 +6,7 @@ using System.Collections.Generic;
 using System.Reflection;
 using System.Threading;
 using System.Threading.Tasks;
+using EpicGames.Horde.Storage.Nodes;
 
 namespace EpicGames.Horde.Storage
 {
@@ -26,31 +27,6 @@ namespace EpicGames.Horde.Storage
 			BlobReader reader = new BlobReader(blobData);
 			return options.GetConverter<T>().Read(reader, options);
 		}
-		/*
-		/// <summary>
-		/// Deserialize an object
-		/// </summary>
-		/// <typeparam name="T">Return type for the deserialized object</typeparam>
-		/// <param name="handle">Handle to the blob to deserialize</param>
-		/// <param name="options">Options to control serialization</param>
-		/// <param name="cancellationToken">Cancellation token for the operation</param>
-		public static async ValueTask<T> DeserializeAsync<T>(IBlobHandle handle, BlobSerializerOptions? options = null, CancellationToken cancellationToken = default)
-		{
-			BlobData data = await handle.ReadAsync(cancellationToken);
-			return Deserialize<T>(data, options);
-		}
-
-		/// <summary>
-		/// Deserialize an object
-		/// </summary>
-		/// <typeparam name="T">Return type for the deserialized object</typeparam>
-		/// <param name="handle">Handle to the blob to deserialize</param>
-		/// <param name="options">Options to control serialization</param>
-		/// <param name="cancellationToken">Cancellation token for the operation</param>
-		public static ValueTask<T> DeserializeAsync<T>(IBlobHandle<T> handle, BlobSerializerOptions? options = null, CancellationToken cancellationToken = default)
-		{
-			return DeserializeAsync<T>((IBlobHandle)handle, options, cancellationToken);
-		}*/
 
 		/// <summary>
 		/// Serialize an object into a blob
@@ -65,21 +41,6 @@ namespace EpicGames.Horde.Storage
 			options ??= BlobSerializerOptions.Default;
 			return options.GetConverter<T>().Write(writer, value, options);
 		}
-		/*
-		/// <summary>
-		/// Deserialize an object
-		/// </summary>
-		/// <param name="writer">Writer for serialized data</param>
-		/// <param name="value">The object to serialize</param>
-		/// <param name="options">Options to control serialization</param>
-		/// <param name="cancellationToken">Cancellation token for the operation</param>
-		/// <returns>Handle to the serialized blob</returns>
-		public static async ValueTask<IBlobHandle<T>> SerializeAsync<T>(IStorageWriter writer, T value, BlobSerializerOptions? options = null, CancellationToken cancellationToken = default)
-		{
-			BlobWriter blobWriter = new BlobWriter(writer);
-			BlobType blobType = Serialize<T>(blobWriter, value, options);
-			return await writer.WriteBlobAsync<T>(blobType, blobWriter.Length, blobWriter.References, cancellationToken);
-		}*/
 	}
 
 	/// <summary>
@@ -135,6 +96,20 @@ namespace EpicGames.Horde.Storage
 			}
 			throw new NotSupportedException($"No converter is available to handle type {type.Name}");
 		}
+
+		/// <summary>
+		/// Creates options for serializing blobs compatible with a particular server API version
+		/// </summary>
+		/// <param name="version">The server API version</param>
+		public static BlobSerializerOptions Create(HordeApiVersion version)
+		{
+			BlobSerializerOptions options = new BlobSerializerOptions();
+			if (version < HordeApiVersion.AddLengthsToInteriorNodes)
+			{
+				options.Converters.Add(new InteriorChunkedDataNodeConverter(2));
+			}
+			return options;
+		}
 	}
 
 	/// <summary>
