Build fix.

diff --git a/Engine/Source/Programs/AutomationTool/BuildGraph/Tasks/CreateArtifactTask.cs b/Engine/Source/Programs/AutomationTool/BuildGraph/Tasks/CreateArtifactTask.cs
--- a/Engine/Source/Programs/AutomationTool/BuildGraph/Tasks/CreateArtifactTask.cs
+++ b/Engine/Source/Programs/AutomationTool/BuildGraph/Tasks/CreateArtifactTask.cs
@@ -125,7 +125,7 @@ namespace AutomationTool.Tasks
 					int totalCount = files.Count;
 					long totalSize = files.Sum(x => x.Length);
 
-					IBlobHandle<DirectoryNode> outputNodeRef = await writer.WriteFilesAsync(baseDir.ToDirectoryInfo(), files, progress: new CopyStatsLoggerWithTotals(totalCount, totalSize, Logger));
+					IBlobHandle<DirectoryNode> outputNodeRef = await writer.WriteFilesAsync(baseDir.ToDirectoryInfo(), files, progress: new CopyStatsLogger(totalCount, totalSize, Logger));
 					await writer.FlushAsync();
 
 					await client.WriteRefAsync(response.RefName, outputNodeRef);
