- Fix UAT not suggesting -skipcook as an option when Stage is requested and no cook param is provided.
- Note Cook is true when -skipcook is set on command line so the condition itself is valid it just was missing -skipcook in the exception message

diff --git a/Engine/Source/Programs/AutomationTool/AutomationUtils/ProjectParams.cs b/Engine/Source/Programs/AutomationTool/AutomationUtils/ProjectParams.cs
--- a/Engine/Source/Programs/AutomationTool/AutomationUtils/ProjectParams.cs
+++ b/Engine/Source/Programs/AutomationTool/AutomationUtils/ProjectParams.cs
@@ -3092,7 +3092,7 @@ namespace AutomationTool
 
 			if (Stage && !SkipStage && !Cook && !CookOnTheFly && !IsProgramTarget)
 			{
-				throw new AutomationException("Only cooked builds or programs can be staged, use -cook or -cookonthefly.");
+				throw new AutomationException("Only cooked builds or programs can be staged, use -cook, -cookonthefly or -skipcook.");
 			}
 
 			if (Manifests && !Cook && !Stage && !Pak)
