Allow specifying the content description attribute for signed files, and set appropriate names for the Horde MSIs. These descriptions are displayed in UAC prompts, and default to random numbers if not specified explicitly.

diff --git a/Engine/Source/Programs/AutomationTool/AutomationUtils/CommandUtils.cs b/Engine/Source/Programs/AutomationTool/AutomationUtils/CommandUtils.cs
--- a/Engine/Source/Programs/AutomationTool/AutomationUtils/CommandUtils.cs
+++ b/Engine/Source/Programs/AutomationTool/AutomationUtils/CommandUtils.cs
@@ -3026,16 +3026,17 @@ namespace AutomationTool
 		};
 
 		[SupportedOSPlatform("windows")]
-		public static void Sign(FileReference File, SignatureType SignatureType, bool AllowMultipleSignatures = true)
+		public static void Sign(FileReference File, SignatureType SignatureType, bool AllowMultipleSignatures = true, string Description = null)
 		{
 			List<FileReference> Files = new List<FileReference> { File };
-			Sign(Files, SignatureType, AllowMultipleSignatures);
+			Sign(Files, SignatureType, AllowMultipleSignatures, Description);
 		}
 
 		[SupportedOSPlatform("windows")]
-		public static void Sign(List<FileReference> Files, SignatureType SignatureType, bool AllowMultipleSignatures = true)
+		public static void Sign(List<FileReference> Files, SignatureType SignatureType, bool AllowMultipleSignatures = true, string Description = null)
 		{
 			string SignToolPath = GetSignToolPath();
+			string DescriptionArg = String.IsNullOrEmpty(Description) ? "" : $"/d \"{Description}\"";
 			string SpecificStoreArg = bUseMachineStoreForCertificates ? " /sm" : "";
 			string MultipleSignatureArg = AllowMultipleSignatures ? " /as" : "";
 			string SHA1TimestampArg = AllowMultipleSignatures ? " /tr" : " /t";
@@ -3053,11 +3054,11 @@ namespace AutomationTool
 					StringBuilder CommandLine = new StringBuilder();
 					if(SignatureType == SignatureType.SHA1)
 					{
-						CommandLine.AppendFormat("sign{0} /a /n \"{1}\" {2} {3} /v {4}", SpecificStoreArg, SigningIdentity, SHA1TimestampArg, TimestampServersSHA1[NumAttempts % TimestampServersSHA1.Length], MultipleSignatureArg);
+						CommandLine.AppendFormat("sign{0} /a /n \"{1}\" {2} {3} /v {4} {5}", SpecificStoreArg, SigningIdentity, SHA1TimestampArg, TimestampServersSHA1[NumAttempts % TimestampServersSHA1.Length], MultipleSignatureArg, DescriptionArg);
 					}
 					else if(SignatureType == SignatureType.SHA256)
 					{
-						CommandLine.AppendFormat("sign{0} /a /fd sha256 /td sha256 /as /n \"{1}\" /tr {2}", SpecificStoreArg, SigningIdentity, TimestampServersSHA256[NumAttempts % TimestampServersSHA256.Length]);
+						CommandLine.AppendFormat("sign{0} /a /fd sha256 /td sha256 /as /n \"{1}\" /tr {2} {3}", SpecificStoreArg, SigningIdentity, TimestampServersSHA256[NumAttempts % TimestampServersSHA256.Length], DescriptionArg);
 					}
 					else
 					{
@@ -3140,7 +3141,7 @@ namespace AutomationTool
 		/// <summary>
 		/// Code signs the specified file
 		/// </summary>
-		public static void SignSingleExecutableIfEXEOrDLL(string Filename, bool bIgnoreExtension = false)
+		public static void SignSingleExecutableIfEXEOrDLL(string Filename, bool bIgnoreExtension = false, string Description = null)
 		{
             if (!OperatingSystem.IsWindows())
             {
@@ -3177,11 +3178,11 @@ namespace AutomationTool
 
 			TargetFileInfo.IsReadOnly = false;
 
-			CodeSignWindows.Sign(new FileReference(TargetFileInfo), CodeSignWindows.SignatureType.SHA1);
+			CodeSignWindows.Sign(new FileReference(TargetFileInfo), CodeSignWindows.SignatureType.SHA1, Description: Description);
 			// MSI files can only have one signature; prefer SHA1 for compatibility, so don't run SHA256 on msi files.
 			if (!TargetFileInfo.FullName.EndsWith(".msi", StringComparison.InvariantCultureIgnoreCase))
 			{
-				CodeSignWindows.Sign(new FileReference(TargetFileInfo), CodeSignWindows.SignatureType.SHA256);
+				CodeSignWindows.Sign(new FileReference(TargetFileInfo), CodeSignWindows.SignatureType.SHA256, Description: Description);
 			}
 		}
 
@@ -3302,7 +3303,7 @@ namespace AutomationTool
 		/// Will automatically skip signing if -NoSign is specified in the command line.
 		/// </summary>
 		/// <param name="Files">List of files to sign</param>
-		public static void SignMultipleIfEXEOrDLL(BuildCommand Command, IEnumerable<string> Files)
+		public static void SignMultipleIfEXEOrDLL(BuildCommand Command, IEnumerable<string> Files, string Description = null)
 		{
 			if (!Command.ParseParam("NoSign"))
 			{
@@ -3322,7 +3323,7 @@ namespace AutomationTool
 					{
 						FilesToSign.Add(new FileReference(File));
 					}
-					SignMultipleFilesIfEXEOrDLL(FilesToSign);
+					SignMultipleFilesIfEXEOrDLL(FilesToSign, Description: Description);
 				}
 			}
 			else
@@ -3331,7 +3332,7 @@ namespace AutomationTool
 			}
 		}
 
-		public static void SignMultipleFilesIfEXEOrDLL(List<FileReference> Files, bool bIgnoreExtension = false)
+		public static void SignMultipleFilesIfEXEOrDLL(List<FileReference> Files, bool bIgnoreExtension = false, string Description = null)
 		{
 			if (!OperatingSystem.IsWindows())
 			{
@@ -3367,8 +3368,8 @@ namespace AutomationTool
 					FinalFiles.Add(new FileReference(TargetFileInfo));
 				}
 			}
-			CodeSignWindows.Sign(FinalFiles, CodeSignWindows.SignatureType.SHA1, !FinalFiles.Any(x => x.HasExtension(".msi"))); // By default we append signatures, but disable if there are MSI files in the list
-			CodeSignWindows.Sign(FinalFiles.Where(x => !x.HasExtension(".msi")).ToList(), CodeSignWindows.SignatureType.SHA256); // MSI files can only have one signature; prefer SHA1 for compatibility
+			CodeSignWindows.Sign(FinalFiles, CodeSignWindows.SignatureType.SHA1, !FinalFiles.Any(x => x.HasExtension(".msi")), Description: Description); // By default we append signatures, but disable if there are MSI files in the list
+			CodeSignWindows.Sign(FinalFiles.Where(x => !x.HasExtension(".msi")).ToList(), CodeSignWindows.SignatureType.SHA256, Description: Description); // MSI files can only have one signature; prefer SHA1 for compatibility
 		}
 	}
 
diff --git a/Engine/Source/Programs/AutomationTool/BuildGraph/Tasks/SignTask.cs b/Engine/Source/Programs/AutomationTool/BuildGraph/Tasks/SignTask.cs
--- a/Engine/Source/Programs/AutomationTool/BuildGraph/Tasks/SignTask.cs
+++ b/Engine/Source/Programs/AutomationTool/BuildGraph/Tasks/SignTask.cs
@@ -25,6 +25,12 @@ namespace AutomationTool.Tasks
 		[TaskParameter(ValidationType = TaskParameterValidationType.FileSpec)]
 		public string Files;
 
+		/// <summary>
+		/// Optional description for the signed content
+		/// </summary>
+		[TaskParameter(Optional = true)]
+		public string Description;
+
 		/// <summary>
 		/// Tag to be applied to build products of this task.
 		/// </summary>
@@ -64,7 +70,7 @@ namespace AutomationTool.Tasks
 			FileReference[] Files = ResolveFilespec(Unreal.RootDirectory, Parameters.Files, TagNameToFileSet).OrderBy(x => x.FullName).ToArray();
 
 			// Sign all the files
-			CodeSign.SignMultipleIfEXEOrDLL(Job.OwnerCommand, (Files.Select(x => x.FullName).ToList()));
+			CodeSign.SignMultipleIfEXEOrDLL(Job.OwnerCommand, (Files.Select(x => x.FullName).ToList()), Description: Parameters.Description);
 
 			// Apply the optional tag to the build products
 			foreach(string TagName in FindTagNamesFromList(Parameters.Tag))
diff --git a/Engine/Source/Programs/Horde/BuildHorde.xml b/Engine/Source/Programs/Horde/BuildHorde.xml
--- a/Engine/Source/Programs/Horde/BuildHorde.xml
+++ b/Engine/Source/Programs/Horde/BuildHorde.xml
@@ -192,7 +192,7 @@
 			<Spawn Exe="$(WixDir)/light.exe" Arguments="-out &quot;$(BinDir)/HordeAgent.msi&quot; -sw1076 -pdbout &quot;$(BinDir)/HordeAgent.pdb&quot; &quot;$(ObjDir)/AgentInstaller.wixobj&quot; &quot;$(ObjDir)/HordeAgentFiles.wixobj&quot; -ext WixUtilExtension -ext WixUIExtension" WorkingDir="$(InstallerDir)"/>
 
 			<!-- We need to sign here as will be included in the server msi -->
-			<Sign Files="$(InstallerOutputDir)/HordeAgent.msi" If="'$(SignServerToolsAgent)' == 'true'"/>
+			<Sign Files="$(InstallerOutputDir)/HordeAgent.msi" Description="Horde Agent" If="'$(SignServerToolsAgent)' == 'true'"/>
 			
 			<!-- Copy files to the output directory -->
 			<Copy Files="HordeAgent.msi" From="$(InstallerOutputDir)" To="$(ServerToolsStagingDir_HordeAgentMsi)"/>
