Localization:
- Allow Localize UAT command to skip plugins that either have a .uplugin that isn't updated with a localization config generation policy or a config generation policy of Never.
- Added a log to inform user if no localization batches matching the user input criteria were found and early out of the build command.
- Cleaned up the logging for cleaning up auto-generated config files and directories. Now the logs print only when there actually were auto-generated logs and directories.
- All plugins that are a part of the user provided include list that does not have any localization targets or should not be gathered are now removed from the localization list and those plugins are not enabled.

diff --git a/Engine/Source/Programs/AutomationTool/Scripts/Localisation.Automation.cs b/Engine/Source/Programs/AutomationTool/Scripts/Localisation.Automation.cs
--- a/Engine/Source/Programs/AutomationTool/Scripts/Localisation.Automation.cs
+++ b/Engine/Source/Programs/AutomationTool/Scripts/Localisation.Automation.cs
@@ -249,11 +249,26 @@ class Localize : BuildCommand
 		if (bShouldGatherPlugins)
 		{
 			HashSet<string> PluginsAdded = new();
-			AddPluginLocalizationBatches(LocalizationBatches, PluginsAdded);
+			HashSet<string> AdditionalPluginsToExclude = new HashSet<string>();
+			AddPluginLocalizationBatches(LocalizationBatches, PluginsAdded, AdditionalPluginsToExclude);
 			if (PluginsAdded.Count > 0)
 			{
 				AutoGeneratePluginLocalizationFiles(PluginsAdded);
 			}
+			if (AdditionalPluginsToExclude.Count > 0)
+			{
+				// We now remove these plugins from the include plugins list so that we don't enable them unnecessarily down the line 
+				Logger.LogInformation($"Removing {AdditionalPluginsToExclude.Count} additional plugins that do not need to be localized.");
+				// Technically we could just set the include plugins to be the added plugisn, but we leave this exclusion explicit to make things clearer
+				// Further, additional logic can be included next time to modify what plugisn should be removed 
+				IncludePlugins = IncludePlugins.Except(AdditionalPluginsToExclude.ToList()).ToList();
+			}
+		}
+
+		if (LocalizationBatches.Count == 0)
+		{
+			Logger.LogWarning("No localization batches found from input parameters provided. Nothing will be localized.");
+			return;
 		}
 
 		// Create a single changelist to use for all changes
@@ -494,7 +509,7 @@ class Localize : BuildCommand
 		}
 	}
 
-	private void AddPluginLocalizationBatches(List<LocalizationBatch> LocalizationBatches, HashSet<string> PluginsAdded)
+	private void AddPluginLocalizationBatches(List<LocalizationBatch> LocalizationBatches, HashSet<string> PluginsAdded, HashSet<string> AdditionalPluginsToExclude)
 	{
 		var PluginsRootDirectory = new DirectoryReference(CombinePaths(UEProjectRoot, UEProjectDirectory));
 		IReadOnlyList<PluginInfo> AllPlugins = Plugins.ReadPluginsFromDirectory(PluginsRootDirectory, "Plugins", UEProjectName.Length == 0 ? PluginType.Engine : PluginType.Project);
@@ -509,17 +524,31 @@ class Localize : BuildCommand
 			bool bPluginHasLocalizationTarget = PluginInfo.Descriptor.LocalizationTargets != null && PluginInfo.Descriptor.LocalizationTargets.Length > 0;
 			if (bShouldIncludePlugin && bPluginHasLocalizationTarget)
 			{
-				PluginsAdded.Add(PluginInfo.Name);
 				var RootRelativePluginPath = PluginInfo.Directory.MakeRelativeTo(new DirectoryReference(UEProjectRoot));
 				RootRelativePluginPath = RootRelativePluginPath.Replace('\\', '/'); // Make sure we use / as these paths are used with P4
 
 				var PluginTargetNames = new List<string>();
 				foreach (var LocalizationTarget in PluginInfo.Descriptor.LocalizationTargets)
 				{
-					PluginTargetNames.Add(LocalizationTarget.Name);
+					// we only add the plugin target if the config generation policy is auto or user, i.e it should actually participate in the localization gather 
+					if (LocalizationTarget.ConfigGenerationPolicy != LocalizationConfigGenerationPolicy.Never)
+					{
+						PluginTargetNames.Add(LocalizationTarget.Name);
+					}
+					
 				}
 
-				LocalizationBatches.Add(new LocalizationBatch(UEProjectDirectory, RootRelativePluginPath, PluginInfo.Name, PluginTargetNames));
+				// If the plugin has no valid loc targets to gather, we don't create the batch 
+				if (PluginTargetNames.Count > 0)
+				{
+					LocalizationBatches.Add(new LocalizationBatch(UEProjectDirectory, RootRelativePluginPath, PluginInfo.Name, PluginTargetNames));
+					PluginsAdded.Add(PluginInfo.Name);
+				}
+				else
+				{
+					// The plugin doesn't actually have to be included. 
+					AdditionalPluginsToExclude.Add(PluginInfo.Name);
+				}
 			}
 		}
 
@@ -530,6 +559,17 @@ class Localize : BuildCommand
 			{
 				Logger.LogWarning("The plugin '{PluginName}' specified by -IncludePlugins wasn't found and will be skipped.", PluginName);
 			}
+			else if (!PluginsAdded.Contains(PluginName))
+			{
+				// If the plugin was in the include list of plugins but they weren't added, this means that those plugisn failed the filter. 
+				// Those plugisn should now be removed from the include list to prevent them from being enabled for no reason 
+				AdditionalPluginsToExclude.Add(PluginName);
+			}
+		}
+
+		foreach (string AdditionalPluginToExclude in AdditionalPluginsToExclude)
+		{
+			Logger.LogInformation($"The plugin '{AdditionalPluginToExclude}' has no localization targets that needs localizing and will be skipped.");
 		}
 	}
 
@@ -855,33 +895,39 @@ class Localize : BuildCommand
 
 	private void CleanUpAutoGeneratedFiles()
 	{
-		Logger.LogInformation("Cleaning up auto-generated files.");
-		foreach (string AutoGeneratedFile in AutoGeneratedFiles)
+		if (AutoGeneratedFiles.Count > 0)
 		{
-			if (File.Exists(AutoGeneratedFile))
+			Logger.LogInformation("Cleaning up auto-generated files.");
+			foreach (string AutoGeneratedFile in AutoGeneratedFiles)
 			{
-				Logger.LogInformation($"Deleting auto-generated file '{AutoGeneratedFile}'");
-				if (!bIsRunningInPreview)
+				if (File.Exists(AutoGeneratedFile))
 				{
-					File.Delete(AutoGeneratedFile);
+					Logger.LogInformation($"Deleting auto-generated file '{AutoGeneratedFile}'");
+					if (!bIsRunningInPreview)
+					{
+						File.Delete(AutoGeneratedFile);
+					}
 				}
 			}
+			Logger.LogInformation("Finished cleaning up all auto-generated files.");
 		}
-		Logger.LogInformation("Finished cleaning up all auto-generated files.");
 
-		Logger.LogInformation("Cleaning up auto-generated directories.");
-		foreach (string AutoGeneratedDirectory in AutoGeneratedDirectories)
+		if (AutoGeneratedDirectories.Count > 0)
 		{
-			if (Directory.Exists(AutoGeneratedDirectory))
+			Logger.LogInformation("Cleaning up auto-generated directories.");
+			foreach (string AutoGeneratedDirectory in AutoGeneratedDirectories)
 			{
-				Logger.LogInformation($"Deleting auto-generated directory '{AutoGeneratedDirectory}'.");
-				if (!bIsRunningInPreview)
+				if (Directory.Exists(AutoGeneratedDirectory))
 				{
-					Directory.Delete(AutoGeneratedDirectory);
+					Logger.LogInformation($"Deleting auto-generated directory '{AutoGeneratedDirectory}'.");
+					if (!bIsRunningInPreview)
+					{
+						Directory.Delete(AutoGeneratedDirectory);
+					}
 				}
 			}
+			Logger.LogInformation("Finished cleaning up auto-generated directories.");
 		}
-		Logger.LogInformation("Finished cleaning up auto-generated directories.");
 	}
 
 	private ProjectInfo GenerateProjectInfo(string RootWorkingDirectory, string ProjectName, IReadOnlyList<string> LocalizationStepNames)
