Override Horde testdata filename through ArtifactName option

#rnx

diff --git a/Engine/Source/Programs/AutomationTool/Gauntlet/Unreal/Base/Gauntlet.UnrealTestNode.cs b/Engine/Source/Programs/AutomationTool/Gauntlet/Unreal/Base/Gauntlet.UnrealTestNode.cs
--- a/Engine/Source/Programs/AutomationTool/Gauntlet/Unreal/Base/Gauntlet.UnrealTestNode.cs
+++ b/Engine/Source/Programs/AutomationTool/Gauntlet/Unreal/Base/Gauntlet.UnrealTestNode.cs
@@ -1524,9 +1524,10 @@ namespace Gauntlet
 			}
 
 			// write test data collection for Horde
+			string FileName = FileUtils.SanitizeFilename(string.IsNullOrEmpty(Context.Options.ArtifactName) ? Name : Context.Options.ArtifactName);
 			string HordeTestDataFilePath = Path.Combine(
 				string.IsNullOrEmpty(GetCachedConfiguration().HordeTestDataPath) ? HordeReport.DefaultTestDataDir : GetCachedConfiguration().HordeTestDataPath,
-				FileUtils.SanitizeFilename(Name) + ".TestData.json"
+				FileName + ".TestData.json"
 			);
 			HordeReport.TestDataCollection HordeTestDataCollection = new HordeReport.TestDataCollection();
 			HordeTestDataCollection.AddNewTestReport(Report, GetCachedConfiguration().HordeTestDataKey);
