Scape '(' and ')' in the session name when launching on Android.
Otherwise it will fail to launch from project launcher since session name is like `DeviceName (SerialNumber)` with the error `/system/bin/sh: syntax error: unexpected '('`

diff --git a/Engine/Source/Programs/AutomationTool/Android/AndroidPlatform.Automation.cs b/Engine/Source/Programs/AutomationTool/Android/AndroidPlatform.Automation.cs
--- a/Engine/Source/Programs/AutomationTool/Android/AndroidPlatform.Automation.cs
+++ b/Engine/Source/Programs/AutomationTool/Android/AndroidPlatform.Automation.cs
@@ -4066,7 +4066,13 @@ public class AndroidPlatform : Platform
 			var CommandLine = "shell am start -n " + PackageName + "/" + GetLaunchableActivityName();
 			if (canReadClientCmdLineViaAmStart)
 			{
-				var ClientSessionCmdLineEscaped = ClientCmdLine.Replace(" ", "\\ ").Replace("\"", "\\\\\\\"");
+				var ToScapeChars = new string[]{" ", "(", ")", "`", "$", "%", "&"};
+				
+				var ClientSessionCmdLineEscaped = ClientCmdLine.Replace("\"", "\\\\\\\"");
+				foreach( var ToScape in ToScapeChars )
+				{
+					ClientSessionCmdLineEscaped = ClientSessionCmdLineEscaped.Replace(ToScape, "\\" + ToScape);
+				}
 				CommandLine += " --es cmdline \"" + ClientSessionCmdLineEscaped + "\"";
 			}
 			RunAdbCommand(DeviceName, CommandLine);
