[Gauntlet] Adds new UnrealSession.LaunchSession flow using new TargetDevice functions
This new flow is disabled by default and requires -ExperimentalLaunchFlow to be present to run

Changes:
- Current implementation of LaunchSession has been moved to Legacy_LaunchSession(). Adding further references should be avoided.
- Exception guarding happens in several steps, improving clarity on which step failures occur in LaunchSession
- Tests that run several iterations will now clean the device artifacts between each run without needing a full re-install of the build
- UnrealDeviceReservation can now choose to release only problem devices
- When a problem device is encountered during LaunchSession, only the problem device(s) will be released to the pool instead of all devices
- Overall improvements to logging

diff --git a/Engine/Source/Programs/AutomationTool/Gauntlet/Framework/Devices/Gauntlet.UnrealDeviceReservation.cs b/Engine/Source/Programs/AutomationTool/Gauntlet/Framework/Devices/Gauntlet.UnrealDeviceReservation.cs
--- a/Engine/Source/Programs/AutomationTool/Gauntlet/Framework/Devices/Gauntlet.UnrealDeviceReservation.cs
+++ b/Engine/Source/Programs/AutomationTool/Gauntlet/Framework/Devices/Gauntlet.UnrealDeviceReservation.cs
@@ -141,16 +141,47 @@ namespace Gauntlet
 		{
 			if ((ReservedDevices != null) && (ReservedDevices.Count() > 0))
 			{
-				foreach (ITargetDevice device in ReservedDevices)
+				foreach (ITargetDevice Device in ReservedDevices)
 				{
-					IDeviceUsageReporter.RecordEnd(device.Name, device.Platform, IDeviceUsageReporter.EventType.Device);
+					IDeviceUsageReporter.RecordEnd(Device.Name, Device.Platform, IDeviceUsageReporter.EventType.Device);
 				}
 				DevicePool.Instance.ReleaseDevices(ReservedDevices);
 				ReservedDevices.Clear();
 			}
 		}
 
-		public void MarkProblemDevice(ITargetDevice Device)
+		public IEnumerable<ITargetDevice> ReleaseProblemDevices()
+		{
+			List<ITargetDevice> Devices = new();
+			bool bReservedDevices = ReservedDevices != null && ReservedDevices.Any();
+			bool bProblemDevices = ProblemDevices != null && ProblemDevices.Any();
+
+			if(bReservedDevices && bProblemDevices)
+			{
+				foreach(ProblemDevice Problem in ProblemDevices)
+				{
+					foreach(ITargetDevice Device in ReservedDevices)
+					{
+						if (Problem.Name == Device.Name && Problem.Platform == Device.Platform)
+						{
+							Devices.Add(Device);
+						}
+					}
+				}
+
+				DevicePool.Instance.ReleaseDevices(Devices);
+				ProblemDevices.Clear();
+
+				foreach (ITargetDevice Device in Devices)
+				{
+					ReservedDevices.Remove(Device);
+					IDeviceUsageReporter.RecordEnd(Device.Name, Device.Platform, IDeviceUsageReporter.EventType.Device);
+				}
+			}
+			return Devices;
+		}
+
+		public void MarkProblemDevice(ITargetDevice Device, string ErrorMessage = "MarkProblemDevice")
 		{
 			if (ProblemDevices.Where(D => D.Name == Device.Name && D.Platform == Device.Platform).Count() > 0)
 			{
@@ -158,7 +189,7 @@ namespace Gauntlet
 			}
 
 			// report device has a problem to the pool
-			DevicePool.Instance.ReportDeviceError(Device, "MarkProblemDevice");
+			DevicePool.Instance.ReportDeviceError(Device, ErrorMessage);
 
 			if (Device.Platform != null)
 			{
@@ -166,5 +197,4 @@ namespace Gauntlet
 			}
 		}
 	}
-}
- 
\ No newline at end of file
+}
\ No newline at end of file
diff --git a/Engine/Source/Programs/AutomationTool/Gauntlet/Unreal/Base/Gauntlet.UnrealSession.cs b/Engine/Source/Programs/AutomationTool/Gauntlet/Unreal/Base/Gauntlet.UnrealSession.cs
--- a/Engine/Source/Programs/AutomationTool/Gauntlet/Unreal/Base/Gauntlet.UnrealSession.cs
+++ b/Engine/Source/Programs/AutomationTool/Gauntlet/Unreal/Base/Gauntlet.UnrealSession.cs
@@ -95,11 +95,12 @@
 		public List<UnrealFileToCopy> FilesToCopy;
 
 		/// <summary>
-		/// 
+		/// Additional UE directories to copy from when saving artifacts
 		/// </summary>
 		public List<EIntendedBaseCopyDirectory> AdditionalArtifactDirectories;
+
 		/// <summary>
-		/// Role device configuration 
+		/// Role device configuration
 		/// </summary>
 		public ConfigureDeviceHandler ConfigureDevice;
 
@@ -119,7 +120,7 @@
 		public ERoleModifier RoleModifier;
 
 		/// <summary>
-		/// Is this a dummy executable? 
+		/// Is this a dummy executable?
 		/// </summary>
 		public bool IsDummy() { return RoleModifier == ERoleModifier.Dummy; }
 
@@ -134,7 +135,7 @@
 		public bool DeferredLaunch { get; set; }
 
 		/// <summary>
-		/// Is this role Null? 
+		/// Is this role Null?
 		/// </summary>
 		public bool IsNullRole() { return RoleModifier == ERoleModifier.Null; }
 
@@ -587,14 +588,43 @@
 		public string Sandbox { get; set; }
 
 		/// <summary>
-		/// Whether or not we should retain our devices this pass
+		/// Whether or not devices should be retained between each test iteration
 		/// </summary>
 		public bool ShouldRetainDevices { get; set; }
 
+		[AutoParam(false)]
+		public bool ReinstallPerPass { get; set; }
+
 		/// <summary>
+		/// Number of attempts when launching a session.
+		/// Failed installs and runs will trigger a re-try
+		/// </summary>
+		public int LaunchSessionAttempts { get; set; }
+
+		/// <summary>
+		/// Number of attempts when trying to reserve devices
+		/// </summary>
+		public int DeviceReservationAttempts { get; set; }
+
+		/// <summary>
+		/// Number of seconds to wait between failed device reservation attempts
+		/// </summary>
+		public int DeviceReservationRetryTime { get; set; }
+
+		/// <summary>
+		/// Record of each ITargetDevice assigned to a given role
+		/// </summary>
+		public Dictionary<UnrealSessionRole, ITargetDevice> RolesToDevices { get; private set; }
+
+		/// <summary>
+		/// Record of each UnrealAppConfig created for each role
+		/// </summary>
+		public Dictionary<UnrealSessionRole, UnrealAppConfig> RolesToConfigs { get; private set; }
+
+		/// <summary>
 		/// Record of our installations in case we want to re-use them in a later pass
 		/// </summary>
-		public Dictionary<UnrealSessionRole, IAppInstall> RolesToInstalls;
+		public Dictionary<UnrealSessionRole, IAppInstall> RolesToInstalls { get; private set; }
 
 		/// <summary>
 		/// Constructor that takes a build source and a number of roles
@@ -603,8 +633,19 @@
 		/// <param name="InSessionRoles"></param>
 		public UnrealSession(UnrealBuildSource InSource, IEnumerable<UnrealSessionRole> InSessionRoles)
 		{
+			AutoParam.ApplyParamsAndDefaults(this, Globals.Params.AllArguments);
+
 			BuildSource = InSource;
 			SessionRoles = InSessionRoles.ToArray();
+			ShouldRetainDevices = !Globals.Params.ParseParam("ReacquireDevicesPerPass");
 
+			RolesToDevices = new Dictionary<UnrealSessionRole, ITargetDevice>();
+			RolesToConfigs = new Dictionary<UnrealSessionRole, UnrealAppConfig>();
+			RolesToInstalls = new Dictionary<UnrealSessionRole, IAppInstall>();
+
+			LaunchSessionAttempts = 3;
+			DeviceReservationAttempts = 5;
+			DeviceReservationRetryTime = 120;
+
 			if (SessionRoles.Length == 0)
 			{
@@ -662,30 +703,47 @@
 		/// <returns></returns>
 		public bool TryReserveDevices()
 		{
+			if(ShouldRetainDevices && HasAcquiredDevices())
+			{
+				return true;
+			}
+
 			// figure out how many of each device we need
 			Dictionary<UnrealDeviceTargetConstraint, int> RequiredDeviceTypes = new Dictionary<UnrealDeviceTargetConstraint, int>();
-			IEnumerable<UnrealSessionRole> RolesNeedingDevices = SessionRoles.Where(R => !R.IsNullRole());
+			IEnumerable<UnrealSessionRole> RolesThatRequireDevice = SessionRoles.Where(R => !R.IsNullRole());
 
 			// Get a count of the number of devices required for each platform
-			RolesNeedingDevices.ToList().ForEach(C =>
+			foreach(UnrealSessionRole Role in RolesThatRequireDevice)
 			{
-				if (!RequiredDeviceTypes.ContainsKey(C.Constraint))
+				if(RequiredDeviceTypes.ContainsKey(Role.Constraint))
 				{
-					RequiredDeviceTypes[C.Constraint] = 0;
+					++RequiredDeviceTypes[Role.Constraint];
 				}
-				RequiredDeviceTypes[C.Constraint]++;
-			});
+				else
+				{
+					RequiredDeviceTypes.Add(Role.Constraint, 1);
+				}
+			}
 
-
-			if (UnrealDeviceReservation.ReservedDevices != null && UnrealDeviceReservation.ReservedDevices.Count > 0
-				&& ShouldRetainDevices)
-			{
-				return true;
-			}
-			else
-			{
-				return UnrealDeviceReservation.TryReserveDevices(RequiredDeviceTypes, RolesNeedingDevices.Count());
-			}
+			return UnrealDeviceReservation.TryReserveDevices(RequiredDeviceTypes, RolesThatRequireDevice.Count());
 		}
 
+		public bool TryReserveDevices(int Attempts)
+		{
+			for (; Attempts > 0; --Attempts)
+			{
+				if(Globals.CancelSignalled || TryReserveDevices())
+				{
+					return true;
+				}
+				else
+				{
+					Thread.Sleep(1000 * DeviceReservationRetryTime);
+				}
+			}
+
+			Log.Error("Failed to reserve devices after {Attempts}", Attempts);
+			return false;
+		}
+
 		/// <summary>
@@ -715,244 +773,96 @@
 		/// <returns></returns>
 		public UnrealSessionInstance LaunchSession()
 		{
+			if(!Globals.Params.ParseParam("ExperimentalLaunchFlow"))
+			{
+				return Legacy_LaunchSession();
+			}
+
+			// Clear any existing session from a previous iteration
 			SessionInstance = null;
 
-			// The number of retries when launching session to avoid an endless loop if package can't be installed, network timeouts, etc
-			int SessionRetries = 2;
+			// When launching, issues with devices may be encountered.
+			// When these issues occur, those devices will be marked as problem devices and returned to the pool.
+			// A new set of devices will then be reserved and another attempt at launching the processes will be made.
+			// If LaunchSession() fails LaunchSessionAttempts amount of times, an exception is thrown.
+			for (int RemainingAttempts = LaunchSessionAttempts; RemainingAttempts > 0; --RemainingAttempts)
+			{
+				// Reserve devices, if needed
+				if (!TryReserveDevices(DeviceReservationAttempts))
+				{
+					// If device reservation fails, the device pool cannot support this launch.
+					DevicePool.Instance.ReportDeviceReservationState();
+					throw new AutomationException("Failed to acquire all devices for launch. See above for details.");
+				}
 
-			// tries to find devices and launch our session. Will loop until we succeed, we run out of devices/retries, or
-			// something fatal occurs..
-			while (SessionInstance == null && Globals.CancelSignalled == false)
-			{
-				int ReservationRetries = 5;
-				int ReservationRetryWait = 120;
+				if(Globals.CancelSignalled)
+				{
+					return null;
+				}
 
-				IEnumerable<UnrealSessionRole> RolesNeedingDevices = SessionRoles.Where(R => R.IsNullRole() == false);
+				if(!TryAssignDevicesToRoles())
+				{
+					// If device assignment fails, reservations were likely deleted at an unexpected time.
+					ReleaseSessionDevices();
+					continue;
+				}
 
-				while (UnrealDeviceReservation.ReservedDevices.Count() < RolesNeedingDevices.Count())
+				// All roles should now be assigned a device.
+				// Install necessary builds, clear stale test artifacts, and copy any additional files
+				try
 				{
-					// get devices
-					TryReserveDevices();
-
-					if (Globals.CancelSignalled)
+					ReadyDevicesForSession();
+				}
+				catch (Exception Ex)
+				{
+					if(IsOutOfSpaceException(Ex))
 					{
-						break;
+						RemainingAttempts = 0;
+						ReleaseSessionDevices();
+						continue;
+					}
+					else if (RemainingAttempts > 1)
+					{
+						Log.Info("A new device will be selected and another attempt at launching session will be made.");
 					}
 
-					// if we failed to get enough devices, show a message and wait
-					if (UnrealDeviceReservation.ReservedDevices.Count() != SessionRoles.Count())
-					{
-						if (ReservationRetries == 0)
-						{
-							DevicePool.Instance.ReportDeviceReservationState();
-							throw new AutomationException("Unable to acquire all devices for test.");
-						}
-						Log.Info("\nUnable to find enough device(s). Waiting {0} secs (retries left={1})\n", ReservationRetryWait, --ReservationRetries);
-						Thread.Sleep(ReservationRetryWait * 1000);
-					}
+					ReleaseProblemDevices();
+					continue;
 				}
 
 				if (Globals.CancelSignalled)
 				{
+					ReleaseSessionDevices();
 					return null;
 				}
 
-				Dictionary<IAppInstall, UnrealSessionRole> InstallsToRoles = new Dictionary<IAppInstall, UnrealSessionRole>();
-
-				// create a copy of our list
-				IEnumerable<ITargetDevice> DevicesToInstallOn = UnrealDeviceReservation.ReservedDevices.ToArray();
-
-				bool InstallSuccess = true;
-
-				// sort by constraints, so that we pick constrained devices first
-				List<UnrealSessionRole> SortedRoles = SessionRoles.OrderBy(R => R.Constraint.IsIdentity() ? 1 : 0).ToList();
-
-				// first install all roles on these devices
-				foreach (UnrealSessionRole Role in SortedRoles)
+				// All roles should now be assigned device with an associated IAppInstall.
+				// Launch all the processes!
+				try
 				{
-					ITargetDevice Device = null;
-
-					if (Role.IsNullRole() == false)
+					SessionInstance = LaunchProcesses();
+				}
+				catch
+				{
+					if (RemainingAttempts > 1)
 					{
-						Device = DevicesToInstallOn.Where(D => D.IsConnected && D.Platform == Role.Platform
-															&& ( Role.Constraint.IsIdentity() || DevicePool.Instance.GetConstraint(D) == Role.Constraint)).First();
-
-						DevicesToInstallOn = DevicesToInstallOn.Where(D => D != Device);
-					}
-					else
-					{
-						Device = new TargetDeviceNull(string.Format("Null{0}", Role.RoleType));
+						Log.Info("A new device will be selected and another attempt at launching session will be made.");
 					}
 
-					IEnumerable<UnrealSessionRole> OtherRoles = SortedRoles.Where(R => R != Role);
-
-					// create a config from the build source (this also applies the role options)
-					UnrealAppConfig AppConfig = BuildSource.CreateConfiguration(Role, OtherRoles);
-
-					// todo - should this be elsewhere?
-					AppConfig.Sandbox = Sandbox;
-
-					IAppInstall Install = null;
-					bool bReinstallPerPass = Globals.Params.ParseParam("ReinstallPerPass");
-					if (RolesToInstalls == null || !RolesToInstalls.ContainsKey(Role) || bReinstallPerPass)
-					{
-						// Tag the device for report result
-						if(BuildHostPlatform.Current.Platform != Device.Platform)
-						{
-							AppConfig.CommandLineParams.Add("DeviceTag", Device.Name);
-						}
-
-						IDeviceUsageReporter.RecordStart(Device.Name, Device.Platform, IDeviceUsageReporter.EventType.Device, IDeviceUsageReporter.EventState.Success);
-						IDeviceUsageReporter.RecordStart(Device.Name, Device.Platform, IDeviceUsageReporter.EventType.Install, IDeviceUsageReporter.EventState.Success, BuildSource.BuildName);
-						try
-						{
-							Install = Device.InstallApplication(AppConfig);
-							IDeviceUsageReporter.RecordEnd(Device.Name, Device.Platform, IDeviceUsageReporter.EventType.Install, IDeviceUsageReporter.EventState.Success);
-						}
-						catch (Exception Ex)
-						{
-							// Warn, ignore the device, and do not continue
-							string ErrorMessage = string.Format("Encountered error setting up device {0} for role {1}. {2}. Will retry with new device", Device, Role, Ex);
-							if (ErrorMessage.Contains("not enough space"))
-							{
-								Log.Error(KnownLogEvents.Gauntlet_DeviceEvent, ErrorMessage);
-								if (Device.Platform == BuildHostPlatform.Current.Platform)
-								{
-									// If on desktop platform, we are not retrying.
-									// It is unlikely that space is going to be made and InstallBuildParallel has marked the build path as problematic.
-									SessionRetries = 0;
-								}
-							}
-							else
-							{
-								Log.Warning(KnownLogEvents.Gauntlet_DeviceEvent, ErrorMessage);
-							}
-							UnrealDeviceReservation.MarkProblemDevice(Device);
-							InstallSuccess = false;
-							IDeviceUsageReporter.RecordEnd(Device.Name, Device.Platform, IDeviceUsageReporter.EventType.Install, IDeviceUsageReporter.EventState.Failure);
-							break;
-						}
-
-
-						if (Globals.CancelSignalled)
-						{
-							break;
-						}
-
-						// Device has app installed, give role a chance to configure device
-						Role.ConfigureDevice?.Invoke(Device);
-
-						InstallsToRoles[Install] = Role;
-
-						if(RolesToInstalls == null)
-						{
-							RolesToInstalls = new Dictionary<UnrealSessionRole, IAppInstall>();
-						}
-						if (!bReinstallPerPass)
-						{
-							RolesToInstalls[Role] = Install;
-						}
-					}
-					else
-					{
-						Install = RolesToInstalls[Role];
-						InstallsToRoles[Install] = Role;
-						Log.Info("Using previous install of {0} on {1}", Install.Name, Install.Device.Name);
-					}
+					ReleaseProblemDevices();
+					continue;
 				}
 
-				if (InstallSuccess == false)
+				if (Globals.CancelSignalled)
 				{
 					ReleaseSessionDevices();
-
-					if (SessionRetries == 0)
-					{
-						throw new AutomationException("Unable to install application for session.");
-					}
-
-					Log.Info("\nUnable to install application for session (retries left={0})\n", --SessionRetries);
+					return null;
 				}
 
-				if (InstallSuccess && Globals.CancelSignalled == false)
-				{
-					List<UnrealSessionInstance.RoleInstance> AllRoles = new List<UnrealSessionInstance.RoleInstance>();
-					Dictionary<UnrealSessionInstance.RoleInstance, IAppInstall> DeferredRoleToAppInstall = new Dictionary<UnrealSessionInstance.RoleInstance, IAppInstall>();
-
-					// Now try to run all installs on their devices
-					foreach (var InstallRoleKV in InstallsToRoles)
-					{
-						IAppInstall CurrentInstall = InstallRoleKV.Key;
-						if (InstallRoleKV.Value.InstallOnly)
-						{
-							AllRoles.Add(new UnrealSessionInstance.RoleInstance(InstallRoleKV.Value, null));
-							continue;
-						}
-						if (InstallRoleKV.Value.DeferredLaunch)
-						{
-							UnrealSessionInstance.RoleInstance DeferredRoleInstance = new UnrealSessionInstance.RoleInstance(InstallRoleKV.Value, null);
-
-							DeferredRoleToAppInstall.Add(DeferredRoleInstance, CurrentInstall);
-							AllRoles.Add(DeferredRoleInstance);
-							continue;
-						}
-						bool Success = false;
-
-						try
-						{
-							Log.Info("Starting {0} on {1}", InstallRoleKV.Value, CurrentInstall.Device);
-							IAppInstance Instance = CurrentInstall.Run();
-							IDeviceUsageReporter.RecordStart(Instance.Device.Name, Instance.Device.Platform, IDeviceUsageReporter.EventType.Test);
-
-							if (Instance != null || Globals.CancelSignalled)
-							{
-								AllRoles.Add(new UnrealSessionInstance.RoleInstance(InstallRoleKV.Value, Instance));
-							}
-
-							Success = true;
-						}
-						catch (DeviceException Ex)
-						{
-							// shutdown all 
-							Log.Warning(KnownLogEvents.Gauntlet_DeviceEvent, "Device {Name} threw an exception during launch. \nException={Exception}", CurrentInstall.Device, Ex.Message);
-							Success = false;
-						}
-
-						if (Success == false)
-						{
-							Log.Warning(KnownLogEvents.Gauntlet_DeviceEvent, "Failed to start build on {Name}. Marking as problem device and retrying with new set", CurrentInstall.Device);
-
-							// terminate anything that's running
-							foreach (UnrealSessionInstance.RoleInstance RunningRole in AllRoles.Where(X => X.AppInstance != null) )
-							{
-								Log.Info("Shutting down {0}", RunningRole.AppInstance.Device);
-								RunningRole.AppInstance.Kill();
-								RunningRole.AppInstance.Device.Disconnect();
-							}
-
-							// mark that device as a problem
-							UnrealDeviceReservation.MarkProblemDevice(CurrentInstall.Device);
-
-							ReleaseSessionDevices();
-
-							if (SessionRetries == 0)
-							{
-								throw new AutomationException("Unable to start application for session, see warnings for details.");
-							}
-
-							Log.Info("\nUnable to start application for session (retries left={0})\n", --SessionRetries);
-
-							break; // do not continue loop
-						}
-					}
-
-					if (AllRoles.Count() == SessionRoles.Count())
-					{
-						SessionInstance = new UnrealSessionInstance(AllRoles.ToArray(), DeferredRoleToAppInstall);
-					}
-				}
+				return SessionInstance;
 			}
 
-			return SessionInstance;
+			return null;
 		}
 
 		/// <summary>
@@ -996,55 +906,58 @@
 
 		public void ReleaseSessionDevices()
 		{
-			if (UnrealDeviceReservation != null)
+			// Terminate any running apps
+			if (SessionInstance != null && SessionInstance.RunningRoles != null)
 			{
-				UnrealDeviceReservation.ReleaseDevices();
+				foreach (UnrealSessionInstance.RoleInstance RunningRole in SessionInstance.RunningRoles)
+				{
+					Log.Info("Shutting down {0}", RunningRole.AppInstance.Device);
+					RunningRole.AppInstance.Kill();
+					RunningRole.AppInstance.Device.Disconnect();
+				}
 			}
-			if (RolesToInstalls != null)
+
+			if (UnrealDeviceReservation == null)
 			{
-				RolesToInstalls.Clear();
+				return;
+			}
+
+			UnrealDeviceReservation.ReleaseDevices();
+
+			RolesToDevices.Clear();
+			RolesToConfigs.Clear();
+			RolesToInstalls.Clear();
+		}
+
+		public void ReleaseProblemDevices()
+		{
+			// Terminate any running apps
+			if (SessionInstance != null && SessionInstance.RunningRoles != null)
+			{
+				foreach (UnrealSessionInstance.RoleInstance RunningRole in SessionInstance.RunningRoles.Where(Role => Role.AppInstance != null))
+				{
+					Log.Info("Shutting down {0}", RunningRole.AppInstance.Device);
+					RunningRole.AppInstance.Kill();
+					RunningRole.AppInstance.Device.Disconnect();
+				}
+			}
+
+			if (UnrealDeviceReservation == null)
+			{
+				return;
+			}
+
+			IEnumerable<ITargetDevice> ProblemDevices = UnrealDeviceReservation.ReleaseProblemDevices();
+			IEnumerable<UnrealSessionRole> RolesToClear = RolesToDevices.Keys.Where(Role => ProblemDevices.Contains(RolesToDevices[Role]));
+
+			Log.Info("Released problem devices...");
+			foreach(UnrealSessionRole Role in RolesToClear)
+			{
+				Log.Info("\t {Device}", RolesToDevices[Role].Name);
+				RolesToDevices.Remove(Role);
+				RolesToConfigs.Remove(Role);
+				RolesToInstalls.Remove(Role);
 			}
 		}
 
-		private string GenerateNotTakenFilePath(string DesiredPath)
-		{
-			string ResultPath = null;
-			
-			FileInfo PotentialPathFileInfo = new FileInfo(DesiredPath);
-
-			for (int NumericPostfix = 0; (string.IsNullOrEmpty(ResultPath)) && (NumericPostfix < int.MaxValue); NumericPostfix++)
-			{
-				string PotentialPath = DesiredPath;
-
-				if (NumericPostfix > 0)
-				{
-					PotentialPath = Path.Combine(
-						PotentialPathFileInfo.DirectoryName,
-						string.Format("{0}_{1}", Path.GetFileNameWithoutExtension(PotentialPathFileInfo.Name), NumericPostfix));
-
-					if (!string.IsNullOrEmpty(PotentialPathFileInfo.Extension))
-					{
-						PotentialPath += PotentialPathFileInfo.Extension;
-					}
-				}
-
-				bool PathIsTaken = File.Exists(PotentialPath);
-				if (PathIsTaken)
-				{
-					Log.VeryVerbose("File already exists at {0}", PotentialPath);
-				}
-				else
-				{
-					ResultPath = PotentialPath;
-				}
-			}
-
-			if (string.IsNullOrEmpty(ResultPath))
-			{
-				throw new AutomationException("Cannot generate not taken file path for the path {0}", DesiredPath);
-			}
-
-			return ResultPath;
-		}
-
 		/// <summary>
@@ -1378,6 +1291,482 @@
 			return AllArtifacts;
 		}
 
+		private UnrealSessionInstance Legacy_LaunchSession()
+		{
+			SessionInstance = null;
+
+			// The number of retries when launching session to avoid an endless loop if package can't be installed, network timeouts, etc
+			int SessionRetries = 2;
+
+			// tries to find devices and launch our session. Will loop until we succeed, we run out of devices/retries, or
+			// something fatal occurs..
+			while (SessionInstance == null && Globals.CancelSignalled == false)
+			{
+				int ReservationRetries = 5;
+				int ReservationRetryWait = 120;
+
+				IEnumerable<UnrealSessionRole> RolesNeedingDevices = SessionRoles.Where(R => R.IsNullRole() == false);
+
+				while (UnrealDeviceReservation.ReservedDevices.Count() < RolesNeedingDevices.Count())
+				{
+					// get devices
+					TryReserveDevices();
+
+					if (Globals.CancelSignalled)
+					{
+						break;
+					}
+
+					// if we failed to get enough devices, show a message and wait
+					if (UnrealDeviceReservation.ReservedDevices.Count() != SessionRoles.Count())
+					{
+						if (ReservationRetries == 0)
+						{
+							DevicePool.Instance.ReportDeviceReservationState();
+							throw new AutomationException("Unable to acquire all devices for test.");
+						}
+						Log.Info("\nUnable to find enough device(s). Waiting {0} secs (retries left={1})\n", ReservationRetryWait, --ReservationRetries);
+						Thread.Sleep(ReservationRetryWait * 1000);
+					}
+				}
+
+				if (Globals.CancelSignalled)
+				{
+					return null;
+				}
+
+				Dictionary<IAppInstall, UnrealSessionRole> InstallsToRoles = new Dictionary<IAppInstall, UnrealSessionRole>();
+
+				// create a copy of our list
+				IEnumerable<ITargetDevice> DevicesToInstallOn = UnrealDeviceReservation.ReservedDevices.ToArray();
+
+				bool InstallSuccess = true;
+
+				// sort by constraints, so that we pick constrained devices first
+				List<UnrealSessionRole> SortedRoles = SessionRoles.OrderBy(R => R.Constraint.IsIdentity() ? 1 : 0).ToList();
+
+				// first install all roles on these devices
+				foreach (UnrealSessionRole Role in SortedRoles)
+				{
+					ITargetDevice Device = null;
+
+					if (Role.IsNullRole() == false)
+					{
+						Device = DevicesToInstallOn.Where(D => D.IsConnected && D.Platform == Role.Platform
+															&& (Role.Constraint.IsIdentity() || DevicePool.Instance.GetConstraint(D) == Role.Constraint)).First();
+
+						DevicesToInstallOn = DevicesToInstallOn.Where(D => D != Device);
+					}
+					else
+					{
+						Device = new TargetDeviceNull(string.Format("Null{0}", Role.RoleType));
+					}
+
+					IEnumerable<UnrealSessionRole> OtherRoles = SortedRoles.Where(R => R != Role);
+
+					// create a config from the build source (this also applies the role options)
+					UnrealAppConfig AppConfig = BuildSource.CreateConfiguration(Role, OtherRoles);
+
+					// todo - should this be elsewhere?
+					AppConfig.Sandbox = Sandbox;
+
+					IAppInstall Install = null;
+					if (RolesToInstalls == null || !RolesToInstalls.ContainsKey(Role) || ReinstallPerPass)
+					{
+						// Tag the device for report result
+						if (BuildHostPlatform.Current.Platform != Device.Platform)
+						{
+							AppConfig.CommandLineParams.Add("DeviceTag", Device.Name);
+						}
+
+						IDeviceUsageReporter.RecordStart(Device.Name, Device.Platform, IDeviceUsageReporter.EventType.Device, IDeviceUsageReporter.EventState.Success);
+						IDeviceUsageReporter.RecordStart(Device.Name, Device.Platform, IDeviceUsageReporter.EventType.Install, IDeviceUsageReporter.EventState.Success, BuildSource.BuildName);
+						try
+						{
+							Install = Device.InstallApplication(AppConfig);
+							IDeviceUsageReporter.RecordEnd(Device.Name, Device.Platform, IDeviceUsageReporter.EventType.Install, IDeviceUsageReporter.EventState.Success);
+						}
+						catch (Exception Ex)
+						{
+							// Warn, ignore the device, and do not continue
+							string ErrorMessage = string.Format("Encountered error setting up device {0} for role {1}. {2}", Device, Role, Ex);
+							if (ErrorMessage.Contains("not enough space"))
+							{
+								Log.Error(KnownLogEvents.Gauntlet_DeviceEvent, ErrorMessage);
+								if (Device.Platform == BuildHostPlatform.Current.Platform)
+								{
+									// If on desktop platform, we are not retrying.
+									// It is unlikely that space is going to be made and InstallBuildParallel has marked the build path as problematic.
+									SessionRetries = 0;
+								}
+							}
+							else
+							{
+								Log.Warning(KnownLogEvents.Gauntlet_DeviceEvent, ErrorMessage);
+							}
+							UnrealDeviceReservation.MarkProblemDevice(Device);
+							InstallSuccess = false;
+							IDeviceUsageReporter.RecordEnd(Device.Name, Device.Platform, IDeviceUsageReporter.EventType.Install, IDeviceUsageReporter.EventState.Failure);
+							break;
+						}
+
+
+						if (Globals.CancelSignalled)
+						{
+							break;
+						}
+
+						// Device has app installed, give role a chance to configure device
+						Role.ConfigureDevice?.Invoke(Device);
+
+						InstallsToRoles[Install] = Role;
+
+						if (ReinstallPerPass)
+						{
+							RolesToInstalls[Role] = Install;
+						}
+					}
+					else
+					{
+						Install = RolesToInstalls[Role];
+						InstallsToRoles[Install] = Role;
+						Log.Info("Using previous install of {0} on {1}", Install.Name, Install.Device.Name);
+					}
+				}
+
+				if (InstallSuccess == false)
+				{
+					ReleaseSessionDevices();
+
+					if (SessionRetries == 0)
+					{
+						throw new AutomationException("Unable to install application for session.");
+					}
+
+					Log.Info("\nUnable to install application for session (retries left={0})\n", --SessionRetries);
+				}
+
+				if (InstallSuccess && Globals.CancelSignalled == false)
+				{
+					List<UnrealSessionInstance.RoleInstance> AllRoles = new List<UnrealSessionInstance.RoleInstance>();
+					Dictionary<UnrealSessionInstance.RoleInstance, IAppInstall> DeferredRoleToAppInstall = new Dictionary<UnrealSessionInstance.RoleInstance, IAppInstall>();
+
+					// Now try to run all installs on their devices
+					foreach (var InstallRoleKV in InstallsToRoles)
+					{
+						IAppInstall CurrentInstall = InstallRoleKV.Key;
+						if (InstallRoleKV.Value.InstallOnly)
+						{
+							AllRoles.Add(new UnrealSessionInstance.RoleInstance(InstallRoleKV.Value, null));
+							continue;
+						}
+						if (InstallRoleKV.Value.DeferredLaunch)
+						{
+							UnrealSessionInstance.RoleInstance DeferredRoleInstance = new UnrealSessionInstance.RoleInstance(InstallRoleKV.Value, null);
+
+							DeferredRoleToAppInstall.Add(DeferredRoleInstance, CurrentInstall);
+							AllRoles.Add(DeferredRoleInstance);
+							continue;
+						}
+						bool Success = false;
+
+						try
+						{
+							Log.Info("Starting {0} on {1}", InstallRoleKV.Value, CurrentInstall.Device);
+							IAppInstance Instance = CurrentInstall.Run();
+							IDeviceUsageReporter.RecordStart(Instance.Device.Name, Instance.Device.Platform, IDeviceUsageReporter.EventType.Test);
+
+							if (Instance != null || Globals.CancelSignalled)
+							{
+								AllRoles.Add(new UnrealSessionInstance.RoleInstance(InstallRoleKV.Value, Instance));
+							}
+
+							Success = true;
+						}
+						catch (DeviceException Ex)
+						{
+							// shutdown all 
+							Log.Warning(KnownLogEvents.Gauntlet_DeviceEvent, "Device {Name} threw an exception during launch. \nException={Exception}", CurrentInstall.Device, Ex.Message);
+							Success = false;
+						}
+
+						if (Success == false)
+						{
+							Log.Warning(KnownLogEvents.Gauntlet_DeviceEvent, "Failed to start build on {Name}. Marking as problem device and retrying with new set", CurrentInstall.Device);
+
+							// terminate anything that's running
+							foreach (UnrealSessionInstance.RoleInstance RunningRole in AllRoles.Where(X => X.AppInstance != null))
+							{
+								Log.Info("Shutting down {0}", RunningRole.AppInstance.Device);
+								RunningRole.AppInstance.Kill();
+								RunningRole.AppInstance.Device.Disconnect();
+							}
+
+							// mark that device as a problem
+							UnrealDeviceReservation.MarkProblemDevice(CurrentInstall.Device);
+
+							ReleaseSessionDevices();
+
+							if (SessionRetries == 0)
+							{
+								throw new AutomationException("Unable to start application for session, see warnings for details.");
+							}
+
+							Log.Info("\nUnable to start application for session (retries left={0})\n", --SessionRetries);
+
+							break; // do not continue loop
+						}
+					}
+
+					if (AllRoles.Count() == SessionRoles.Count())
+					{
+						SessionInstance = new UnrealSessionInstance(AllRoles.ToArray(), DeferredRoleToAppInstall);
+					}
+				}
+			}
+
+			return SessionInstance;
+		}
+
+		/// <summary>
+		/// Returns true if the number of reserved devices matches the number on non-null session roles
+		/// </summary>
+		private bool HasAcquiredDevices()
+		{
+			IEnumerable<UnrealSessionRole> RolesThatRequireDevice = SessionRoles.Where(Role => !Role.IsNullRole());
+			return UnrealDeviceReservation != null
+				&& UnrealDeviceReservation.ReservedDevices != null
+				&& UnrealDeviceReservation.ReservedDevices.Count() == RolesThatRequireDevice.Count();
+		}
+
+		/// <summary>
+		/// Returns true if the provided session role has not yet been assigned a device
+		/// </summary>
+		private bool RoleNeedsDevice(UnrealSessionRole Role)
+		{
+			return !(RolesToDevices.ContainsKey(Role) && RolesToDevices[Role] != null);
+		}
+
+		/// <summary>
+		/// Returns true if the provided session role has not yet had an install performed on it's assigned device
+		/// </summary>
+		private bool RoleNeedsInstall(UnrealSessionRole Role)
+		{
+			return !ReinstallPerPass
+				&& !(RolesToConfigs.ContainsKey(Role) && RolesToConfigs[Role] != null)
+				&& !(RolesToInstalls.ContainsKey(Role) && RolesToInstalls[Role] != null);
+		}
+
+		/// <summary>
+		/// Returns true if the provided target device matches the constraint requested by the session role
+		/// </summary>
+		private bool DeviceMatchesRoleConstraint(UnrealSessionRole Role, ITargetDevice Device)
+		{
+			bool bRoleMatchesConstraint = DevicePool.Instance.GetConstraint(Device) == Role.Constraint;
+
+			return Device.IsConnected
+				&& Device.Platform == Role.Platform
+				&& Role.Constraint.IsIdentity() || bRoleMatchesConstraint;
+		}
+
+		/// <summary>
+		/// From the existing reserved device pool, assign each role a device that matches it's requested constraint
+		/// This will cache a map of each role and the target device it's using in this UnrealSession
+		/// </summary>
+		private bool TryAssignDevicesToRoles()
+		{
+			// Order by constraint. This ensures roles with constraints have their devices selected first.
+			IEnumerable<UnrealSessionRole> RolesSortedByConstraint = SessionRoles.OrderBy(R => R.Constraint.IsIdentity() ? 1 : 0);
+
+			foreach (UnrealSessionRole Role in RolesSortedByConstraint)
+			{
+				if (RoleNeedsDevice(Role))
+				{
+					ITargetDevice DeviceToAssign = null;
+
+					if (Role.IsNullRole())
+					{
+						DeviceToAssign = new TargetDeviceNull($"Null{Role.RoleType}");
+					}
+					else
+					{
+						try
+						{
+							DeviceToAssign = UnrealDeviceReservation.ReservedDevices.Where(Device => DeviceMatchesRoleConstraint(Role, Device)).First();
+							IDeviceUsageReporter.RecordStart(DeviceToAssign.Name, DeviceToAssign.Platform, IDeviceUsageReporter.EventType.Device, IDeviceUsageReporter.EventState.Success);
+						}
+						catch (Exception Ex)
+						{
+							Log.Warning("Failed to assign a reserved device to role {Role}. " +
+								"This usually means devices were unexpectedly released mid session\n{Exception}", Role, Ex);
+							return false;
+						}
+					}
+
+					RolesToDevices.Add(Role, DeviceToAssign);
+				}
+			}
+
+			return true;
+		}
+
+		/// <summary>
+		/// Prepares each device for launch by performing the following
+		///		- Install builds
+		///		- Clean up old artifacts
+		///		- Copy additional files requested by a test
+		///		- Specific role configurations
+		///	This will create and cache both an UnrealAppConfig and an IAppInstall for future reference
+		/// </summary>
+		private void ReadyDevicesForSession()
+		{
+			foreach(UnrealSessionRole Role in SessionRoles)
+			{
+				UnrealAppConfig AppConfig = null;
+				ITargetDevice Device = RolesToDevices[Role];
+
+				if (Globals.CancelSignalled)
+				{
+					return;
+				}
+
+				try
+				{
+					if (RoleNeedsInstall(Role))
+					{
+						// Create the app config
+						IEnumerable<UnrealSessionRole> OtherRoles = SessionRoles.Where(Other => Other != Role);
+						AppConfig = BuildSource.CreateConfiguration(Role, OtherRoles);
+						AppConfig.Sandbox = Sandbox;
+						AppConfig.CommandLineParams.AddUnique("DeviceTag", Device.Name);
+						RolesToConfigs.Add(Role, AppConfig);
+
+						// Install the build
+						if (AppConfig.FullClean)
+						{
+							Log.Info("Fully cleaning device before install...");
+							Device.FullClean();
+						}
+
+						if (AppConfig.SkipInstall)
+						{
+							Log.Info("Skipping install due to SkipInstall");
+						}
+						else
+						{
+							// Telemetry
+							DateTimeStopwatch Stopwatch = DateTimeStopwatch.Start();
+							Log.Info("Installing {BuildName} of type {BuildType} to {Device}...", BuildSource.BuildName, AppConfig.Build.GetType().Name, Device);
+							IDeviceUsageReporter.RecordStart(Device.Name, Device.Platform, IDeviceUsageReporter.EventType.Install, IDeviceUsageReporter.EventState.Success, BuildSource.BuildName);
+
+							try
+							{
+								Device.InstallBuild(AppConfig);
+								IDeviceUsageReporter.RecordEnd(Device.Name, Device.Platform, IDeviceUsageReporter.EventType.Install, IDeviceUsageReporter.EventState.Success);
+								Log.Info("Installation completed in {InstallTime}", GetInstallTime(Stopwatch.ElapsedTime));
+							}
+							catch
+							{
+								IDeviceUsageReporter.RecordEnd(Device.Name, Device.Platform, IDeviceUsageReporter.EventType.Install, IDeviceUsageReporter.EventState.Failure);
+								throw;
+							}
+						}
+
+						IAppInstall Install = Device.CreateAppInstall(AppConfig);
+						RolesToInstalls.Add(Role, Install);
+					}
+
+					Device.CleanArtifacts();
+					Device.CopyAdditionalFiles(AppConfig.FilesToCopy);
+					Role.ConfigureDevice?.Invoke(Device);
+				}
+				catch(Exception Ex)
+				{
+					string Message = $"Encountered {Ex.GetType()} when creating installation on device {Device}.\n{Ex.Message}";
+
+					if (IsOutOfSpaceException(Ex) && Device.Platform == BuildHostPlatform.Current.Platform)
+					{
+						// If on desktop platform, we are not retrying.
+						// It is unlikely that space is going to be made and InstallBuildParallel has marked the build path as problematic.
+						Log.Error(KnownLogEvents.Gauntlet_DeviceEvent, Message);
+						throw;
+					}
+					else
+					{
+						UnrealDeviceReservation.MarkProblemDevice(Device, Message);
+						Log.Warning(KnownLogEvents.Gauntlet_DeviceEvent, Message);
+					}
+					throw; // Can consider not throwing here - this would let every device complete setup before releasing the problem devices
+				}
+			}
+		}
+
+		/// <summary>
+		/// Launches the Unreal Engine processes for each role requested by this session.
+		/// Roles marked InstallOnly will not have a process launched.
+		/// Roles marked DeferredLaunch will not have a process launched.
+		/// Deferred roles can be launched at anytime (usually in TickTest()) by calling UnrealSessionInstance.LaunchDeferredRole
+		/// </summary>
+		private UnrealSessionInstance LaunchProcesses()
+		{
+			List<UnrealSessionInstance.RoleInstance> RoleInstances = new();
+			Dictionary<UnrealSessionInstance.RoleInstance, IAppInstall> DeferredRolesToInstalls = new();
+
+			foreach (KeyValuePair<UnrealSessionRole, IAppInstall> RoleInstall in RolesToInstalls)
+			{
+				UnrealSessionRole Role = RoleInstall.Key;
+				IAppInstall Install = RoleInstall.Value;
+
+				// InstallOnly roles don't execute a process
+				if (Role.InstallOnly)
+				{
+					RoleInstances.Add(new UnrealSessionInstance.RoleInstance(Role, null));
+					continue;
+				}
+
+				// DeferredLaunch roles don't immediately execute a process.
+				// We cache deferred roles so users can launch the process at the desired time.
+				else if (Role.DeferredLaunch)
+				{
+					UnrealSessionInstance.RoleInstance DeferredRoleInstance = new(Role, null);
+					RoleInstances.Add(DeferredRoleInstance);
+					DeferredRolesToInstalls.Add(DeferredRoleInstance, Install);
+					continue;
+				}
+
+				try
+				{
+					Log.Info("Launching {Install} on {Device}", Install, RolesToDevices[Role]);
+					IAppInstance AppInstance = Install.Run();
+
+					if (AppInstance == null)
+					{
+						throw new AutomationException("Failed to create an IAppInstance after attempting Run() on {Install}", Install);
+					}
+					else
+					{
+						RoleInstances.Add(new UnrealSessionInstance.RoleInstance(Role, AppInstance));
+					}
+				}
+				catch(Exception Ex)
+				{
+					// Kill any processes that were started
+					foreach (UnrealSessionInstance.RoleInstance Instance in RoleInstances.Where(Role => Role.AppInstance != null))
+					{
+						Log.Info("Shutting down {AppInstance}", Instance.AppInstance);
+						Instance.AppInstance.Kill();
+					}
+
+					string WarningMessage = $"Encountered {Ex.GetType()} when attempting to run install {Ex.Message}";
+					UnrealDeviceReservation.MarkProblemDevice(Install.Device, WarningMessage);
+					Log.Warning(WarningMessage);
+					throw;
+				}
+			}
+
+			return new UnrealSessionInstance(RoleInstances.ToArray(), DeferredRolesToInstalls);
+		}
+
 		private void SavePSOs(UnrealTestContext InContext, UnrealSessionInstance.RoleInstance InRunningRole, string DestSavedDir)
 		{
 			if (InRunningRole.Role.RoleType.IsServer())
@@ -1468,6 +1857,47 @@
 			}
 		}
 
+		private string GenerateNotTakenFilePath(string DesiredPath)
+		{
+			string ResultPath = null;
+
+			FileInfo PotentialPathFileInfo = new FileInfo(DesiredPath);
+
+			for (int NumericPostfix = 0; (string.IsNullOrEmpty(ResultPath)) && (NumericPostfix < int.MaxValue); NumericPostfix++)
+			{
+				string PotentialPath = DesiredPath;
+
+				if (NumericPostfix > 0)
+				{
+					PotentialPath = Path.Combine(
+						PotentialPathFileInfo.DirectoryName,
+						string.Format("{0}_{1}", Path.GetFileNameWithoutExtension(PotentialPathFileInfo.Name), NumericPostfix));
+
+					if (!string.IsNullOrEmpty(PotentialPathFileInfo.Extension))
+					{
+						PotentialPath += PotentialPathFileInfo.Extension;
+					}
+				}
+
+				bool PathIsTaken = File.Exists(PotentialPath);
+				if (PathIsTaken)
+				{
+					Log.VeryVerbose("File already exists at {0}", PotentialPath);
+				}
+				else
+				{
+					ResultPath = PotentialPath;
+				}
+			}
+
+			if (string.IsNullOrEmpty(ResultPath))
+			{
+				throw new AutomationException("Cannot generate not taken file path for the path {0}", DesiredPath);
+			}
+
+			return ResultPath;
+		}
+
 		/// <summary>
 		/// Filter that Truncate long file paths such as CrashReporter files. //UECC-Windows-F0DD9BB04C3C9250FAF39D8AB4A88556//
 		/// These are particularly problematic with testflights which append a long random name to the destination folder, 
@@ -1496,5 +1926,20 @@
 
 			return LongFilePath;
 		}
+
+		private string GetInstallTime(TimeSpan Time)
+		{
+			string Hours = Time.Hours > 0 ? string.Format("{0} hrs, ", Time.Hours) : string.Empty;
+			string Minutes = Time.Minutes > 0 ? string.Format("{0} mins, ", Time.Minutes) : string.Empty;
+			string Seconds = string.Format("{0} secs", Time.Seconds);
+
+			return Hours + Minutes + Seconds;
+		}
+
+		private bool IsOutOfSpaceException(Exception Ex)
+		{
+			return Ex.Message.Contains("not enough space", StringComparison.OrdinalIgnoreCase);
+		}
 	}
+
 }
diff --git a/Engine/Source/Programs/AutomationTool/Gauntlet/Unreal/Base/Gauntlet.UnrealTestNode.cs b/Engine/Source/Programs/AutomationTool/Gauntlet/Unreal/Base/Gauntlet.UnrealTestNode.cs
--- a/Engine/Source/Programs/AutomationTool/Gauntlet/Unreal/Base/Gauntlet.UnrealTestNode.cs
+++ b/Engine/Source/Programs/AutomationTool/Gauntlet/Unreal/Base/Gauntlet.UnrealTestNode.cs
@@ -777,13 +777,6 @@ namespace Gauntlet
 				{
 					throw new AutomationException("Node already has a null UnrealApp, was PrepareUnrealSession or IsReadyToStart called?");
 				}
-				else
-				{
-					bool bReacquireDevicesPerPass = Globals.Params.ParseParam("ReacquireDevicesPerPass");
-					bool ShouldRetain = (!bReacquireDevicesPerPass && InNumPasses > 1 && Pass + 1 < InNumPasses) ? true : false;
-					UnrealApp.ShouldRetainDevices = ShouldRetain;
-					Log.Verbose("ShouldRetainDevices: {0}", UnrealApp.ShouldRetainDevices);
-				}
 			}
 
 			// ensure we reset things
@@ -809,8 +802,13 @@ namespace Gauntlet
 
 			// Launch the test
 			TestInstance = UnrealApp.LaunchSession();
+			if(TestInstance == null)
+			{
+				return false;
+			}
+
 			// Add info from test context to device usage log
-			foreach(IAppInstance AppInstance in TestInstance.ClientApps)
+			foreach (IAppInstance AppInstance in TestInstance.ClientApps)
 			{
 				if (AppInstance != null)
 				{
@@ -818,8 +816,6 @@ namespace Gauntlet
 					IDeviceUsageReporter.RecordComment(AppInstance.Device.Name, AppInstance.Device.Platform, IDeviceUsageReporter.EventType.Test, this.GetType().Name);
 				}
 			}
-			
-			
 
 			// track the overall session time
 			if (SessionStartTime == DateTime.MinValue)
@@ -827,16 +823,13 @@ namespace Gauntlet
 				SessionStartTime = DateTime.Now;
 			}
 
-			if (TestInstance != null)
-			{
-				// Update these for the executor
-				MaxDuration = Config.MaxDuration;
-				MaxDurationReachedResult = Config.MaxDurationReachedResult;
-				MaxRetries = Config.MaxRetries;
-				MarkTestStarted();
-			}
-			
-			return TestInstance != null;
+			// Update these for the executor
+			MaxDuration = Config.MaxDuration;
+			MaxDurationReachedResult = Config.MaxDurationReachedResult;
+			MaxRetries = Config.MaxRetries;
+			MarkTestStarted();
+
+			return true;
 		}
 
 		public virtual void PopulateCommandlineInfo()
