[Gauntlet] Adds new LaunchSession to LowLevelTests flow using new TargetDevice functions
This new flow is disabled by default and requires -ExperimentalLaunchFlow to be present to run

diff --git a/Engine/Source/Programs/AutomationTool/LowLevelTests/RunLowLevelTests.cs b/Engine/Source/Programs/AutomationTool/LowLevelTests/RunLowLevelTests.cs
--- a/Engine/Source/Programs/AutomationTool/LowLevelTests/RunLowLevelTests.cs
+++ b/Engine/Source/Programs/AutomationTool/LowLevelTests/RunLowLevelTests.cs
@@ -295,6 +295,24 @@ namespace LowLevelTests
 			Containerized = InOptions.Containerized;
 		}
 
+		#region IDisposable Support
+		public void Dispose()
+		{
+			if (Instance != null)
+			{
+				Instance.Kill();
+				IDeviceUsageReporter.RecordEnd(Instance.Device.Name, (UnrealTargetPlatform)Instance.Device.Platform, IDeviceUsageReporter.EventType.Test, IDeviceUsageReporter.EventState.Success);
+
+				if (Instance.HasExited)
+				{
+					IDeviceUsageReporter.RecordEnd(Instance.Device.Name, (UnrealTargetPlatform)Instance.Device.Platform, IDeviceUsageReporter.EventType.Test, IDeviceUsageReporter.EventState.Failure);
+				}
+
+				UnrealDeviceReservation?.ReleaseDevices();
+			}
+		}
+		#endregion
+
 		public bool TryReserveDevices()
 		{
 			// Low level tests require exactly one device of the build source's platform.
@@ -303,7 +321,6 @@ namespace LowLevelTests
 				{ new UnrealDeviceTargetConstraint(BuildSource.Platform), 1 }
 			};
 
-
 			if (Containerized)
 			{
 				// Linux platforms, including arm64, can be run through Docker on Windows machines.
@@ -327,6 +344,110 @@ namespace LowLevelTests
 		/// No packaging required.
 		/// </summary>
 		public IAppInstance InstallAndRunNativeTestApp()
+		{
+			if(!Globals.Params.ParseParam("ExperimentalLaunchFlow"))
+			{
+				return Legacy_InstallAndRunNativeTestApp();
+			}
+
+			// TargetDevice<Platform> classes have a hard dependency on UnrealAppConfig instead of IAppConfig.
+			// More refactoring needed to support non-packaged applications that can be run natively from a path on the device.
+			UnrealAppConfig AppConfig = BuildSource.GetUnrealAppConfig(Tags, Sleep, AttachToDebugger, ReportType, PerTestTimeout, TestExtraArgs, Containerized);
+
+			IEnumerable<ITargetDevice> DevicesToInstallOn = UnrealDeviceReservation.ReservedDevices.ToArray();
+			ITargetDevice Device = DevicesToInstallOn.Where(D => D.IsConnected && D.Platform == BuildSource.Platform).First();
+
+			if (AppConfig.FullClean)
+			{
+				Device.FullClean();
+			}
+
+			// Install the build onto the device
+			if (AppConfig.SkipInstall)
+			{
+				Log.Info("Skipping install due to SkipInstall");
+			}
+			else
+			{
+				// Telemetry
+				DateTimeStopwatch Stopwatch = DateTimeStopwatch.Start();
+				Log.Info("Installing {BuildName} of type {BuildType} to {Device}...", BuildSource.BuildName, AppConfig.Build.GetType().Name, Device);
+				IDeviceUsageReporter.RecordStart(Device.Name, Device.Platform, IDeviceUsageReporter.EventType.Install, IDeviceUsageReporter.EventState.Success, BuildSource.BuildName);
+
+				try
+				{
+					Device.InstallBuild(AppConfig);
+					IDeviceUsageReporter.RecordEnd(Device.Name, Device.Platform, IDeviceUsageReporter.EventType.Install, IDeviceUsageReporter.EventState.Success);
+					Log.Info("Installation completed in {InstallTime}", GetInstallTime(Stopwatch.ElapsedTime));
+				}
+				catch (Exception Ex)
+				{
+					Log.Info("Failed to install low level tests app onto device {Device}: {Exception}", Device, Ex.ToString());
+
+					IDeviceUsageReporter.RecordEnd(Device.Name, Device.Platform, IDeviceUsageReporter.EventType.Install, IDeviceUsageReporter.EventState.Failure);
+					UnrealDeviceReservation.MarkProblemDevice(Device);
+					UnrealDeviceReservation.ReleaseDevices();
+				}
+			}
+
+			// Create the IAppInstall instance
+			try
+			{
+				Install = Device.CreateAppInstall(AppConfig);
+			}
+			catch (Exception Ex)
+			{
+				Log.Info("Failed to create IAppInstall for low level tests on device {Device}: {Exception}", Device, Ex.ToString());
+				UnrealDeviceReservation.MarkProblemDevice(Device);
+				UnrealDeviceReservation.ReleaseDevices();
+			}
+
+			// Clean/Copy files to the device
+			try
+			{
+				Device.CleanArtifacts();
+				Device.CopyAdditionalFiles(AppConfig.FilesToCopy);
+			}
+			catch (Exception Ex)
+			{
+				Log.Info("Failed to CleanArtifacts or CopyAdditionalFiles for low level tests on device {Device}: {Exception}", Device, Ex.ToString());
+				UnrealDeviceReservation.MarkProblemDevice(Device);
+				UnrealDeviceReservation.ReleaseDevices();
+			}
+
+			// Run the application
+			try
+			{
+				if (Device is IRunningStateOptions DeviceWithStateOptions)
+				{
+					// Don't wait to detect running state and query for running state every second
+					DeviceWithStateOptions.WaitForRunningState = false;
+					DeviceWithStateOptions.CachedStateRefresh = QUERY_STATE_INTERVAL;
+				}
+
+				Instance = Device.Run(Install);
+				IDeviceUsageReporter.RecordStart(Instance.Device.Name, Instance.Device.Platform, IDeviceUsageReporter.EventType.Test);
+			}
+			catch (DeviceException DeviceEx)
+			{
+				Log.Warning("Device {0} threw an exception during launch. \nException={1}", Install.Device, DeviceEx.Message);
+				Log.Warning("Failed to start low level test on {0}. Marking as problem device. Will not retry.", Device);
+
+				if (Instance != null)
+				{
+					Instance.Kill();
+				}
+
+				UnrealDeviceReservation.MarkProblemDevice(Device);
+				UnrealDeviceReservation.ReleaseDevices();
+
+				throw new AutomationException("Unable to start low level tests app, see warnings for details.");
+			}
+
+			return Instance;
+		}
+
+		private IAppInstance Legacy_InstallAndRunNativeTestApp()
 		{
 			bool InstallSuccess = false;
 			bool RunSuccess = false;
@@ -413,20 +534,13 @@ namespace LowLevelTests
 			return Instance;
 		}
 
-		public void Dispose()
+		private string GetInstallTime(TimeSpan Time)
 		{
-			if (Instance != null)
-			{
-				Instance.Kill();
-				IDeviceUsageReporter.RecordEnd(Instance.Device.Name, (UnrealTargetPlatform)Instance.Device.Platform, IDeviceUsageReporter.EventType.Test, IDeviceUsageReporter.EventState.Success);
+			string Hours = Time.Hours > 0 ? string.Format("{0} hrs, ", Time.Hours) : string.Empty;
+			string Minutes = Time.Minutes > 0 ? string.Format("{0} mins, ", Time.Minutes) : string.Empty;
+			string Seconds = string.Format("{0} secs", Time.Seconds);
 
-				if (Instance.HasExited)
-				{
-					IDeviceUsageReporter.RecordEnd(Instance.Device.Name, (UnrealTargetPlatform)Instance.Device.Platform, IDeviceUsageReporter.EventType.Test, IDeviceUsageReporter.EventState.Failure);
-				}
-				
-				UnrealDeviceReservation?.ReleaseDevices();				
-			}
+			return Hours + Minutes + Seconds;
 		}
 	}
 
