BuildGraph: Add a task to check links within a set of markdown files.

diff --git a/Engine/Source/Programs/AutomationTool/BuildGraph/Tasks/CheckMarkdownTask.cs b/Engine/Source/Programs/AutomationTool/BuildGraph/Tasks/CheckMarkdownTask.cs
new file mode 100644
--- /dev/null
+++ b/Engine/Source/Programs/AutomationTool/BuildGraph/Tasks/CheckMarkdownTask.cs
@@ -0,0 +1,133 @@
+using EpicGames.Core;
+using Microsoft.Extensions.Logging;
+using System;
+using System.Collections.Generic;
+using System.Linq;
+using System.Text;
+using System.Text.RegularExpressions;
+using System.Threading.Tasks;
+using System.Xml;
+using UnrealBuildBase;
+
+namespace AutomationTool.Tasks
+{
+	/// <summary>
+	/// Parameters for a CheckMarkdown task
+	/// </summary>
+	public class CheckMarkdownTaskParameters
+	{
+		/// <summary>
+		/// Optional filter to be applied to the list of input files.
+		/// </summary>
+		[TaskParameter(Optional = true, ValidationType = TaskParameterValidationType.FileSpec)]
+		public string Files;
+	}
+
+	/// <summary>
+	/// Checks that all markdown links between the given files are valid.
+	/// </summary>
+	[TaskElement("CheckMarkdown", typeof(CheckMarkdownTaskParameters))]
+	public class CheckMarkdownTask : BgTaskImpl
+	{
+		readonly CheckMarkdownTaskParameters _parameters;
+
+		/// <summary>
+		/// Constructor
+		/// </summary>
+		/// <param name="parameters">Parameters for this task</param>
+		public CheckMarkdownTask(CheckMarkdownTaskParameters parameters)
+		{
+			_parameters = parameters;
+		}
+
+		/// <inheritdoc/>
+		public override async Task ExecuteAsync(JobContext job, HashSet<FileReference> buildProducts, Dictionary<string, HashSet<FileReference>> tagNameToFileSet)
+		{
+			HashSet<FileReference> files = ResolveFilespec(Unreal.RootDirectory, _parameters.Files, tagNameToFileSet);
+
+			HashSet<FileReference> markdownFiles = files.Where(x => x.HasExtension(".md")).ToHashSet();
+			Logger.LogInformation("Checking {NumFiles} markdown files...", markdownFiles.Count);
+
+			// Build a set of valid link targets
+			HashSet<string> validLinks = new HashSet<string>(StringComparer.Ordinal);
+			foreach (FileReference file in files)
+			{
+				validLinks.Add(file.FullName);
+			}
+
+			// Find the anchors in any markdown files
+			foreach(FileReference file in markdownFiles)
+			{
+				string[] lines = await FileReference.ReadAllLinesAsync(file);
+				foreach (string line in lines)
+				{
+					Match match = Regex.Match(line, @"^#+\s+(.*)");
+					if (match.Success)
+					{
+						string anchor = match.Groups[1].Value.ToLowerInvariant();
+						anchor = Regex.Replace(anchor, "[^a-z0-9]+", "-");
+						anchor = Regex.Replace(anchor, "^-|-$", "");
+						validLinks.Add($"{file.FullName}#{anchor}");
+					}
+				}
+			}
+
+			// Check the links 
+			foreach (FileReference file in markdownFiles)
+			{
+				string[] lines = await FileReference.ReadAllLinesAsync(file);
+				for (int lineIdx = 0; lineIdx < lines.Length; lineIdx++)
+				{
+					string line = lines[lineIdx];
+					foreach (Match match in Regex.Matches(line, @"\[([^\]]+)\]\(([^\)]*)\)"))
+					{
+						string link = match.Groups[2].Value;
+						if (!Regex.IsMatch(link, "^[a-z]+://"))
+						{
+							int hashIdx = link.IndexOf('#');
+							if (hashIdx == -1)
+							{
+								link = FileReference.Combine(file.Directory, link).FullName;
+							}
+							else if (hashIdx == 0)
+							{
+								link = $"{file}{link[hashIdx..]}";
+							}
+							else
+							{
+								link = $"{FileReference.Combine(file.Directory, link[0..hashIdx])}{link[hashIdx..]}";
+							}
+
+							bool validLink;
+							try
+							{
+								validLink = validLinks.Contains(link);
+							}
+							catch
+							{
+								validLink = false;
+							}
+
+							if (!validLink)
+							{
+								Logger.LogWarning("{File}({Line}): Invalid link '{Link}'", file, lineIdx + 1, match.Groups[0].Value);
+							}
+						}
+					}
+				}
+			}
+		}
+
+		/// <inheritdoc/>
+		public override void Write(XmlWriter Writer)
+			=> Write(Writer, _parameters);
+
+		/// <inheritdoc/>
+		public override IEnumerable<string> FindConsumedTagNames()
+			=> FindTagNamesFromFilespec(_parameters.Files);
+
+		/// <inheritdoc/>
+		public override IEnumerable<string> FindProducedTagNames()
+			=> Array.Empty<string>();
+	}
+}
diff --git a/Engine/Source/Programs/Horde/BuildHorde.xml b/Engine/Source/Programs/Horde/BuildHorde.xml
--- a/Engine/Source/Programs/Horde/BuildHorde.xml
+++ b/Engine/Source/Programs/Horde/BuildHorde.xml
@@ -32,7 +32,11 @@
 			Engine/Source/Programs/UnrealBuildTool/Matchers/...
 		</Property>
 
-		<Node Name="Analyze HordeServer">
+		<Node Name="Check Documentation">
+			<CheckMarkdown Files="$(RootDir)/Engine/Source/Programs/Horde/Docs/..."/>
+		</Node>
+
+		<Node Name="Analyze HordeServer" Requires="Check Documentation">
 			<DotNet Arguments="build &quot;$(RootDir)/Engine/Source/Programs/Horde/Horde.Server/Horde.Server.csproj&quot; /p:Configuration=Analyze $(VersionArguments)" DotNetPath="$(DotNetPath)"/>
 			<DotNet Arguments="build &quot;$(RootDir)/Engine/Source/Programs/Horde/Horde.Server.Tests/Horde.Server.Tests.csproj&quot; /p:Configuration=Analyze $(VersionArguments)" DotNetPath="$(DotNetPath)"/>
 		</Node>
