Add x86_64 architecture to Android installed build

#android

diff --git a/Engine/Build/InstalledEngineBuild.xml b/Engine/Build/InstalledEngineBuild.xml
--- a/Engine/Build/InstalledEngineBuild.xml
+++ b/Engine/Build/InstalledEngineBuild.xml
@@ -97,7 +97,7 @@
 	<Property Name="ToolchainDir" Value="$(RootDir)/Engine/Extras/ThirdPartyNotUE/SDKs"/>
 
 	<!-- Architectures that we build for Android. NOTE: This should match the architectures specified in BuildPluginCommand. -->
-	<Property Name="AndroidArchitectures" Value="arm64"/>
+	<Property Name="AndroidArchitectures" Value="arm64+x64"/>
 	
 	<Property Name="DefaultArchitectureCommandsWin64" Value=""/>
 
diff --git a/Engine/Source/Programs/AutomationTool/Scripts/BuildPluginCommand.Automation.cs b/Engine/Source/Programs/AutomationTool/Scripts/BuildPluginCommand.Automation.cs
--- a/Engine/Source/Programs/AutomationTool/Scripts/BuildPluginCommand.Automation.cs
+++ b/Engine/Source/Programs/AutomationTool/Scripts/BuildPluginCommand.Automation.cs
@@ -27,7 +27,7 @@ using static AutomationTool.CommandUtils;
 public sealed class BuildPlugin : BuildCommand
 {
 	const string MacDefaultArchitectures = "arm64+x64";
-	const string AndroidDefaultArchitectures = "arm64";
+	const string AndroidDefaultArchitectures = "arm64+x64";
 
 	string UnrealBuildToolDllRelativePath = @"Engine/Binaries/DotNET/UnrealBuildTool/UnrealBuildTool.dll";
 	FileReference UnrealBuildToolDll;
