Localization:
Fix for Auto localization generating absolute file paths instead of relative files paths causing downloads/uploads to fail for localization targets.

diff --git a/Engine/Source/Programs/AutomationTool/Scripts/Localisation.Automation.cs b/Engine/Source/Programs/AutomationTool/Scripts/Localisation.Automation.cs
--- a/Engine/Source/Programs/AutomationTool/Scripts/Localisation.Automation.cs
+++ b/Engine/Source/Programs/AutomationTool/Scripts/Localisation.Automation.cs
@@ -849,6 +849,9 @@ class Localize : BuildCommand
 
 	private void AutoGeneratePluginLocalizationFiles(HashSet<string> PluginNames)
 	{
+		string PluginsRootDirectory = Path.Combine(UEProjectRoot, UEProjectDirectory);
+		DirectoryReference PluginsRootDirectoryReference = new DirectoryReference(PluginsRootDirectory);
+
 		foreach (string PluginName in PluginNames)
 		{
 			PluginInfo Info = Plugins.GetPlugin(PluginName);
@@ -871,7 +874,7 @@ class Localize : BuildCommand
 					LocalizationConfigFileGenerator Generator = LocalizationConfigFileGenerator.GetGeneratorForFileFormat(LocalizationConfigFileFormat.Latest);
 					LocalizationConfigFileGeneratorParams GeneratorParams = new();
 					GeneratorParams.LocalizationTargetName = Target.Name;
-					GeneratorParams.LocalizationTargetRootDirectory = Info.File.Directory.FullName;
+					GeneratorParams.LocalizationTargetRootDirectory = Info.Directory.MakeRelativeTo(PluginsRootDirectoryReference);
 					List<LocalizationConfigFile> PluginLocalizationConfigFiles = Generator.GenerateDefaultSettingsConfigFiles(GeneratorParams);
 					foreach (LocalizationConfigFile PluginLocalizationConfigFile in PluginLocalizationConfigFiles)
 					{
