Localization:
Localize UAT command now supports taking multiple directories to include and exclude. The -IncludePluginsDirectory and -ExcludePluginsDirectory now both take multiple relative paths separated by the ";" character.
- Examples uses: RunUAT.bat Localize -IncludePluginsDirectory="Plugins/A;Plugins/B" -ExcludePluginsDirectory="Plugins/C;Plugins/D;Plugins/E"

#test Locally ran the Localize command with multiple directories for -IncludePluginsDirectory and -ExcludePluginsDirectory. Set breakoints to verify that the directories were being processed and the plugisn were indeed being added to the appropriate lists.

diff --git a/Engine/Source/Programs/AutomationTool/Scripts/Localisation.Automation.cs b/Engine/Source/Programs/AutomationTool/Scripts/Localisation.Automation.cs
--- a/Engine/Source/Programs/AutomationTool/Scripts/Localisation.Automation.cs
+++ b/Engine/Source/Programs/AutomationTool/Scripts/Localisation.Automation.cs
@@ -24,9 +24,9 @@ using Microsoft.Extensions.Logging;
 [Help("LocalizationProvider", "Optional localization provide override.")]
 [Help("LocalizationSteps", "Optional comma separated list of localization steps to perform [Download, Gather, Import, Export, Compile, GenerateReports, Upload] (default is all). Only valid for projects using a modular config.")]
 [Help("IncludePlugins", "Optional flag to include plugins from within the given UEProjectDirectory as part of the gather. This may optionally specify a comma separated list of the specific plugins to gather (otherwise all plugins will be gathered).")]
-[Help("IncludePluginsDirectory", "Optional parameter that is a relative path to a directory under UEProjectDirectory. All plugins under this directory will be gathered from (if not excluded).")]
+[Help("IncludePluginsDirectory", "Optional parameter that is a list of relative paths to a directory under UEProjectDirectory separated by a ';' character. All plugins under this directory will be gathered from (if not excluded) E.g -IncludePluginsDirectory=\"Plugins/A;Plugins/B;Plugins/C\"")]
 [Help("ExcludePlugins", "Optional comma separated list of plugins to exclude from the gather.")]
-[Help("ExcludePluginsDirectory", "Optional relative path to a directory under UEProjectDirectory. All plugins under this directory will be excluded from gather.")]
+[Help("ExcludePluginsDirectory", "Optional list of relative paths to a directory under UEProjectDirectory separated by the ';' character. All plugins under this directory will be excluded from gather. E.g -ExcludePluginsDirectory=\"Plugins/A;Plugins/B;Plugins/C\"")]
 [Help("EnableIncludedPlugins", "Optional flag that passes all included plugins that aren't excluded to the -EnablePlugins editor argument to ensure content and metadata for plugins are loaded for gathering.")]
 [Help("IncludePlatforms", "Optional flag to include platforms from within the given UEProjectDirectory as part of the gather.")]
 [Help("AdditionalCommandletArguments", "Optional arguments to pass to the gather process.")]
@@ -412,14 +412,22 @@ class Localize : BuildCommand
 		if (!string.IsNullOrEmpty(IncludePluginsUnderDirectoryStr))
 		{
 			bShouldGatherPlugins = true;
-			string AbsolutePathToIncludePluginsDirectory = Path.Combine(PluginsRootPath, IncludePluginsUnderDirectoryStr);
-			IncludePlugins.AddRange(LocalizationUtilities.GetPluginNamesUnderDirectory(AbsolutePathToIncludePluginsDirectory, PluginsRootPath, UEProjectName.Length == 0 ? PluginType.Engine : PluginType.Project));
+			foreach (string IncludePluginDirectory in IncludePluginsUnderDirectoryStr.Split(";"))
+			{
+				string AbsolutePathToIncludePluginsDirectory = Path.Combine(PluginsRootPath, IncludePluginDirectory.Trim());
+				IncludePlugins.AddRange(LocalizationUtilities.GetPluginNamesUnderDirectory(AbsolutePathToIncludePluginsDirectory, PluginsRootPath, UEProjectName.Length == 0 ? PluginType.Engine : PluginType.Project));
+			}
 		}
+
 		string ExcludePluginsUnderDirectoryStr = ParseParamValue("ExcludePluginsDirectory");
 		if (!string.IsNullOrEmpty(ExcludePluginsUnderDirectoryStr))
 		{
-			string AbsolutePathToExcludePluginsDirectory = Path.Combine(PluginsRootPath, ExcludePluginsUnderDirectoryStr);
-			ExcludePlugins.AddRange(LocalizationUtilities.GetPluginNamesUnderDirectory(AbsolutePathToExcludePluginsDirectory, PluginsRootPath, UEProjectName.Length == 0 ? PluginType.Engine : PluginType.Project));
+			foreach (string ExcludePluginDirectory in ExcludePluginsUnderDirectoryStr.Split(";"))
+			{
+				string AbsolutePathToExcludePluginsDirectory = Path.Combine(PluginsRootPath, ExcludePluginDirectory.Trim());
+				ExcludePlugins.AddRange(LocalizationUtilities.GetPluginNamesUnderDirectory(AbsolutePathToExcludePluginsDirectory, PluginsRootPath, UEProjectName.Length == 0 ? PluginType.Engine : PluginType.Project));
+			}
+
 		}
 
 		if (bShouldGatherPlugins)
