Fix some dotnet warnings

#rnx

diff --git a/Engine/Plugins/Developer/VisualStudioCodeSourceCodeAccess/Source/VisualStudioCodeSourceCodeAccess/VisualStudioCodeSourceCodeAccess.Build.cs b/Engine/Plugins/Developer/VisualStudioCodeSourceCodeAccess/Source/VisualStudioCodeSourceCodeAccess/VisualStudioCodeSourceCodeAccess.Build.cs
--- a/Engine/Plugins/Developer/VisualStudioCodeSourceCodeAccess/Source/VisualStudioCodeSourceCodeAccess/VisualStudioCodeSourceCodeAccess.Build.cs
+++ b/Engine/Plugins/Developer/VisualStudioCodeSourceCodeAccess/Source/VisualStudioCodeSourceCodeAccess/VisualStudioCodeSourceCodeAccess.Build.cs
@@ -37,7 +37,10 @@ namespace UnrealBuildTool.Rules
 							throw new Exception("Unknown visual studio version when mapping to DTEKey: " +
 												Target.WindowsPlatform.ToolChain.ToString());
 					}
-					bHasVisualStudioDTE = RegistryKey.OpenBaseKey(RegistryHive.ClassesRoot, RegistryView.Registry32).OpenSubKey(DTEKey) != null;
+					if (OperatingSystem.IsWindows())
+					{
+						bHasVisualStudioDTE = RegistryKey.OpenBaseKey(RegistryHive.ClassesRoot, RegistryView.Registry32).OpenSubKey(DTEKey) != null;
+					}
 				}
 			}
 			catch
diff --git a/Engine/Source/Programs/AutomationTool/AutomationUtils/ProcessUtils.cs b/Engine/Source/Programs/AutomationTool/AutomationUtils/ProcessUtils.cs
--- a/Engine/Source/Programs/AutomationTool/AutomationUtils/ProcessUtils.cs
+++ b/Engine/Source/Programs/AutomationTool/AutomationUtils/ProcessUtils.cs
@@ -573,15 +573,18 @@ namespace AutomationTool
 			{
 				VisitedPids.Add(PossiblyRelatedId);
 				Process Parent = null;
-				using (ManagementObject ManObj = new ManagementObject(string.Format("win32_process.handle='{0}'", PossiblyRelatedId)))
+				if (OperatingSystem.IsWindows())
 				{
-					ManObj.Get();
-					int ParentId = Convert.ToInt32(ManObj["ParentProcessId"]);
-					if (ParentId == 0 || VisitedPids.Contains(ParentId))
+					using (ManagementObject ManObj = new ManagementObject(string.Format("win32_process.handle='{0}'", PossiblyRelatedId)))
 					{
-						return false;
+						ManObj.Get();
+						int ParentId = Convert.ToInt32(ManObj["ParentProcessId"]);
+						if (ParentId == 0 || VisitedPids.Contains(ParentId))
+						{
+							return false;
+						}
+						Parent = Process.GetProcessById(ParentId);  // will throw an exception if not spawned by us or not running
 					}
-					Parent = Process.GetProcessById(ParentId);  // will throw an exception if not spawned by us or not running
 				}
 				if (Parent != null)
 				{
diff --git a/Engine/Source/Programs/AutomationTool/AutomationUtils/Utils.cs b/Engine/Source/Programs/AutomationTool/AutomationUtils/Utils.cs
--- a/Engine/Source/Programs/AutomationTool/AutomationUtils/Utils.cs
+++ b/Engine/Source/Programs/AutomationTool/AutomationUtils/Utils.cs
@@ -963,13 +963,16 @@ namespace AutomationTool
 				DriveInfo Drive = new(InPath);
 				if (Drive.DriveType == DriveType.Network)
 				{
-					using (ManagementObject ManObj = new($"Win32_LogicalDisk='{Drive.Name.TrimEnd(Path.DirectorySeparatorChar)}'")) // Win32_LogicalDisk='X:'
+					if (OperatingSystem.IsWindows())
 					{
-						string UNCRoot = ManObj["ProviderName"].ToString(); // e.g. \\MyServer\Root
-						string SharedPathFragment = InPath.Replace(Drive.Name, "", StringComparison.InvariantCultureIgnoreCase);
+						using (ManagementObject ManObj = new($"Win32_LogicalDisk='{Drive.Name.TrimEnd(Path.DirectorySeparatorChar)}'")) // Win32_LogicalDisk='X:'
+						{
+							string UNCRoot = ManObj["ProviderName"].ToString(); // e.g. \\MyServer\Root
+							string SharedPathFragment = InPath.Replace(Drive.Name, "", StringComparison.InvariantCultureIgnoreCase);
 
-						UNCPath = Path.Combine( UNCRoot, SharedPathFragment);
-						return true;
+							UNCPath = Path.Combine(UNCRoot, SharedPathFragment);
+							return true;
+						}
 					}
 				}
 			}
diff --git a/Engine/Source/Programs/AutomationTool/Gauntlet/Framework/FileLockInfo.cs b/Engine/Source/Programs/AutomationTool/Gauntlet/Framework/FileLockInfo.cs
--- a/Engine/Source/Programs/AutomationTool/Gauntlet/Framework/FileLockInfo.cs
+++ b/Engine/Source/Programs/AutomationTool/Gauntlet/Framework/FileLockInfo.cs
@@ -15,11 +15,15 @@ namespace FileLockInfo
 		{
 			try
 			{
-				using (ManagementObjectSearcher searcher = new ManagementObjectSearcher("SELECT CommandLine FROM Win32_Process WHERE ProcessId = " + process.Id))
-				using (ManagementObjectCollection objects = searcher.Get())
+				if (OperatingSystem.IsWindows())
 				{
-					return objects.Cast<ManagementBaseObject>().SingleOrDefault()?["CommandLine"]?.ToString();
+					using (ManagementObjectSearcher searcher = new ManagementObjectSearcher("SELECT CommandLine FROM Win32_Process WHERE ProcessId = " + process.Id))
+					using (ManagementObjectCollection objects = searcher.Get())
+					{
+						return objects.Cast<ManagementBaseObject>().SingleOrDefault()?["CommandLine"]?.ToString();
+					}
 				}
+				return string.Format("<error retrieving commandline: non-windows platform>");
 			}
 			catch (Exception Ex)
 			{
diff --git a/Engine/Source/Programs/UnrealBuildTool.Tests/ArtifactsTest.cs b/Engine/Source/Programs/UnrealBuildTool.Tests/ArtifactsTest.cs
--- a/Engine/Source/Programs/UnrealBuildTool.Tests/ArtifactsTest.cs
+++ b/Engine/Source/Programs/UnrealBuildTool.Tests/ArtifactsTest.cs
@@ -1,6 +1,5 @@
 // Copyright Epic Games, Inc. All Rights Reserved.
 
-using System.Reflection;
 using System.Threading;
 using System.Threading.Tasks;
 using EpicGames.Core;
diff --git a/Engine/Source/Programs/UnrealBuildTool.Tests/PluginDescriptorTests.cs b/Engine/Source/Programs/UnrealBuildTool.Tests/PluginDescriptorTests.cs
--- a/Engine/Source/Programs/UnrealBuildTool.Tests/PluginDescriptorTests.cs
+++ b/Engine/Source/Programs/UnrealBuildTool.Tests/PluginDescriptorTests.cs
@@ -25,7 +25,7 @@ namespace UnrealBuildToolTests
 		}
 
 		[ClassInitialize]
-		public static void Setup(TestContext testContext)
+		public static void Setup()
 		{
 			PluginDescriptorTests.s_tempDirectory = CreateTempDir();
 		}
diff --git a/Engine/Source/Programs/UnrealBuildTool.Tests/UnrealHeaderToolTests.cs b/Engine/Source/Programs/UnrealBuildTool.Tests/UnrealHeaderToolTests.cs
--- a/Engine/Source/Programs/UnrealBuildTool.Tests/UnrealHeaderToolTests.cs
+++ b/Engine/Source/Programs/UnrealBuildTool.Tests/UnrealHeaderToolTests.cs
@@ -1,12 +1,14 @@
 // Copyright Epic Games, Inc. All Rights Reserved.
 
+#if DISABLE
 using System.Reflection;
 using EpicGames.Core;
 using EpicGames.UHT.Utils;
 using Microsoft.Extensions.Logging;
-using Microsoft.VisualStudio.TestTools.UnitTesting;
 using UnrealBuildBase;
 using UnrealBuildTool.Modes;
+#endif
+using Microsoft.VisualStudio.TestTools.UnitTesting;
 
 namespace UnrealBuildToolTests
 {
