BuildGraph: Log the tool id when deploying to Horde.

diff --git a/Engine/Source/Programs/AutomationTool/BuildGraph/Tasks/DeployToolTask.cs b/Engine/Source/Programs/AutomationTool/BuildGraph/Tasks/DeployToolTask.cs
--- a/Engine/Source/Programs/AutomationTool/BuildGraph/Tasks/DeployToolTask.cs
+++ b/Engine/Source/Programs/AutomationTool/BuildGraph/Tasks/DeployToolTask.cs
@@ -141,6 +141,8 @@ namespace AutomationTool.Tasks
 				throw new AutomationException($"Missing 'server' key from {settingsFile}");
 			}
 
+			ToolId toolId = new ToolId(Parameters.Id);
+
 			ServiceCollection serviceCollection = new ServiceCollection();
 			serviceCollection.Configure<HordeOptions>(options =>
 			{
@@ -156,14 +158,12 @@ namespace AutomationTool.Tasks
 			using HordeHttpClient hordeHttpClient = hordeClient.CreateHttpClient();
 
 			GetServerInfoResponse infoResponse = await hordeHttpClient.GetServerInfoAsync();
-			Logger.LogInformation("Uploading tool to {ServerUrl} (Version: {Version}, API v{ApiVersion})...", settings.Server, infoResponse.ServerVersion, (int)infoResponse.ApiVersion);
+			Logger.LogInformation("Uploading {ToolId} to {ServerUrl} (Version: {Version}, API v{ApiVersion})...", toolId, settings.Server, infoResponse.ServerVersion, (int)infoResponse.ApiVersion);
 
 			BlobSerializerOptions serializerOptions = BlobSerializerOptions.Create(infoResponse.ApiVersion);
 
 			IBlobRef handle;
 
-			ToolId toolId = new ToolId(Parameters.Id);
-
 			using IStorageClient storageClient = hordeClient.CreateStorageClient(toolId);
 			await using (IBlobWriter blobWriter = storageClient.CreateBlobWriter(serializerOptions: serializerOptions))
 			{
