Corrected documentation on nuget license check task

diff --git a/Engine/Source/Programs/AutomationTool/BuildGraph/Tasks/NuGetLicenseCheckTask.cs b/Engine/Source/Programs/AutomationTool/BuildGraph/Tasks/NuGetLicenseCheckTask.cs
--- a/Engine/Source/Programs/AutomationTool/BuildGraph/Tasks/NuGetLicenseCheckTask.cs
+++ b/Engine/Source/Programs/AutomationTool/BuildGraph/Tasks/NuGetLicenseCheckTask.cs
@@ -58,7 +58,7 @@ namespace AutomationTool.Tasks
 	}
 
 	/// <summary>
-	/// Spawns Docker and waits for it to complete.
+	/// Verifies which licenses are in use by nuget dependencies
 	/// </summary>
 	[TaskElement("NuGet-LicenseCheck", typeof(NuGetLicenseCheckTaskParameters))]
 	public class NuGetLicenseCheckTask : SpawnTaskBase
@@ -74,7 +74,7 @@ namespace AutomationTool.Tasks
 		NuGetLicenseCheckTaskParameters Parameters;
 
 		/// <summary>
-		/// Construct a Docker task
+		/// Construct a NuGetLicenseCheckTask task
 		/// </summary>
 		/// <param name="InParameters">Parameters for the task</param>
 		public NuGetLicenseCheckTask(NuGetLicenseCheckTaskParameters InParameters)
