Adding an option to run code signing via signtool.exe in parallel, up to 16 concurrent instances.

* Exposing the option to the InstalledEngineBuild.xml, off by default
* Moving to a signing service requires parallel calls for performance
* Some clean up of log printing to improved structured logging output

#rnx

#changelist validated
#virtualized

diff --git a/Engine/Build/InstalledEngineBuild.xml b/Engine/Build/InstalledEngineBuild.xml
--- a/Engine/Build/InstalledEngineBuild.xml
+++ b/Engine/Build/InstalledEngineBuild.xml
@@ -52,6 +52,7 @@
 
 	<!-- Whether to sign any executables produced -->
 	<Option Name="SignExecutables" Restrict="true|false" DefaultValue="false" Description="Sign the executables produced where signing is available"/>
+	<Option Name="SignWindowsExecutablesInParallel" Restrict="true|false" DefaultValue="false" Description="Sign the executables in parallel, Windows only"/>
 
 	<!-- The analytics type -->
 	<Option Name="AnalyticsTypeOverride" DefaultValue="" Description="Identifier for analytic events to send"/>
@@ -194,7 +195,7 @@
 			<Do If="'$(SignExecutables)' == true">
 				<Tag Files="#UnrealEditor Win64" Filter="$(WindowsSignFilter)" Except="$(PluginsExceptions)" With="#UnrealEditor Win64 Unsigned"/>
 				<Copy Files="#UnrealEditor Win64 Unsigned" From="$(RootDir)" To="$(SavedOutput)" Tag="#Editor_ToSign_Win64"/>
-				<Sign Files="#Editor_ToSign_Win64" Tag="#UnrealEditor Win64 Signed"/>
+				<Sign Files="#Editor_ToSign_Win64" Parallel="$(SignWindowsExecutablesInParallel)" Tag="#UnrealEditor Win64 Signed"/>
 			</Do>
 			<Tag Files="#UnrealEditor Win64" Except="$(PluginsExceptions)" With="#UnrealEditor Win64 WithExceptions"/>
 
@@ -302,7 +303,7 @@
 			<Do If="'$(SignExecutables)' == true">
 				<Tag Files="#UnrealGame Win64" Filter="$(WindowsSignFilter)" With="#UnrealGame Win64 Unsigned"/>
 				<Copy Files="#UnrealGame Win64 Unsigned" From="$(RootDir)" To="$(SavedOutput)" Tag="#Game_ToSign_Win64"/>
-				<Sign Files="#Game_ToSign_Win64" Tag="#UnrealGame Win64 Signed"/>
+				<Sign Files="#Game_ToSign_Win64" Parallel="$(SignWindowsExecutablesInParallel)" Tag="#UnrealGame Win64 Signed"/>
 			</Do>
 		</Node>
 	</Agent>
@@ -553,7 +554,7 @@
 			<Tag Files="$(CsToolsDir)/Engine/..." With="#SourceForDocumentation" />
 
 			<!-- Sign the binaries -->
-			<Sign Files="#Sign Binaries" If="'$(SignExecutables)' == true"/>
+			<Sign Files="#Sign Binaries" Parallel="$(SignWindowsExecutablesInParallel)" If="'$(SignExecutables)' == true"/>
 		</Node>
 		<Node Name="Build Tools Linux" Requires="Update Version Files">
 			<Compile Target="CrashReportClient" Platform="Linux" Configuration="Shipping" Tag="#Build Tools Linux" Arguments="$(CrashReporterCompileArgs)"/>
@@ -1096,7 +1097,7 @@
 				<!-- Copy then Sign required files in their final location -->
 				<Copy Files="#Sign Tools for Installed Win64" From="$(RootDir)" To="$(LocalInstalledDir)" Tag="#Files to Sign"/>
 				<Copy Files="#Sign CS Tools for Installed Win64" From="$(CsToolsDir)" To="$(LocalInstalledDir)" Tag="#Files to Sign"/>
-				<Sign Files="#Files to Sign" Tag="#Installed Build Win64 Files"/>
+				<Sign Files="#Files to Sign" Parallel="$(SignWindowsExecutablesInParallel)" Tag="#Installed Build Win64 Files"/>
 			</Do>
 
 			<!-- Copy the compressed DDC -->
diff --git a/Engine/Source/Programs/AutomationTool/AutomationUtils/CommandUtils.cs b/Engine/Source/Programs/AutomationTool/AutomationUtils/CommandUtils.cs
--- a/Engine/Source/Programs/AutomationTool/AutomationUtils/CommandUtils.cs
+++ b/Engine/Source/Programs/AutomationTool/AutomationUtils/CommandUtils.cs
@@ -1,28 +1,24 @@
 // Copyright Epic Games, Inc. All Rights Reserved.
+using Microsoft.Extensions.Logging;
 using System;
 using System.Collections;
 using System.Collections.Generic;
 using System.Diagnostics;
 using System.IO;
-using System.Net;
+using System.IO.Compression;
+using System.Linq;
 using System.Runtime.Serialization;
-using System.Runtime.Serialization.Json;
+using System.Runtime.Versioning;
 using System.Text;
 using System.Threading;
-using UnrealBuildTool;
-using System.Runtime.CompilerServices;
-using System.Linq;
-using System.Reflection;
-using System.Security.Permissions;
 using System.Threading.Tasks;
-using EpicGames.Core;
 using System.Xml;
 using System.Xml.Serialization;
-using System.IO.Compression;
-using JetBrains.Annotations;
+
+using EpicGames.Core;
 using UnrealBuildBase;
-using Microsoft.Extensions.Logging;
-using System.Runtime.Versioning;
+using UnrealBuildTool;
+using JetBrains.Annotations;
 
 using static AutomationTool.CommandUtils;
 
@@ -2301,8 +2297,8 @@ namespace AutomationTool
 					if (AgeDays > MaximumDaysToKeepTempStorage)
                     {
 						Logger.LogInformation("Deleting formal build directory {Directory}, because it is {Age} days old (maximum {MaximumDaysToKeepTempStorage}).", ThisDirInfo.FullName, (int)AgeDays, MaximumDaysToKeepTempStorage);
-                        DeleteDirectory_NoExceptions(true, ThisDirInfo.FullName);
-                    }
+						DeleteDirectory_NoExceptions(true, ThisDirInfo.FullName);
+					}
                     else
                     {
 						Logger.LogDebug("Not deleting formal build directory {Directory}, because it is {Age} days old (maximum {MaximumDaysToKeepTempStorage}).", ThisDirInfo.FullName, (int)AgeDays, MaximumDaysToKeepTempStorage);
@@ -2879,8 +2875,8 @@ namespace AutomationTool
                     Name = AlternateName;
                 }
 
-                var OutputStr = String.Format("UAT,{0},{1},{2}" + Environment.NewLine, Name, StartTime, DateTime.Now);
-                Logger.LogDebug("{Text}", OutputStr);
+                var OutputStr = $"UAT,{Name},{StartTime},{DateTime.Now}" + Environment.NewLine;
+                Logger.LogDebug("{CSVHeader}", OutputStr);
                 if (CommandUtils.IsBuildMachine && !String.IsNullOrEmpty(CommandUtils.CmdEnv.CSVFile) && CommandUtils.CmdEnv.CSVFile != "nul")
                 {
                     try
@@ -2889,7 +2885,7 @@ namespace AutomationTool
                     }
                     catch (Exception Ex)
                     {
-                        Logger.LogWarning("Could not append to csv file ({Arg0}) : {Arg1}", CommandUtils.CmdEnv.CSVFile, Ex.ToString());
+                        Logger.LogWarning("Could not append to csv file ({File}) : {Exception}", CommandUtils.CmdEnv.CSVFile, Ex.ToString());
                     }
                 }
             }
@@ -3020,81 +3016,116 @@ namespace AutomationTool
 
 		static readonly string[] TimestampServersSHA256 =
 		{
+			"http://timestamp.digicert.com",
 			"http://sha256timestamp.ws.symantec.com/sha256/timestamp",
 			"http://timestamp.comodoca.com/?td=sha256",
 			"http://rfc3161timestamp.globalsign.com/advanced"
 		};
 
+		static readonly int LongestSHA1 = TimestampServersSHA1.Select(x => x.Length).Max();
+		static readonly int LongestSHA256 = TimestampServersSHA1.Select(x => x.Length).Max();
+
 		[SupportedOSPlatform("windows")]
-		public static void Sign(FileReference File, SignatureType SignatureType, bool AllowMultipleSignatures = true, string Description = null)
+		public static void Sign(FileReference File, SignatureType SignatureType, bool AllowMultipleSignatures = true, string Description = null, bool RunInParallel = false)
 		{
 			List<FileReference> Files = new List<FileReference> { File };
-			Sign(Files, SignatureType, AllowMultipleSignatures, Description);
+			Sign(Files, SignatureType, AllowMultipleSignatures, Description, RunInParallel);
 		}
 
 		[SupportedOSPlatform("windows")]
-		public static void Sign(List<FileReference> Files, SignatureType SignatureType, bool AllowMultipleSignatures = true, string Description = null)
+		public static void Sign(List<FileReference> Files, SignatureType SignatureType, bool AllowMultipleSignatures = true, string Description = null, bool RunInParallel = false)
 		{
+			const int MaxAttempts = 6;
 			string SignToolPath = GetSignToolPath();
 			string DescriptionArg = String.IsNullOrEmpty(Description) ? "" : $"/d \"{Description}\"";
 			string SpecificStoreArg = bUseMachineStoreForCertificates ? " /sm" : "";
 			string MultipleSignatureArg = AllowMultipleSignatures ? " /as" : "";
 			string SHA1TimestampArg = AllowMultipleSignatures ? " /tr" : " /t";
+			string SHA256TimestampArg = " /tr";
+
+			List<Task> SignTaskList = new List<Task>();
 
 			for(int FileIdx = 0; FileIdx < Files.Count; )
 			{
-				Stopwatch Timer = Stopwatch.StartNew();
+				//@TODO: Verbosity choosing
+				//  /v will spew lots of info
+				//  /q does nothing on success and minimal output on failure
+				StringBuilder CommandLine = new StringBuilder();
+				StringBuilder FileList = new StringBuilder();
+				if(SignatureType == SignatureType.SHA1)
+				{
+					CommandLine.Append($"sign{SpecificStoreArg} /a /n \"{SigningIdentity}\" /v {MultipleSignatureArg} {DescriptionArg}");
+				}
+				else if(SignatureType == SignatureType.SHA256)
+				{
+					CommandLine.Append($"sign{SpecificStoreArg} /a /fd sha256 /td sha256 /as /n \"{SigningIdentity}\" {DescriptionArg}");
+				}
+				else
+				{
+					throw new ArgumentException($"Invalid signature type type ({SignatureType})");
+				}
 
-				int NumAttempts = 0;
-				for(;;)
+				// Append the files for this batch
+				// per: https://docs.microsoft.com/en-us/dotnet/api/system.diagnostics.process.start?view=net-6.0#system-diagnostics-process-start(system-diagnostics-processstartinfo)
+				// The length of the application + command line arguments cannot exceed 2080 characters, otherwise a Win32Exception will be thrown.
+				int NextFileIdx = FileIdx;
+				int FileCount = 0;
+				int TimeStampLength = SignatureType == SignatureType.SHA1 ? SHA1TimestampArg.Length + LongestSHA1 : SHA256TimestampArg.Length + LongestSHA256;
+				while(NextFileIdx < Files.Count && (SignToolPath.Length + TimeStampLength + CommandLine.Length + FileList.Length + Files[NextFileIdx].FullName.Length < 2080))
 				{
-					//@TODO: Verbosity choosing
-					//  /v will spew lots of info
-					//  /q does nothing on success and minimal output on failure
-					StringBuilder CommandLine = new StringBuilder();
-					if(SignatureType == SignatureType.SHA1)
-					{
-						CommandLine.AppendFormat("sign{0} /a /n \"{1}\" {2} {3} /v {4} {5}", SpecificStoreArg, SigningIdentity, SHA1TimestampArg, TimestampServersSHA1[NumAttempts % TimestampServersSHA1.Length], MultipleSignatureArg, DescriptionArg);
-					}
-					else if(SignatureType == SignatureType.SHA256)
-					{
-						CommandLine.AppendFormat("sign{0} /a /fd sha256 /td sha256 /as /n \"{1}\" /tr {2} {3}", SpecificStoreArg, SigningIdentity, TimestampServersSHA256[NumAttempts % TimestampServersSHA256.Length], DescriptionArg);
-					}
-					else
-					{
-						throw new ArgumentException(String.Format("Invalid signature type type ({0})", SignatureType));
-					}
+					FileList.Append($" \"{Files[NextFileIdx]}\"");
+					NextFileIdx++;
+					FileCount++;
+				}
 
-					// Append the files for this batch
-					// per: https://docs.microsoft.com/en-us/dotnet/api/system.diagnostics.process.start?view=net-6.0#system-diagnostics-process-start(system-diagnostics-processstartinfo)
-					// The length of the application + command line arguments cannot exceed 2080 characters, otherwise a Win32Exception will be thrown.
-					int NextFileIdx = FileIdx;
-					while(NextFileIdx < Files.Count && SignToolPath.Length + CommandLine.Length + Files[NextFileIdx].FullName.Length < 2080)
+				Action SignToolAction = new Action(() =>
+				{
+					Stopwatch TaskTimer = Stopwatch.StartNew();
+					for (int Attempts = 0; ; ++Attempts)
 					{
-						CommandLine.AppendFormat(" \"{0}\"", Files[NextFileIdx]);
-						NextFileIdx++;
-					}
+						StringBuilder FinalCommandline = new StringBuilder();
+						FinalCommandline.Append(CommandLine);
+
+						if (SignatureType == SignatureType.SHA1)
+						{
+							FinalCommandline.Append($"{SHA1TimestampArg} {TimestampServersSHA1[Attempts % TimestampServersSHA1.Length]} ");
+						}
+						else
+						{
+							FinalCommandline.Append($"{SHA256TimestampArg} {TimestampServersSHA256[Attempts % TimestampServersSHA256.Length]} ");
+						}
 
-					IProcessResult Result = CommandUtils.Run(SignToolPath, CommandLine.ToString(), null, CommandUtils.ERunOptions.AllowSpew);
-					NumAttempts++;
+						FinalCommandline.Append(FileList);
 
-					if (Result.ExitCode != 1)
-					{
-						if (Result.ExitCode == 2)
+
+						IProcessResult Result = CommandUtils.Run(SignToolPath, FinalCommandline.ToString(), null, CommandUtils.ERunOptions.AllowSpew);
+
+						if (Result.ExitCode != 1)
 						{
-							Logger.LogError("{Text}", String.Format("Signtool returned a warning."));
+							if (Result.ExitCode == 2)
+							{
+								Logger.LogWarning("Signtool returned a warning.");
+							}
+							// Success!
+							break;
 						}
-						// Success!
-						FileIdx = NextFileIdx;
-						break;
-					}
 
-					if (Timer.Elapsed.TotalMinutes > 3.0 && NumAttempts >= 6)
-					{
-						throw new AutomationException("Failed to sign files {0} times over a period of {1}", NumAttempts, Timer.Elapsed);
+						if (TaskTimer.Elapsed.TotalMinutes > 3.0 && Attempts >= MaxAttempts)
+						{
+							throw new AutomationException($"Failed to sign files {Attempts} times over a period of {TaskTimer.Elapsed}");
+						}
 					}
-				}
+				});
+
+				FileIdx = NextFileIdx;
+
+				SignTaskList.Add(new Task(SignToolAction));
 			}
+
+			// If running in parallel, Limit to between 4 and 16 concurrent. Otherwise 1 instance.
+			int MaxParallelism = RunInParallel ? Math.Max(16, Math.Min(Environment.ProcessorCount, 4)) : 1;
+			Logger.LogInformation("Running {Count}, {Concurrent} max concurrent, signtool instances", SignTaskList.Count, MaxParallelism);
+			Parallel.ForEach(SignTaskList, new ParallelOptions() { MaxDegreeOfParallelism = MaxParallelism }, x => { x.Start(); x.Wait(); });
 		}
 
 		/// <summary>
@@ -3145,7 +3176,7 @@ namespace AutomationTool
 		{
             if (!OperatingSystem.IsWindows())
             {
-                Logger.LogDebug("{Text}", String.Format("Can't sign '{0}' on non-Windows platform.", Filename));
+                Logger.LogDebug("Can't sign '{File}' on non-Windows platform.", Filename);
                 return;
             }
             if (!CommandUtils.FileExists(Filename))
@@ -3156,9 +3187,11 @@ namespace AutomationTool
 			FileInfo TargetFileInfo = new FileInfo(Filename);
 
 			// Executable extensions
-			List<string> Extensions = new List<string>();
-			Extensions.Add(".dll");
-			Extensions.Add(".exe");
+			List<string> Extensions = new List<string>
+			{
+				".dll",
+				".exe"
+			};
 
 			bool IsExecutable = bIgnoreExtension;
 
@@ -3172,7 +3205,7 @@ namespace AutomationTool
 			}
 			if (!IsExecutable)
 			{
-				Logger.LogDebug("{Text}", String.Format("Won't sign '{0}', not an executable.", TargetFileInfo.FullName));
+				Logger.LogDebug("Won't sign '{File}', not an executable.", TargetFileInfo.FullName);
 				return;
 			}
 
@@ -3226,12 +3259,14 @@ namespace AutomationTool
 			}
 
 			// Executable extensions
-			List<string> Extensions = new List<string>();
-			Extensions.Add(".dylib");
-			Extensions.Add(".so");
-			Extensions.Add(".app");
-			Extensions.Add(".framework");
-			Extensions.Add(".bundle");
+			List<string> Extensions = new List<string>
+			{
+				".dylib",
+				".so",
+				".app",
+				".framework",
+				".bundle"
+			};
 
 			bool bIsExecutable = bIgnoreExtension || (!bIsDirectory && Path.GetExtension(InPath) == "" && !InPath.EndsWith("PkgInfo"));
 
@@ -3247,7 +3282,7 @@ namespace AutomationTool
 			{
 				if (!bIsDirectory)
 				{
-					Logger.LogDebug("{Text}", String.Format("Won't sign '{0}', not an executable.", InPath));
+					Logger.LogDebug("Won't sign '{File}', not an executable.", InPath);
 				}
 				return;
 			}
@@ -3303,11 +3338,12 @@ namespace AutomationTool
 		/// Will automatically skip signing if -NoSign is specified in the command line.
 		/// </summary>
 		/// <param name="Files">List of files to sign</param>
-		public static void SignMultipleIfEXEOrDLL(BuildCommand Command, IEnumerable<string> Files, string Description = null)
+		/// <param name="bRunInParallel">Run the sign operation in paralle, Windows only</param>
+		public static void SignMultipleIfEXEOrDLL(BuildCommand Command, IEnumerable<string> Files, string Description = null, bool bRunInParallel = false)
 		{
 			if (!Command.ParseParam("NoSign"))
 			{
-				Logger.LogInformation("Signing up to {Arg0} files...", Files.Count());
+				Logger.LogInformation("Signing up to {FileCount} files...", Files.Count());
 				UnrealBuildTool.UnrealTargetPlatform TargetPlatform = UnrealBuildTool.BuildHostPlatform.Current.Platform;
 				if (TargetPlatform == UnrealBuildTool.UnrealTargetPlatform.Mac)
 				{
@@ -3323,20 +3359,20 @@ namespace AutomationTool
 					{
 						FilesToSign.Add(new FileReference(File));
 					}
-					SignMultipleFilesIfEXEOrDLL(FilesToSign, Description: Description);
+					SignMultipleFilesIfEXEOrDLL(FilesToSign, Description: Description, bRunInParallel: bRunInParallel);
 				}
 			}
 			else
 			{
-				Logger.LogDebug("Skipping signing {Arg0} files due to -nosign.", Files.Count());
+				Logger.LogDebug("Skipping signing {FileCount} files due to -nosign.", Files.Count());
 			}
 		}
 
-		public static void SignMultipleFilesIfEXEOrDLL(List<FileReference> Files, bool bIgnoreExtension = false, string Description = null)
+		public static void SignMultipleFilesIfEXEOrDLL(List<FileReference> Files, bool bIgnoreExtension = false, string Description = null, bool bRunInParallel = false)
 		{
 			if (!OperatingSystem.IsWindows())
 			{
-				Logger.LogDebug("{Text}", String.Format("Can't sign on non-Windows platform."));
+				Logger.LogDebug("Can't sign on non-Windows platform.");
 				return;
 			}
 			List<FileReference> FinalFiles = new List<FileReference>();
@@ -3345,11 +3381,13 @@ namespace AutomationTool
 				FileInfo TargetFileInfo = new FileInfo(Filename);
 
 				// Executable extensions
-				List<string> Extensions = new List<string>();
-				Extensions.Add(".dll");
-				Extensions.Add(".exe");
-				Extensions.Add(".msi");
-				Extensions.Add(".dle");
+				List<string> Extensions = new List<string>
+				{
+					".dll",
+					".exe",
+					".msi",
+					".dle"
+				};
 
 				bool IsExecutable = bIgnoreExtension;
 
@@ -3368,8 +3406,8 @@ namespace AutomationTool
 					FinalFiles.Add(new FileReference(TargetFileInfo));
 				}
 			}
-			CodeSignWindows.Sign(FinalFiles, CodeSignWindows.SignatureType.SHA1, !FinalFiles.Any(x => x.HasExtension(".msi")), Description: Description); // By default we append signatures, but disable if there are MSI files in the list
-			CodeSignWindows.Sign(FinalFiles.Where(x => !x.HasExtension(".msi")).ToList(), CodeSignWindows.SignatureType.SHA256, Description: Description); // MSI files can only have one signature; prefer SHA1 for compatibility
+			CodeSignWindows.Sign(FinalFiles, CodeSignWindows.SignatureType.SHA1, !FinalFiles.Any(x => x.HasExtension(".msi")), Description: Description, bRunInParallel); // By default we append signatures, but disable if there are MSI files in the list
+			CodeSignWindows.Sign(FinalFiles.Where(x => !x.HasExtension(".msi")).ToList(), CodeSignWindows.SignatureType.SHA256, Description: Description, RunInParallel: bRunInParallel); // MSI files can only have one signature; prefer SHA1 for compatibility
 		}
 	}
 
diff --git a/Engine/Source/Programs/AutomationTool/BuildGraph/Tasks/SignTask.cs b/Engine/Source/Programs/AutomationTool/BuildGraph/Tasks/SignTask.cs
--- a/Engine/Source/Programs/AutomationTool/BuildGraph/Tasks/SignTask.cs
+++ b/Engine/Source/Programs/AutomationTool/BuildGraph/Tasks/SignTask.cs
@@ -1,15 +1,10 @@
 // Copyright Epic Games, Inc. All Rights Reserved.
 
-using AutomationTool;
-using EpicGames.BuildGraph;
-using System;
+using EpicGames.Core;
 using System.Collections.Generic;
 using System.Linq;
-using System.Text;
 using System.Threading.Tasks;
 using System.Xml;
-using EpicGames.Core;
-using UnrealBuildTool;
 using UnrealBuildBase;
 
 namespace AutomationTool.Tasks
@@ -36,6 +31,12 @@ namespace AutomationTool.Tasks
 		/// </summary>
 		[TaskParameter(Optional = true, ValidationType = TaskParameterValidationType.TagList)]
 		public string Tag;
+
+		/// <summary>
+		/// If true, the calls to the signing tool will be performed in parallel.
+		/// </summary>
+		[TaskParameter(Optional = true)]
+		public bool Parallel;
 	}
 
 	/// <summary>
@@ -70,7 +71,7 @@ namespace AutomationTool.Tasks
 			FileReference[] Files = ResolveFilespec(Unreal.RootDirectory, Parameters.Files, TagNameToFileSet).OrderBy(x => x.FullName).ToArray();
 
 			// Sign all the files
-			CodeSign.SignMultipleIfEXEOrDLL(Job.OwnerCommand, (Files.Select(x => x.FullName).ToList()), Description: Parameters.Description);
+			CodeSign.SignMultipleIfEXEOrDLL(Job.OwnerCommand, (Files.Select(x => x.FullName).ToList()), Description: Parameters.Description, Parameters.Parallel);
 
 			// Apply the optional tag to the build products
 			foreach(string TagName in FindTagNamesFromList(Parameters.Tag))
