Horde: Fix errors publishing artifacts using the CreateArtifact task due to Horde services not being registered with DI container.

diff --git a/Engine/Source/Programs/AutomationTool/BuildGraph/Tasks/CreateArtifactTask.cs b/Engine/Source/Programs/AutomationTool/BuildGraph/Tasks/CreateArtifactTask.cs
--- a/Engine/Source/Programs/AutomationTool/BuildGraph/Tasks/CreateArtifactTask.cs
+++ b/Engine/Source/Programs/AutomationTool/BuildGraph/Tasks/CreateArtifactTask.cs
@@ -99,6 +99,7 @@ namespace AutomationTool.Tasks
 		{
 			// Create a DI container for building the graph
 			ServiceCollection serviceCollection = new ServiceCollection();
+			serviceCollection.AddHorde();
 			serviceCollection.AddLogging(builder => builder.AddProvider(new LoggerProviderAdapter(Log.Logger)));
 			serviceCollection.Configure<HordeOptions>(x => x.AllowAuthPrompt = !Automation.IsBuildMachine);
 			serviceCollection.Configure<LoggerFilterOptions>(options => options.AddFilter(typeof(HttpClient).FullName, LogLevel.Warning));
