Suppress the warning message in WebTests.

diff --git a/Engine/Source/Programs/AutomationTool/AutomationUtils/Matchers/UnsecureHttpMatcher.cs b/Engine/Source/Programs/AutomationTool/AutomationUtils/Matchers/UnsecureHttpMatcher.cs
new file mode 100644
--- /dev/null
+++ b/Engine/Source/Programs/AutomationTool/AutomationUtils/Matchers/UnsecureHttpMatcher.cs
@@ -0,0 +1,23 @@
+// Copyright Epic Games, Inc. All Rights Reserved.
+
+using EpicGames.Core;
+using Microsoft.Extensions.Logging;
+using System.Text.RegularExpressions;
+
+#nullable enable
+
+namespace AutomationUtils.Matchers
+{
+	class UnsecureHttpMatcher : ILogEventMatcher
+	{
+		public LogEventMatch? Match(ILogCursor cursor)
+		{
+			if (cursor.Contains("WARNING: Insecure HTTP request"))
+			{
+				return new LogEventBuilder(cursor).ToMatch(LogEventPriority.Normal, LogLevel.Information, KnownLogEvents.Systemic_SignToolTimeStampServer);
+			}
+
+			return null;
+		}
+	}
+}
