[Gauntlet] Return false when cancelling during TryReserveDevies

diff --git a/Engine/Source/Programs/AutomationTool/Gauntlet/Unreal/Base/Gauntlet.UnrealSession.cs b/Engine/Source/Programs/AutomationTool/Gauntlet/Unreal/Base/Gauntlet.UnrealSession.cs
--- a/Engine/Source/Programs/AutomationTool/Gauntlet/Unreal/Base/Gauntlet.UnrealSession.cs
+++ b/Engine/Source/Programs/AutomationTool/Gauntlet/Unreal/Base/Gauntlet.UnrealSession.cs
@@ -732,7 +732,13 @@ namespace Gauntlet
 		{
 			for (; Attempts > 0; --Attempts)
 			{
-				if(Globals.CancelSignalled || TryReserveDevices())
+				if(Globals.CancelSignalled)
+				{
+					ReleaseSessionDevices();
+					return false;
+				}
+
+				if(TryReserveDevices())
 				{
 					return true;
 				}
@@ -743,6 +749,7 @@ namespace Gauntlet
 			}
 
 			Log.Error("Failed to reserve devices after {Attempts}", Attempts);
+			ReleaseSessionDevices();
 			return false;
 		}
 
@@ -790,6 +797,11 @@ namespace Gauntlet
 				// Reserve devices, if needed
 				if (!TryReserveDevices(DeviceReservationAttempts))
 				{
+					if(Globals.CancelSignalled)
+					{
+						return null;
+					}
+
 					// If device reservation fails, the device pool cannot support this launch.
 					DevicePool.Instance.ReportDeviceReservationState();
 					throw new AutomationException("Failed to acquire all devices for launch. See above for details.");
