UAT - include container files (.utoc/.ucas) when checking for existing staged files

diff --git a/Engine/Source/Programs/AutomationTool/Scripts/CopyBuildToStagingDirectory.Automation.cs b/Engine/Source/Programs/AutomationTool/Scripts/CopyBuildToStagingDirectory.Automation.cs
--- a/Engine/Source/Programs/AutomationTool/Scripts/CopyBuildToStagingDirectory.Automation.cs
+++ b/Engine/Source/Programs/AutomationTool/Scripts/CopyBuildToStagingDirectory.Automation.cs
@@ -3473,7 +3473,20 @@ namespace AutomationScripts
 						string UnrealPakResponseFileName = CombinePaths(CmdEnv.EngineSavedFolder, "ResponseFiles", "PakList_" + OutputLocation.GetFileNameWithoutExtension() + ".txt");
 						if (File.Exists(UnrealPakResponseFileName) && FileReference.GetLastWriteTimeUtc(OutputLocation) > File.GetLastWriteTimeUtc(UnrealPakResponseFileName))
 						{
-							bCopiedExistingPak = true;
+							if (ShouldCreateIoStoreContainerFiles(Params, SC))
+							{
+								string UnrealPakIoStoreResponseFileName = CombinePaths(CmdEnv.EngineSavedFolder, "ResponseFiles", "PakListIoStore_" + OutputLocation.GetFileNameWithoutExtension() + ".txt");
+								FileReference[] ContainerOutputLocations = new FileReference[] { OutputLocation.ChangeExtension(".utoc"), OutputLocation.ChangeExtension(".ucas") };
+								if (File.Exists(UnrealPakIoStoreResponseFileName) &&
+									ContainerOutputLocations.All(ContainerOutputLocation => FileReference.GetLastWriteTimeUtc(ContainerOutputLocation) > File.GetLastWriteTimeUtc(UnrealPakIoStoreResponseFileName)))
+								{
+									bCopiedExistingPak = true;
+								}
+							}
+							else
+							{
+								bCopiedExistingPak = true;
+							}
 						}
 					}
 					if (!bCopiedExistingPak)
