Fix archive fail when packaging a BP only project for distribution

diff --git a/Engine/Source/Programs/AutomationTool/Apple/ApplePlatform.Automation.cs b/Engine/Source/Programs/AutomationTool/Apple/ApplePlatform.Automation.cs
--- a/Engine/Source/Programs/AutomationTool/Apple/ApplePlatform.Automation.cs
+++ b/Engine/Source/Programs/AutomationTool/Apple/ApplePlatform.Automation.cs
@@ -152,13 +152,21 @@ public abstract class ApplePlatform : Platform
 		DateDirs.SortBy(x => Directory.GetCreationTime(x.FullName));
 		DateDirs.Reverse();
 
+		string ArchiveName = Target.Receipt.TargetName;
+		if (!SC.IsCodeBasedProject)
+		{
+			if (SC.RawProjectPath != null)
+			{
+				ArchiveName = MakeContentOnlyTargetName(Target.Receipt, SC.ShortProjectName);
+			}
+		}
 
 		// go through each folder, starting at most recent, looking for an archive for the target
 		foreach (DirectoryReference DateDir in DateDirs)
 		{
 			// find the most recent archive for this target (based on name of target, this ignores Development vs Shipping, but 
 			// since Distribution is meant only for Shipping it's ok
-			string Wildcard = AppleExports.MakeBinaryFileName(SC.ShortProjectName, Target.Receipt.Platform, UnrealTargetConfiguration.Development,
+			string Wildcard = AppleExports.MakeBinaryFileName(ArchiveName, Target.Receipt.Platform, UnrealTargetConfiguration.Development,
 				Target.Receipt.Architectures, UnrealTargetConfiguration.Development, null) + " *.xcarchive";
 
 			Logger.LogInformation("Looking in Xcode archive dir {0} for {1}", DateDir, Wildcard);
