PR #11468: Default parameters should be added to the list of valid arguments in BuildGraph scripts.

diff --git a/Engine/Source/Programs/AutomationTool/BuildGraph/BgScriptReader.cs b/Engine/Source/Programs/AutomationTool/BuildGraph/BgScriptReader.cs
--- a/Engine/Source/Programs/AutomationTool/BuildGraph/BgScriptReader.cs
+++ b/Engine/Source/Programs/AutomationTool/BuildGraph/BgScriptReader.cs
@@ -9,10 +9,12 @@ using System.Text;
 using System.Text.RegularExpressions;
 using System.Threading.Tasks;
 using System.Xml;
+using System.Xml.Linq;
 using System.Xml.Schema;
 using EpicGames.BuildGraph;
 using EpicGames.Core;
 using Microsoft.Extensions.Logging;
+using Newtonsoft.Json.Linq;
 using UnrealBuildBase;
 
 #nullable enable
@@ -395,14 +397,27 @@ namespace AutomationTool
 			HashSet<string> validArgumentNames = new HashSet<string>(reader._graph.Options.Select(x => x.Name), StringComparer.OrdinalIgnoreCase);
 			validArgumentNames.Add("PreflightChange");
 
+			// All default properties are valid arguments too
+			foreach (string defaultPropertyKey in defaultProperties.Keys)
+			{
+				validArgumentNames.Add(defaultPropertyKey);
+			}
+
+			bool hasInvalidArguments = false;
 			foreach (string argumentName in arguments.Keys)
 			{
 				if (!validArgumentNames.Contains(argumentName))
 				{
+					hasInvalidArguments = true;
 					logger.LogWarning("Unknown argument '{ArgumentName}' for '{Script}'", argumentName, file);
 				}
 			}
 
+			if (hasInvalidArguments)
+			{
+				logger.LogInformation("Valid arguments for '{Script}': {Arguments}", file, String.Join(", ", validArgumentNames.OrderBy(x => x, StringComparer.OrdinalIgnoreCase)));
+			}
+
 			// Return the constructed graph
 			return reader._graph;
 		}
