[Gauntlet] Introduce a new class (UnrealLogStreamParser) to handle Unreal Structure logging output natively and use it during test tick to avoid parsing the output several times.

#rnx

diff --git a/Engine/Source/Programs/AutomationTool/Gauntlet/Unreal/Automation/UE.Automation.cs b/Engine/Source/Programs/AutomationTool/Gauntlet/Unreal/Automation/UE.Automation.cs
--- a/Engine/Source/Programs/AutomationTool/Gauntlet/Unreal/Automation/UE.Automation.cs
+++ b/Engine/Source/Programs/AutomationTool/Gauntlet/Unreal/Automation/UE.Automation.cs
@@ -572,20 +572,20 @@ namespace UE
 				IdleTimeout = Config.LogIdleTimeout;
 			}
 
-			List<string> ChannelEntries = new List<string>();
 			// We are primarily interested in what the editor is doing
 			var AppInstance = TestInstance.EditorApp;
 
-			UnrealLogParser Parser = new UnrealLogParser(AppInstance.StdOut);
-			ChannelEntries.AddRange(Parser.GetEditorBusyChannels());
+			UnrealLogStreamParser Parser = new UnrealLogStreamParser();
+			LastAutomationEntryCount += Parser.ReadStream(AppInstance.StdOut, LastAutomationEntryCount);
+
+			IEnumerable<string> ChannelEntries = Parser.GetLogFromEditorBusyChannels();
 
 			// Any new entries?
-			if (ChannelEntries.Count > LastAutomationEntryCount)
+			if (ChannelEntries.Any())
 			{
 				// log new entries so people have something to look at
-				ChannelEntries.Skip(LastAutomationEntryCount).ToList().ForEach(S => Log.Info(S));
+				ChannelEntries.ToList().ForEach(S => Log.Info(S));
 				LastAutomationEntryTime = DateTime.Now;
-				LastAutomationEntryCount = ChannelEntries.Count;
 			}
 			else
 			{
diff --git a/Engine/Source/Programs/AutomationTool/Gauntlet/Unreal/Automation/UE.BootTest.cs b/Engine/Source/Programs/AutomationTool/Gauntlet/Unreal/Automation/UE.BootTest.cs
--- a/Engine/Source/Programs/AutomationTool/Gauntlet/Unreal/Automation/UE.BootTest.cs
+++ b/Engine/Source/Programs/AutomationTool/Gauntlet/Unreal/Automation/UE.BootTest.cs
@@ -97,17 +97,15 @@ namespace UE
 			// Get the log of the first client app
 			IAppInstance RunningInstance = this.TestInstance.RunningRoles.First().AppInstance;
 
-			UnrealLogParser LogParser = new UnrealLogParser(RunningInstance.StdOut);
+			UnrealLogStreamParser LogParser = new UnrealLogStreamParser();
+			LogLinesLastTick += LogParser.ReadStream(RunningInstance.StdOut, LogLinesLastTick);
 
-			IEnumerable<string> BusyLogLines = LogParser.GetEditorBusyChannels();
-			int BusyLineCount = BusyLogLines.Count();
-
-			if (BusyLineCount > LogLinesLastTick)
+			IEnumerable<string> BusyLogLines = LogParser.GetLogFromEditorBusyChannels();
+			if (BusyLogLines.Any())
 			{
 				LastLogTime = DateTime.Now;
 				// log new entries so people have something to look at
-				BusyLogLines.Skip(LogLinesLastTick).ToList().ForEach(S => Log.Info("{0}", S));
-				LogLinesLastTick = BusyLineCount;
+				BusyLogLines.ToList().ForEach(S => Log.Info("{0}", S));
 			}
 
 			// Gauntlet will timeout tests based on the -timeout argument, but we have greater insight here so can bail earlier to save
@@ -119,14 +117,11 @@ namespace UE
 				SetUnrealTestResult(TestResult.TimedOut);
 			}
 
-			// now see if the game has brought the first world up for play
-			IEnumerable<string> LogWorldLines = LogParser.GetLogChannel("World");
-
 			string CompletionString = GetCompletionString();
 
 			if (!string.IsNullOrEmpty(CompletionString))
 			{
-				if (LogParser.Content.IndexOf(CompletionString, StringComparison.OrdinalIgnoreCase) > 0)
+				if (LogParser.GetLogLinesContaining(CompletionString).Any())
 				{
 					Log.Info("Found '{0}'. Ending Test", GetCompletionString());
 					MarkTestComplete();
diff --git a/Engine/Source/Programs/AutomationTool/Gauntlet/Unreal/Automation/UE.CookOnTheFly.cs b/Engine/Source/Programs/AutomationTool/Gauntlet/Unreal/Automation/UE.CookOnTheFly.cs
--- a/Engine/Source/Programs/AutomationTool/Gauntlet/Unreal/Automation/UE.CookOnTheFly.cs
+++ b/Engine/Source/Programs/AutomationTool/Gauntlet/Unreal/Automation/UE.CookOnTheFly.cs
@@ -24,8 +24,9 @@ namespace UE
 		bool ClientConnected = false;
 		bool CookingRequest = false;
 		bool CookOnTheFlyModeChange = false;
-		int LastCookEntryCount = 0;
-		IEnumerable<string> LogCategories = new string[] { "CookOnTheFly", "Cook" };
+		int LastEditorLogLine = 0;
+		int LastClientLogLine = 0;
+		IEnumerable<string> LogCategories = new string[] { "LogCookOnTheFly", "LogCook" };
 
 		public CookOnTheFly(Gauntlet.UnrealTestContext InContext)
 			: base(InContext)
@@ -115,7 +116,8 @@ namespace UE
 				return false;
 			}
 
-			LastCookEntryCount = 0;
+			LastEditorLogLine = 0;
+			LastClientLogLine = 0;
 			StartTime = DateTime.Now;
 
 			return true;
@@ -124,7 +126,6 @@ namespace UE
 		public override void TickTest()
 		{
 			base.TickTest();
-			List<string> CookEntries = new List<string>();
 
 			const int TimeoutDuration = 5;
 			if ((DateTime.Now - StartTime).TotalMinutes > TimeoutDuration)
@@ -142,20 +143,19 @@ namespace UE
 				SetUnrealTestResult(Gauntlet.TestResult.Failed);
 			}
 
-			UnrealLogParser EditorLogParser = new UnrealLogParser(EditorInstance.StdOut);
-			IEnumerable<string> EditorCookEntries = EditorLogParser.GetLogChannels(LogCategories, true);
-			CookEntries.AddRange(EditorCookEntries);
+			UnrealLogStreamParser EditorLogParser = new UnrealLogStreamParser();
+			LastEditorLogLine += EditorLogParser.ReadStream(EditorInstance.StdOut, LastEditorLogLine);
+			IEnumerable<string> EditorCookEntries = EditorLogParser.GetLogFromChannels(LogCategories);
 
-			if (CookEntries.Count > LastCookEntryCount)
+			if (EditorCookEntries.Any())
 			{
-				CookEntries.Skip(LastCookEntryCount).ToList().ForEach(S => Log.Info("Editor: " + S));
-				LastCookEntryCount = CookEntries.Count;
+				EditorCookEntries.ToList().ForEach(S => Log.Info("[Editor] " + S));
 			}
 			string CompletionString = GetStartedCookServerString();
 
 			if (!ServerStarted)
 			{
-				if (EditorLogParser.Content.IndexOf(CompletionString, StringComparison.OrdinalIgnoreCase) > 0)
+				if (EditorLogParser.GetLogLinesContaining(CompletionString).Any())
 				{
 					Log.Info("Found '{0}'. Cook Server Started", CompletionString);
 					ServerStarted = true;
@@ -194,7 +194,7 @@ namespace UE
 			{
 				string ClientConnectedString = GetClientConnectedString();
 
-				if (EditorLogParser.Content.IndexOf(ClientConnectedString, StringComparison.OrdinalIgnoreCase) > 0)
+				if (EditorLogParser.GetLogLinesContaining(ClientConnectedString).Any())
 				{
 					Log.Info("Found '{0}'. Client connected", ClientConnectedString);
 					ClientConnected = true;
@@ -205,7 +205,7 @@ namespace UE
 			{
 				string CookingRequestString = GetCookingRequestString();
 
-				if (EditorLogParser.Content.IndexOf(CookingRequestString, StringComparison.OrdinalIgnoreCase) > 0)
+				if (EditorLogParser.GetLogLinesContaining(CookingRequestString).Any())
 				{
 					Log.Info("Found '{0}'. Cooking request exists", CookingRequestString);
 					CookingRequest = true;
@@ -246,31 +246,31 @@ namespace UE
 					SetUnrealTestResult(Gauntlet.TestResult.Failed);
 				}
 
-				UnrealLogParser ClientLogParser = new UnrealLogParser(ClientInstance.StdOut);
+				UnrealLogStreamParser ClientLogParser = new UnrealLogStreamParser();
+				ClientLogParser.ReadStream(ClientInstance.StdOut);
 
 				string EngineInitializedString = GetEngineInitializedString();
-				if (ClientLogParser.Content.IndexOf(EngineInitializedString, StringComparison.OrdinalIgnoreCase) > 0)
+				if (ClientLogParser.GetLogLinesContaining(EngineInitializedString).Any())
 				{
 					string GameStartedString = GetGameStartString();
 					string CookingProcessString = GetCookingProcessString();
 					string ReceivedPackagesCookedString = GetReceivedPackagesCookedString();
 					string TransportCreatedString = GetCreatedTransportString();
-					IEnumerable<string> ClientCookEntries = ClientLogParser.GetLogChannels(LogCategories, true);
-					CookEntries.AddRange(ClientCookEntries);
+					IEnumerable<string> ClientCookEntries = ClientLogParser.GetLogFromChannels(LogCategories);
 
-					if (CookEntries.Count > LastCookEntryCount)
+					if (ClientCookEntries.Count() > LastClientLogLine)
 					{
-						CookEntries.Skip(LastCookEntryCount).ToList().ForEach(S => Log.Info("Client: " + S));
-						LastCookEntryCount = CookEntries.Count;
+						ClientCookEntries.Skip(LastClientLogLine).ToList().ForEach(S => Log.Info("[Client] " + S));
+						LastClientLogLine = ClientCookEntries.Count();
 					}
 
-					if (ClientLogParser.Content.IndexOf(GameStartedString, StringComparison.OrdinalIgnoreCase) > 0
-						&& EditorLogParser.Content.IndexOf(CookingProcessString, StringComparison.OrdinalIgnoreCase) > 0
+					if (ClientLogParser.GetLogLinesContaining(GameStartedString).Any()
+						&& EditorLogParser.GetLogLinesContaining(CookingProcessString).Any()
 						&& ClientCookEntries.Any()
-						&& ClientLogParser.Content.IndexOf(ReceivedPackagesCookedString, StringComparison.OrdinalIgnoreCase) > 0
-						&& ClientLogParser.Content.IndexOf(TransportCreatedString, StringComparison.OrdinalIgnoreCase) > 0)
+						&& ClientLogParser.GetLogLinesContaining(ReceivedPackagesCookedString).Any()
+						&& ClientLogParser.GetLogLinesContaining(TransportCreatedString).Any())
 					{
-						Log.Info("Found '{0}', '{1}', '{2}', '{3}'. The CookOnTheFly log channel takes place. The cooking process takes place", GameStartedString, CookingProcessString, ReceivedPackagesCookedString, TransportCreatedString);
+						Log.Info("Found '{0}', '{1}', '{2}', '{3}'. The CookOnTheFly log channel is active. The cooking process is taking place.", GameStartedString, CookingProcessString, ReceivedPackagesCookedString, TransportCreatedString);
 						MarkTestComplete();
 						SetUnrealTestResult(Gauntlet.TestResult.Passed);
 					}
diff --git a/Engine/Source/Programs/AutomationTool/Gauntlet/Unreal/Base/Gauntlet.UnrealTestNode.cs b/Engine/Source/Programs/AutomationTool/Gauntlet/Unreal/Base/Gauntlet.UnrealTestNode.cs
--- a/Engine/Source/Programs/AutomationTool/Gauntlet/Unreal/Base/Gauntlet.UnrealTestNode.cs
+++ b/Engine/Source/Programs/AutomationTool/Gauntlet/Unreal/Base/Gauntlet.UnrealTestNode.cs
@@ -1033,30 +1033,27 @@ namespace Gauntlet
 			{
 				if (App != null)
 				{
-					UnrealLogParser Parser = new UnrealLogParser(App.StdOut);
-					List<string> TestLines = new List<string>();
-					// ONLY ADD RANGE ONCE. Ordering is important and will be skewed if multiple ranges are added which can skew how logs are pulled out.
-					TestLines.AddRange(Parser.GetLogChannels(LogCategories, true));
+					UnrealLogStreamParser Parser = new UnrealLogStreamParser();
+					LastLogCount += Parser.ReadStream(App.StdOut, LastLogCount);
 
-					for (int i = LastLogCount; i < TestLines.Count(); i++)
+					foreach (string TestLine in Parser.GetLogFromShortNameChannels(LogCategories))
 					{
-						Log.Info(string.Format("{0}: {1}", AppPrefix, TestLines[i]));
+						Log.Info(string.Format("{0}: {1}", AppPrefix, TestLine));
 
 						if (bUpdateHeartbeatTime)
 						{
-							if (Regex.IsMatch(TestLines[i], @".*GauntletHeartbeat\: Active.*"))
+							if (Regex.IsMatch(TestLine, @".*GauntletHeartbeat\: Active.*"))
 							{
 								LastHeartbeatTime = DateTime.Now;
 								LastActiveHeartbeatTime = DateTime.Now;
 							}
-							else if (Regex.IsMatch(TestLines[i], @".*GauntletHeartbeat\: Idle.*"))
+							else if (Regex.IsMatch(TestLine, @".*GauntletHeartbeat\: Idle.*"))
 							{
 								LastHeartbeatTime = DateTime.Now;
 							}
 						}
 					}
 
-					LastLogCount = TestLines.Count();
 				}
 			}
 
diff --git a/Engine/Source/Programs/AutomationTool/Gauntlet/Unreal/Utils/Gauntlet.UnrealLogParser.cs b/Engine/Source/Programs/AutomationTool/Gauntlet/Unreal/Utils/Gauntlet.UnrealLogParser.cs
--- a/Engine/Source/Programs/AutomationTool/Gauntlet/Unreal/Utils/Gauntlet.UnrealLogParser.cs
+++ b/Engine/Source/Programs/AutomationTool/Gauntlet/Unreal/Utils/Gauntlet.UnrealLogParser.cs
@@ -2,18 +2,12 @@
 
 using System;
 using System.Collections.Generic;
-using System.IO;
-using AutomationTool;
-using UnrealBuildTool;
-using System.Threading;
 using System.Text.RegularExpressions;
-using System.Drawing;
 using System.Linq;
 using System.Text;
 using EpicGames.Core;
-using System.Text.Json;
-using System.Runtime.CompilerServices;
 using Microsoft.Extensions.Logging;
+using Logging = Microsoft.Extensions.Logging;
 
 namespace Gauntlet
 {
@@ -37,6 +31,11 @@ namespace Gauntlet
 			Fatal
 		}
 
+		/// <summary>
+		/// Set of log channels that are used to monitor Editor processing
+		/// </summary>
+		public static string[] EditorBusyChannels = new string[] { "Automation", "FunctionalTest", "Material", "DerivedDataCache", "ShaderCompilers", "Texture", "SkeletalMesh", "StaticMesh", "Python" };
+
 		/// <summary>
 		/// Represents an entry in an Unreal logfile with and contails the associated category, level, and message
 		/// </summary>
@@ -209,6 +208,350 @@ namespace Gauntlet
 		}
 	}
 
+	/// <summary>
+	/// Parse Unreal log from string chunk and aggregate lines as LogEvents
+	/// Support structure logging output and legacy logging style
+	/// </summary>
+	public class UnrealLogStreamParser
+	{
+
+		static class LogEventProperties
+		{
+			public static readonly string Channel = "_channel";
+			public static readonly string Severity = "_severity";
+		}
+
+		protected List<LogEvent> LogEvents { get; private set; }
+
+		private List<string> UnidentifiedLogLevels { get; set; }
+
+		public UnrealLogStreamParser()
+		{
+			LogEvents = new List<LogEvent>();
+			UnidentifiedLogLevels = new List<string>();
+		}
+
+		/// <summary>
+		/// Clear aggregated log events 
+		/// </summary>
+		public void Clear()
+		{
+			LogEvents.Clear();
+		}
+
+		/// <summary>
+		/// Parse a string as log and aggregate identified unreal log lines
+		/// </summary>
+		/// <param name="InContent"></param>
+		/// <param name="LineOffset">Line offset to start parsing and aggregate</param>
+		/// <returns>The number of line parsed</returns>
+		public int ReadStream(string InContent, int LineOffset = 0)
+		{
+			int Cursor = 0;
+
+			for (int BaseIdx = 0; BaseIdx < InContent.Length;)
+			{
+				// Extract the next line
+				int EndIdx = InContent.IndexOf('\n', BaseIdx);
+				if (EndIdx == -1)
+				{
+					break;
+				}
+
+				// Position Cursor
+				if (Cursor < LineOffset)
+				{
+					// Skip to next line
+					Cursor++;
+					BaseIdx = EndIdx + 1;
+					continue;
+				}
+
+				// Skip over any windows CR-LF line endings
+				int LineEndIdx = EndIdx;
+				if (LineEndIdx > BaseIdx && InContent[LineEndIdx - 1] == '\r')
+				{
+					LineEndIdx--;
+				}
+
+				// Grab a line
+				string Line = InContent.Substring(BaseIdx, LineEndIdx - BaseIdx);
+				// Move to the next line
+				Cursor++;
+				BaseIdx = EndIdx + 1;
+
+				// Parse the line
+				if (Line.Length > 0 && Line[0] == '{')
+				{
+
+					try
+					{
+						byte[] Buffer = Encoding.UTF8.GetBytes(Line);
+						LogEvent LogEntry = LogEvent.Read(JsonLogEvent.Parse(Buffer).Data.Span);
+						LogEvents.Add(LogEntry);
+						continue;
+					}
+					catch { /* Ignore the exception and handle the line like a regular log line then */ }
+				}
+
+				// Parse the line as Unreal legacy line
+				Match MatchLine = Regex.Match(Line, @"(?<channel>[A-Za-z][\w\d]+):\s(?:(?<level>Display|Verbose|VeryVerbose|Warning|Error|Fatal):\s)?(?<message>.*)");
+
+				if (MatchLine.Success)
+				{
+					string Channel = MatchLine.Groups["channel"].ToString();
+					string LevelStr = MatchLine.Groups["level"].ToString();
+					string Message = MatchLine.Groups["message"].ToString();
+
+					Logging.LogLevel Level = Logging.LogLevel.Information;
+					switch (LevelStr)
+					{
+						case "Display":
+							Level = Logging.LogLevel.Information;
+							break;
+
+						case "Verbose":
+						case "VeryVerbose":
+							Level = Logging.LogLevel.Debug;
+							break;
+
+						default:
+							if (!string.IsNullOrEmpty(LevelStr))
+							{
+								if (!Enum.TryParse(LevelStr, out Level))
+								{
+									// only show a warning once
+									if (!UnidentifiedLogLevels.Contains(LevelStr))
+									{
+										UnidentifiedLogLevels.Add(LevelStr);
+										Log.Warning("Failed to match log level {0} to enum!", LevelStr);
+									}
+								}
+							}
+							break;
+					}
+
+					Dictionary<string, object> Properties = new Dictionary<string, object>();
+					string Format = Message;
+					if (!string.IsNullOrEmpty(LevelStr))
+					{
+						Properties.Add(LogEventProperties.Severity, new LogValue(LogValueType.Severity, LevelStr));
+						Format = $"{{{LogEventProperties.Severity}}}: {Format}";
+					}
+					Properties.Add(LogEventProperties.Channel, new LogValue(LogValueType.Channel, Channel));
+					Format = $"{{{LogEventProperties.Channel}}}: {Format}";
+
+					LogEvents.Add(new LogEvent(new DateTime(0), Level, KnownLogEvents.Engine_LogChannel, 0, 1, MatchLine.Groups[0].ToString(), Format, Properties, null));
+				}
+				else
+				{
+					// Not an Unreal Engine log line
+					LogEvents.Add(new LogEvent(new DateTime(0), Logging.LogLevel.Information, new EventId(0), 0, 1, Line, Line, null, null));
+				}
+			}
+			return Cursor - LineOffset;
+		}
+
+		/// <summary>
+		/// Return All the LogEvent instances
+		/// </summary>
+		/// <returns></returns>
+		public IEnumerable<LogEvent> GetEvents()
+		{
+			return LogEvents;
+		}
+
+		/// <summary>
+		/// Return the LogEvent instances which property type and value match requirements
+		/// </summary>
+		/// <param name="PropertyName">The property name type to match the value</param>
+		/// <param name="InValues">The property values to match</param>
+		/// <param name="ExactMatch">Whether to use an exact match or partial match</param>
+		/// <returns></returns>
+		public IEnumerable<LogEvent> GetEventsFromProperty(string PropertyName, IEnumerable<string> InValues, bool ExactMatch = false)
+		{
+			IEnumerable<LogEvent> Events;
+
+			if (ExactMatch)
+			{
+				Events = LogEvents.Where(E =>
+				{
+					LogValue PropertyValue = null;
+					if (E.TryGetProperty(PropertyName, out PropertyValue))
+					{
+						return InValues.Contains(PropertyValue.Text, StringComparer.OrdinalIgnoreCase);
+					}
+
+					return false;
+
+				});
+			}
+			else
+			{
+				// partial match
+				Events = LogEvents.Where(E =>
+				{
+					LogValue PropertyValue = null;
+					if (!E.TryGetProperty(PropertyName, out PropertyValue))
+					{
+						return false;
+					}
+
+					foreach (string Value in InValues)
+					{
+						if (PropertyValue.Text.IndexOf(Value, StringComparison.OrdinalIgnoreCase) >= 0)
+						{
+							return true;
+						}
+					}
+
+					return false;
+				});
+			}
+			return Events;
+		}
+
+		/// <summary>
+		/// Return the LogEvent instances which property value meet predicate
+		/// </summary>
+		/// <param name="PropertyName"></param>
+		/// <param name="InPredicate"></param>
+		/// <returns></returns>
+		public IEnumerable<LogEvent> GetEventsFromProperty(string PropertyName, Predicate<string> InPredicate)
+		{
+			IEnumerable<LogEvent> Events;
+
+			Events = LogEvents.Where(E =>
+			{
+				LogValue PropertyValue = null;
+				if (E.TryGetProperty(PropertyName, out PropertyValue))
+				{
+					return InPredicate(PropertyValue.Text);
+				}
+
+				return false;
+			});
+
+			return Events;
+		}
+
+		/// <summary>
+		/// Return the LogEvent instances that match the channel names
+		/// </summary>
+		/// <param name="Channels">The names of the channel to match</param>
+		/// <param name="ExactMatch"></param>
+		/// <returns></returns>
+		public IEnumerable<LogEvent> GetEventsFromChannels(IEnumerable<string> Channels, bool ExactMatch = true)
+		{
+			return GetEventsFromProperty(LogEventProperties.Channel, Channels, ExactMatch);
+		}
+
+		/// <summary>
+		/// Return the LogEvent instances which Channel property meet predicate
+		/// </summary>
+		/// <param name="InPredicate"></param>
+		/// <returns></returns>
+		public IEnumerable<LogEvent> GetEventsFromChannels(Predicate<string> InPredicate)
+		{
+			return GetEventsFromProperty(LogEventProperties.Channel, InPredicate);
+		}
+
+		/// <summary>
+		/// Return the LogEvent instances that match the Editor busy channels
+		/// </summary>
+		/// <returns></returns>
+		public IEnumerable<LogEvent> GetEventsFromEditorBusyChannels()
+		{
+			return GetEventsFromChannels(UnrealLog.EditorBusyChannels, false);
+		}
+
+		/// <summary>
+		/// Return the log lines that match the channel names
+		/// </summary>
+		/// <param name="Channels">The names of the channel to match</param>
+		/// <param name="ExactMatch"></param>
+		/// <returns></returns>
+		public IEnumerable<string> GetLogFromChannels(IEnumerable<string> Channels, bool ExactMatch = true)
+		{
+			return GetEventsFromChannels(Channels, ExactMatch).Select(E => E.Message);
+		}
+
+		/// <summary>
+		/// Return the log lines that match the channel names ignoring the "log" prefix
+		/// </summary>
+		/// <param name="Channels"></param>
+		/// <returns></returns>
+		public IEnumerable<string> GetLogFromShortNameChannels(IEnumerable<string> Channels)
+		{
+			return GetEventsFromChannels(
+				V => Channels.Contains(
+					V.StartsWith("log", StringComparison.OrdinalIgnoreCase)? V.Substring(3) : V,
+					StringComparer.OrdinalIgnoreCase
+				)
+			).Select(E => E.Message);
+		}
+
+		/// <summary>
+		/// Return the log lines that match the Editor busy channels
+		/// </summary>
+		/// <returns></returns>
+		public IEnumerable<string> GetLogFromEditorBusyChannels()
+		{
+			return GetLogFromChannels(UnrealLog.EditorBusyChannels, false);
+		}
+
+		/// <summary>
+		/// Return the log lines that match the channel name
+		/// </summary>
+		/// <param name="Channel">The channel name to match</param>
+		/// <param name="ExactMatch"></param>
+		/// <returns></returns>
+		public IEnumerable<string> GetLogFromChannel(string Channel, bool ExactMatch = true)
+		{
+			return GetLogFromChannels(new string[] { Channel }, ExactMatch);
+		}
+
+		/// <summary>
+		/// Return the log lines that match the severity name
+		/// </summary>
+		/// <param name="Severity">The severity string to match</param>
+		/// <returns></returns>
+		public IEnumerable<string> GetLogFromSeverity(string Severity)
+		{
+			return GetEventsFromProperty(LogEventProperties.Severity, new string[] { Severity }, true).Select(E => E.Message);
+		}
+
+		/// <summary>
+		/// Return all the aggregated log lines 
+		/// </summary>
+		/// <returns></returns>
+		public IEnumerable<string> GetLogLines()
+		{
+			return LogEvents.Select(E => E.Message);
+		}
+
+		/// <summary>
+		/// Return the log lines that contain a string
+		/// </summary>
+		/// <param name="Text">The text to match in the line</param>
+		/// <returns></returns>
+		public IEnumerable<string> GetLogLinesContaining(string Text)
+		{
+			return GetLogLines().Where(L => L.IndexOf(Text, StringComparison.OrdinalIgnoreCase) >= 0);
+		}
+
+		/// <summary>
+		/// Return the log lines that match the regex pattern
+		/// </summary>
+		/// <param name="Pattern">The regex pattern to match in the line</param>
+		/// <returns></returns>
+		public IEnumerable<string> GetLogLinesMatchingPattern(string Pattern)
+		{
+			Regex RegexPattern = new Regex(Pattern, RegexOptions.IgnoreCase);
+			return GetLogLines().Where(L => RegexPattern.IsMatch(L));
+		}
+	}
+
 	/// <summary>
 	/// Helper class for parsing logs
 	/// </summary>
@@ -539,7 +882,7 @@ namespace Gauntlet
 		/// <returns></returns>
 		public IEnumerable<string> GetEditorBusyChannels()
 		{
-			return GetLogChannels(new string[] { "Automation", "FunctionalTest", "Material", "DerivedDataCache", "ShaderCompilers", "Texture", "SkeletalMesh", "StaticMesh", "Python" }, false);
+			return GetLogChannels(UnrealLog.EditorBusyChannels, false);
 		}
 
 		/// <summary>
