Disable report for Horde on UE.InstallOnly dummy test

diff --git a/Engine/Source/Programs/AutomationTool/Gauntlet/Unreal/Automation/UE.InstallOnly.cs b/Engine/Source/Programs/AutomationTool/Gauntlet/Unreal/Automation/UE.InstallOnly.cs
--- a/Engine/Source/Programs/AutomationTool/Gauntlet/Unreal/Automation/UE.InstallOnly.cs
+++ b/Engine/Source/Programs/AutomationTool/Gauntlet/Unreal/Automation/UE.InstallOnly.cs
@@ -40,6 +40,9 @@ namespace UE
 				Config.RequireRole(UnrealTargetRole.Server).InstallOnly = true;
 			}
 
+			// Force disable test results for Horde as there is not actual test running. It will always pass.
+			Config.WriteTestResultsForHorde = false;
+
 			return Config;
 		}
 
