Fixed move task so it respects ErrorIfNotFound option.

diff --git a/Engine/Source/Programs/AutomationTool/BuildGraph/Tasks/MoveTask.cs b/Engine/Source/Programs/AutomationTool/BuildGraph/Tasks/MoveTask.cs
--- a/Engine/Source/Programs/AutomationTool/BuildGraph/Tasks/MoveTask.cs
+++ b/Engine/Source/Programs/AutomationTool/BuildGraph/Tasks/MoveTask.cs
@@ -102,33 +102,57 @@ namespace AutomationTool.Tasks
 			}
 
 			// Build the file mapping
-			Dictionary<FileReference, FileReference> TargetFileToSourceFile = FilePattern.CreateMapping(Files, ref SourcePattern, ref TargetPattern);
+			Dictionary<FileReference, FileReference> TargetFileToSourceFile;
 
-			// Check we got some files
-			if(TargetFileToSourceFile.Count == 0)
+			try
 			{
-				if (Parameters.ErrorIfNotFound)
-				{
-					Logger.LogError("No files found matching '{SourcePattern}'", SourcePattern);
-				}
-				else
+				TargetFileToSourceFile = FilePattern.CreateMapping(Files, ref SourcePattern, ref TargetPattern);
+
+				// Check we got some files
+				if (TargetFileToSourceFile.Count == 0)
 				{
-					Logger.LogInformation("No files found matching '{SourcePattern}'", SourcePattern);
+					if (Parameters.ErrorIfNotFound)
+					{
+						Logger.LogError("No files found matching '{SourcePattern}'", SourcePattern);
+					}
+					else
+					{
+						Logger.LogInformation("No files found matching '{SourcePattern}'", SourcePattern);
+					}
+					return Task.CompletedTask;
 				}
-				return Task.CompletedTask;
-			}
+				// Copy them all
+				Logger.LogInformation("Moving {Arg0} file{Arg1} from {Arg2} to {Arg3}...", TargetFileToSourceFile.Count, (TargetFileToSourceFile.Count == 1) ? "" : "s", SourcePattern.BaseDirectory, TargetPattern.BaseDirectory);
+				CommandUtils.ParallelMoveFiles(TargetFileToSourceFile.Select(x => new KeyValuePair<FileReference, FileReference>(x.Value, x.Key)), Parameters.Overwrite);
 
-			// Copy them all
-			Logger.LogInformation("Moving {Arg0} file{Arg1} from {Arg2} to {Arg3}...", TargetFileToSourceFile.Count, (TargetFileToSourceFile.Count == 1)? "" : "s", SourcePattern.BaseDirectory, TargetPattern.BaseDirectory);
-			CommandUtils.ParallelMoveFiles(TargetFileToSourceFile.Select(x => new KeyValuePair<FileReference, FileReference>(x.Value, x.Key)), Parameters.Overwrite);
+				// Update the list of build products
+				BuildProducts.UnionWith(TargetFileToSourceFile.Keys);
 
-			// Update the list of build products
-			BuildProducts.UnionWith(TargetFileToSourceFile.Keys);
+				// Apply the optional output tag to them
+				foreach (string TagName in FindTagNamesFromList(Parameters.Tag))
+				{
+					FindOrAddTagSet(TagNameToFileSet, TagName).UnionWith(TargetFileToSourceFile.Keys);
+				}
 
-			// Apply the optional output tag to them
-			foreach(string TagName in FindTagNamesFromList(Parameters.Tag))
+			}
+			catch (FilePatternException Ex)
 			{
-				FindOrAddTagSet(TagNameToFileSet, TagName).UnionWith(TargetFileToSourceFile.Keys);
+				if (Ex.GetFilePatternExceptionType() == FilePatternException.FilePatternExceptionType.SourceFileDoesNotExist)
+				{
+					if (Parameters.ErrorIfNotFound)
+					{
+						Logger.LogError("Error while trying to create file pattern match for '{SourcePattern}', error {ExceptionString}", SourcePattern, Ex.ToString());
+					}
+					else
+					{
+						Logger.LogInformation("Error while trying to create file pattern match for '{SourcePattern}', error {ExceptionString}", SourcePattern, Ex.ToString());
+					}
+					return Task.CompletedTask;
+				}
+				else
+				{
+					throw Ex;
+				}
 			}
 
 			return Task.CompletedTask;
diff --git a/Engine/Source/Programs/Shared/EpicGames.Core/FilePattern.cs b/Engine/Source/Programs/Shared/EpicGames.Core/FilePattern.cs
--- a/Engine/Source/Programs/Shared/EpicGames.Core/FilePattern.cs
+++ b/Engine/Source/Programs/Shared/EpicGames.Core/FilePattern.cs
@@ -21,6 +21,7 @@ namespace EpicGames.Core
 		public FilePatternException(string message)
 			: base(message)
 		{
+			Type = FilePatternExceptionType.Unknown;
 		}
 
 		/// <summary>
@@ -29,13 +30,31 @@ namespace EpicGames.Core
 		public FilePatternException(string format, params object[] args)
 			: base(String.Format(format, args))
 		{
+			Type = FilePatternExceptionType.Unknown;
 		}
 
+		public FilePatternException(FilePatternExceptionType InType, string format, params object[] args)
+			: base(String.Format(format, args))
+		{
+			Type = InType;
+		}
+
+		public FilePatternExceptionType GetFilePatternExceptionType() 
+		{ 
+			return Type; 
+		}
 		/// <inheritdoc/>
 		public override string ToString()
 		{
 			return Message;
 		}
+		public enum FilePatternExceptionType
+		{
+			Unknown,
+			SourceFileDoesNotExist,
+			SourceFileNotUnderBaseDir
+		};
+		FilePatternExceptionType Type;
 	}
 
 	/// <summary>
@@ -390,7 +409,7 @@ namespace EpicGames.Core
 				}
 				else
 				{
-					throw new FilePatternException("Source file '{0}' does not exist", sourceFile);
+					throw new FilePatternException(FilePatternException.FilePatternExceptionType.SourceFileDoesNotExist, "Source file '{0}' does not exist", sourceFile);
 				}
 			}
 
@@ -420,11 +439,11 @@ namespace EpicGames.Core
 			{
 				if(!inputFile.IsUnderDirectory(baseDirectory))
 				{
-					throw new FilePatternException("Source file '{0}' is not under '{1}'", inputFile, baseDirectory);
+					throw new FilePatternException(FilePatternException.FilePatternExceptionType.SourceFileNotUnderBaseDir, "Source file '{0}' is not under '{1}'", inputFile, baseDirectory);
 				}
 				else if(!FileReference.Exists(inputFile))
 				{
-					throw new FilePatternException("Source file '{0}' does not exist", inputFile);
+					throw new FilePatternException(FilePatternException.FilePatternExceptionType.SourceFileDoesNotExist, "Source file '{0}' does not exist", inputFile);
 				}
 				else
 				{
