BuildGraph: Allow #gh-light-mode-only and #gh-dark-mode-only tags on images in Markdown files.

diff --git a/Engine/Source/Programs/AutomationTool/BuildGraph/Tasks/CheckMarkdownTask.cs b/Engine/Source/Programs/AutomationTool/BuildGraph/Tasks/CheckMarkdownTask.cs
--- a/Engine/Source/Programs/AutomationTool/BuildGraph/Tasks/CheckMarkdownTask.cs
+++ b/Engine/Source/Programs/AutomationTool/BuildGraph/Tasks/CheckMarkdownTask.cs
@@ -74,6 +74,16 @@ namespace AutomationTool.Tasks
 				}
 			}
 
+			// Add anchors for any images
+			foreach (FileReference file in files)
+			{
+				if (file.HasExtension(".png") || file.HasExtension(".jpg") || file.HasExtension(".jpeg") || file.HasExtension(".gif"))
+				{
+					validLinks.Add($"{file.FullName}#gh-dark-mode-only");
+					validLinks.Add($"{file.FullName}#gh-light-mode-only");
+				}
+			}
+
 			// Check the links 
 			foreach (FileReference file in markdownFiles)
 			{
