Lower log severy on failing attempt to install application on device.

#rnx

diff --git a/Engine/Source/Programs/AutomationTool/Gauntlet/Framework/Devices/Gauntlet.DevicePool.cs b/Engine/Source/Programs/AutomationTool/Gauntlet/Framework/Devices/Gauntlet.DevicePool.cs
--- a/Engine/Source/Programs/AutomationTool/Gauntlet/Framework/Devices/Gauntlet.DevicePool.cs
+++ b/Engine/Source/Programs/AutomationTool/Gauntlet/Framework/Devices/Gauntlet.DevicePool.cs
@@ -1337,7 +1337,7 @@ namespace Gauntlet
 			}
 			catch (Exception Ex)
 			{
-				Log.Warning(KnownLogEvents.Gauntlet_DeviceEvent, "Failed to create device {Name}. {Message}\n{StackTrace}", Def.ToString(), Ex.Message, Ex.StackTrace);
+				Log.Info(KnownLogEvents.Gauntlet_DeviceEvent, "Failed to create device {Name}. {Message}\n{StackTrace}", Def.ToString(), Ex.Message, Ex.StackTrace);
 			}
 
 			return NewDevice;
@@ -1402,7 +1402,7 @@ namespace Gauntlet
 		{
 			if (MaxDeviceErrorReports == 0)
 			{
-				Log.Warning(KnownLogEvents.Gauntlet_DeviceEvent, "Maximum device errors reported to backend, {Name} : {Message} ignored", ServiceDeviceName, ErrorMessage);
+				Log.Info(KnownLogEvents.Gauntlet_DeviceEvent, "Maximum device errors reported to backend, {Name} : {Message} ignored", ServiceDeviceName, ErrorMessage);
 				return;
 			}
 
