- Split up the UAT BuildCommand into 2 agendas, one for project targets, one for shared, so per-project SDKs won't cause conflicts with shared targets

diff --git a/Engine/Source/Programs/AutomationTool/Scripts/BuildProjectCommand.Automation.cs b/Engine/Source/Programs/AutomationTool/Scripts/BuildProjectCommand.Automation.cs
--- a/Engine/Source/Programs/AutomationTool/Scripts/BuildProjectCommand.Automation.cs
+++ b/Engine/Source/Programs/AutomationTool/Scripts/BuildProjectCommand.Automation.cs
@@ -74,7 +74,9 @@ namespace AutomationScripts
 			var StartTime = DateTime.UtcNow;
 
 			var UnrealBuild = new UnrealBuild(Command);
-			var Agenda = new UnrealBuild.BuildAgenda();
+			var NoProjectAgenda = new UnrealBuild.BuildAgenda();
+			// if there's a project, make a new Agenda. we will build one or both at the end
+			var ProjectAgenda = Params.CodeBasedUprojectPath == null ? NoProjectAgenda : new UnrealBuild.BuildAgenda();
 			var CrashReportPlatforms = new HashSet<UnrealTargetPlatform>();
 
 
@@ -91,18 +93,18 @@ namespace AutomationScripts
 				UnrealTargetPlatform EditorPlatform = HostPlatform.Current.HostEditorPlatform;
 				const UnrealTargetConfiguration EditorConfiguration = UnrealTargetConfiguration.Development;
 
-				Agenda.AddTargets(Params.EditorTargets.ToArray(), EditorPlatform, EditorConfiguration, Params.CodeBasedUprojectPath, InAddArgs: EditorBuildArgs);
+				ProjectAgenda.AddTargets(Params.EditorTargets.ToArray(), EditorPlatform, EditorConfiguration, Params.CodeBasedUprojectPath, InAddArgs: EditorBuildArgs);
 
 				if (!Unreal.IsEngineInstalled())
 				{
 					CrashReportPlatforms.Add(EditorPlatform);
 					if (Params.EditorTargets.Contains("ShaderCompileWorker") == false)
 					{
-						Agenda.AddTargets(new string[] { "ShaderCompileWorker" }, EditorPlatform, EditorConfiguration, InAddArgs: ProgramBuildArgs);
+						ProjectAgenda.AddTargets(new string[] { "ShaderCompileWorker" }, EditorPlatform, EditorConfiguration, Params.CodeBasedUprojectPath, InAddArgs: ProgramBuildArgs);
 					}
 					if (Params.FileServer && Params.EditorTargets.Contains("UnrealFileServer") == false)
 					{
-						Agenda.AddTargets(new string[] { "UnrealFileServer" }, EditorPlatform, EditorConfiguration, InAddArgs: ProgramBuildArgs);
+						NoProjectAgenda.AddTargets(new string[] { "UnrealFileServer" }, EditorPlatform, EditorConfiguration, InAddArgs: ProgramBuildArgs);
 					}
 				}
 			}
@@ -112,7 +114,7 @@ namespace AutomationScripts
 			{
 				if (!Params.HasEditorTargets || Params.EditorTargets.Contains("UnrealPak") == false)
 				{
-					Agenda.AddTargets(new string[] { "UnrealPak" }, HostPlatform.Current.HostEditorPlatform, UnrealTargetConfiguration.Development, Params.CodeBasedUprojectPath, InAddArgs: ProgramBuildArgs);
+					ProjectAgenda.AddTargets(new string[] { "UnrealPak" }, HostPlatform.Current.HostEditorPlatform, UnrealTargetConfiguration.Development, Params.CodeBasedUprojectPath, InAddArgs: ProgramBuildArgs);
 				}
 			}
 
@@ -161,7 +163,7 @@ namespace AutomationScripts
 						UnrealTargetPlatform CrashReportPlatform = Platform.GetPlatform(ClientPlatformType).CrashReportPlatform ?? ClientPlatformType;
 						CrashReportPlatforms.Add(CrashReportPlatform);
 						string Arch = Params.IsProgramTarget ? ProgramBuildArgs : ClientBuildArgs;
-						Agenda.AddTargets(Params.ClientCookedTargets.ToArray(), ClientPlatformType, BuildConfig, Params.CodeBasedUprojectPath, 
+						ProjectAgenda.AddTargets(Params.ClientCookedTargets.ToArray(), ClientPlatformType, BuildConfig, Params.CodeBasedUprojectPath, 
 							InAddArgs: $" -remoteini=\"{Params.RawProjectPath.Directory}\" {AdditionalArgs} {Arch}");
 					}
 				}
@@ -176,7 +178,7 @@ namespace AutomationScripts
 					{
 						UnrealTargetPlatform CrashReportPlatform = Platform.GetPlatform(ServerPlatformType).CrashReportPlatform ?? ServerPlatformType;
 						CrashReportPlatforms.Add(CrashReportPlatform);
-						Agenda.AddTargets(Params.ServerCookedTargets.ToArray(), ServerPlatformType, BuildConfig, Params.CodeBasedUprojectPath, 
+						ProjectAgenda.AddTargets(Params.ServerCookedTargets.ToArray(), ServerPlatformType, BuildConfig, Params.CodeBasedUprojectPath, 
 							InAddArgs: $" -remoteini=\"{Params.RawProjectPath.Directory}\" {AdditionalArgs} {ServerBuildArgs}");
 					}
 				}
@@ -188,7 +190,7 @@ namespace AutomationScripts
 				{
 					if (Params.ClientTargetPlatforms.Contains(new TargetPlatformDescriptor(BootstrapPackagedGamePlatformType)))
 					{
-						Agenda.AddTarget("BootstrapPackagedGame", BootstrapPackagedGamePlatformType, UnrealBuildTool.UnrealTargetConfiguration.Shipping, InAddArgs: ClientBuildArgs);
+						NoProjectAgenda.AddTarget("BootstrapPackagedGame", BootstrapPackagedGamePlatformType, UnrealBuildTool.UnrealTargetConfiguration.Shipping, InAddArgs: ClientBuildArgs);
 					}
 				}
 			}
@@ -198,7 +200,7 @@ namespace AutomationScripts
 				{
 					if (PlatformSupportsCrashReporter(CrashReportPlatform))
 					{
-						Agenda.AddTarget("CrashReportClient", CrashReportPlatform, UnrealTargetConfiguration.Shipping, 
+						NoProjectAgenda.AddTarget("CrashReportClient", CrashReportPlatform, UnrealTargetConfiguration.Shipping, 
 							InAddArgs: $" -remoteini=\"{Params.RawProjectPath.Directory}\" {ProgramBuildArgs}");
 					}
 				}
@@ -211,7 +213,7 @@ namespace AutomationScripts
 				{
 					foreach (var ClientPlatformType in UniquePlatformTypes)
 					{
-						Agenda.AddTargets(Params.ProgramTargets.ToArray(), ClientPlatformType, BuildConfig, Params.CodeBasedUprojectPath, ProgramBuildArgs);
+						ProjectAgenda.AddTargets(Params.ProgramTargets.ToArray(), ClientPlatformType, BuildConfig, Params.CodeBasedUprojectPath, ProgramBuildArgs);
 					}
 				}
 			}
@@ -222,10 +224,20 @@ namespace AutomationScripts
 			UniquePlatforms.UnionWith(Params.ServerTargetPlatforms.Select(x => x.Type));
 			foreach (UnrealTargetPlatform TargetPlatform in UniquePlatforms)
 			{
-				Platform.GetPlatform(TargetPlatform).PreBuildAgenda(UnrealBuild, Agenda, Params);
+				// @todo: pass in both project and no project
+				Platform.GetPlatform(TargetPlatform).PreBuildAgenda(UnrealBuild, ProjectAgenda, Params);
 			}
 
-			UnrealBuild.Build(Agenda, InDeleteBuildProducts: Params.Clean, InUpdateVersionFiles: WorkingCL > 0);
+			UnrealBuild.Build(ProjectAgenda, InDeleteBuildProducts: Params.Clean, InUpdateVersionFiles: WorkingCL > 0);
+			// if we had two separate agendas, build the other one now
+			if (NoProjectAgenda != ProjectAgenda)
+			{
+				foreach (UnrealTargetPlatform TargetPlatform in UniquePlatforms)
+				{
+					Platform.GetPlatform(TargetPlatform).PreBuildAgenda(UnrealBuild, NoProjectAgenda, Params);
+				}
+				UnrealBuild.Build(NoProjectAgenda, InDeleteBuildProducts: Params.Clean, InUpdateVersionFiles: WorkingCL > 0);
+			}
 
 			if (WorkingCL > 0) // only move UAT files if we intend to check in some build products
 			{
