- Added a shared GetEditorForProject() function which should be used in UAT to get the path to the binary to run, becasue UnrealEditor-Cmd.exe is not always correct (if the editor is a Unique build environment, the editor to use to cook etc is the project editor. not UnrealEditor)

diff --git a/Engine/Source/Programs/AutomationTool/AutomationUtils/ProjectParams.cs b/Engine/Source/Programs/AutomationTool/AutomationUtils/ProjectParams.cs
--- a/Engine/Source/Programs/AutomationTool/AutomationUtils/ProjectParams.cs
+++ b/Engine/Source/Programs/AutomationTool/AutomationUtils/ProjectParams.cs
@@ -1599,17 +1599,7 @@ namespace AutomationTool
 			{
 				if (SpecifiedUnrealExe == null)
 				{
-					SpecifiedUnrealExe = "UnrealEditor-Cmd.exe";
-					if (CodeBasedUprojectPath != null)
-					{
-						FileReference ReceiptLocation = TargetReceipt.GetDefaultPath(CodeBasedUprojectPath.Directory, EditorTargets[0], HostPlatform.Platform, UnrealTargetConfiguration.Development, null);
-						TargetReceipt Receipt;
-						if (!TargetReceipt.TryRead(ReceiptLocation, out Receipt))
-						{
-							throw new AutomationException($"Missing {ReceiptLocation} receipt. Editor needs to be built first.");
-						}
-						SpecifiedUnrealExe = Receipt.LaunchCmd.FullName;
-					}
+					SpecifiedUnrealExe = ProjectUtils.GetEditorForProject(RawProjectPath).FullName;
 				}
 
 				return SpecifiedUnrealExe;
diff --git a/Engine/Source/Programs/AutomationTool/AutomationUtils/ProjectUtils.cs b/Engine/Source/Programs/AutomationTool/AutomationUtils/ProjectUtils.cs
--- a/Engine/Source/Programs/AutomationTool/AutomationUtils/ProjectUtils.cs
+++ b/Engine/Source/Programs/AutomationTool/AutomationUtils/ProjectUtils.cs
@@ -687,6 +687,24 @@ namespace AutomationTool
 			
 			return TargetReceipt.Launch;
 		}
+
+		// Helper function to get the editor executable to run for cooking, etc the given project
+		public static FileSystemReference GetEditorForProject(FileReference ProjectFile)
+		{
+			FileSystemReference EditorExe = null;
+			if (ProjectFile != null)
+			{
+				EditorExe = ProjectUtils.GetProjectTarget(ProjectFile, UnrealBuildTool.TargetType.Editor, BuildHostPlatform.Current.Platform, UnrealTargetConfiguration.Development, true);
+			}
+
+			// if we haven't gotten an editor yet, use the engine shared binary
+			if (EditorExe == null)
+			{
+				EditorExe = new FileReference(HostPlatform.Current.GetUnrealExePath("UnrealEditor-Cmd.exe"));
+			}
+
+			return EditorExe;
+		}
 	}
 
     public class BranchInfo
diff --git a/Engine/Source/Programs/AutomationTool/BuildGraph/Tasks/CommandletTask.cs b/Engine/Source/Programs/AutomationTool/BuildGraph/Tasks/CommandletTask.cs
--- a/Engine/Source/Programs/AutomationTool/BuildGraph/Tasks/CommandletTask.cs
+++ b/Engine/Source/Programs/AutomationTool/BuildGraph/Tasks/CommandletTask.cs
@@ -98,18 +98,10 @@ namespace AutomationTool.Tasks
 			}
 
 			// Get the path to the editor, and check it exists
-			FileSystemReference EditorExe;
-			if(Parameters.EditorExe == null)
+			FileSystemReference EditorExe = Parameters.EditorExe;
+			if (EditorExe == null)
 			{
-				EditorExe = ProjectUtils.GetProjectTarget(ProjectFile, UnrealBuildTool.TargetType.Editor, BuildHostPlatform.Current.Platform, UnrealTargetConfiguration.Development, true);
-				if (EditorExe == null)
-				{
-					EditorExe = new FileReference(HostPlatform.Current.GetUnrealExePath("UnrealEditor-Cmd.exe"));
-				}
-			}
-			else
-			{
-				EditorExe = Parameters.EditorExe;
+				EditorExe = ProjectUtils.GetEditorForProject(ProjectFile);
 			}
 
 			// Run the commandlet
diff --git a/Engine/Source/Programs/AutomationTool/BuildGraph/Tasks/CookTask.cs b/Engine/Source/Programs/AutomationTool/BuildGraph/Tasks/CookTask.cs
--- a/Engine/Source/Programs/AutomationTool/BuildGraph/Tasks/CookTask.cs
+++ b/Engine/Source/Programs/AutomationTool/BuildGraph/Tasks/CookTask.cs
@@ -117,7 +117,7 @@ namespace AutomationTool.Tasks
 				Scope.Span.SetTag("platform", Parameters.Platform);
 				string[] Maps = (Parameters.Maps == null)? null : Parameters.Maps.Split(new char[]{ '+' });
 				string Arguments = (Parameters.Versioned ? "" : "-Unversioned ") + "-LogCmds=\"LogSavePackage Warning\" " + Parameters.Arguments;
-				string EditorExe = (string.IsNullOrWhiteSpace(Parameters.EditorExe) ? "UnrealEditor-Cmd.exe" : Parameters.EditorExe);
+				string EditorExe = (string.IsNullOrWhiteSpace(Parameters.EditorExe) ? ProjectUtils.GetEditorForProject(ProjectFile).FullName : Parameters.EditorExe);
 				CommandUtils.CookCommandlet(ProjectFile, EditorExe, Maps, null, null, null, Parameters.Platform, Arguments);
 			}
 
