[Gauntlet] Re-order TargetDeviceAndroid members/defintions

diff --git a/Engine/Source/Programs/AutomationTool/Gauntlet/Platform/Android/Gauntlet.TargetDeviceAndroid.cs b/Engine/Source/Programs/AutomationTool/Gauntlet/Platform/Android/Gauntlet.TargetDeviceAndroid.cs
--- a/Engine/Source/Programs/AutomationTool/Gauntlet/Platform/Android/Gauntlet.TargetDeviceAndroid.cs
+++ b/Engine/Source/Programs/AutomationTool/Gauntlet/Platform/Android/Gauntlet.TargetDeviceAndroid.cs
@@ -14,370 +14,60 @@
 
 namespace Gauntlet
 {
-	// device data from json
-	public sealed class AndroidDeviceData
-	{
-		// remote device settings (wifi)
-
-		// host of PC which is tethered 
-		public string hostIP { get; set; }
-
-		// public key 
-		public string publicKey { get; set; }
-
-		// private key
-		public string privateKey { get; set; }
-	}
-
-	// become IAppInstance when implemented enough
-	class AndroidAppInstance : IAppInstance
-	{
-		protected TargetDeviceAndroid AndroidDevice;
-
-		protected AndroidAppInstall Install;
-
-		internal IProcessResult LaunchProcess;
-
-		internal bool bHaveSavedArtifacts;
-
-		public string CommandLine { get { return Install.CommandLine; } }
-
-		public AndroidAppInstance(TargetDeviceAndroid InDevice, AndroidAppInstall InInstall, IProcessResult InProcess)
-		{
-			AndroidDevice = InDevice;
-			Install = InInstall;
-			LaunchProcess = InProcess;
-		}
-
-		public string ArtifactPath
-		{
-			get
-			{
-				if (HasExited)
-				{
-					SaveArtifacts();
-				}
-				
-				return Path.Combine(AndroidDevice.LocalCachePath, "Saved");
-			}
-		}
-
-		public ITargetDevice Device
-		{
-			get
-			{
-				return AndroidDevice;
-			}
-		}
-
-		public bool HasExited
-		{
-			get
-			{
-				try
-				{
-					if (!LaunchProcess.HasExited)
-					{
-						return false;
-					}
-				}
-				catch (System.InvalidOperationException)
-				{
-					return true;
-				}
-
-				return IsActivityRunning();
-			}
-		}
-
-		/// <summary>
-		/// Checks on device whether the activity is running, this is an expensive shell with output operation
-		/// the result is cached, with checks at ActivityCheckDelta seconds
-		/// </summary>
-		private bool IsActivityRunning()
-		{
-			if (ActivityExited)
-			{
-				return true;
-			}
-
-			if ((DateTime.UtcNow - ActivityCheckTime) < ActivityCheckDelta)
-			{
-				return false;
-			}
-
-			if(AndroidDevice.Disposed)
-			{
-				return false;
-			}
-
-			ActivityCheckTime = DateTime.UtcNow;
-
-			// get activities filtered by our package name
-			IProcessResult ActivityQuery = AndroidDevice.RunAdbDeviceCommand("shell dumpsys activity -p " + Install.AndroidPackageName + " a");
-
-			// We have exited if our activity doesn't appear in the activity query or is not the focused activity.
-			bool bActivityPresent = ActivityQuery.Output.Contains(Install.AndroidPackageName);
-			bool bActivityInForeground = ActivityQuery.Output.Contains("mResumedActivity");
-			bool bHasExited = !bActivityPresent || !bActivityInForeground;
-			if (bHasExited)
-			{
-				ActivityExited = true;
-				// The activity has exited, make sure entire activity log has been captured, sleep to allow time for the log to flush
-				Thread.Sleep(5000);
-				UpdateCachedLog(true);
-				Log.VeryVerbose("{0}: process exited, Activity running={1}, Activity in foreground={2} ", ToString(), bActivityPresent.ToString(), bActivityInForeground.ToString());
-			}
-
-			return bHasExited;
-
-		}
-
-		private static readonly TimeSpan ActivityCheckDelta = TimeSpan.FromSeconds(10);
-		private DateTime ActivityCheckTime = DateTime.UtcNow;
-		private bool ActivityExited = false;
-
-		public bool WasKilled { get; protected set; }
-
-		/// <summary>
-		/// The output of the test activity
-		/// </summary>
-		public string StdOut
-		{
-			get
-			{
-				UpdateCachedLog();
-				return String.IsNullOrEmpty(ActivityLogCached) ? String.Empty : ActivityLogCached;
-			}
-		}
-		/// <summary>
-		/// Updates cached activity log by running a shell command returning the full log from device (possibly over wifi)
-		/// The result is cached and updated at ActivityLogDelta frequency
-		/// </summary>
-		private void UpdateCachedLog(bool ForceUpdate = false)
-		{
-			if (!ForceUpdate && (ActivityLogTime == DateTime.MinValue || ((DateTime.UtcNow - ActivityLogTime) < ActivityLogDelta)))
-			{
-				return;
-			}
-
-			if (Install.AndroidDevice != null && Install.AndroidDevice.Disposed)
-			{
-				Log.Warning(KnownLogEvents.Gauntlet_DeviceEvent, "Attempting to cache log using disposed Android device");
-				return;
-			}
-
-			// Note: DeviceLogPath has the correct location for the current app. Apps can be configured to override the default log storage path. 
-			string GetLogCommand = string.Format("shell cat {0}", Install.AndroidDevice.DeviceLogPath);
-			IProcessResult LogQuery = Install.AndroidDevice.RunAdbDeviceCommand(GetLogCommand, true);
-
-			if (LogQuery.ExitCode != 0)
-			{
-				Log.VeryVerbose("Unable to query activity stdout on device {0}", Install.AndroidDevice.Name);
-			}
-			else
-			{
-				ActivityLogCached = LogQuery.Output;
-			}
-
-			ActivityLogTime = DateTime.UtcNow;
-
-			// the activity has exited, mark final log sentinel 
-			if (ActivityExited)
-			{
-				ActivityLogTime = DateTime.MinValue;
-			}
-
-		}
-
-		private static readonly TimeSpan ActivityLogDelta = TimeSpan.FromSeconds(15);
-		private DateTime ActivityLogTime = DateTime.UtcNow - ActivityLogDelta;
-		private string ActivityLogCached = string.Empty;
-		
-
-		public int WaitForExit()
-		{
-			if (!HasExited)
-			{
-				LaunchProcess.WaitForExit();
-			}
-
-			return ExitCode;
-		}
-
-		public void Kill()
-		{
-			if (!HasExited && !AndroidDevice.Disposed)
-			{
-				WasKilled = true;
-				Install.AndroidDevice.KillRunningProcess(Install.AndroidPackageName);
-			}
-		}
-		public int ExitCode { get { return LaunchProcess.ExitCode; } }
-
-		protected void SaveArtifacts()
-		{
-			if (bHaveSavedArtifacts)
-				return;
-
-			// copy remote artifacts to local
-			if (Directory.Exists(Install.AndroidDevice.LocalCachePath))
-			{
-
-				Log.Verbose("Deleting {0}", Install.AndroidDevice.LocalCachePath);
-
-				try
-				{
-					// don't consider this fatal, people often have the directory or a file open
-					Directory.Delete(Install.AndroidDevice.LocalCachePath, true);
-				}
-				catch
-				{
-					Log.Warning(KnownLogEvents.Gauntlet_DeviceEvent, "Failed to remove old cache folder {Folder}", Install.AndroidDevice.LocalCachePath);
-				}
-			}
-
-			// mark it as a temp dir (will also create it)
-			try
-			{
-				Utils.SystemHelpers.MarkDirectoryForCleanup(Install.AndroidDevice.LocalCachePath);
-			}
-			catch (Exception Ex)
-			{
-				Log.Warning(KnownLogEvents.Gauntlet_DeviceEvent, "Exception marking directory for cleanup {Exception}", Ex.Message);
-			}			
-
-			string LocalSaved = Path.Combine(Install.AndroidDevice.LocalCachePath, "Saved");
-			Log.Verbose("Creating {0}", LocalSaved);
-			Directory.CreateDirectory(LocalSaved);
-
-
-			
-
-
-			// pull all the artifacts
-			string ArtifactPullCommand = string.Format("pull {0} \"{1}\"", Install.AndroidDevice.DeviceArtifactPath, Install.AndroidDevice.LocalCachePath);
-			IProcessResult PullCmd = Install.AndroidDevice.RunAdbDeviceCommand(ArtifactPullCommand, bShouldLogCommand: Log.IsVerbose);
-
-			if (PullCmd.ExitCode != 0)
-			{
-				Log.Warning(KnownLogEvents.Gauntlet_DeviceEvent, "Failed to retrieve artifacts. {Output}", PullCmd.Output);
-			}
-			else
-			{
-				// update final cached stdout property
-				string LogFilename = Install.AndroidDevice.DeviceLogPath;
-				if (File.Exists(LogFilename))
-				{
-					ActivityLogCached = File.ReadAllText(LogFilename);
-					ActivityLogTime = DateTime.MinValue;
-				}
-
-			}
-
-			// pull the logcat over from device.
-			IProcessResult LogcatResult = Install.AndroidDevice.RunAdbDeviceCommand("logcat -d", bShouldLogCommand: Log.IsVerbose);
-
-			string LogcatFilename = "Logcat.log";
-			// Save logcat dump to local artifact path.
-			File.WriteAllText(Path.Combine(LocalSaved, LogcatFilename), LogcatResult.Output);
-
-			Install.AndroidDevice.PostRunCleanup();
-
-			bHaveSavedArtifacts = true;
-		}
-	}
-
-	class AndroidAppInstall : IAppInstall
-	{
-		public string Name { get; protected set; }
-
-		public string AndroidPackageName { get; protected set; }
-
-		public TargetDeviceAndroid AndroidDevice { get; protected set; }
-
-		public ITargetDevice Device { get { return AndroidDevice; } }
-
-		public string CommandLine { get; protected set; }
-
-		public IAppInstance Run()
-		{
-			return AndroidDevice.Run(this);
-		}
-
-		public AndroidAppInstall(TargetDeviceAndroid InDevice, string InName, string InAndroidPackageName, string InCommandLine)
-		{
-			AndroidDevice = InDevice;
-			Name = InName;
-			AndroidPackageName = InAndroidPackageName;
-			CommandLine = InCommandLine;
-		}
-	}
-
-	public class DefaultAndroidDevices : IDefaultDeviceSource
-	{
-		public bool CanSupportPlatform(UnrealTargetPlatform? Platform)
-		{
-			return Platform == UnrealTargetPlatform.Android;
-		}
-
-		public ITargetDevice[] GetDefaultDevices()
-		{
-			return TargetDeviceAndroid.GetDefaultDevices();
-		}
-	}
-
-	public class AndroidBuildSupport : BaseBuildSupport
-	{
-		protected override BuildFlags SupportedBuildTypes => BuildFlags.Packaged | BuildFlags.CanReplaceCommandLine | BuildFlags.CanReplaceExecutable | BuildFlags.Bulk | BuildFlags.NotBulk;
-		protected override UnrealTargetPlatform? Platform => UnrealTargetPlatform.Android;
-	}
-
-	public class AndroidDeviceFactory : IDeviceFactory
-	{
-		public bool CanSupportPlatform(UnrealTargetPlatform? Platform)
-		{
-			return Platform == UnrealTargetPlatform.Android;
-		}
-
-		public ITargetDevice CreateDevice(string InRef, string InCachePath, string InParam = null)
-		{
-			AndroidDeviceData DeviceData = null;
-
-			if (!String.IsNullOrEmpty(InParam))
-			{
-				DeviceData = fastJSON.JSON.Instance.ToObject<AndroidDeviceData>(InParam);
-			}
-
-			return new TargetDeviceAndroid(InRef, DeviceData, InCachePath);
-		}
-	}
-
 	/// <summary>
 	/// Android implementation of a device that can run applications
 	/// </summary>
 	public class TargetDeviceAndroid : ITargetDevice
-	{ 
+	{
+		#region ITargetDevice
+		public string Name { get; protected set; }
+
+		public UnrealTargetPlatform? Platform => UnrealTargetPlatform.Android;
+
+		public CommandUtils.ERunOptions RunOptions { get; set; }
+
+		public bool IsConnected => IsAvailable;
+
+		public bool IsOn
+		{
+			get
+			{
+				string CommandLine = "shell dumpsys power";
+				IProcessResult OnAndUnlockedQuery = RunAdbDeviceCommand(CommandLine, bPauseErrorParsing: true);
+				return OnAndUnlockedQuery.Output.Contains("mHoldingDisplaySuspendBlocker=true")
+					&& OnAndUnlockedQuery.Output.Contains("mHoldingWakeLockSuspendBlocker=true");
+			}
+		}
+
+		public bool IsAvailable
+		{
+			get
+			{
+				// ensure our device is present in 'adb devices' output.
+				var AllDevices = GetAllConnectedDevices();
+
+				if (AllDevices.Keys.Contains(DeviceName) == false)
+				{
+					return false;
+				}
+
+				if (AllDevices[DeviceName] != "device")
+				{
+					Log.Warning(KnownLogEvents.Gauntlet_DeviceEvent, "Device {Name} is '{State}'", DeviceName, AllDevices[DeviceName]);
+					return false;
+				}
+
+				// any device will do, but only one at a time.
+				return true;
+			}
+		}
+		#endregion
+
 		/// <summary>
-		/// Friendly name for this target
-		/// </summary>
-		public string Name { get; protected set; }
-	
-		/// <summary>
 		/// Low-level device name
 		/// </summary>
 		public string DeviceName { get; protected set; }
 
 		/// <summary>
-		/// Platform type.
-		/// </summary>
-		public UnrealTargetPlatform? Platform { get { return UnrealTargetPlatform.Android; } }
-
-		/// <summary>
-		/// Options for executing commands
-		/// </summary>
-		public CommandUtils.ERunOptions RunOptions { get; set; }
-
-		/// <summary>
 		/// Temp path we use to push/pull things from the device
 		/// </summary>
@@ -386,7 +76,7 @@
 		/// <summary>
 		/// Artifact (e.g. Saved) path on the device
 		/// </summary>
-		public string DeviceArtifactPath { get; protected set;  }
+		public string DeviceArtifactPath { get; protected set; }
 
 		/// <summary>
 		/// External storage path on the device, e.g. /sdcard/UEGame/ etc..
@@ -401,56 +91,28 @@
 		/// <summary>
 		/// Path to the log file. (This takes public logs setting of the package in to account.)
 		/// </summary>
-		public string DeviceLogPath { get; protected set; }		
+		public string DeviceLogPath { get; protected set; }
 
 		/// <summary>
 		/// Path to a command line if installed
 		/// </summary>
 		protected string CommandLineFilePath { get; set; }
 
-		public bool IsAvailable
+		protected bool IsExistingDevice = false;
+
+		protected Dictionary<EIntendedBaseCopyDirectory, string> LocalDirectoryMappings;
+
+		public void PopulateDirectoryMappings(string ProjectDir)
 		{
-			get
-			{
-                // ensure our device is present in 'adb devices' output.
-				var AllDevices = GetAllConnectedDevices();
-
-				if (AllDevices.Keys.Contains(DeviceName) == false)
-				{
-					return false;
-				}
-
-				if (AllDevices[DeviceName] != "device")
-				{
-					Log.Warning(KnownLogEvents.Gauntlet_DeviceEvent, "Device {Name} is '{State}'", DeviceName, AllDevices[DeviceName]);
-					return false;
-				}
-
-				// any device will do, but only one at a time.
-				return true;
-            }
-		}
-
-		protected Dictionary<EIntendedBaseCopyDirectory, string> LocalDirectoryMappings { get; set; }
-        void SetUpDirectoryMappings()
-        {
-            LocalDirectoryMappings = new Dictionary<EIntendedBaseCopyDirectory, string>();
-        }
-
-        public void PopulateDirectoryMappings(string ProjectDir)
-		{
 			LocalDirectoryMappings.Add(EIntendedBaseCopyDirectory.Build, Path.Combine(ProjectDir, "Build"));
 			LocalDirectoryMappings.Add(EIntendedBaseCopyDirectory.Binaries, Path.Combine(ProjectDir, "Binaries"));
 			LocalDirectoryMappings.Add(EIntendedBaseCopyDirectory.Config, Path.Combine(ProjectDir, "Config"));
-            LocalDirectoryMappings.Add(EIntendedBaseCopyDirectory.Content, Path.Combine(ProjectDir, "Content"));
-            LocalDirectoryMappings.Add(EIntendedBaseCopyDirectory.Demos, Path.Combine(ProjectDir, "Demos"));
+			LocalDirectoryMappings.Add(EIntendedBaseCopyDirectory.Content, Path.Combine(ProjectDir, "Content"));
+			LocalDirectoryMappings.Add(EIntendedBaseCopyDirectory.Demos, Path.Combine(ProjectDir, "Demos"));
 			LocalDirectoryMappings.Add(EIntendedBaseCopyDirectory.PersistentDownloadDir, Path.Combine(ProjectDir, "Saved", "PersistentDownloadDir"));
 			LocalDirectoryMappings.Add(EIntendedBaseCopyDirectory.Profiling, Path.Combine(ProjectDir, "Profiling"));
-            LocalDirectoryMappings.Add(EIntendedBaseCopyDirectory.Saved, ProjectDir);
-        }
-        public bool IsConnected { get	{ return IsAvailable; }	}
-
-		protected bool IsExistingDevice = false;
+			LocalDirectoryMappings.Add(EIntendedBaseCopyDirectory.Saved, ProjectDir);
+		}
 
 		/// <summary>
 		/// Constructor
@@ -467,7 +129,7 @@
 			// If no device name or its 'default' then use the first default device
 			if (string.IsNullOrEmpty(DeviceName) || DeviceName.Equals("default", StringComparison.OrdinalIgnoreCase))
 			{
-				var DefaultDevices = GetAllAvailableDevices();	
+				var DefaultDevices = GetAllAvailableDevices();
 
 				if (DefaultDevices.Count() == 0)
 				{
@@ -535,10 +197,8 @@
 
 			ConnectedDevices = GetAllConnectedDevices();
 
-            SetUpDirectoryMappings();
-
-            // sanity check that it was now dound
-            if (ConnectedDevices.Keys.Contains(DeviceName) == false)
+			// sanity check that it was now dound
+			if (ConnectedDevices.Keys.Contains(DeviceName) == false)
 			{
 				throw new AutomationException("Failed to find new device {0} in connection list", DeviceName);
 			}
@@ -556,11 +216,11 @@
 		}
 
 		#region IDisposable Support
-		private bool disposedValue = false; // To detect redundant calls
+		private bool DisposedValue = false; // To detect redundant calls
 
 		protected virtual void Dispose(bool disposing)
 		{
-			if (!disposedValue)
+			if (!DisposedValue)
 			{
 				try
 				{
@@ -570,7 +230,7 @@
 						RunAdbGlobalCommand(string.Format("disconnect {0}", DeviceName), true, false, true);
 
 						Log.Info("Disconnected {0}", DeviceName);
-					}					
+					}
 				}
 				catch (Exception Ex)
 				{
@@ -578,7 +238,7 @@
 				}
 				finally
 				{
-					disposedValue = true;
+					DisposedValue = true;
 					AdbCredentialCache.RemoveInstance();
 				}
 
@@ -598,38 +258,200 @@
 		{
 			get
 			{
-				return disposedValue;
+				return DisposedValue;
 			}
-			
+
 		}
 
 		#endregion
 
-		/// <summary>
-		/// Returns a list of locally connected devices (e.g. 'adb devices'). 
-		/// </summary>
-		/// <returns></returns>
-		static private Dictionary<string, string> GetAllConnectedDevices()
+		#region ITargetDevice
+		public bool Connect()
 		{
-           var Result = RunAdbGlobalCommand("devices");
-
-            MatchCollection DeviceMatches = Regex.Matches(Result.Output, @"^([\d\w\.\:\-]{6,32})\s+(\w+)", RegexOptions.Multiline);
-
-            var DeviceList = DeviceMatches.Cast<Match>().ToDictionary(
-                M => M.Groups[1].ToString(),
-                M => M.Groups[2].ToString().ToLower()
-            );
-
-            return DeviceList;
+			AllowDeviceSleepState(true);
+			return true;
 		}
 
-		static private IEnumerable<string> GetAllAvailableDevices()
+		public bool Disconnect(bool bForce = false)
 		{
-			var AllDevices = GetAllConnectedDevices();
-			return AllDevices.Keys.Where(D => AllDevices[D] == "device");
+			AllowDeviceSleepState(false);
+			return true;
 		}
 
-		static public ITargetDevice[] GetDefaultDevices()
+		public bool PowerOn()
 		{
+			Log.Verbose("{0}: Powering on", ToString());
+			string CommandLine = "shell \"input keyevent KEYCODE_WAKEUP && input keyevent KEYCODE_MENU\"";
+			IProcessResult PowerOnQuery = RunAdbDeviceCommand(CommandLine, bPauseErrorParsing: true);
+			return !PowerOnQuery.Output.Contains("error");
+		}
+		public bool PowerOff()
+		{
+			Log.Verbose("{0}: Powering off", ToString());
+			string CommandLine = "shell \"input keyevent KEYCODE_SLEEP\"";
+			IProcessResult PowerOffQuery = RunAdbDeviceCommand(CommandLine, bPauseErrorParsing: true);
+			return !PowerOffQuery.Output.Contains("error");
+		}
+
+		public bool Reboot()
+		{
+			return true;
+		}
+
+		public Dictionary<EIntendedBaseCopyDirectory, string> GetPlatformDirectoryMappings()
+		{
+			return LocalDirectoryMappings;
+		}
+
+		public override string ToString()
+		{
+			// TODO: device id
+			if (Name == DeviceName)
+			{
+				return Name;
+			}
+			return string.Format("{0} ({1})", Name, DeviceName);
+		}
+
+		public void FullClean()
+		{
+
+		}
+
+		public void CleanArtifacts()
+		{
+
+		}
+
+		public void InstallBuild(UnrealAppConfig AppConfiguration)
+		{
+
+		}
+
+		public IAppInstall CreateAppInstall(UnrealAppConfig AppConfig)
+		{
+			return null;
+		}
+
+		public void CopyAdditionalFiles(IEnumerable<UnrealFileToCopy> FilesToCopy)
+		{
+
+		}
+
+		public IAppInstance Run(IAppInstall App)
+		{
+			AndroidAppInstall DroidAppInstall = App as AndroidAppInstall;
+
+			if (DroidAppInstall == null)
+			{
+				throw new Exception("AppInstance is of incorrect type!");
+			}
+
+			// wake the device - we can install while its asleep but not run
+			PowerOn();
+
+			// kill any currently running instance:
+			KillRunningProcess(DroidAppInstall.AndroidPackageName);
+
+			string LaunchActivity = AndroidPlatform.GetLaunchableActivityName();
+
+			Log.Info("Launching {0} on '{1}' ", DroidAppInstall.AndroidPackageName + "/" + LaunchActivity, ToString());
+			Log.Verbose("\t{0}", DroidAppInstall.CommandLine);
+
+			// Clear the device's logcat in preparation for the test..
+			RunAdbDeviceCommand("logcat --clear");
+
+			// start the app on device!
+			string CommandLine = "shell am start -W -S -n " + DroidAppInstall.AndroidPackageName + "/" + LaunchActivity;
+			IProcessResult Process = RunAdbDeviceCommand(CommandLine, false, true);
+
+			return new AndroidAppInstance(this, DroidAppInstall, Process);
+		}
+		#endregion
+
+		/// <summary>
+		/// Runs an ADB command, automatically adding the name of the current device to
+		/// the arguments sent to adb
+		/// </summary>
+		/// <param name="Args"></param>
+		/// <param name="Wait"></param>
+		/// <param name="Input"></param>
+		/// <returns></returns>
+		public IProcessResult RunAdbDeviceCommand(string Args, bool Wait = true, bool bShouldLogCommand = false, bool bPauseErrorParsing = false)
+		{
+			if (string.IsNullOrEmpty(DeviceName) == false)
+			{
+				Args = string.Format("-s {0} {1}", DeviceName, Args);
+			}
+
+			return RunAdbGlobalCommand(Args, Wait, bShouldLogCommand, bPauseErrorParsing);
+		}
+
+		/// <summary>
+		/// Runs an ADB command, automatically adding the name of the current device to
+		/// the arguments sent to adb
+		/// </summary>
+		/// <param name="Args"></param>
+		/// <param name="Wait"></param>
+		/// <param name="Input"></param>
+		/// <returns></returns>
+		public string RunAdbDeviceCommandAndGetOutput(string Args)
+		{
+			if (string.IsNullOrEmpty(DeviceName) == false)
+			{
+				Args = string.Format("-s {0} {1}", DeviceName, Args);
+			}
+
+			IProcessResult Result = RunAdbGlobalCommand(Args);
+
+			if (Result.ExitCode != 0)
+			{
+				throw new AutomationException("adb command {0} failed. {1}", Args, Result.Output);
+			}
+
+			return Result.Output;
+		}
+
+		/// <summary>
+		/// Runs an ADB command at the global scope
+		/// </summary>
+		/// <param name="Args"></param>
+		/// <param name="Wait"></param>
+		/// <returns></returns>
+		public static IProcessResult RunAdbGlobalCommand(string Args, bool Wait = true, bool bShouldLogCommand = false, bool bPauseErrorParsing = false)
+		{
+			CommandUtils.ERunOptions RunOptions = CommandUtils.ERunOptions.AppMustExist | CommandUtils.ERunOptions.NoWaitForExit | CommandUtils.ERunOptions.SpewIsVerbose;
+
+			if (Log.IsVeryVerbose)
+			{
+				RunOptions |= CommandUtils.ERunOptions.AllowSpew;
+			}
+			else
+			{
+				RunOptions |= CommandUtils.ERunOptions.NoLoggingOfRunCommand;
+			}
+
+			if (bShouldLogCommand)
+			{
+				Log.Verbose("Running ADB Command: adb {0}", Args);
+			}
+
+			IProcessResult Process;
+
+			using (bPauseErrorParsing ? new ScopedSuspendECErrorParsing() : null)
+			{
+				Process = AndroidPlatform.RunAdbCommand(null, null, Args, null, RunOptions);
+
+				if (Wait)
+				{
+					Process.WaitForExit();
+				}
+			}
+
+			return Process;
+		}
+
+		public static ITargetDevice[] GetDefaultDevices()
+		{
 			var Result = RunAdbGlobalCommand("devices");
 
@@ -646,7 +468,7 @@
 			return Devices.ToArray();
 		}
 
-		internal void PostRunCleanup()
+		public void PostRunCleanup()
 		{
 			// Delete the commandline file, if someone installs an APK on top of ours
 			// they will get very confusing behavior...
@@ -658,88 +480,29 @@
 			}
 		}
 
-		public bool IsOn
-		{
-			get
-			{
-				string CommandLine = "shell dumpsys power";
-				IProcessResult OnAndUnlockedQuery = RunAdbDeviceCommand(CommandLine, bPauseErrorParsing: true);
-				return OnAndUnlockedQuery.Output.Contains("mHoldingDisplaySuspendBlocker=true")
-					&& OnAndUnlockedQuery.Output.Contains("mHoldingWakeLockSuspendBlocker=true");
-			}
-		}
-
-		public bool PowerOn()
-		{
-			Log.Verbose("{0}: Powering on", ToString());
-			string CommandLine = "shell \"input keyevent KEYCODE_WAKEUP && input keyevent KEYCODE_MENU\"";
-			IProcessResult PowerOnQuery = RunAdbDeviceCommand(CommandLine, bPauseErrorParsing: true);
-			return !PowerOnQuery.Output.Contains("error");
-		}
-		public bool PowerOff()
-		{
-			Log.Verbose("{0}: Powering off", ToString());
-			string CommandLine = "shell \"input keyevent KEYCODE_SLEEP\"";
-			IProcessResult PowerOffQuery = RunAdbDeviceCommand(CommandLine, bPauseErrorParsing: true);
-			return !PowerOffQuery.Output.Contains("error");
-		}
-
-		public bool Reboot()
-		{
-			return true;
-		}
-
-		public bool Connect()
-		{
-			AllowDeviceSleepState(true);
-			return true;
-		}
-
-		public bool Disconnect(bool bForce=false)
-		{
-			AllowDeviceSleepState(false);
-			return true;
-		}
-
-		public override string ToString()
-		{
-            // TODO: device id
-			if (Name == DeviceName)
-			{
-				return Name;
-			}
-			return string.Format("{0} ({1})", Name, DeviceName);
-		}
-
-		protected bool DeleteFileFromDevice(string DestPath)
-		{
-			var AdbResult = RunAdbDeviceCommand(string.Format("shell rm -f {0}", DestPath));
-			return AdbResult.ExitCode == 0;
-		}
-
 		public bool CopyFileToDevice(string PackageName, string SourcePath, string DestPath, bool IgnoreDependencies = false)
 		{
 			bool IsAPK = string.Equals(Path.GetExtension(SourcePath), ".apk", StringComparison.OrdinalIgnoreCase);
 
-            // for the APK there's no easy/reliable way to get the date of the version installed, so 
-            // we write this out to a dependency file in the demote dir and check it each time.
-            // current file time
-            DateTime LocalModifiedTime = File.GetLastWriteTime(SourcePath);
+			// for the APK there's no easy/reliable way to get the date of the version installed, so 
+			// we write this out to a dependency file in the demote dir and check it each time.
+			// current file time
+			DateTime LocalModifiedTime = File.GetLastWriteTime(SourcePath);
 
-            string QuotedSourcePath = SourcePath;
-            if (SourcePath.Contains(" "))
-            {
-                QuotedSourcePath = '"' + SourcePath + '"';
-            }
+			string QuotedSourcePath = SourcePath;
+			if (SourcePath.Contains(" "))
+			{
+				QuotedSourcePath = '"' + SourcePath + '"';
+			}
 
-            // dependency info is a hash of the destination name, saved under a folder on /sdcard
-            string DestHash = ContentHash.MD5(DestPath).ToString();
+			// dependency info is a hash of the destination name, saved under a folder on /sdcard
+			string DestHash = ContentHash.MD5(DestPath).ToString();
 			string DependencyCacheDir = "/sdcard/gdeps";
 			string DepFile = string.Format("{0}/{1}", DependencyCacheDir, DestHash);
 
 			IProcessResult AdbResult = null;
 
-	
+
 			// get info from the device about this file
 			string CurrentFileInfo = null;
 
@@ -873,6 +636,63 @@
 			return true;
 		}
 
+		public void AllowDeviceSleepState(bool bAllowSleep)
+		{
+			string CommandLine = "shell svc power stayon " + (bAllowSleep ? "false" : "usb");
+			RunAdbDeviceCommand(CommandLine, true, false, true);
+		}
+		/// <summary>
+		/// Enable Android permissions which would otherwise block automation with permimssion requests
+		/// </summary>
+		public void EnablePermissions(string AndroidPackageName)
+		{
+			List<string> Permissions = new List<string> { "WRITE_EXTERNAL_STORAGE", "READ_EXTERNAL_STORAGE", "GET_ACCOUNTS", "RECORD_AUDIO" };
+			Permissions.ForEach(Permission =>
+			{
+				string CommandLine = string.Format("shell pm grant {0} android.permission.{1}", AndroidPackageName, Permission);
+				Log.Verbose(string.Format("Enabling permission: {0} {1}", AndroidPackageName, Permission));
+				RunAdbDeviceCommand(CommandLine, true, false, true);
+			});
+		}
+
+		public void KillRunningProcess(string AndroidPackageName)
+		{
+			Log.Verbose("{0}: Killing process '{1}' ", ToString(), AndroidPackageName);
+			string KillProcessCommand = string.Format("shell am force-stop {0}", AndroidPackageName);
+			RunAdbDeviceCommand(KillProcessCommand);
+		}
+
+		protected bool DeleteFileFromDevice(string DestPath)
+		{
+			var AdbResult = RunAdbDeviceCommand(string.Format("shell rm -f {0}", DestPath));
+			return AdbResult.ExitCode == 0;
+		}
+
+		/// <summary>
+		/// Returns a list of locally connected devices (e.g. 'adb devices').
+		/// </summary>
+		/// <returns></returns>
+		private static Dictionary<string, string> GetAllConnectedDevices()
+		{
+			var Result = RunAdbGlobalCommand("devices");
+
+			MatchCollection DeviceMatches = Regex.Matches(Result.Output, @"^([\d\w\.\:\-]{6,32})\s+(\w+)", RegexOptions.Multiline);
+
+			var DeviceList = DeviceMatches.Cast<Match>().ToDictionary(
+				M => M.Groups[1].ToString(),
+				M => M.Groups[2].ToString().ToLower()
+			);
+
+			return DeviceList;
+		}
+
+		private static IEnumerable<string> GetAllAvailableDevices()
+		{
+			var AllDevices = GetAllConnectedDevices();
+			return AllDevices.Keys.Where(D => AllDevices[D] == "device");
+		}
+
+		#region Legacy Implementations
 		public IAppInstall InstallApplication(UnrealAppConfig AppConfig)
 		{
 			// todo - pass this through
@@ -913,7 +733,7 @@
 			// clear all file store paths between installs:
 			RunAdbDeviceCommand(string.Format("shell rm -r {0}", DeviceExternalStorageSavedPath));
 			RunAdbDeviceCommand(string.Format("shell rm -r {0}", DeviceExternalFilesSavedPath));
-		
+
 			if (AppConfig.FullClean)
 			{
 				Log.Info("Fully cleaning console before install...");
@@ -921,7 +741,7 @@
 				RunAdbDeviceCommand(string.Format("shell rm -r {0}/Android/data/{1}/*", StorageLocation, Build.AndroidPackageName));
 				RunAdbDeviceCommand(string.Format("shell rm -r {0}/Android/obb/{1}/*", StorageLocation, Build.AndroidPackageName));
 				RunAdbDeviceCommand(string.Format("shell rm -r {0}/Download/*", StorageLocation));
-			}			
+			}
 
 			if (!AppConfig.SkipInstall)
 			{
@@ -978,30 +798,30 @@
 			Dictionary<string, string> FilesToInstall = new Dictionary<string, string>();
 
 			Console.WriteLine("trying to copy files over.");
-            if (AppConfig.FilesToCopy != null)
-            {
-                if (LocalDirectoryMappings.Count == 0)
-                {
-                    Console.WriteLine("Populating Directory");
-                    PopulateDirectoryMappings(DeviceArtifactPath);
-                }
-                Console.WriteLine("trying to copy files over.");
-                foreach (UnrealFileToCopy FileToCopy in AppConfig.FilesToCopy)
-                {
-                    string PathToCopyTo = Path.Combine(LocalDirectoryMappings[FileToCopy.TargetBaseDirectory], FileToCopy.TargetRelativeLocation);
-                    if (File.Exists(FileToCopy.SourceFileLocation))
-                    {
-                        FileInfo SrcInfo = new FileInfo(FileToCopy.SourceFileLocation);
-                        SrcInfo.IsReadOnly = false;
-                        FilesToInstall.Add(FileToCopy.SourceFileLocation, PathToCopyTo.Replace("\\", "/"));
-                    }
+			if (AppConfig.FilesToCopy != null)
+			{
+				if (LocalDirectoryMappings.Count == 0)
+				{
+					Console.WriteLine("Populating Directory");
+					PopulateDirectoryMappings(DeviceArtifactPath);
+				}
+				Console.WriteLine("trying to copy files over.");
+				foreach (UnrealFileToCopy FileToCopy in AppConfig.FilesToCopy)
+				{
+					string PathToCopyTo = Path.Combine(LocalDirectoryMappings[FileToCopy.TargetBaseDirectory], FileToCopy.TargetRelativeLocation);
+					if (File.Exists(FileToCopy.SourceFileLocation))
+					{
+						FileInfo SrcInfo = new FileInfo(FileToCopy.SourceFileLocation);
+						SrcInfo.IsReadOnly = false;
+						FilesToInstall.Add(FileToCopy.SourceFileLocation, PathToCopyTo.Replace("\\", "/"));
+					}
 
-                    else
-                    {
-                        Log.Warning(KnownLogEvents.Gauntlet_DeviceEvent, "File to copy {File} not found", FileToCopy);
-                    }
-                }
-            }
+					else
+					{
+						Log.Warning(KnownLogEvents.Gauntlet_DeviceEvent, "File to copy {File} not found", FileToCopy);
+					}
+				}
+			}
 
 			if (!AppConfig.SkipInstall)
 			{
@@ -1095,176 +915,329 @@
 
 			return AppInstall;
 		}
+		#endregion
+	}
 
-		public void FullClean()
+	class AndroidAppInstall : IAppInstall
+	{
+		public string Name { get; protected set; }
+
+		public TargetDeviceAndroid AndroidDevice { get; protected set; }
+
+		public ITargetDevice Device => AndroidDevice;
+
+		public string AndroidPackageName { get; protected set; }
+
+		public string CommandLine { get; protected set; }
+
+		public AndroidAppInstall(TargetDeviceAndroid InDevice, string InName, string InAndroidPackageName, string InCommandLine)
 		{
+			AndroidDevice = InDevice;
+			Name = InName;
+			AndroidPackageName = InAndroidPackageName;
+			CommandLine = InCommandLine;
+		}
+
+		public IAppInstance Run()
+		{
+			return AndroidDevice.Run(this);
+		}
+	}
+
+	public class DefaultAndroidDevices : IDefaultDeviceSource
+	{
+		public bool CanSupportPlatform(UnrealTargetPlatform? Platform)
+		{
+			return Platform == UnrealTargetPlatform.Android;
+		}
+
+		public ITargetDevice[] GetDefaultDevices()
+		{
+			return TargetDeviceAndroid.GetDefaultDevices();
+		}
+	}
+
+	// become IAppInstance when implemented enough
+	class AndroidAppInstance : IAppInstance
+	{
+		public string ArtifactPath
+		{
+			get
+			{
+				if (HasExited && !bHaveSavedArtifacts)
+				{
+					bHaveSavedArtifacts = true;
+					SaveArtifacts();
+				}
+
+				return Path.Combine(AndroidDevice.LocalCachePath, "Saved");
+			}
+		}
+
+		public bool HasExited
+		{
+			get
+			{
+				try
+				{
+					if (!LaunchProcess.HasExited)
+					{
+						return false;
+					}
+				}
+				catch (InvalidOperationException)
+				{
+					return true;
+				}
+
+				return IsActivityRunning();
+			}
+		}
+
+		public string StdOut
+		{
+			get
+			{
+				UpdateCachedLog();
+				return String.IsNullOrEmpty(ActivityLogCached) ? String.Empty : ActivityLogCached;
+			}
+		}
+
+		public string CommandLine => Install.CommandLine;
+
+		public ITargetDevice Device => AndroidDevice;
+
+		public int ExitCode => LaunchProcess.ExitCode;
+
+		public bool WasKilled { get; protected set; }
+
+		public IProcessResult LaunchProcess;
+
+		protected TargetDeviceAndroid AndroidDevice;
+
+		protected AndroidAppInstall Install;
+
+		protected bool bHaveSavedArtifacts;
+
+		private static readonly TimeSpan ActivityCheckDelta = TimeSpan.FromSeconds(10);
+		private DateTime ActivityCheckTime = DateTime.UtcNow;
+		private bool ActivityExited = false;
+
+		private static readonly TimeSpan ActivityLogDelta = TimeSpan.FromSeconds(15);
+		private DateTime ActivityLogTime = DateTime.UtcNow - ActivityLogDelta;
+		private string ActivityLogCached = string.Empty;
+
+		public AndroidAppInstance(TargetDeviceAndroid InDevice, AndroidAppInstall InInstall, IProcessResult InProcess)
+		{
+			AndroidDevice = InDevice;
+			Install = InInstall;
+			LaunchProcess = InProcess;
+		}
+
+		public int WaitForExit()
+		{
+			if (!HasExited)
+			{
+				LaunchProcess.WaitForExit();
+			}
+
+			return ExitCode;
+		}
+
+		public void Kill()
+		{
+			if (!HasExited && !AndroidDevice.Disposed)
+			{
+				WasKilled = true;
+				Install.AndroidDevice.KillRunningProcess(Install.AndroidPackageName);
+			}
+		}
+
+		protected void SaveArtifacts()
+		{
+			// copy remote artifacts to local
+			if (Directory.Exists(Install.AndroidDevice.LocalCachePath))
+			{
+
+				Log.Verbose("Deleting {0}", Install.AndroidDevice.LocalCachePath);
+
+				try
+				{
+					// don't consider this fatal, people often have the directory or a file open
+					Directory.Delete(Install.AndroidDevice.LocalCachePath, true);
+				}
+				catch
+				{
+					Log.Warning(KnownLogEvents.Gauntlet_DeviceEvent, "Failed to remove old cache folder {Folder}", Install.AndroidDevice.LocalCachePath);
+				}
+			}
+
+			// mark it as a temp dir (will also create it)
+			try
+			{
+				Utils.SystemHelpers.MarkDirectoryForCleanup(Install.AndroidDevice.LocalCachePath);
+			}
+			catch (Exception Ex)
+			{
+				Log.Warning(KnownLogEvents.Gauntlet_DeviceEvent, "Exception marking directory for cleanup {Exception}", Ex.Message);
+			}
+
+			string LocalSaved = Path.Combine(Install.AndroidDevice.LocalCachePath, "Saved");
+			Log.Verbose("Creating {0}", LocalSaved);
+			Directory.CreateDirectory(LocalSaved);
+
+			// pull all the artifacts
+			string ArtifactPullCommand = string.Format("pull {0} \"{1}\"", Install.AndroidDevice.DeviceArtifactPath, Install.AndroidDevice.LocalCachePath);
+			IProcessResult PullCmd = Install.AndroidDevice.RunAdbDeviceCommand(ArtifactPullCommand, bShouldLogCommand: Log.IsVerbose);
+
+			if (PullCmd.ExitCode != 0)
+			{
+				Log.Warning(KnownLogEvents.Gauntlet_DeviceEvent, "Failed to retrieve artifacts. {Output}", PullCmd.Output);
+			}
+			else
+			{
+				// update final cached stdout property
+				string LogFilename = Install.AndroidDevice.DeviceLogPath;
+				if (File.Exists(LogFilename))
+				{
+					ActivityLogCached = File.ReadAllText(LogFilename);
+					ActivityLogTime = DateTime.MinValue;
+				}
+
+			}
+
+			// pull the logcat over from device.
+			IProcessResult LogcatResult = Install.AndroidDevice.RunAdbDeviceCommand("logcat -d", bShouldLogCommand: Log.IsVerbose);
+
+			string LogcatFilename = "Logcat.log";
+			// Save logcat dump to local artifact path.
+			File.WriteAllText(Path.Combine(LocalSaved, LogcatFilename), LogcatResult.Output);
+
+			Install.AndroidDevice.PostRunCleanup();
+
+			bHaveSavedArtifacts = true;
+		}
+
+		/// <summary>
+		/// Checks on device whether the activity is running, this is an expensive shell with output operation
+		/// the result is cached, with checks at ActivityCheckDelta seconds
+		/// </summary>
+		private bool IsActivityRunning()
+		{
+			if (ActivityExited)
+			{
+				return true;
+			}
+
+			if ((DateTime.UtcNow - ActivityCheckTime) < ActivityCheckDelta)
+			{
+				return false;
+			}
+
+			if (AndroidDevice.Disposed)
+			{
+				return false;
+			}
+
+			ActivityCheckTime = DateTime.UtcNow;
+
+			// get activities filtered by our package name
+			IProcessResult ActivityQuery = AndroidDevice.RunAdbDeviceCommand("shell dumpsys activity -p " + Install.AndroidPackageName + " a");
+
+			// We have exited if our activity doesn't appear in the activity query or is not the focused activity.
+			bool bActivityPresent = ActivityQuery.Output.Contains(Install.AndroidPackageName);
+			bool bActivityInForeground = ActivityQuery.Output.Contains("mResumedActivity");
+			bool bHasExited = !bActivityPresent || !bActivityInForeground;
+			if (bHasExited)
+			{
+				ActivityExited = true;
+				// The activity has exited, make sure entire activity log has been captured, sleep to allow time for the log to flush
+				Thread.Sleep(5000);
+				UpdateCachedLog(true);
+				Log.VeryVerbose("{0}: process exited, Activity running={1}, Activity in foreground={2} ", ToString(), bActivityPresent.ToString(), bActivityInForeground.ToString());
+			}
+
+			return bHasExited;
 
 		}
 
-		public void CleanArtifacts()
+		/// <summary>
+		/// Updates cached activity log by running a shell command returning the full log from device (possibly over wifi)
+		/// The result is cached and updated at ActivityLogDelta frequency
+		/// </summary>
+		private void UpdateCachedLog(bool ForceUpdate = false)
 		{
-
-		}
-
-		public void InstallBuild(UnrealAppConfig AppConfiguration)
-		{
-
-		}
-
-		public IAppInstall CreateAppInstall(UnrealAppConfig AppConfig)
-		{
-			return null;
-		}
-
-		public void CopyAdditionalFiles(IEnumerable<UnrealFileToCopy> FilesToCopy)
-		{
-
-		}
-
-		public IAppInstance Run(IAppInstall App)
-		{
-			AndroidAppInstall DroidAppInstall = App as AndroidAppInstall;
-
-			if (DroidAppInstall == null)
+			if (!ForceUpdate && (ActivityLogTime == DateTime.MinValue || ((DateTime.UtcNow - ActivityLogTime) < ActivityLogDelta)))
 			{
-				throw new Exception("AppInstance is of incorrect type!");
+				return;
 			}
 
-			// wake the device - we can install while its asleep but not run
-			PowerOn();
-
-			// kill any currently running instance:
-			KillRunningProcess(DroidAppInstall.AndroidPackageName);
-
-			string LaunchActivity = AndroidPlatform.GetLaunchableActivityName();
-
-			Log.Info("Launching {0} on '{1}' ", DroidAppInstall.AndroidPackageName + "/" + LaunchActivity, ToString());
-			Log.Verbose("\t{0}", DroidAppInstall.CommandLine);
-
-			// Clear the device's logcat in preparation for the test..
-			RunAdbDeviceCommand("logcat --clear");
-
-			// start the app on device!
-			string CommandLine = "shell am start -W -S -n " + DroidAppInstall.AndroidPackageName + "/" + LaunchActivity;
-			IProcessResult Process = RunAdbDeviceCommand(CommandLine, false, true);
-
-			return new AndroidAppInstance(this, DroidAppInstall, Process);
-		}
-
-		/// <summary>
-		/// Runs an ADB command, automatically adding the name of the current device to
-		/// the arguments sent to adb
-		/// </summary>
-		/// <param name="Args"></param>
-		/// <param name="Wait"></param>
-		/// <param name="Input"></param>
-		/// <returns></returns>
-		public IProcessResult RunAdbDeviceCommand(string Args, bool Wait=true, bool bShouldLogCommand = false, bool bPauseErrorParsing = false)
-		{
-			if (string.IsNullOrEmpty(DeviceName) == false)
+			if (Install.AndroidDevice != null && Install.AndroidDevice.Disposed)
 			{
-				Args = string.Format("-s {0} {1}", DeviceName, Args);
+				Log.Warning(KnownLogEvents.Gauntlet_DeviceEvent, "Attempting to cache log using disposed Android device");
+				return;
 			}
 
-			return RunAdbGlobalCommand(Args, Wait, bShouldLogCommand, bPauseErrorParsing);
-		}
+			// Note: DeviceLogPath has the correct location for the current app. Apps can be configured to override the default log storage path. 
+			string GetLogCommand = string.Format("shell cat {0}", Install.AndroidDevice.DeviceLogPath);
+			IProcessResult LogQuery = Install.AndroidDevice.RunAdbDeviceCommand(GetLogCommand, true);
 
-		/// <summary>
-		/// Runs an ADB command, automatically adding the name of the current device to
-		/// the arguments sent to adb
-		/// </summary>
-		/// <param name="Args"></param>
-		/// <param name="Wait"></param>
-		/// <param name="Input"></param>
-		/// <returns></returns>
-		public string RunAdbDeviceCommandAndGetOutput(string Args)
-		{
-			if (string.IsNullOrEmpty(DeviceName) == false)
+			if (LogQuery.ExitCode != 0)
 			{
-				Args = string.Format("-s {0} {1}", DeviceName, Args);
-			}
-
-			IProcessResult Result = RunAdbGlobalCommand(Args);
-
-			if (Result.ExitCode != 0)
-			{
-				throw new AutomationException("adb command {0} failed. {1}", Args, Result.Output);
-			}
-
-			return Result.Output;
-		}
-
-		/// <summary>
-		/// Runs an ADB command at the global scope
-		/// </summary>
-		/// <param name="Args"></param>
-		/// <param name="Wait"></param>
-		/// <returns></returns>
-		public static IProcessResult RunAdbGlobalCommand(string Args, bool Wait = true, bool bShouldLogCommand = false, bool bPauseErrorParsing = false)
-		{
-			CommandUtils.ERunOptions RunOptions = CommandUtils.ERunOptions.AppMustExist | CommandUtils.ERunOptions.NoWaitForExit | CommandUtils.ERunOptions.SpewIsVerbose;
-
-			if (Log.IsVeryVerbose)
-			{
-				RunOptions |= CommandUtils.ERunOptions.AllowSpew;
+				Log.VeryVerbose("Unable to query activity stdout on device {0}", Install.AndroidDevice.Name);
 			}
 			else
 			{
-				RunOptions |= CommandUtils.ERunOptions.NoLoggingOfRunCommand;
-			}
-	
-			if (bShouldLogCommand)
-			{
-				Log.Verbose("Running ADB Command: adb {0}", Args);
+				ActivityLogCached = LogQuery.Output;
 			}
 
-			IProcessResult Process;
+			ActivityLogTime = DateTime.UtcNow;
 
-			using (bPauseErrorParsing ? new ScopedSuspendECErrorParsing() : null)
+			// the activity has exited, mark final log sentinel
+			if (ActivityExited)
 			{
-				Process = AndroidPlatform.RunAdbCommand(null, null, Args, null, RunOptions);
+				ActivityLogTime = DateTime.MinValue;
+			}
 
-				if (Wait)
-				{
-					Process.WaitForExit();
-				}
-			}
-			
-			return Process;
-		}
-
-		public void AllowDeviceSleepState(bool bAllowSleep)
-		{
-			string CommandLine = "shell svc power stayon " + (bAllowSleep ? "false" : "usb");
-			RunAdbDeviceCommand(CommandLine, true, false, true);
-		}
-		/// <summary>
-		/// Enable Android permissions which would otherwise block automation with permimssion requests
-		/// </summary>
-		public void EnablePermissions(string AndroidPackageName)
-		{
-			List<string> Permissions = new List<string>{ "WRITE_EXTERNAL_STORAGE", "READ_EXTERNAL_STORAGE", "GET_ACCOUNTS", "RECORD_AUDIO" };
-			Permissions.ForEach(Permission => {
-				string CommandLine = string.Format("shell pm grant {0} android.permission.{1}", AndroidPackageName, Permission);
-				Log.Verbose(string.Format("Enabling permission: {0} {1}", AndroidPackageName, Permission));
-				RunAdbDeviceCommand(CommandLine, true, false, true);
-			});
-		}
-
-		public void KillRunningProcess(string AndroidPackageName)
-		{
-			Log.Verbose("{0}: Killing process '{1}' ", ToString(), AndroidPackageName);
-			string KillProcessCommand = string.Format("shell am force-stop {0}", AndroidPackageName);
-			RunAdbDeviceCommand(KillProcessCommand);
-		}
-
-		public Dictionary<EIntendedBaseCopyDirectory, string> GetPlatformDirectoryMappings()
-		{
-			if (LocalDirectoryMappings.Count == 0)
-			{
-				Log.Warning("Platform directory mappings have not been populated for this platform! This should be done within InstallApplication()");
-			}
-			return LocalDirectoryMappings;
 		}
 	}
 
+	// Device data from json
+	public sealed class AndroidDeviceData
+	{
+		// Host of PC which is tethered
+		public string HostIP { get; set; }
+
+		// public key
+		public string PublicKey { get; set; }
+
+		// private key
+		public string PrivateKey { get; set; }
+	}
+
+	public class AndroidBuildSupport : BaseBuildSupport
+	{
+		protected override BuildFlags SupportedBuildTypes => BuildFlags.Packaged | BuildFlags.CanReplaceCommandLine | BuildFlags.CanReplaceExecutable | BuildFlags.Bulk | BuildFlags.NotBulk;
+		protected override UnrealTargetPlatform? Platform => UnrealTargetPlatform.Android;
+	}
+
+	public class AndroidDeviceFactory : IDeviceFactory
+	{
+		public bool CanSupportPlatform(UnrealTargetPlatform? Platform)
+		{
+			return Platform == UnrealTargetPlatform.Android;
+		}
+
+		public ITargetDevice CreateDevice(string InRef, string InCachePath, string InParam = null)
+		{
+			AndroidDeviceData DeviceData = null;
+
+			if (!string.IsNullOrEmpty(InParam))
+			{
+				DeviceData = fastJSON.JSON.Instance.ToObject<AndroidDeviceData>(InParam);
+			}
+
+			return new TargetDeviceAndroid(InRef, DeviceData, InCachePath);
+		}
+	}
+
@@ -1273,7 +1246,6 @@
 	/// </summary>
 	internal static class AdbCredentialCache
 	{
-
 		private static int InstanceCount = 0;
 		private static bool bUsingCustomKeys = false;
 
@@ -1282,17 +1254,71 @@
 
 		private const string KeyBackupExt = ".gauntlet.bak";
 
-		private static void Reset()
+		static AdbCredentialCache()
 		{
-			if (InstanceCount != 0)
+			Reset();
+		}
+
+		public static void RemoveInstance()
+		{
+			lock (Globals.MainLock)
 			{
-				throw new AutomationException("AdbCredentialCache.Reset() called with outstanding instances");
+				InstanceCount--;
+
+				if (InstanceCount == 0 && bUsingCustomKeys)
+				{
+					Reset();
+					KillAdbServer();
+
+					// Kill ADB server, just as a safety measure to ensure it closes
+					IEnumerable<Process> ADBProcesses = Process.GetProcesses().Where(p => p.ProcessName.Equals("adb"));
+					if (ADBProcesses.Count() > 0)
+					{
+						Log.Info("Terminating {0} ADB Process(es)", ADBProcesses.Count());
+						foreach (Process ADBProcess in ADBProcesses)
+						{
+							Log.Info("Killing ADB process {0}", ADBProcess.Id);
+							ADBProcess.Kill();
+						}
+					}
+				}
+			}
+		}
+
+		public static void RestoreBackupKeys()
+		{
+			string LocalKeyPath = Path.Combine(Environment.GetEnvironmentVariable("USERPROFILE"), ".android");
+			string LocalKeyFile = Path.Combine(LocalKeyPath, "adbkey");
+			string LocalPubKeyFile = Path.Combine(LocalKeyPath, "adbkey.pub");
+			string BackupSentry = Path.Combine(LocalKeyPath, "gauntlet.inuse");
+
+			if (File.Exists(BackupSentry))
+			{
+				Log.Info("Restoring original adb keys");
+
+				if (File.Exists(LocalKeyFile + KeyBackupExt))
+				{
+					File.Copy(LocalKeyFile + KeyBackupExt, LocalKeyFile, true);
+					File.Delete(LocalKeyFile + KeyBackupExt);
+				}
+				else
+				{
+					File.Delete(LocalKeyFile);
+				}
+
+				if (File.Exists(LocalPubKeyFile + KeyBackupExt))
+				{
+					File.Copy(LocalPubKeyFile + KeyBackupExt, LocalPubKeyFile, true);
+					File.Delete(LocalPubKeyFile + KeyBackupExt);
+				}
+				else
+				{
+					File.Delete(LocalPubKeyFile);
+				}
+
+				File.Delete(BackupSentry);
 			}
 
-			PrivateKey = PublicKey = String.Empty;
-			bUsingCustomKeys = false;
-
-			RestoreBackupKeys();
 		}
 
 		public static void AddInstance(AndroidDeviceData DeviceData = null)
@@ -1302,29 +1328,29 @@
 				string KeyPath = Globals.Params.ParseValue("adbkeys", null);
 
 				// setup key store from device data
-				if (String.IsNullOrEmpty(KeyPath) && DeviceData != null)
+				if (string.IsNullOrEmpty(KeyPath) && DeviceData != null)
 				{
 					// checked that cached keys are the same
-					if (!String.IsNullOrEmpty(PrivateKey))
+					if (!string.IsNullOrEmpty(PrivateKey))
 					{
-						if (PrivateKey != DeviceData.privateKey)
+						if (PrivateKey != DeviceData.PrivateKey)
 						{
 							throw new AutomationException("ADB device private keys must match");
 						}
 					}
 
-					if (!String.IsNullOrEmpty(PublicKey))
+					if (!string.IsNullOrEmpty(PublicKey))
 					{
-						if (PublicKey != DeviceData.publicKey)
+						if (PublicKey != DeviceData.PublicKey)
 						{
 							throw new AutomationException("ADB device public keys must match");
 						}
 					}
 
-					PrivateKey = DeviceData.privateKey;
-					PublicKey = DeviceData.publicKey;
+					PrivateKey = DeviceData.PrivateKey;
+					PublicKey = DeviceData.PublicKey;
 
-					if (String.IsNullOrEmpty(PublicKey) || String.IsNullOrEmpty(PrivateKey))
+					if (string.IsNullOrEmpty(PublicKey) || string.IsNullOrEmpty(PrivateKey))
 					{
 						throw new AutomationException("Invalid key in device data");
 					}
@@ -1401,87 +1427,30 @@
 
 		}
 
+		private static void Reset()
+		{
+			if (InstanceCount != 0)
+			{
+				throw new AutomationException("AdbCredentialCache.Reset() called with outstanding instances");
+			}
+
+			PrivateKey = PublicKey = string.Empty;
+			bUsingCustomKeys = false;
+
+			RestoreBackupKeys();
+		}
+
 		private static void KillAdbServer()
-		{			
+		{
 			using (new ScopedSuspendECErrorParsing())
 			{
 				Log.Info("Running adb kill-server to refresh credentials");
 				TargetDeviceAndroid.RunAdbGlobalCommand("kill-server");
-				// killing the adb server restarts it and can surface superfluous device errors
+
+				// Killing the adb server restarts it and can surface superfluous device errors
 				int SleepTime = CommandUtils.IsBuildMachine ? 15000 : 5000;
 				Thread.Sleep(SleepTime);
 			}
-
 		}
-
-		public static void RemoveInstance()
-		{
-			lock (Globals.MainLock)
-			{
-
-				InstanceCount--;
-
-				if (InstanceCount == 0 && bUsingCustomKeys)
-				{
-					Reset();
-					KillAdbServer();
-					// Kill ADB server, just as a safety measure to ensure it closes
-					IEnumerable<Process> ADBProcesses = Process.GetProcesses().Where(p => p.ProcessName.Equals("adb"));
-					if (ADBProcesses.Count() > 0)
-					{
-						Log.Info("Terminating {0} ADB Process(es)", ADBProcesses.Count());
-						foreach (Process ADBProcess in ADBProcesses)
-						{
-							Log.Info("Killing ADB process {0}", ADBProcess.Id);
-							ADBProcess.Kill();
-						}
-					}
-				}
-			}
-		}
-
-		public static void RestoreBackupKeys()
-		{			
-			string LocalKeyPath = Path.Combine(Environment.GetEnvironmentVariable("USERPROFILE"), ".android");
-			string LocalKeyFile = Path.Combine(LocalKeyPath, "adbkey");
-			string LocalPubKeyFile = Path.Combine(LocalKeyPath, "adbkey.pub");
-			string BackupSentry = Path.Combine(LocalKeyPath, "gauntlet.inuse");
-
-			if (File.Exists(BackupSentry))
-			{
-				Log.Info("Restoring original adb keys");
-
-				if (File.Exists(LocalKeyFile + KeyBackupExt))
-				{
-					File.Copy(LocalKeyFile + KeyBackupExt, LocalKeyFile, true);
-					File.Delete(LocalKeyFile + KeyBackupExt);
-				}
-				else
-				{
-					File.Delete(LocalKeyFile);
-				}
-
-				if (File.Exists(LocalPubKeyFile + KeyBackupExt))
-				{
-					File.Copy(LocalPubKeyFile + KeyBackupExt, LocalPubKeyFile, true);
-					File.Delete(LocalPubKeyFile + KeyBackupExt);
-				}
-				else
-				{
-					File.Delete(LocalPubKeyFile);
-				}
-
-				File.Delete(BackupSentry);
-			}
-
-		}
-
-		static AdbCredentialCache()
-		{
-			Reset();
-		}
-
 	}
-
-
 }
