- Cleaned up remaining UnrealEditor-Cmd.exe hardcodes, and using the recently checked in GetEditorForProject function
- Renamed some "FileReference ProjectName" to "FileReference ProjectFile"
- Passing in "null" for a lot of UnrealExe params, instead of "UnrealEditor-Cmd", so allow GetEditorForProject to be called (we must still allow for a true override, so we can't get rid of the param)
- Added FileSystemReference.Exists

diff --git a/Engine/Source/Programs/AutomationTool/AutomationUtils/CommandletUtils.cs b/Engine/Source/Programs/AutomationTool/AutomationUtils/CommandletUtils.cs
--- a/Engine/Source/Programs/AutomationTool/AutomationUtils/CommandletUtils.cs
+++ b/Engine/Source/Programs/AutomationTool/AutomationUtils/CommandletUtils.cs
@@ -49,7 +49,7 @@ namespace AutomationTool
 		/// <summary>
 		/// Runs Cook commandlet.
 		/// </summary>
-		/// <param name="ProjectName">Project name.</param>
+		/// <param name="ProjectFile">Project file.</param>
 		/// <param name="UnrealExe">The name of the Unreal Editor executable to use.</param>
 		/// <param name="Maps">List of maps to cook, can be null in which case -MapIniSection=AllMaps is used.</param>
 		/// <param name="Dirs">List of directories to cook, can be null</param>
@@ -57,7 +57,7 @@ namespace AutomationTool
         /// <param name="CulturesToCook">List of culture names whose localized assets should be cooked, can be null (implying defaults should be used).</param>
 		/// <param name="TargetPlatform">Target platform.</param>
 		/// <param name="Parameters">List of additional parameters.</param>
-		public static void CookCommandlet(FileReference ProjectName, string UnrealExe = "UnrealEditor-Cmd.exe", string[] Maps = null, string[] Dirs = null, string InternationalizationPreset = "", string[] CulturesToCook = null, string TargetPlatform = "Windows", string Parameters = "-Unversioned")
+		public static void CookCommandlet(FileReference ProjectFile, string UnrealExe = null, string[] Maps = null, string[] Dirs = null, string InternationalizationPreset = "", string[] CulturesToCook = null, string TargetPlatform = "Windows", string Parameters = "-Unversioned")
 		{
             string CommandletArguments = "";
 
@@ -90,18 +90,18 @@ namespace AutomationTool
                 CommandletArguments += (CommandletArguments.Length > 0 ? " " : "") + CulturesToCookArg;
             }
 
-            RunCommandlet(ProjectName, UnrealExe, "Cook", String.Format("{0} -TargetPlatform={1} {2}",  CommandletArguments, TargetPlatform, Parameters));
+            RunCommandlet(ProjectFile, UnrealExe, "Cook", String.Format("{0} -TargetPlatform={1} {2}",  CommandletArguments, TargetPlatform, Parameters));
 		}
 
-        /// <summary>
-        /// Runs DDC commandlet.
-        /// </summary>
-        /// <param name="ProjectName">Project name.</param>
-        /// <param name="UnrealExe">The name of the Unreal Editor executable to use.</param>
-        /// <param name="Maps">List of maps to cook, can be null in which case -MapIniSection=AllMaps is used.</param>
-        /// <param name="TargetPlatform">Target platform.</param>
-        /// <param name="Parameters">List of additional parameters.</param>
-        public static void DDCCommandlet(FileReference ProjectName, string UnrealExe = "UnrealEditor-Cmd.exe", string[] Maps = null, string TargetPlatform = "Windows", string Parameters = "")
+		/// <summary>
+		/// Runs DDC commandlet.
+		/// </summary>
+		/// <param name="ProjectFile">Project file.</param>
+		/// <param name="UnrealExe">The name of the Unreal Editor executable to use.</param>
+		/// <param name="Maps">List of maps to cook, can be null in which case -MapIniSection=AllMaps is used.</param>
+		/// <param name="TargetPlatform">Target platform.</param>
+		/// <param name="Parameters">List of additional parameters.</param>
+		public static void DDCCommandlet(FileReference ProjectFile, string UnrealExe = null, string[] Maps = null, string TargetPlatform = "Windows", string Parameters = "")
         {
             string MapsToCook = "";
             if (!IsNullOrEmpty(Maps))
@@ -109,17 +109,17 @@ namespace AutomationTool
                 MapsToCook = "-Map=" + CombineCommandletParams(Maps).Trim();
             }
 
-            RunCommandlet(ProjectName, UnrealExe, "DerivedDataCache", String.Format("{0} -TargetPlatform={1} {2}", MapsToCook, TargetPlatform, Parameters));
+			RunCommandlet(ProjectFile, UnrealExe, "DerivedDataCache", String.Format("{0} -TargetPlatform={1} {2}", MapsToCook, TargetPlatform, Parameters));
         }
 
 		/// <summary>
 		/// Runs RebuildLightMaps commandlet.
 		/// </summary>
-		/// <param name="ProjectName">Project name.</param>
+		/// <param name="ProjectFile">Project file.</param>
 		/// <param name="UnrealExe">The name of the Unreal Editor executable to use.</param>
 		/// <param name="Maps">List of maps to rebuild light maps for. Can be null in which case -MapIniSection=AllMaps is used.</param>
 		/// <param name="Parameters">List of additional parameters.</param>
-		public static void RebuildLightMapsCommandlet(FileReference ProjectName, string UnrealExe = "UnrealEditor-Cmd.exe", string[] Maps = null, string Parameters = "")
+		public static void RebuildLightMapsCommandlet(FileReference ProjectFile, string UnrealExe = null, string[] Maps = null, string Parameters = "")
 		{
 			string MapsToRebuildLighting = "";
 			if (!IsNullOrEmpty(Maps))
@@ -127,15 +127,15 @@ namespace AutomationTool
 				MapsToRebuildLighting = "-Map=" + CombineCommandletParams(Maps).Trim();
 			}
 
-			RunCommandlet(ProjectName, UnrealExe, "ResavePackages", String.Format("-buildtexturestreaming -buildlighting -MapsOnly -ProjectOnly -AllowCommandletRendering -SkipSkinVerify {0} {1}", MapsToRebuildLighting, Parameters));
+			RunCommandlet(ProjectFile, UnrealExe, "ResavePackages", String.Format("-buildtexturestreaming -buildlighting -MapsOnly -ProjectOnly -AllowCommandletRendering -SkipSkinVerify {0} {1}", MapsToRebuildLighting, Parameters));
 		}
 
-		public static void RebuildHLODCommandlet(FileReference ProjectName, string UnrealExe = "UnrealEditor-Cmd.exe", string[] Maps = null, string Parameters = "")
+		public static void RebuildHLODCommandlet(FileReference ProjectName, string UnrealExe = null, string[] Maps = null, string Parameters = "")
 		{
 			RebuildHLODCommandlet(ProjectName, out string LogFile, UnrealExe, Maps, Parameters);
 		}
 
-		public static void RebuildHLODCommandlet(FileReference ProjectName, out string DestLogFile, string UnrealExe = "UnrealEditor-Cmd.exe", string[] Maps = null, string Parameters = "")
+		public static void RebuildHLODCommandlet(FileReference ProjectFile, out string DestLogFile, string UnrealExe = null, string[] Maps = null, string Parameters = "")
         {
             string MapsToRebuildHLODs = "";
             if (!IsNullOrEmpty(Maps))
@@ -143,17 +143,17 @@ namespace AutomationTool
                 MapsToRebuildHLODs = "-Map=" + CombineCommandletParams(Maps).Trim();
             }
 
-            RunCommandlet(ProjectName, UnrealExe, "ResavePackages", String.Format("-BuildHLOD -ProjectOnly -AllowCommandletRendering -SkipSkinVerify {0} {1}", MapsToRebuildHLODs, Parameters), out DestLogFile);
+			RunCommandlet(ProjectFile, UnrealExe, "ResavePackages", String.Format("-BuildHLOD -ProjectOnly -AllowCommandletRendering -SkipSkinVerify {0} {1}", MapsToRebuildHLODs, Parameters), out DestLogFile);
         }
 
-        /// <summary>
-        /// Runs RebuildLightMaps commandlet.
-        /// </summary>
-        /// <param name="ProjectName">Project name.</param>
-        /// <param name="UnrealExe">The name of the Unreal Editor executable to use.</param>
-        /// <param name="Maps">List of maps to rebuild light maps for. Can be null in which case -MapIniSection=AllMaps is used.</param>
-        /// <param name="Parameters">List of additional parameters.</param>
-        public static void ResavePackagesCommandlet(FileReference ProjectName, string UnrealExe = "UnrealEditor-Cmd.exe", string[] Maps = null, string Parameters = "")
+		/// <summary>
+		/// Runs RebuildLightMaps commandlet.
+		/// </summary>
+		/// <param name="ProjectFile">Project file.</param>
+		/// <param name="UnrealExe">The name of the Unreal Editor executable to use.</param>
+		/// <param name="Maps">List of maps to rebuild light maps for. Can be null in which case -MapIniSection=AllMaps is used.</param>
+		/// <param name="Parameters">List of additional parameters.</param>
+		public static void ResavePackagesCommandlet(FileReference ProjectFile, string UnrealExe = null, string[] Maps = null, string Parameters = "")
         {
             string MapsToRebuildLighting = "";
             if (!IsNullOrEmpty(Maps))
@@ -161,18 +161,18 @@ namespace AutomationTool
                 MapsToRebuildLighting = "-Map=" + CombineCommandletParams(Maps).Trim();
             }
 
-            RunCommandlet(ProjectName, UnrealExe, "ResavePackages", String.Format((!String.IsNullOrEmpty(MapsToRebuildLighting) ? "-MapsOnly" : "") + "-ProjectOnly {0} {1}", MapsToRebuildLighting, Parameters));
+            RunCommandlet(ProjectFile, UnrealExe, "ResavePackages", String.Format((!String.IsNullOrEmpty(MapsToRebuildLighting) ? "-MapsOnly" : "") + "-ProjectOnly {0} {1}", MapsToRebuildLighting, Parameters));
         }
 
-        /// <summary>
-        /// Runs GenerateDistillFileSets commandlet.
-        /// </summary>
-        /// <param name="ProjectName">Project name.</param>
-        /// <param name="UnrealExe">The name of the Unreal Editor executable to use.</param>
-        /// <param name="Maps">List of maps to cook, can be null in which case -MapIniSection=AllMaps is used.</param>
-        /// <param name="TargetPlatform">Target platform.</param>
-        /// <param name="Parameters">List of additional parameters.</param>
-        public static List<FileReference> GenerateDistillFileSetsCommandlet(FileReference ProjectName, string ManifestFile, string UnrealExe = "UnrealEditor-Cmd.exe", string[] Maps = null, string Parameters = "")
+		/// <summary>
+		/// Runs GenerateDistillFileSets commandlet.
+		/// </summary>
+		/// <param name="ProjectFile">Project file.</param>
+		/// <param name="UnrealExe">The name of the Unreal Editor executable to use.</param>
+		/// <param name="Maps">List of maps to cook, can be null in which case -MapIniSection=AllMaps is used.</param>
+		/// <param name="TargetPlatform">Target platform.</param>
+		/// <param name="Parameters">List of additional parameters.</param>
+		public static List<FileReference> GenerateDistillFileSetsCommandlet(FileReference ProjectFile, string ManifestFile, string UnrealExe = null, string[] Maps = null, string Parameters = "")
         {
             string MapsToCook = "";
             if (!IsNullOrEmpty(Maps))
@@ -191,16 +191,16 @@ namespace AutomationTool
                 DeleteFile(ManifestFile);
             }
 
-            RunCommandlet(ProjectName, UnrealExe, "GenerateDistillFileSets", String.Format("{0} -OutputFolder={1} -Output={2} {3}", MapsToCook, CommandUtils.MakePathSafeToUseWithCommandLine(Dir), Filename, Parameters));
+            RunCommandlet(ProjectFile, UnrealExe, "GenerateDistillFileSets", String.Format("{0} -OutputFolder={1} -Output={2} {3}", MapsToCook, CommandUtils.MakePathSafeToUseWithCommandLine(Dir), Filename, Parameters));
 
             if (!FileExists_NoExceptions(ManifestFile))
             {
-                throw new AutomationException("GenerateDistillFileSets did not produce a manifest for {0}.", ProjectName);
+                throw new AutomationException("GenerateDistillFileSets did not produce a manifest for {0}.", ProjectFile);
             }
             var Lines = new List<string>(ReadAllLines(ManifestFile));
             if (Lines.Count < 1)
             {
-                throw new AutomationException("GenerateDistillFileSets for {0} did not produce any files.", ProjectName);
+                throw new AutomationException("GenerateDistillFileSets for {0} did not produce any files.", ProjectFile);
             }
             var Result = new List<FileReference>();
             foreach (var ThisFile in Lines)
@@ -222,63 +222,63 @@ namespace AutomationTool
             return Result;
         }
 
-        /// <summary>
-        /// Runs UpdateGameProject commandlet.
-        /// </summary>
-        /// <param name="ProjectName">Project name.</param>
-        /// <param name="UnrealExe">The name of the Unreal Editor executable to use.</param>
-        /// <param name="Parameters">List of additional parameters.</param>
-        public static void UpdateGameProjectCommandlet(FileReference ProjectName, string UnrealExe = "UnrealEditor-Cmd.exe", string Parameters = "")
+		/// <summary>
+		/// Runs UpdateGameProject commandlet.
+		/// </summary>
+		/// <param name="ProjectFile">Project file.</param>
+		/// <param name="UnrealExe">The name of the Unreal Editor executable to use.</param>
+		/// <param name="Parameters">List of additional parameters.</param>
+		public static void UpdateGameProjectCommandlet(FileReference ProjectFile, string UnrealExe = null, string Parameters = "")
         {
-            RunCommandlet(ProjectName, UnrealExe, "UpdateGameProject", Parameters);
+            RunCommandlet(ProjectFile, UnrealExe, "UpdateGameProject", Parameters);
         }
 
 		/// <summary>
 		/// Runs a commandlet using Engine/Binaries/Win64/UnrealEditor-Cmd.exe.
 		/// </summary>
-		/// <param name="ProjectName">Project name.</param>
+		/// <param name="ProjectFile">Project file.</param>
 		/// <param name="UnrealExe">The name of the Unreal Editor executable to use.</param>
 		/// <param name="Commandlet">Commandlet name.</param>
 		/// <param name="Parameters">Command line parameters (without -run=)</param>
 		/// <param name="ErrorLevel">The minimum exit code, which is treated as an error.</param>
-		public static void RunCommandlet(FileReference ProjectName, string UnrealExe, string Commandlet, string Parameters = null, uint ErrorLevel = 1)
+		public static void RunCommandlet(FileReference ProjectFile, string UnrealExe, string Commandlet, string Parameters = null, uint ErrorLevel = 1)
 		{
 			string LogFile;
-			RunCommandlet(ProjectName, UnrealExe, Commandlet, Parameters, out LogFile, ErrorLevel);
+			RunCommandlet(ProjectFile, UnrealExe, Commandlet, Parameters, out LogFile, ErrorLevel);
 		}
 
 		/// <summary>
 		/// Runs a commandlet using Engine/Binaries/Win64/UnrealEditor-Cmd.exe.
 		/// </summary>
-		/// <param name="ProjectName">Project name.</param>
+		/// <param name="ProjectFile">Project file.</param>
 		/// <param name="UnrealExe">The name of the Unreal Editor executable to use.</param>
 		/// <param name="Commandlet">Commandlet name.</param>
 		/// <param name="Parameters">Command line parameters (without -run=)</param>
 		/// <param name="ErrorLevel">The minimum exit code, which is treated as an error.</param>
-		public static void RunCommandlet(FileReference ProjectName, string UnrealExe, string Commandlet, string Parameters, int ErrorLevel)
+		public static void RunCommandlet(FileReference ProjectFile, string UnrealExe, string Commandlet, string Parameters, int ErrorLevel)
 		{
 			string LogFile;
-			RunCommandlet(ProjectName, UnrealExe, Commandlet, Parameters, out LogFile, (uint)ErrorLevel);
+			RunCommandlet(ProjectFile, UnrealExe, Commandlet, Parameters, out LogFile, (uint)ErrorLevel);
 		}
 
 		/// <summary>
 		/// Runs a commandlet using Engine/Binaries/Win64/UnrealEditor-Cmd.exe.
 		/// </summary>
-		/// <param name="ProjectName">Project name.</param>
+		/// <param name="ProjectFile">Project file.</param>
 		/// <param name="UnrealExe">The name of the Unreal Editor executable to use.</param>
 		/// <param name="Commandlet">Commandlet name.</param>
 		/// <param name="Parameters">Command line parameters (without -run=)</param>
 		/// <param name="DestLogFile">Log file after completion</param>
 		/// <param name="ErrorLevel">The minimum exit code, which is treated as an error.</param>
-		public static void RunCommandlet(FileReference ProjectName, string UnrealExe, string Commandlet, string Parameters, out string DestLogFile, uint ErrorLevel = 1)
+		public static void RunCommandlet(FileReference ProjectFile, string UnrealExe, string Commandlet, string Parameters, out string DestLogFile, uint ErrorLevel = 1)
 		{
 			string LocalLogFile;
 			IProcessResult RunResult;
 
 			DateTime StartTime = DateTime.UtcNow;
 
-			StartRunCommandlet(ProjectName, UnrealExe, Commandlet, Parameters, ERunOptions.Default, out LocalLogFile, out RunResult);
-			FinishRunCommandlet(ProjectName, Commandlet, StartTime, RunResult, LocalLogFile, out DestLogFile, ErrorLevel);
+			StartRunCommandlet(ProjectFile, UnrealExe, Commandlet, Parameters, ERunOptions.Default, out LocalLogFile, out RunResult);
+			FinishRunCommandlet(ProjectFile, Commandlet, StartTime, RunResult, LocalLogFile, out DestLogFile, ErrorLevel);
 		}
 
 		/// <summary>
@@ -290,22 +290,27 @@ namespace AutomationTool
 		/// <param name="Parameters">Command line parameters (without -run=)</param>
 		/// <param name="DestLogFile">Log file after completion</param>
 		/// <param name="ErrorLevel">The minimum exit code, which is treated as an error.</param>
-		public static void RunCommandlet(FileReference ProjectName, string UnrealExe, string Commandlet, string Parameters, out string DestLogFile, int ErrorLevel)
+		public static void RunCommandlet(FileReference ProjectFile, string UnrealExe, string Commandlet, string Parameters, out string DestLogFile, int ErrorLevel)
 		{
-			RunCommandlet(ProjectName, UnrealExe, Commandlet, Parameters, out DestLogFile, (uint)ErrorLevel);
+			RunCommandlet(ProjectFile, UnrealExe, Commandlet, Parameters, out DestLogFile, (uint)ErrorLevel);
 		}
 
-		public static void StartRunCommandlet(FileReference ProjectName, string UnrealExe, string Commandlet, string Parameters, ERunOptions RunOptions, out string LocalLogFile, out IProcessResult RunResult, ProcessResult.SpewFilterCallbackType SpewFilterCallback=null)
+		public static void StartRunCommandlet(FileReference ProjectFile, string UnrealExe, string Commandlet, string Parameters, ERunOptions RunOptions, out string LocalLogFile, out IProcessResult RunResult, ProcessResult.SpewFilterCallbackType SpewFilterCallback=null)
 		{
-			Logger.LogInformation("Running UnrealEditor {Commandlet} for project {ProjectName}", Commandlet, ProjectName);
+			Logger.LogInformation("Running UnrealEditor {Commandlet} for project {ProjectName}", Commandlet, ProjectFile);
 
 			var CWD = Path.GetDirectoryName(UnrealExe);
 
 			string EditorExe = UnrealExe;
 
+			if (EditorExe == null)
+			{
+				EditorExe = ProjectUtils.GetEditorForProject(ProjectFile).FullName;
+			}
+
 			if (String.IsNullOrEmpty(CWD))
 			{
-				EditorExe = HostPlatform.Current.GetUnrealExePath(UnrealExe);
+				EditorExe = HostPlatform.Current.GetUnrealExePath(EditorExe);
 				CWD = CombinePaths(CmdEnv.LocalRoot, HostPlatform.Current.RelativeBinariesFolder);
 			}
 
@@ -315,7 +320,7 @@ namespace AutomationTool
 			Logger.LogInformation("Commandlet log file is {LocalLogFile}", LocalLogFile);
 			string Args = String.Format(
 				"{0} -run={1} {2} -abslog={3} -stdout -CrashForUAT -unattended -NoLogTimes {5}{4}",
-				(ProjectName == null) ? "" : CommandUtils.MakePathSafeToUseWithCommandLine(ProjectName.FullName),
+				(ProjectFile == null) ? "" : CommandUtils.MakePathSafeToUseWithCommandLine(ProjectFile.FullName),
 				Commandlet,
 				String.IsNullOrEmpty(Parameters) ? "" : Parameters,
 				CommandUtils.MakePathSafeToUseWithCommandLine(LocalLogFile),
@@ -331,12 +336,12 @@ namespace AutomationTool
 			PopDir();
 		}
 
-		public static void FinishRunCommandlet(FileReference ProjectName, string Commandlet, DateTime StartTime, IProcessResult RunResult, string LocalLogFile, out string DestLogFile, uint ErrorLevel = 1)
+		public static void FinishRunCommandlet(FileReference ProjectFile, string Commandlet, DateTime StartTime, IProcessResult RunResult, string LocalLogFile, out string DestLogFile, uint ErrorLevel = 1)
 		{
 			// If we're running on a Windows build machine, copy any crash dumps into the log folder
 			if(HostPlatform.Current.HostEditorPlatform == UnrealTargetPlatform.Win64 && IsBuildMachine)
 			{
-				DirectoryInfo CrashesDir = new DirectoryInfo(GetCrashesDirectory(ProjectName).FullName);
+				DirectoryInfo CrashesDir = new DirectoryInfo(GetCrashesDirectory(ProjectFile).FullName);
 				if(CrashesDir.Exists)
 				{
 					foreach(DirectoryInfo CrashDir in CrashesDir.EnumerateDirectories())
@@ -434,7 +439,7 @@ namespace AutomationTool
 			{
 				Logger.LogWarning("Commandlet {Commandlet} failed to copy the local log file from {LocalLogFile} to {DestLogFile}. The log file will be lost.", Commandlet, LocalLogFile, DestLogFile);
 			}
-            string ProjectStatsDirectory = CombinePaths((ProjectName == null)? CombinePaths(CmdEnv.LocalRoot, "Engine") : Path.GetDirectoryName(ProjectName.FullName), "Saved", "Stats");
+            string ProjectStatsDirectory = CombinePaths((ProjectFile == null)? CombinePaths(CmdEnv.LocalRoot, "Engine") : Path.GetDirectoryName(ProjectFile.FullName), "Saved", "Stats");
             if (Directory.Exists(ProjectStatsDirectory))
             {
                 string DestCookerStats = CmdEnv.LogFolder;
@@ -469,13 +474,13 @@ namespace AutomationTool
 						ExitCodeDesc = String.Format(" (signal {0})", RunResult.ExitCode - 128);
 					}
 				}
-				throw new CommandletException(DestLogFile, RunResult.ExitCode, "Editor terminated with exit code {0}{1} while running {2}{3}; see log {4}", RunResult.ExitCode, ExitCodeDesc, Commandlet, (ProjectName == null)? "" : String.Format(" for {0}", ProjectName), DestLogFile) { OutputFormat = AutomationExceptionOutputFormat.Minimal };
+				throw new CommandletException(DestLogFile, RunResult.ExitCode, "Editor terminated with exit code {0}{1} while running {2}{3}; see log {4}", RunResult.ExitCode, ExitCodeDesc, Commandlet, (ProjectFile == null)? "" : String.Format(" for {0}", ProjectFile), DestLogFile) { OutputFormat = AutomationExceptionOutputFormat.Minimal };
 			}
 		}
 		
-		public static void FinishRunCommandlet(FileReference ProjectName, string Commandlet, DateTime StartTime, IProcessResult RunResult, string LocalLogFile, out string DestLogFile, int ErrorLevel)
+		public static void FinishRunCommandlet(FileReference ProjectFile, string Commandlet, DateTime StartTime, IProcessResult RunResult, string LocalLogFile, out string DestLogFile, int ErrorLevel)
 		{
-			FinishRunCommandlet(ProjectName, Commandlet, StartTime, RunResult, LocalLogFile, out DestLogFile, (uint)ErrorLevel);
+			FinishRunCommandlet(ProjectFile, Commandlet, StartTime, RunResult, LocalLogFile, out DestLogFile, (uint)ErrorLevel);
 		}
 
 		/// <summary>
diff --git a/Engine/Source/Programs/AutomationTool/Scripts/BenchmarkBuild/BenchmarkCookTask.cs b/Engine/Source/Programs/AutomationTool/Scripts/BenchmarkBuild/BenchmarkCookTask.cs
--- a/Engine/Source/Programs/AutomationTool/Scripts/BenchmarkBuild/BenchmarkCookTask.cs
+++ b/Engine/Source/Programs/AutomationTool/Scripts/BenchmarkBuild/BenchmarkCookTask.cs
@@ -70,7 +70,7 @@ namespace AutomationTool.Benchmark
 		{
 			string Arguments = GetBasicEditorCommandLine(false);
 			// will throw an exception if it fails
-			CommandUtils.RunCommandlet(ProjectTarget.ProjectFile, "UnrealEditor-Cmd.exe", "Cook", String.Format("-TargetPlatform={0} {1}", CookPlatformName, Arguments));
+			CommandUtils.RunCommandlet(ProjectTarget.ProjectFile, UnrealExe:null, "Cook", String.Format("-TargetPlatform={0} {1}", CookPlatformName, Arguments));
 			return true;
 		}
 	}
diff --git a/Engine/Source/Programs/AutomationTool/Scripts/ReplaceAssetsUsingManifest.cs b/Engine/Source/Programs/AutomationTool/Scripts/ReplaceAssetsUsingManifest.cs
--- a/Engine/Source/Programs/AutomationTool/Scripts/ReplaceAssetsUsingManifest.cs
+++ b/Engine/Source/Programs/AutomationTool/Scripts/ReplaceAssetsUsingManifest.cs
@@ -17,7 +17,8 @@ public class ReplaceAssetsUsingManifest : BuildCommand
 		// Command parameters - not all required if using existing manifest
 		string ProjectPath = ParseParamValue("ProjectPath");
 		string ManifestFile = ParseParamValue("ManifestFile");
-		string UnrealExe = ParseParamValue("UnrealExe", "UnrealEditor-Cmd.exe", ObsoleteParam: "UE4Exe");
+		// by defaulting to null here, we will perform a lookup of what editor the project wants to cook, etc with, based on target receipt
+		string UnrealExe = ParseParamValue("UnrealExe", null, ObsoleteParam: "UE4Exe");
 		string ReplacedPaths = ParseParamValue("ReplacedPaths", "");
 		string ReplacedClasses = ParseParamValue("ReplacedClasses", "");
 		string ExcludedPaths = ParseParamValue("ExcludedPaths", "");
diff --git a/Engine/Source/Programs/AutomationTool/Scripts/WorldPartitionBuilder.Automation.cs b/Engine/Source/Programs/AutomationTool/Scripts/WorldPartitionBuilder.Automation.cs
--- a/Engine/Source/Programs/AutomationTool/Scripts/WorldPartitionBuilder.Automation.cs
+++ b/Engine/Source/Programs/AutomationTool/Scripts/WorldPartitionBuilder.Automation.cs
@@ -46,13 +46,10 @@ namespace AutomationScripts.Automation
 				}
 			}
 
-			string EditorExe = "UnrealEditor-Cmd.exe";
-			EditorExe = AutomationTool.HostPlatform.Current.GetUnrealExePath(EditorExe);
-
 			FileReference ProjectPath = ParseProjectParam();
 
 			// Execute the commandlet - Will throw an exception on failures
-			RunCommandlet(ProjectPath, EditorExe, "WorldPartitionBuilderCommandlet", CommandletArgs);
+			RunCommandlet(ProjectPath, UnrealExe:null, "WorldPartitionBuilderCommandlet", CommandletArgs);
 			
 			if (bShelveResult)
 			{
diff --git a/Engine/Source/Programs/Shared/EpicGames.Core/FileSystemReference.cs b/Engine/Source/Programs/Shared/EpicGames.Core/FileSystemReference.cs
--- a/Engine/Source/Programs/Shared/EpicGames.Core/FileSystemReference.cs
+++ b/Engine/Source/Programs/Shared/EpicGames.Core/FileSystemReference.cs
@@ -130,6 +130,14 @@ namespace EpicGames.Core
 		/// <returns>True if this path is under the given directory</returns>
 		public bool IsUnderDirectory(DirectoryReference other) => FullName.StartsWith(other.FullName, Comparison) && (FullName.Length == other.FullName.Length || FullName[other.FullName.Length] == Path.DirectorySeparatorChar || other.IsRootDirectory());
 
+		/// <summary>
+		/// Checks to see if this exists as either a file or directory
+		/// This is helpful for Mac, because a binary may be a .app which is a directory
+		/// </summary>
+		/// <param name="location">FileSsytem object to check</param>
+		/// <returns>True if a file or a directory exists</returns>
+		public static bool Exists(FileSystemReference location) => File.Exists(location.FullName) || Directory.Exists(location.FullName);
+
 		/// <summary>
 		/// Searches the path fragments for the given name. Only complete fragments are considered a match.
 		/// </summary>
