Fix Unable to locate DeviceSupport directory with suffix 'DeveloperDiskImage.dmg' error when upgrade to xcode15 for iOS 17
PR #11491

diff --git a/Engine/Source/Programs/AutomationTool/Gauntlet/Platform/IOS/Gauntlet.TargetDeviceIOS.cs b/Engine/Source/Programs/AutomationTool/Gauntlet/Platform/IOS/Gauntlet.TargetDeviceIOS.cs
--- a/Engine/Source/Programs/AutomationTool/Gauntlet/Platform/IOS/Gauntlet.TargetDeviceIOS.cs
+++ b/Engine/Source/Programs/AutomationTool/Gauntlet/Platform/IOS/Gauntlet.TargetDeviceIOS.cs
@@ -277,27 +277,72 @@ namespace Gauntlet
 			// ios-deploy notes: -L launches detached, -I non-interactive  (exits when app exits), -r uninstalls before install (removes app Documents folder)
 			// -t <seconds> number of seconds to wait for device to be connected
 
-			// setup symbols if available
-			string DSymBundle = "";
-			string DSymDir = Path.Combine(GauntletAppCache, "Symbols");
+			// check xcodebuild -version, to see if it larger than 15.0
+			bool bIsXcode15 = false;
+			IProcessResult XcodeVersionResult = CommandUtils.Run("xcodebuild", "-version");
+			if (XcodeVersionResult.ExitCode == 0)
+			{
+				// XcodeVersionResult.Output is: Xcode 15, try get version number
+				Match VersionMatch = Regex.Match(XcodeVersionResult.Output, @"Xcode\ (?<Version>\d+)");
+				if (VersionMatch.Success)
+				{
+					int Version = 0;
+					if (int.TryParse(VersionMatch.Groups["Version"].Value, out Version))
+					{
+						bIsXcode15 = Version >= 15;
+					}
+				}
+			}
 
-			if (Directory.Exists(DSymDir))
+			IProcessResult Result;
+			if (bIsXcode15)
 			{
-				DSymBundle = Directory.GetDirectories(DSymDir).Where(D => Path.GetExtension(D).ToLower() == ".dsym").FirstOrDefault();
-				DSymBundle = string.IsNullOrEmpty(DSymBundle) ? "" : DSymBundle = " -S \"" + DSymBundle + "\"";
+				CommandLine = CommandLine.Replace("\\\\", "");
+				string CL = String.Format("device process launch --device {0} {1} '{2}'", DeviceName, IOSApp.PackageName, CommandLine.Trim());
+				Result = ExecuteIOSDeployCommand(CL, 0, true, false, true);
+				Thread.Sleep(5000);
+				// check if IOSApp.PackageName process running using devicectl device info processes
+				while(true)
+				{
+					CL = String.Format("device info processes --device {0}", DeviceName);
+					Result = ExecuteIOSDeployCommand(CL, 0, true, false, true);
+					// Get BundleName from IOSApp.PackageName
+					string BundleName = IOSApp.PackageName.Substring(IOSApp.PackageName.LastIndexOf('.') + 1);
+					Log.Info("Process running: " + BundleName);
+					// check if IOSApp.PackageName process not running
+					if (!Result.Output.Contains(BundleName))
+					{
+						break;
+					}
+					Thread.Sleep(5000);
+				}
+
 			}
+			else
+			{
+
+				// setup symbols if available
+				string DSymBundle = "";
+				string DSymDir = Path.Combine(GauntletAppCache, "Symbols");
 
-			string CL = "--noinstall -I" + DSymBundle + " -b \"" + LocalAppBundle + "\" --args '" + CommandLine.Trim() + "'";
+				if (Directory.Exists(DSymDir))
+				{
+					DSymBundle = Directory.GetDirectories(DSymDir).Where(D => Path.GetExtension(D).ToLower() == ".dsym").FirstOrDefault();
+					DSymBundle = string.IsNullOrEmpty(DSymBundle) ? "" : DSymBundle = " -S \"" + DSymBundle + "\"";
+				}
 
-			IProcessResult Result = ExecuteIOSDeployCommand(CL, 0);
+				string CL = "--noinstall -I" + DSymBundle + " -b \"" + LocalAppBundle + "\" --args '" + CommandLine.Trim() + "'";
 
-			Thread.Sleep(5000);
+				Result = ExecuteIOSDeployCommand(CL, 0);
 
-			// Give ios-deploy a chance to throw out any errors...
-			if (Result.HasExited)
-			{
-				Log.Warning(KnownLogEvents.Gauntlet_DeviceEvent, "ios-deploy exited early: " + Result.Output);
-				throw new DeviceException("Failed to launch on {0}. {1}", Name, Result.Output);
+				Thread.Sleep(5000);
+
+				// Give ios-deploy a chance to throw out any errors...
+				if (Result.HasExited)
+				{
+					Log.Warning(KnownLogEvents.Gauntlet_DeviceEvent, "ios-deploy exited early: " + Result.Output);
+					throw new DeviceException("Failed to launch on {0}. {1}", Name, Result.Output);
+				}
 			}
 
 			return new IOSAppInstance(IOSApp, Result, IOSApp.CommandLine);
@@ -912,14 +957,22 @@ namespace Gauntlet
 			return LocalDirectoryMappings;
 		}
 
-		public IProcessResult ExecuteIOSDeployCommand(String CommandLine, int WaitTime = 60, bool WarnOnTimeout = true, bool UseDeviceID = true)
+		public IProcessResult ExecuteIOSDeployCommand(String CommandLine, int WaitTime = 60, bool WarnOnTimeout = true, bool UseDeviceID = true, bool UseDeviceCtl = false)
 		{
-			if (UseDeviceID && !IsDefaultDevice)
+			String IOSDeployPath = Path.Combine(Globals.UnrealRootDir, "Engine/Extras/ThirdPartyNotUE/ios-deploy/bin/ios-deploy");
+			if (UseDeviceCtl)
 			{
-				CommandLine = String.Format("--id {0} {1}", DeviceName, CommandLine);
+				CommandLine = String.Format("devicectl {0}", CommandLine);
+				IOSDeployPath = "/usr/bin/xcrun";		
 			}
-
-			String IOSDeployPath = Path.Combine(Globals.UnrealRootDir, "Engine/Extras/ThirdPartyNotUE/ios-deploy/bin/ios-deploy");
+			else
+			{
+				if (UseDeviceID && !IsDefaultDevice)
+				{
+					CommandLine = String.Format("--id {0} {1}", DeviceName, CommandLine);
+				}
+			}
+	
 
 			if (!File.Exists(IOSDeployPath))
 			{
