Change log severity level for role for better report readability. Info in the header, and an Error in the role summary.
+ Better reflect severity base on role processed status

diff --git a/Engine/Source/Programs/AutomationTool/Gauntlet/Unreal/Base/Gauntlet.UnrealTestNode.cs b/Engine/Source/Programs/AutomationTool/Gauntlet/Unreal/Base/Gauntlet.UnrealTestNode.cs
--- a/Engine/Source/Programs/AutomationTool/Gauntlet/Unreal/Base/Gauntlet.UnrealTestNode.cs
+++ b/Engine/Source/Programs/AutomationTool/Gauntlet/Unreal/Base/Gauntlet.UnrealTestNode.cs
@@ -1927,10 +1927,17 @@ namespace Gauntlet
 
 			Log.Info(" #### Role: {0} ({1} {2})", RoleArtifacts.SessionRole.RoleType, RoleArtifacts.SessionRole.Platform, RoleArtifacts.SessionRole.Configuration);
 
-			bool HasFailed = InRoleResult.ExitCode != 0 && InRoleResult.LogSummary.HasAbnormalExit;
+			bool HasFailed = InRoleResult.ProcessResult != UnrealProcessResult.ExitOk;
 			string RoleState = HasFailed ? "failed:" : "completed:";
 			string StatusMessage = string.Format(" #### {0} {1} {2} ({3}, ExitCode={4})", RoleArtifacts.SessionRole.RoleType, RoleState, InRoleResult.Summary, InRoleResult.ProcessResult, InRoleResult.ExitCode);
-			Log.Info(StatusMessage);
+			if (HasFailed)
+			{
+				Log.Error(KnownLogEvents.Gauntlet_TestEvent, StatusMessage);
+			}
+			else
+			{
+				Log.Info(StatusMessage);
+			}
 
 			// log command line up here for visibility
 
@@ -2287,14 +2294,7 @@ namespace Gauntlet
 			{
 				string RoleName = RoleResult.Artifacts.SessionRole.RoleType.ToString();
 				string LogMessage = string.Format(" {0} Role: {1} ({2}, ExitCode={3})", RoleName, RoleResult.Summary, RoleResult.ProcessResult, RoleResult.ExitCode);
-				if (RoleResult.ExitCode != 0)
-				{
-					Log.Error(KnownLogEvents.Gauntlet_TestEvent, LogMessage);
-				}
-				else
-				{
-					Log.Info(LogMessage);
-				}
+				Log.Info(LogMessage);
 
 				FatalErrors += RoleResult.LogSummary.FatalError != null ? 1 : 0;
 				Ensures += RoleResult.LogSummary.Ensures.Count();
