clean up comments

diff --git a/Engine/Source/Programs/AutomationTool/AutomationUtils/MacHostPlatform.cs b/Engine/Source/Programs/AutomationTool/AutomationUtils/MacHostPlatform.cs
--- a/Engine/Source/Programs/AutomationTool/AutomationUtils/MacHostPlatform.cs
+++ b/Engine/Source/Programs/AutomationTool/AutomationUtils/MacHostPlatform.cs
@@ -127,7 +127,7 @@ namespace AutomationTool
 						AppName = AppName + ".app/Contents/MacOS/" + AppFilename;
 					}
 				}
-				// some of our C# applications are converted to dotnet core, do not run those via mono
+				// some of our C# applications are converted to dotnet core
 				else if (AppName.Contains("UnrealBuildTool") || AppName.Contains("AutomationTool"))
 				{
 					Options &= ~CommandUtils.ERunOptions.AppMustExist;
diff --git a/Engine/Source/Programs/AutomationTool/AutomationUtils/ProcessUtils.cs b/Engine/Source/Programs/AutomationTool/AutomationUtils/ProcessUtils.cs
--- a/Engine/Source/Programs/AutomationTool/AutomationUtils/ProcessUtils.cs
+++ b/Engine/Source/Programs/AutomationTool/AutomationUtils/ProcessUtils.cs
@@ -41,7 +41,6 @@ namespace AutomationTool
 	{
 		public delegate bool CtrlHandlerDelegate(CtrlTypes EventType);
 
-		// @todo: Add mono support
 		[DllImport("Kernel32")]
 		public static extern bool SetConsoleCtrlHandler(CtrlHandlerDelegate Handler, bool Add);
 
