Changed the staging option '-Upload=LocalZen' to use IasTool rather than UnrealPak.

#rnx

- Instead of passing in '-Upload=<PathToPaks>' the IasTool expects first a command 'Upload' followed by '<PathToPaks>' as it's own arg.
- Rather than adding an IasTool version of ::GetUnrealPakLocation, I have opted to keep the code all together in a single function as we are likely to remove this from the general staging scripts at some point.

diff --git a/Engine/Source/Programs/AutomationTool/Scripts/CopyBuildToStagingDirectory.Automation.cs b/Engine/Source/Programs/AutomationTool/Scripts/CopyBuildToStagingDirectory.Automation.cs
--- a/Engine/Source/Programs/AutomationTool/Scripts/CopyBuildToStagingDirectory.Automation.cs
+++ b/Engine/Source/Programs/AutomationTool/Scripts/CopyBuildToStagingDirectory.Automation.cs
@@ -3783,7 +3783,7 @@ namespace AutomationScripts
 					string CloudDir = CombinePaths(SC.StageDirectory.FullName, "Cloud");
 					string ConfigFilePath = CombinePaths(CloudDir, "IoStoreOnDemand.ini");
 
-					UploadArgs.Append(string.Format("-Upload={0}", PakPath));
+					UploadArgs.Append(string.Format("Upload {0}", PakPath));
 					UploadArgs.Append(string.Format(" -ServiceUrl={0}/obj/bucket", ServiceUrl));
 					UploadArgs.Append(" -Bucket=" + SC.ShortProjectName.ToLower());
 					UploadArgs.Append(" -ConfigFilePath=\"" + ConfigFilePath + "\"");
@@ -3810,9 +3810,16 @@ namespace AutomationScripts
 
 				if (UploadArgs.Length > 0)
 				{
-					string PakArgs = UploadArgs.ToString();
-					Logger.LogInformation("Running UnrealPak with arguments: {CommandletParams}", PakArgs);
-					RunAndLog(CmdEnv, GetUnrealPakLocation().FullName, PakArgs, Options: ERunOptions.Default | ERunOptions.UTF8Output);
+					string Args = UploadArgs.ToString();
+					Logger.LogInformation("Running IasTool with arguments: {CommandletParams}", Args);
+
+					FileReference IasToolPath = FileReference.Combine(Unreal.EngineDirectory, "Binaries", HostPlatform.Current.HostEditorPlatform.ToString(), "IasTool");
+					if (HostPlatform.Current.HostEditorPlatform == UnrealTargetPlatform.Win64)
+					{
+						IasToolPath += ".exe";
+					}
+
+					RunAndLog(CmdEnv, IasToolPath.FullName, Args, Options: ERunOptions.Default | ERunOptions.UTF8Output);
 				}
 			}
 
