Update allowed .NET versions to fix build warnings, and add support for parsing NET core versions.

diff --git a/Engine/Source/Programs/AutomationTool/Scripts/CheckCsprojDotNetVersion.cs b/Engine/Source/Programs/AutomationTool/Scripts/CheckCsprojDotNetVersion.cs
--- a/Engine/Source/Programs/AutomationTool/Scripts/CheckCsprojDotNetVersion.cs
+++ b/Engine/Source/Programs/AutomationTool/Scripts/CheckCsprojDotNetVersion.cs
@@ -9,6 +9,7 @@ using UnrealBuildBase;
 using Microsoft.Extensions.Logging;
 
 using static AutomationTool.CommandUtils;
+using System.Collections.Generic;
 
 class CheckCsprojDotNetVersion : BuildCommand
 {
@@ -21,7 +22,7 @@ class CheckCsprojDotNetVersion : BuildCommand
 			throw new AutomationException("-TargetVersion was not specified.");
 		}
 
-		string[] DesiredTargetVersions = DesiredTargetVersionParam.Split('+');
+		HashSet<string> DesiredTargetVersions = new HashSet<string>(DesiredTargetVersionParam.Split('+', StringSplitOptions.RemoveEmptyEntries | StringSplitOptions.TrimEntries));
 
 		Logger.LogInformation("Scanning for all csproj's...");
 		// Check for all csproj's in the engine dir
@@ -31,7 +32,7 @@ class CheckCsprojDotNetVersion : BuildCommand
 		Regex FrameworkRegex = new Regex("<TargetFrameworkVersion>v(\\d\\.\\d\\.?\\d?)<\\/TargetFrameworkVersion>");
 		Regex PossibleAppConfigRegex = new Regex("<TargetFrameworkProfile>(.+)<\\/TargetFrameworkProfile>");
 		Regex AppConfigRegex = new Regex("<supportedRuntime version=\"v(\\d\\.\\d\\.?\\d?)\" sku=\"\\.NETFramework,Version=v(\\d\\.\\d\\.?\\d?),Profile=(.+)\"\\/>");
-		Regex DotNetCoreRegex = new Regex("<TargetFramework>(netcoreapp2.0|netstandard2.0)<\\/TargetFramework>");
+		Regex NetCoreRegex = new Regex("<TargetFramework>(.*)<\\/TargetFramework>");
 		foreach (FileReference CsProj in DirectoryReference.EnumerateFiles(EngineDir, "*.csproj", SearchOption.AllDirectories))
 		{
 			if (CsProj.ContainsName("ThirdParty", EngineDir) ||
@@ -41,21 +42,25 @@ class CheckCsprojDotNetVersion : BuildCommand
 				CsProj.ContainsName("DatasmithRevitExporter", EngineDir) ||
 				CsProj.ContainsName("DatasmithNavisworksExporter", EngineDir) ||
 				CsProj.ContainsName("CSVTools", EngineDir))
-
 			{
 				continue;
 			}
 
 			// read in the file
 			string Contents = File.ReadAllText(CsProj.FullName);
-			Match Match = DotNetCoreRegex.Match(Contents);
-			// Check if we're a _NETCore app, ignore these.
+
+			// Check if we're a _NETCore app
+			Match Match = NetCoreRegex.Match(Contents);
 			if (Match.Success)
 			{
+				string TargetedVersion = Regex.Replace(Match.Groups[1].Value, "-.*$", "");
+				if (!DesiredTargetVersions.Contains(TargetedVersion))
+				{
+					Logger.LogWarning("Targeted Framework version for project: {CsProj} was not {Arg1}! Targeted Version: {TargetedVersion}", CsProj, String.Join("/", DesiredTargetVersions), TargetedVersion);
+				}
 				continue;
 			}
 
-
 			Match = FrameworkRegex.Match(Contents);
 			if (Match.Success)
 			{
