Automation support for using an HTTP2-only endpoint when publishing snapshots to the cloud.  Also ensured that the snapshot descriptor can refer to an alternate endpoint that isn't HTTP2-only so that you can publish to one endpoint, but still leave the descriptor referring to another.

diff --git a/Engine/Build/Graph/Tasks/BuildAndTestProject.xml b/Engine/Build/Graph/Tasks/BuildAndTestProject.xml
--- a/Engine/Build/Graph/Tasks/BuildAndTestProject.xml
+++ b/Engine/Build/Graph/Tasks/BuildAndTestProject.xml
@@ -3,6 +3,9 @@
 	<EnvVar Name="UE_HORDE_JOBID"/>
 	
 	<EnvVar Name="UE-CloudPublishHost"/>
+	<EnvVar Name="UE-CloudPublishHttpVersion"/>
+	<EnvVar Name="UE-CloudPublishDescriptorHost"/>
+	<EnvVar Name="UE-CloudPublishDescriptorHttpVersion"/>
 	<EnvVar Name="UE-CloudPublishNamespace"/>
 
 	<Option Name="Horde" Restrict="true|false" DefaultValue="$(IsBuildMachine)" Description="Whether the build is running under Horde"/>
@@ -553,7 +556,7 @@
 			<Do If="$(PublishBuilds) and $(PublishCloudSnapshot) and '$(UE-CloudPublishHost)' != '' and '$(UE-CloudPublishNamespace)' != ''">
 				<Property Name="PlatformCloudPublishNodeName" Value="$(PreNodeName)Cloud Publish $(TargetPlatform)" />
 				<Node Name="$(PlatformCloudPublishNodeName)" Requires="$(PlatformHashNodeName)">
-					<ZenExportSnapshot Project="$(TargetProject)" Platform="$(CookPlatform)" DestinationStorageType="Cloud" DestinationCloudHost="$(UE-CloudPublishHost)" DestinationCloudNamespace="$(UE-CloudPublishNamespace)" DestinationIdentifier="$(EscapedBranch)/$(BuildNamePath)" SnapshotDescriptorFile="$(NetworkOutputDirectory)/$(StagedPlatformFolder)/SnapshotDescriptor.json"/>
+					<ZenExportSnapshot Project="$(TargetProject)" Platform="$(CookPlatform)" DestinationStorageType="Cloud" DestinationCloudHost="$(UE-CloudPublishHost)" DestinationCloudHttpVersion="$(UE-CloudPublishHttpVersion)" SnapshotDescriptorCloudHost="$(UE-CloudPublishDescriptorHost)" SnapshotDescriptorCloudHttpVersion="$(UE-CloudPublishDescriptorHttpVersion)" DestinationCloudNamespace="$(UE-CloudPublishNamespace)" DestinationIdentifier="$(EscapedBranch)/$(BuildNamePath)" SnapshotDescriptorFile="$(NetworkOutputDirectory)/$(StagedPlatformFolder)/SnapshotDescriptor.json"/>
 				</Node>
 				<Property Name="PlatformPublishStageDependencyName" Value="$(PlatformCloudPublishNodeName)" />
 			</Do>
diff --git a/Engine/Source/Programs/AutomationTool/BuildGraph/Tasks/ZenExportSnapshotTask.cs b/Engine/Source/Programs/AutomationTool/BuildGraph/Tasks/ZenExportSnapshotTask.cs
--- a/Engine/Source/Programs/AutomationTool/BuildGraph/Tasks/ZenExportSnapshotTask.cs
+++ b/Engine/Source/Programs/AutomationTool/BuildGraph/Tasks/ZenExportSnapshotTask.cs
@@ -87,6 +87,24 @@ namespace AutomationTool.Tasks
 		[TaskParameter(Optional = true)]
 		public string DestinationCloudHost;
 
+		/// <summary>
+		/// The host name to use when writing a snapshot descriptor for a cloud destination
+		/// </summary>
+		[TaskParameter(Optional = true)]
+		public string SnapshotDescriptorCloudHost;
+
+		/// <summary>
+		/// The http version to use when exporting to a cloud destination
+		/// </summary>
+		[TaskParameter(Optional = true)]
+		public string DestinationCloudHttpVersion;
+
+		/// <summary>
+		/// The http version to use when writing a snapshot descriptor for a cloud destination
+		/// </summary>
+		[TaskParameter(Optional = true)]
+		public string SnapshotDescriptorCloudHttpVersion;
+
 		/// <summary>
 		/// The namespace to use when exporting to a cloud destination
 		/// </summary>
@@ -298,11 +316,27 @@ namespace AutomationTool.Tasks
 						BucketName = ProjectNameAsBucketName;
 					}
 
+					string HostName = Parameters.SnapshotDescriptorCloudHost;
+					if (string.IsNullOrEmpty(HostName))
+					{
+						HostName = Parameters.DestinationCloudHost;
+					}
+
+					string HttpVersion = Parameters.SnapshotDescriptorCloudHttpVersion;
+					if (string.IsNullOrEmpty(HttpVersion))
+					{
+						HostName = Parameters.DestinationCloudHttpVersion;
+					}
+
 					IoHash DestinationKeyHash = IoHash.Compute(Encoding.UTF8.GetBytes(Name));
 					Writer.WriteValue("name", Name);
 					Writer.WriteValue("type", "cloud");
 					Writer.WriteValue("targetplatform", ExportSource.TargetPlatform);
-					Writer.WriteValue("host", Parameters.DestinationCloudHost);
+					Writer.WriteValue("host", HostName);
+					if (!string.IsNullOrEmpty(HttpVersion) && !HttpVersion.Equals("None", StringComparison.InvariantCultureIgnoreCase))
+					{
+						Writer.WriteValue("httpversion", HttpVersion);
+					}
 					Writer.WriteValue("namespace", Parameters.DestinationCloudNamespace);
 					Writer.WriteValue("bucket", BucketName);
 					Writer.WriteValue("key", DestinationKeyHash.ToString().ToLowerInvariant());
@@ -433,6 +467,18 @@ namespace AutomationTool.Tasks
 
 					OplogExportCommandline.AppendFormat(" --cloud {0} --namespace {1} --bucket {2}", Parameters.DestinationCloudHost, Parameters.DestinationCloudNamespace, BucketName);
 
+					if (!string.IsNullOrEmpty(Parameters.DestinationCloudHttpVersion))
+					{
+						if (Parameters.DestinationCloudHttpVersion.Equals("http2-only", StringComparison.InvariantCultureIgnoreCase))
+						{
+							OplogExportCommandline.Append(" --assume-http2");
+						}
+						else
+						{
+							throw new AutomationException("Unexpected destination cloud http version");
+						}
+					}
+
 					ExportIndex = 0;
 					foreach (ExportSourceData ExportSource in ExportSources)
 					{
