Remove P4 syncing from updating version files on buildmachines. Builders can track their own workspace without having to sync files manually.

diff --git a/Engine/Source/Programs/AutomationTool/AutomationUtils/UnrealBuild.cs b/Engine/Source/Programs/AutomationTool/AutomationUtils/UnrealBuild.cs
--- a/Engine/Source/Programs/AutomationTool/AutomationUtils/UnrealBuild.cs
+++ b/Engine/Source/Programs/AutomationTool/AutomationUtils/UnrealBuild.cs
@@ -320,7 +320,7 @@ namespace AutomationTool
 			FileReference BuildVersionFile = BuildVersion.GetDefaultFileName();
 
 			// Get the revision to sync files to before 
-			if(CommandUtils.P4Enabled && ChangelistNumber > 0)
+			if(CommandUtils.P4Enabled && ChangelistNumber > 0 && !CommandUtils.IsBuildMachine)
 			{
 				CommandUtils.P4.Sync(String.Format("-f \"{0}@{1}\"", BuildVersionFile, ChangelistNumber), false, false);
 			}
diff --git a/Engine/Source/Programs/AutomationTool/AutomationUtils/Utils.cs b/Engine/Source/Programs/AutomationTool/AutomationUtils/Utils.cs
--- a/Engine/Source/Programs/AutomationTool/AutomationUtils/Utils.cs
+++ b/Engine/Source/Programs/AutomationTool/AutomationUtils/Utils.cs
@@ -1348,7 +1348,7 @@ namespace AutomationTool
 		{
 			if(CommandUtils.IsReadOnly(FileName))
 			{
-				if(CommandUtils.P4Enabled)
+				if(CommandUtils.P4Enabled && !CommandUtils.IsBuildMachine)
 				{
 					CommandUtils.P4.Sync(String.Format("\"{0}#0\"", FileName), false, false);
 				}
