Fix the IAS staging arg '-Upload=LocalZen' to work with projects that contain spaces in their path.

#rnx

- Make sure that the path to the pakfile directory and the path to the cryptokeys are properly quoted.

diff --git a/Engine/Source/Programs/AutomationTool/Scripts/CopyBuildToStagingDirectory.Automation.cs b/Engine/Source/Programs/AutomationTool/Scripts/CopyBuildToStagingDirectory.Automation.cs
--- a/Engine/Source/Programs/AutomationTool/Scripts/CopyBuildToStagingDirectory.Automation.cs
+++ b/Engine/Source/Programs/AutomationTool/Scripts/CopyBuildToStagingDirectory.Automation.cs
@@ -3783,7 +3783,7 @@ namespace AutomationScripts
 					string CloudDir = CombinePaths(SC.StageDirectory.FullName, "Cloud");
 					string ConfigFilePath = CombinePaths(CloudDir, "IoStoreOnDemand.ini");
 
-					UploadArgs.Append(string.Format("Upload {0}", PakPath));
+					UploadArgs.Append(string.Format("Upload \"{0}\"", PakPath));
 					UploadArgs.Append(string.Format(" -ServiceUrl={0}/obj/bucket", ServiceUrl));
 					UploadArgs.Append(" -Bucket=" + SC.ShortProjectName.ToLower());
 					UploadArgs.Append(" -ConfigFilePath=\"" + ConfigFilePath + "\"");
@@ -3800,7 +3800,7 @@ namespace AutomationScripts
 					FileReference CryptoKeysFilename = FileReference.Combine(SC.MetadataDir, "Crypto.json");
 					if (FileExists_NoExceptions(CryptoKeysFilename.FullName))
 					{
-						UploadArgs.Append(string.Format(" -CryptoKeys={0}", CryptoKeysFilename.FullName));
+						UploadArgs.Append(string.Format(" -CryptoKeys=\"{0}\"", CryptoKeysFilename.FullName));
 					}
 				}
 				else
