Improve .uondemandtoc handling during staging.

#rnx

- There were multiple areas that checked if a cooked file was of a certain extension, this code has been consolidated in a single method ::CanCookedFileBeStaged and as we were excluding .ucas/.utoc I added .uondemandtoc to be excluded too.
- There are three places with code logic that try to copy container files from one location to another but each one is subtley different so I left the code duplication in place but added support to copy any .uondemandtoc file that might exist. Unlike the .ucas/.utoc the new type is optional so that had to be taken into account.
- ::DeletePakFiles has been updated to also delete all .uondemandtoc files present. The only code path calling this does not seem to be easily accessed, I have added a code comment describing the only method I found that could actually cause it to be executed to hopefully help the next person working in this area.
- Note that in the code branch for creating a patch for chunked installed builds (line 3955) I am not confident that this works at all. All if the file copying operations use 'PatchSourceContentPath' as the source which is almost certainly wrong and the path itself seems to always contain a wildcard character which would cause it to be wrong as well. For now I have followed the current code convention and used 'PatchSourceContentPath' when copying the .uondemandtoc but looking into this code path will need to be it's own work item.

diff --git a/Engine/Source/Programs/AutomationTool/Scripts/CopyBuildToStagingDirectory.Automation.cs b/Engine/Source/Programs/AutomationTool/Scripts/CopyBuildToStagingDirectory.Automation.cs
--- a/Engine/Source/Programs/AutomationTool/Scripts/CopyBuildToStagingDirectory.Automation.cs
+++ b/Engine/Source/Programs/AutomationTool/Scripts/CopyBuildToStagingDirectory.Automation.cs
@@ -931,6 +931,29 @@ namespace AutomationScripts
 			}
 		}
 
+		private static bool CanCookedFileBeStaged(FileReference File)
+		{
+			// json files have never been staged
+			if (File.HasExtension(".json"))
+			{
+				return false;
+			}
+
+			// metallib files cannot *currently* be staged as UFS as the Metal API needs to mmap them from files on disk in order to function efficiently
+			if (File.HasExtension(".metallib"))
+			{
+				return false;
+			}
+
+			// Cannot stage
+			if (File.HasExtension(".utoc") || File.HasExtension(".ucas") || File.HasExtension(".uondemandtoc"))
+			{
+				return false;
+			}
+
+			return true;
+		}
+
 		public static bool SetUpStagingSourceDirectories(ProjectParams Params, DeploymentContext SC)
 		{
 			if (SC.CookPlatform == null)
@@ -1128,10 +1151,8 @@ namespace AutomationScripts
 					{
 						continue;
 					}
-
-					// json files have never been staged
-					// metallib files cannot *currently* be staged as UFS as the Metal API needs to mmap them from files on disk in order to function efficiently
-					if (!CookedFile.HasExtension(".json") && !CookedFile.HasExtension(".metallib") && !CookedFile.HasExtension(".utoc") && !CookedFile.HasExtension(".ucas"))
+					
+					if (CanCookedFileBeStaged(CookedFile))
 					{
 						SC.StageFile(StagedFileType.UFS, CookedFile, new StagedFileReference(CookedFile.MakeRelativeTo(SC.PlatformCookDir)));
 					}
@@ -1484,9 +1505,7 @@ namespace AutomationScripts
 								continue;
 							}
 
-							// json files have never been staged
-							// metallib files cannot *currently* be staged as UFS as the Metal API needs to mmap them from files on disk in order to function efficiently
-							if (!CookedFile.HasExtension(".json") && !CookedFile.HasExtension(".metallib") && !CookedFile.HasExtension(".utoc") && !CookedFile.HasExtension(".ucas"))
+							if (CanCookedFileBeStaged(CookedFile))
 							{
 								SC.StageFile(StagedFileType.UFS, CookedFile, new StagedFileReference(CookedFile.MakeRelativeTo(SC.PlatformCookDir)));
 							}
@@ -2728,7 +2747,7 @@ namespace AutomationScripts
 				}
 
 				// Filter I/O store container files
-				if (Src.HasExtension(".ucas") || Src.HasExtension(".utoc"))
+				if (Src.HasExtension(".ucas") || Src.HasExtension(".utoc") || Src.HasExtension(".uondemandtoc"))
 				{
 					Logger.LogInformation("Excluding {Src}", Src);
 					continue;
@@ -2901,6 +2920,25 @@ namespace AutomationScripts
 				Logger.LogInformation("Missing utoc file {InUtocFile}, creating new pak", InUtocFile);
 				bCopiedExistingPak = false;
 			}
+
+			// The .uondemandtoc is optional and should only be looked for if both the .ucas and .utoc were copied successfully.
+			if (bCopiedExistingPak)
+			{
+				FileReference OnDemandTocSrcFile = InUtocFile.ChangeExtension(".uondemandtoc");
+				if (FileReference.Exists(OnDemandTocSrcFile))
+				{
+					FileReference OnDemandTocDstFile = OutputLocation.ChangeExtension(".uondemandtoc");
+
+					Logger.LogInformation("Copying utoc from {Src} to {Dst}", OnDemandTocSrcFile, OnDemandTocDstFile);
+
+					if (!InternalUtils.SafeCopyFile(OnDemandTocSrcFile.FullName, OnDemandTocDstFile.FullName))
+					{
+						Logger.LogInformation("Failed to copy utoc {Src} to {Dst}, creating new pak", OnDemandTocSrcFile, OnDemandTocDstFile);
+						bCopiedExistingPak = false;
+					}
+				}
+			}
+
 			return bCopiedExistingPak;
 		}
 
@@ -3850,6 +3888,13 @@ namespace AutomationScripts
 					{
 						InternalUtils.SafeCopyFile(Path.ChangeExtension(OutputLocation.FullName, ".utoc"), Path.ChangeExtension(ReleaseVersionPath, ".utoc"));
 						InternalUtils.SafeCopyFile(Path.ChangeExtension(OutputLocation.FullName, ".ucas"), Path.ChangeExtension(ReleaseVersionPath, ".ucas"));
+
+						// Check if the optional .uondemandtoc file exists before trying to copy it
+						string OnDemandTocSrcPath = Path.ChangeExtension(OutputLocation.FullName, ".uondemandtoc");
+						if (File.Exists(OnDemandTocSrcPath))
+						{
+							InternalUtils.SafeCopyFile(OnDemandTocSrcPath, Path.ChangeExtension(ReleaseVersionPath, ".uondemandtoc"));
+						}
 					}
 				}
 
@@ -3897,6 +3942,13 @@ namespace AutomationScripts
 						{
 							InternalUtils.SafeCopyFile(Path.ChangeExtension(OutputLocation.FullName, ".utoc"), Path.ChangeExtension(RawDataPakPath, ".utoc"), true);
 							InternalUtils.SafeCopyFile(Path.ChangeExtension(OutputLocation.FullName, ".ucas"), Path.ChangeExtension(RawDataPakPath, ".ucas"), true);
+
+							// Check if the optional .uondemandtoc file exists before trying to copy it
+							string OnDemandTocSrcPath = Path.ChangeExtension(OutputLocation.FullName, ".uondemandtoc");
+							if (File.Exists(OnDemandTocSrcPath))
+							{
+								InternalUtils.SafeCopyFile(OnDemandTocSrcPath, Path.ChangeExtension(RawDataPakPath, ".uondemandtoc"), true);
+							}
 						}
 						InternalUtils.SafeDeleteFile(OutputLocation.FullName, true);
 
@@ -3918,6 +3970,12 @@ namespace AutomationScripts
 							{
 								InternalUtils.SafeCopyFile(PatchSourceContentPath, Path.ChangeExtension(SourceRawDataPakPath, ".utoc"), true);
 								InternalUtils.SafeCopyFile(PatchSourceContentPath, Path.ChangeExtension(SourceRawDataPakPath, ".ucas"), true);
+
+								// Check if the optional .uondemandtoc file exists before trying to copy it
+								if (File.Exists(PatchSourceContentPath))
+								{
+									InternalUtils.SafeCopyFile(PatchSourceContentPath, Path.ChangeExtension(SourceRawDataPakPath, ".uondemandtoc"), true);
+								}
 							}
 						}
 
@@ -3970,6 +4028,7 @@ namespace AutomationScripts
 							{
 								IncludedExtensions.Add(".ucas");
 								IncludedExtensions.Add(".utoc");
+								IncludedExtensions.Add(".uondemandtoc");
 							}
 							IEnumerable<string> PakFileSet = Directory.EnumerateFiles(ExistingPatchSearchPath, PakName + "-" + SC.FinalCookPlatform + "*.*");
 							foreach (string PakFilePath in PakFileSet)
@@ -4770,6 +4829,7 @@ namespace AutomationScripts
 			StagedFilesDir.GetFiles("*.pak", SearchOption.AllDirectories).ToList().ForEach(File => File.Delete());
 			StagedFilesDir.GetFiles("*.ucas", SearchOption.AllDirectories).ToList().ForEach(File => File.Delete());
 			StagedFilesDir.GetFiles("*.utoc", SearchOption.AllDirectories).ToList().ForEach(File => File.Delete());
+			StagedFilesDir.GetFiles("*.uondemandtoc", SearchOption.AllDirectories).ToList().ForEach(File => File.Delete());
 		}
 
 		protected static void CleanDirectoryExcludingPakFiles(DirectoryInfo StagingDirectory)
@@ -4831,6 +4891,9 @@ namespace AutomationScripts
 			}
 			else
 			{
+				// Note: To get to this code we need to run BuildCookRun with no staging parameters at all (-pak, -staging, -skipstaging) and the target platform returns
+				// 'PakType.Always' from Platform::RequiresPak
+
 				Logger.LogInformation("Cleaning PAK files in stage directory: {Arg0}", SC.StageDirectory.FullName);
 				try
 				{
