[Gauntlet] Mark Gauntlet builds with whether or not they have the capability to copy additional files. UnrealSession will only copy files if the build supports it

diff --git a/Engine/Source/Programs/AutomationTool/Gauntlet/Framework/Base/Gauntlet.BuildSource.cs b/Engine/Source/Programs/AutomationTool/Gauntlet/Framework/Base/Gauntlet.BuildSource.cs
--- a/Engine/Source/Programs/AutomationTool/Gauntlet/Framework/Base/Gauntlet.BuildSource.cs
+++ b/Engine/Source/Programs/AutomationTool/Gauntlet/Framework/Base/Gauntlet.BuildSource.cs
@@ -54,6 +54,11 @@ namespace Gauntlet
 		/// </summary>
 		UnrealTargetConfiguration Configuration { get; }
 
+		/// <summary>
+		/// Does this build support additional loose files be copied to the device
+		/// </summary>
+		bool SupportsAdditionalFileCopy { get; }
+
 		/// <summary>
 		/// Check if this buld is able to support the provided role
 		/// </summary>
diff --git a/Engine/Source/Programs/AutomationTool/Gauntlet/Platform/Android/Gauntlet.AndroidBuildSource.cs b/Engine/Source/Programs/AutomationTool/Gauntlet/Platform/Android/Gauntlet.AndroidBuildSource.cs
--- a/Engine/Source/Programs/AutomationTool/Gauntlet/Platform/Android/Gauntlet.AndroidBuildSource.cs
+++ b/Engine/Source/Programs/AutomationTool/Gauntlet/Platform/Android/Gauntlet.AndroidBuildSource.cs
@@ -30,6 +30,8 @@ namespace Gauntlet
 
 		public UnrealTargetPlatform Platform { get { return UnrealTargetPlatform.Android; } }
 
+		public bool SupportsAdditionalFileCopy => true;
+
 		public bool Is32Bit { get; protected set; }
 
 		public bool UsesExternalFilesDir { get; protected set; }
diff --git a/Engine/Source/Programs/AutomationTool/Gauntlet/Platform/IOS/Gauntlet.IOSBuildSource.cs b/Engine/Source/Programs/AutomationTool/Gauntlet/Platform/IOS/Gauntlet.IOSBuildSource.cs
--- a/Engine/Source/Programs/AutomationTool/Gauntlet/Platform/IOS/Gauntlet.IOSBuildSource.cs
+++ b/Engine/Source/Programs/AutomationTool/Gauntlet/Platform/IOS/Gauntlet.IOSBuildSource.cs
@@ -30,6 +30,8 @@ namespace Gauntlet
 
 		public UnrealTargetPlatform Platform { get { return UnrealTargetPlatform.IOS; } }
 
+		public bool SupportsAdditionalFileCopy { get; }
+
 		public IOSBuild(UnrealTargetConfiguration InConfig, string InPackageName, string InIPAPath, Dictionary<string, string> InFilesToInstall, BuildFlags InFlags)
 		{
 			Configuration = InConfig;
@@ -37,6 +39,7 @@ namespace Gauntlet
 			SourceIPAPath = InIPAPath;
 			FilesToInstall = InFilesToInstall;
 			Flags = InFlags;
+			SupportsAdditionalFileCopy = true;
 		}
 
 		public bool CanSupportRole(UnrealTargetRole RoleType)
diff --git a/Engine/Source/Programs/AutomationTool/Gauntlet/Unreal/Base/Gauntlet.UnrealSession.cs b/Engine/Source/Programs/AutomationTool/Gauntlet/Unreal/Base/Gauntlet.UnrealSession.cs
--- a/Engine/Source/Programs/AutomationTool/Gauntlet/Unreal/Base/Gauntlet.UnrealSession.cs
+++ b/Engine/Source/Programs/AutomationTool/Gauntlet/Unreal/Base/Gauntlet.UnrealSession.cs
@@ -1689,7 +1689,12 @@ namespace Gauntlet
 					}
 
 					Device.CleanArtifacts();
-					Device.CopyAdditionalFiles(AppConfig.FilesToCopy);
+
+					if(AppConfig.Build.SupportsAdditionalFileCopy)
+					{
+						Device.CopyAdditionalFiles(AppConfig.FilesToCopy);
+					}
+
 					Role.ConfigureDevice?.Invoke(Device);
 				}
 				catch(Exception Ex)
diff --git a/Engine/Source/Programs/AutomationTool/Gauntlet/Unreal/BuildSource/Gauntlet.StagedBuildSource.cs b/Engine/Source/Programs/AutomationTool/Gauntlet/Unreal/BuildSource/Gauntlet.StagedBuildSource.cs
--- a/Engine/Source/Programs/AutomationTool/Gauntlet/Unreal/BuildSource/Gauntlet.StagedBuildSource.cs
+++ b/Engine/Source/Programs/AutomationTool/Gauntlet/Unreal/BuildSource/Gauntlet.StagedBuildSource.cs
@@ -12,7 +12,7 @@ using System.Linq;
 namespace Gauntlet
 {
 
-	public class EditorBuild : IBuild 
+	public class EditorBuild : IBuild
 	{
 		public int PreferenceOrder { get { return 0; } }
 
@@ -22,6 +22,8 @@ namespace Gauntlet
 
 		public BuildFlags Flags { get { return BuildFlags.CanReplaceCommandLine | BuildFlags.Loose; } }
 
+		public bool SupportsAdditionalFileCopy => true;
+
 		public string Flavor { get { return ""; } }
 
 		public bool CanSupportRole(UnrealTargetRole InRoleType) { return InRoleType.UsesEditor(); }
@@ -49,6 +51,8 @@ namespace Gauntlet
 
 		public string Flavor { get; protected set; }
 
+		public bool SupportsAdditionalFileCopy => false;
+
 		public string BuildPath { get; protected set; }
 
 		public virtual bool CanSupportRole(UnrealTargetRole InRoleType)
@@ -86,6 +90,8 @@ namespace Gauntlet
 
 		public string Flavor { get; protected set; }
 
+		public bool SupportsAdditionalFileCopy { get; }
+
 		public string BuildPath { get; protected set; }
 
 		public string ExecutablePath { get; protected set; }
@@ -108,6 +114,7 @@ namespace Gauntlet
 			ExecutablePath = InExecutablePath;
 			Flags = BuildFlags.CanReplaceCommandLine | BuildFlags.CanReplaceExecutable | BuildFlags.Loose;
 			Flavor = InFlavor;
+			SupportsAdditionalFileCopy = true; // NativeStagedBuilds may want this to be false
 		}
 
 		enum InstallStatus
