[Gauntlet] Fix for basedir command line setting from path containing spaces

diff --git a/Engine/Source/Programs/AutomationTool/Gauntlet/Platform/Windows/Gauntlet.TargetDeviceWindows.cs b/Engine/Source/Programs/AutomationTool/Gauntlet/Platform/Windows/Gauntlet.TargetDeviceWindows.cs
--- a/Engine/Source/Programs/AutomationTool/Gauntlet/Platform/Windows/Gauntlet.TargetDeviceWindows.cs
+++ b/Engine/Source/Programs/AutomationTool/Gauntlet/Platform/Windows/Gauntlet.TargetDeviceWindows.cs
@@ -157,7 +157,7 @@ namespace Gauntlet
 					if (LocalFileExists && LocalFileNewer)
 					{
 						// need to -basedir to have our exe load content from the path
-						WinApp.CommandArguments += string.Format(" -basedir={0}", Path.GetDirectoryName(BinaryPath));
+						WinApp.CommandArguments += $" -basedir=\"{Path.GetDirectoryName(BinaryPath)}\"";
 
 						BinaryPath = LocalBinary;
 					}
