Upload PCBs to Horde using the same changelist number as the P4 PCBs.

diff --git a/Engine/Source/Programs/AutomationTool/BuildGraph/Tasks/CreateArtifactTask.cs b/Engine/Source/Programs/AutomationTool/BuildGraph/Tasks/CreateArtifactTask.cs
--- a/Engine/Source/Programs/AutomationTool/BuildGraph/Tasks/CreateArtifactTask.cs
+++ b/Engine/Source/Programs/AutomationTool/BuildGraph/Tasks/CreateArtifactTask.cs
@@ -59,6 +59,12 @@ namespace AutomationTool.Tasks
 		[TaskParameter]
 		public string? BaseDir;
 
+		/// <summary>
+		/// Changelist number for this artifact
+		/// </summary>
+		[TaskParameter(Optional = true)]
+		public int? Change;
+
 		/// <summary>
 		/// Files to be uploaded.
 		/// </summary>
@@ -116,7 +122,8 @@ namespace AutomationTool.Tasks
 			ArtifactType artifactType = new ArtifactType(Parameters.Type);
 
 			HordeHttpClient hordeHttpClient = serviceProvider.GetRequiredService<HordeHttpClient>();
-			CreateArtifactResponse response = await hordeHttpClient.CreateArtifactAsync(artifactName, artifactType, Parameters.Description);
+			int? change = (Parameters.Change == 0) ? (int?)null : Parameters.Change;
+			CreateArtifactResponse response = await hordeHttpClient.CreateArtifactAsync(artifactName, artifactType, Parameters.Description, change: change);
 			Logger.LogInformation("Creating artifact {ArtifactId} '{ArtifactName}' ({ArtifactType}) (ns: {NamespaceId}, ref: {RefName})", response.ArtifactId, artifactName, artifactType, response.NamespaceId, response.RefName);
 
 			Stopwatch timer = Stopwatch.StartNew();
