[Gauntlet] Implements new launch flow functions for TargetDeviceDesktopCommon
Also implements overrides for each desktop platform (Win64, WinGDK, Linux, Mac)

diff --git a/Engine/Source/Programs/AutomationTool/Gauntlet/Framework/Base/Gauntlet.TargetDeviceDesktopCommon.cs b/Engine/Source/Programs/AutomationTool/Gauntlet/Framework/Base/Gauntlet.TargetDeviceDesktopCommon.cs
--- a/Engine/Source/Programs/AutomationTool/Gauntlet/Framework/Base/Gauntlet.TargetDeviceDesktopCommon.cs
+++ b/Engine/Source/Programs/AutomationTool/Gauntlet/Framework/Base/Gauntlet.TargetDeviceDesktopCommon.cs
@@ -5,12 +5,12 @@ using EpicGames.Core;
 using System;
 using System.Collections.Generic;
 using System.IO;
-using System.Text;
 using System.Threading;
 using System.Linq;
 using UnrealBuildTool;
 using System.Text.RegularExpressions;
-using static AutomationTool.ProcessResult;
+using Gauntlet.Utils;
+
 
 namespace Gauntlet
 {
@@ -71,7 +71,11 @@ namespace Gauntlet
 		[AutoParam]
 		public string InstallRoot { get; protected set; }
 
-		// TODO - move this be part of ITargetDevice
+		/// <summary>
+		/// Cached reference to the install on this device. Used to maintain a soft handle to desktop artifact directories for cleanup
+		/// </summary>
+		protected IAppInstall InstallCache;
+
 		protected Dictionary<EIntendedBaseCopyDirectory, string> LocalDirectoryMappings { get; set; }
 
 		public TargetDeviceDesktopCommon(string InName, string InCacheDir)
@@ -82,7 +86,7 @@ namespace Gauntlet
 			LocalCachePath = InCacheDir;
 			UserDir = Path.Combine(InCacheDir, "UserDir");
 
-			if(string.IsNullOrEmpty(InstallRoot))
+			if (string.IsNullOrEmpty(InstallRoot))
 			{
 				InstallRoot = InCacheDir;
 			}
@@ -90,51 +94,93 @@ namespace Gauntlet
 			LocalDirectoryMappings = new Dictionary<EIntendedBaseCopyDirectory, string>();
 		}
 
-		public virtual IAppInstall InstallApplication(UnrealAppConfig AppConfig)
+		public void FullClean()
 		{
-			switch (AppConfig.Build)
-			{
-				case NativeStagedBuild:
-					return InstallNativeStagedBuild(AppConfig, AppConfig.Build as NativeStagedBuild);
-
-				case StagedBuild:
-					return InstallStagedBuild(AppConfig, AppConfig.Build as StagedBuild);
-
-				case EditorBuild:
-					return InstallEditorBuild(AppConfig, AppConfig.Build as EditorBuild);
+			CleanArtifacts();
+		}
 
-				default:
-					throw new AutomationException("{0} is an invalid build type!", AppConfig.Build.ToString());
+		public void CleanArtifacts()
+		{
+			CleanArtifactDirectory(Path.Combine(UserDir, "Saved"));
+			if (InstallCache != null)
+			{
+				CleanArtifactDirectory(GetInstallArtifactPath());
 			}
 		}
 
-		public void FullClean()
+		public virtual void InstallBuild(UnrealAppConfig AppConfig)
 		{
+			IBuild Build = AppConfig.Build;
+			switch (Build)
+			{
+				case NativeStagedBuild:
+				case EditorBuild:
+				{
+					Log.Info("Skipping installation of {BuildType}", Build.GetType().Name);
+					break;
+				}
 
-		}
+				case StagedBuild:
+				{
+					StagedBuild Staged = Build as StagedBuild;
 
-		public void CleanArtifacts()
-		{
+					if (SystemHelpers.IsNetworkPath(Staged.BuildPath))
+					{
+						string SubDir = string.IsNullOrEmpty(AppConfig.Sandbox) ? AppConfig.ProjectName : AppConfig.Sandbox;
+						string InstallDir = Path.Combine(InstallRoot, SubDir, AppConfig.ProcessType.ToString());
 
+						InstallDir = StagedBuild.InstallBuildParallel(AppConfig, Staged, Staged.BuildPath, InstallDir, ToString());
+						SystemHelpers.MarkDirectoryForCleanup(InstallDir);
+					}
+					else
+					{
+						Log.Info("Build exists on local drive, skipping installation.");
+					}
+					break;
+				}
+
+				default:
+					throw new AutomationException("{0} is not a valid build type for {1}", Build.GetType().Name, Platform);
+			}
 		}
 
-		public void InstallBuild(UnrealAppConfig AppConfiguration)
+		public virtual IAppInstall CreateAppInstall(UnrealAppConfig AppConfig)
 		{
+			IAppInstall Install;
+			IBuild Build = AppConfig.Build;
 
-		}
+			switch (AppConfig.Build)
+			{
+				case NativeStagedBuild:
+					Install = CreateNativeStagedInstall(AppConfig, Build as NativeStagedBuild);
+					break;
+				case StagedBuild:
+					Install = CreateStagedInstall(AppConfig, Build as StagedBuild);
+					break;
+				case EditorBuild:
+					Install = CreateEditorInstall(AppConfig, Build as EditorBuild);
+					break;
+				default:
+					throw new AutomationException("{0} is an invalid build type for {1}!", Build.GetType().Name, Platform);
+			}
 
-		public IAppInstall CreateAppInstall(UnrealAppConfig AppConfig)
-		{
-			return null;
+			InstallCache = Install;
+			return Install;
 		}
 
 		public void CopyAdditionalFiles(IEnumerable<UnrealFileToCopy> FilesToCopy)
 		{
-			if (FilesToCopy == null || FilesToCopy.Any())
+			if (FilesToCopy == null || !FilesToCopy.Any())
 			{
 				return;
 			}
 
+			if (!LocalDirectoryMappings.Any())
+			{
+				throw new AutomationException("Attempted to copy additional files before LocalDirectoryMappings were populated." +
+					"{0} must call PopulateDirectoryMappings before attempting to call CopyAdditionalFiles", this.GetType());
+			}
+
 			foreach (UnrealFileToCopy FileToCopy in FilesToCopy)
 			{
 				string PathToCopyTo = Path.Combine(LocalDirectoryMappings[FileToCopy.TargetBaseDirectory], FileToCopy.TargetRelativeLocation);
@@ -152,8 +198,9 @@ namespace Gauntlet
 						FileInfo ExistingFile = new FileInfo(PathToCopyTo);
 						ExistingFile.IsReadOnly = false;
 					}
+
+					Log.Info("Copying {SourceFile} to {DestinationFile}", FileToCopy.SourceFileLocation, PathToCopyTo);
 					SrcInfo.CopyTo(PathToCopyTo, true);
-					Log.Info("Copying {0} to {1}", FileToCopy.SourceFileLocation, PathToCopyTo);
 				}
 				else
 				{
@@ -162,8 +209,6 @@ namespace Gauntlet
 			}
 		}
 
-		public abstract IAppInstance Run(IAppInstall Install);
-
 		public virtual void PopulateDirectoryMappings(string BaseDirectory)
 		{
 			LocalDirectoryMappings.Clear();
@@ -178,10 +223,49 @@ namespace Gauntlet
             LocalDirectoryMappings.Add(EIntendedBaseCopyDirectory.Saved, Path.Combine(BaseDirectory, "Saved"));
 		}
 
-		// TODO - b.lienau: implement these at desktop level and remove implementations from each desktop
+		public abstract IAppInstance Run(IAppInstall Install);
+
+		protected abstract IAppInstall CreateNativeStagedInstall(UnrealAppConfig AppConfig, NativeStagedBuild Build);
+
+		protected abstract IAppInstall CreateStagedInstall(UnrealAppConfig AppConfig, StagedBuild Build);
+
+		protected abstract IAppInstall CreateEditorInstall(UnrealAppConfig AppConfig, EditorBuild Build);
+
+		protected abstract string GetInstallArtifactPath();
+
+		private void CleanArtifactDirectory(string ArtifactDirectory)
+		{
+			if (!string.IsNullOrEmpty(ArtifactDirectory) && Directory.Exists(ArtifactDirectory))
+			{
+				Log.Info("Cleaning device artifacts path {ArtifactDirectory}", ArtifactDirectory);
+				DirectoryInfo Info = new(ArtifactDirectory);
+				SystemHelpers.Delete(Info, true, true);
+			}
+		}
+
+		#region Legacy Implementations
+		public virtual IAppInstall InstallApplication(UnrealAppConfig AppConfig)
+		{
+			switch (AppConfig.Build)
+			{
+				case NativeStagedBuild:
+					return InstallNativeStagedBuild(AppConfig, AppConfig.Build as NativeStagedBuild);
+
+				case StagedBuild:
+					return InstallStagedBuild(AppConfig, AppConfig.Build as StagedBuild);
+
+				case EditorBuild:
+					return InstallEditorBuild(AppConfig, AppConfig.Build as EditorBuild);
+
+				default:
+					throw new AutomationException("{0} is an invalid build type!", AppConfig.Build.ToString());
+			}
+		}
+
 		protected abstract IAppInstall InstallNativeStagedBuild(UnrealAppConfig AppConfig, NativeStagedBuild Build);
 		protected abstract IAppInstall InstallStagedBuild(UnrealAppConfig AppConfig, StagedBuild Build);
 		protected abstract IAppInstall InstallEditorBuild(UnrealAppConfig AppConfig, EditorBuild Build);
+		#endregion
 	}
 
 	public abstract class DesktopCommonAppInstall<DesktopTargetDevice> : IAppInstall where DesktopTargetDevice : TargetDeviceDesktopCommon
@@ -244,58 +328,30 @@ namespace Gauntlet
 			return Device.Run(this);
 		}
 
+		/// <summary>
+		/// Obsolete! Will be removed in a future release.
+		/// Use ITargetDevice.CleanArtifacts instead
+		/// </summary>
 		public virtual void CleanDeviceArtifacts()
 		{
-			// log file
-			try
-			{
-				if (LogFile != null && File.Exists(LogFile))
-				{
-					EpicGames.Core.FileUtils.ForceDeleteFile(LogFile);
-				}
-			}
-			catch (Exception Ex)
-			{
-				Log.Warning(KnownLogEvents.Gauntlet_DeviceEvent, "Unable to delete existing log file {File}. {Exception}", LogFile, Ex.Message);
-			}
-			// all other artifacts
-			if (!string.IsNullOrEmpty(ArtifactPath) && Directory.Exists(ArtifactPath))
-			{
-				try
-				{
-					Log.Info("Clearing device artifacts path {0} for {1}", ArtifactPath, Device.Name);
-					Directory.Delete(ArtifactPath, true);
-				}
-				catch (Exception Ex)
-				{
-					Log.Info(KnownLogEvents.Gauntlet_DeviceEvent, "First attempt at clearing artifact path {0} failed - trying again", ArtifactPath);
-					if (!ForceCleanDeviceArtifacts())
-					{
-						Log.Warning(KnownLogEvents.Gauntlet_DeviceEvent, "Failed to delete {File}. {Exception}", ArtifactPath, Ex.Message);
-					}
-				}
-			}
+			Device.CleanArtifacts();
 		}
 
+		/// <summary>
+		/// Obsolete! Will be removed in a future release.
+		/// Use ITargetDevice.CleanArtifacts instead
+		/// </summary>
 		public virtual bool ForceCleanDeviceArtifacts()
 		{
-			DirectoryInfo ClientTempDirInfo = new DirectoryInfo(ArtifactPath) { Attributes = FileAttributes.Normal };
-			Log.Info(KnownLogEvents.Gauntlet_DeviceEvent, "Setting files in device artifacts {0} to have normal attributes (no longer read-only).", ArtifactPath);
-			foreach (FileSystemInfo info in ClientTempDirInfo.GetFileSystemInfos("*", SearchOption.AllDirectories))
-			{
-				info.Attributes = FileAttributes.Normal;
-			}
 			try
 			{
-				Log.Info(KnownLogEvents.Gauntlet_DeviceEvent, "Clearing device artifact path {0} (force)", ArtifactPath);
-				Directory.Delete(ArtifactPath, true);
+				Device.CleanArtifacts();
+				return true;
 			}
-			catch (Exception Ex)
+			catch
 			{
-				Log.Warning(KnownLogEvents.Gauntlet_DeviceEvent, "Failed to force delete artifact path {File}. {Exception}", ArtifactPath, Ex.Message);
 				return false;
 			}
-			return true;
 		}
 
 		public virtual void SetDefaultCommandLineArguments(UnrealAppConfig AppConfig, CommandUtils.ERunOptions InRunOptions, string BuildDir)
diff --git a/Engine/Source/Programs/AutomationTool/Gauntlet/Platform/Linux/Gauntlet.TargetDeviceLinux.cs b/Engine/Source/Programs/AutomationTool/Gauntlet/Platform/Linux/Gauntlet.TargetDeviceLinux.cs
--- a/Engine/Source/Programs/AutomationTool/Gauntlet/Platform/Linux/Gauntlet.TargetDeviceLinux.cs
+++ b/Engine/Source/Programs/AutomationTool/Gauntlet/Platform/Linux/Gauntlet.TargetDeviceLinux.cs
@@ -1,14 +1,10 @@
 // Copyright Epic Games, Inc. All Rights Reserved.
 
 using System;
-using System.Collections.Generic;
 using System.IO;
-using System.Linq;
-using System.Threading;
 using AutomationTool;
+using Gauntlet.Utils;
 using UnrealBuildTool;
-using System.Text.RegularExpressions;
-using EpicGames.Core;
 using static AutomationTool.ProcessResult;
 
 namespace Gauntlet
@@ -93,6 +89,69 @@ namespace Gauntlet
 			return new LinuxAppInstance(LinuxApp, Result, LinuxApp.LogFile);
 		}
 
+		protected override IAppInstall CreateNativeStagedInstall(UnrealAppConfig AppConfig, NativeStagedBuild Build)
+		{
+			LinuxAppInstall LinuxApp;
+			if (AppConfig.ContainerInfo != null)
+			{
+				LinuxApp = new LinuxAppContainerInstall(AppConfig.Name, AppConfig.ProjectName, AppConfig.ContainerInfo, this);
+			}
+			else
+			{
+				LinuxApp = new LinuxAppInstall(AppConfig.Name, AppConfig.ProjectName, this);
+			}
+
+			LinuxApp.ExecutablePath = Path.Combine(Build.BuildPath, Build.ExecutablePath);
+			LinuxApp.WorkingDirectory = Build.BuildPath;
+			LinuxApp.SetDefaultCommandLineArguments(AppConfig, RunOptions, Build.BuildPath);
+
+			return LinuxApp;
+		}
+
+		protected override IAppInstall CreateStagedInstall(UnrealAppConfig AppConfig, StagedBuild Build)
+		{
+			string BuildDir = Build.BuildPath;
+			if (SystemHelpers.IsNetworkPath(BuildDir))
+			{
+				string SubDir = string.IsNullOrEmpty(AppConfig.Sandbox) ? AppConfig.ProjectName : AppConfig.Sandbox;
+				string InstallDir = Path.Combine(InstallRoot, SubDir, AppConfig.ProcessType.ToString());
+				BuildDir = InstallDir;
+			}
+
+			PopulateDirectoryMappings(Path.Combine(BuildDir, AppConfig.ProjectName));
+
+			LinuxAppInstall LinuxApp = new LinuxAppInstall(AppConfig.Name, AppConfig.ProjectName, this)
+			{
+				ExecutablePath = Path.IsPathRooted(Build.ExecutablePath)
+					? Build.ExecutablePath
+					: Path.Combine(BuildDir, Build.ExecutablePath)
+			};
+			LinuxApp.SetDefaultCommandLineArguments(AppConfig, RunOptions, BuildDir);
+
+			return LinuxApp;
+		}
+
+		protected override IAppInstall CreateEditorInstall(UnrealAppConfig AppConfig, EditorBuild Build)
+		{
+			PopulateDirectoryMappings(AppConfig.ProjectFile.Directory.FullName);
+
+			LinuxAppInstall LinuxApp = new LinuxAppInstall(AppConfig.Name, AppConfig.ProjectName, this)
+			{
+				ExecutablePath = Build.ExecutablePath,
+				WorkingDirectory = Path.GetDirectoryName(Build.ExecutablePath)
+			};
+			LinuxApp.SetDefaultCommandLineArguments(AppConfig, RunOptions, LinuxApp.WorkingDirectory);
+
+			return LinuxApp;
+		}
+
+		protected override string GetInstallArtifactPath()
+		{
+			return (InstallCache as DesktopCommonAppInstall<TargetDeviceLinux>).ArtifactPath;
+		}
+
+		#region Legacy Implementations
+
 		protected override IAppInstall InstallNativeStagedBuild(UnrealAppConfig AppConfig, NativeStagedBuild InBuild)
 		{
 			LinuxAppInstall LinuxApp;
@@ -198,6 +257,7 @@ namespace Gauntlet
 
 			return LinuxApp;
 		}
+		#endregion
 	}
 
 	public class LinuxAppInstall : DesktopCommonAppInstall<TargetDeviceLinux>
diff --git a/Engine/Source/Programs/AutomationTool/Gauntlet/Platform/Mac/Gauntlet.TargetDeviceMac.cs b/Engine/Source/Programs/AutomationTool/Gauntlet/Platform/Mac/Gauntlet.TargetDeviceMac.cs
--- a/Engine/Source/Programs/AutomationTool/Gauntlet/Platform/Mac/Gauntlet.TargetDeviceMac.cs
+++ b/Engine/Source/Programs/AutomationTool/Gauntlet/Platform/Mac/Gauntlet.TargetDeviceMac.cs
@@ -8,6 +8,7 @@
 using System.Text.RegularExpressions;
 using EpicGames.Core;
 using UnrealBuildBase;
+using Gauntlet.Utils;
 using static AutomationTool.ProcessResult;
 
 namespace Gauntlet
@@ -21,25 +22,78 @@
 			RunOptions = CommandUtils.ERunOptions.NoWaitForExit;
 		}
 
-		public override IAppInstall InstallApplication(UnrealAppConfig AppConfig)
+		public override void InstallBuild(UnrealAppConfig AppConfig)
 		{
+			IBuild Build = AppConfig.Build;
+			switch (Build)
+			{
+				case NativeStagedBuild:
+				case EditorBuild:
+				{
+					Log.Info("Skipping installation of {BuildType}", Build.GetType().Name);
+					break;
+				}
+
+				case StagedBuild:
+				case MacPackagedBuild:
+				{
+					string BuildDir;
+					if (Build is StagedBuild)
+					{
+						BuildDir = (Build as StagedBuild).BuildPath;
+					}
+					else
+					{
+						BuildDir = (Build as MacPackagedBuild).BuildPath;
+					}
+
+					// If the build already exists on this drive, we can skip installation
+					string BuildVolume = GetVolumeName(BuildDir);
+					string LocalRoot = GetVolumeName(Environment.CurrentDirectory);
+					if (BuildVolume.Equals(LocalRoot, StringComparison.OrdinalIgnoreCase))
+					{
+						Log.Info("Build exists on desired volume, skipping installation");
+						return;
+					}
+
+					// Otherwise, determine the desired installation directory and install the build
+					string InstallDir = GetTargetBuildDirectory(BuildDir, AppConfig.Sandbox,
+						AppConfig.ProjectName, AppConfig.ProcessType.ToString());
+
+					SystemHelpers.CopyDirectory(BuildDir, InstallDir, SystemHelpers.CopyOptions.Mirror);
+					break;
+				}
+
+				default:
+					throw new AutomationException("{0} is not a valid build type for {1}", Build.GetType().Name, Platform);
+			}
+		}
+
+		public override IAppInstall CreateAppInstall(UnrealAppConfig AppConfig)
+		{
+			IAppInstall Install;
+			IBuild Build = AppConfig.Build;
+
 			switch (AppConfig.Build)
 			{
 				case NativeStagedBuild:
-					return InstallNativeStagedBuild(AppConfig, AppConfig.Build as NativeStagedBuild);
+					Install = CreateNativeStagedInstall(AppConfig, Build as NativeStagedBuild);
+					break;
+				case StagedBuild:
+					Install = CreateStagedInstall(AppConfig, Build as StagedBuild);
+					break;
+				case EditorBuild:
+					Install = CreateEditorInstall(AppConfig, Build as EditorBuild);
+					break;
+				case MacPackagedBuild:
+					Install = CreatePackagedInstall(AppConfig, Build as MacPackagedBuild);
+					break;
+				default:
+					throw new AutomationException("{0} is an invalid build type for {1}!", Build.GetType().Name, Platform);
+			}
 
-				case StagedBuild:
-					return InstallStagedBuild(AppConfig, AppConfig.Build as StagedBuild);
-
-				case MacPackagedBuild:
-					return InstallPackagedBuild(AppConfig, AppConfig.Build as MacPackagedBuild);
-
-				case EditorBuild:
-					return InstallEditorBuild(AppConfig, AppConfig.Build as EditorBuild);
-
-				default:
-					throw new AutomationException("{0} is an invalid build type!", AppConfig.Build.ToString());
-			}
+			InstallCache = Install;
+			return Install;
 		}
 
 		public override IAppInstance Run(IAppInstall App)
@@ -75,91 +129,64 @@
 			return new MacAppInstance(MacInstall, Result, MacInstall.LogFile);
 		}
 
-		protected override IAppInstall InstallNativeStagedBuild(UnrealAppConfig AppConfig, NativeStagedBuild InBuild)
+		protected override IAppInstall CreateNativeStagedInstall(UnrealAppConfig AppConfig, NativeStagedBuild Build)
 		{
-			MacAppInstall MacApp = new MacAppInstall(AppConfig.Name, AppConfig.ProjectName, this);
-			MacApp.SetDefaultCommandLineArguments(AppConfig, RunOptions, InBuild.BuildPath);
-			MacApp.WorkingDirectory = InBuild.BuildPath;
-			MacApp.ExecutablePath = Path.Combine(InBuild.BuildPath, InBuild.ExecutablePath);
+			PopulateDirectoryMappings(Build.BuildPath);
 
-			CopyAdditionalFiles(AppConfig.FilesToCopy);
+			MacAppInstall MacApp = new MacAppInstall(AppConfig.Name, AppConfig.ProjectName, this)
+			{
+				ExecutablePath = Path.Combine(Build.BuildPath, Build.ExecutablePath),
+				WorkingDirectory = Build.BuildPath
+			};
+			MacApp.SetDefaultCommandLineArguments(AppConfig, RunOptions, Build.BuildPath);
 
 			return MacApp;
 		}
 
-		protected override IAppInstall InstallEditorBuild(UnrealAppConfig AppConfig, EditorBuild Build)
+		protected override IAppInstall CreateStagedInstall(UnrealAppConfig AppConfig, StagedBuild Build)
 		{
-			MacAppInstall MacApp = new MacAppInstall(AppConfig.Name, AppConfig.ProjectName, this);
-			MacApp.SetDefaultCommandLineArguments(AppConfig, RunOptions, Build.ExecutablePath);
-			MacApp.WorkingDirectory = Path.GetFullPath(Build.ExecutablePath);
-			MacApp.ExecutablePath = GetExecutableIfBundle(Build.ExecutablePath);
-
-			if (LocalDirectoryMappings.Count == 0)
-			{
-				PopulateDirectoryMappings(AppConfig.ProjectFile.Directory.FullName);
-			}
-
-			CopyAdditionalFiles(AppConfig.FilesToCopy);
-
-			return MacApp;
-		}
-
-		protected override IAppInstall InstallStagedBuild(UnrealAppConfig AppConfig, StagedBuild Build)
-		{
-			string BuildPath = CopyBuildIfNecessary(AppConfig, Build.BuildPath);
+			string BuildPath = GetTargetBuildDirectory(Build.BuildPath, AppConfig.Sandbox,
+				AppConfig.ProjectName, AppConfig.ProcessType.ToString());
 			string BundlePath = Path.Combine(BuildPath, Build.ExecutablePath);
-			return InstallBuild(AppConfig, BuildPath, BundlePath);
+
+			return CreateMacInstall(AppConfig, BuildPath, BundlePath);
 		}
 
-		protected IAppInstall InstallPackagedBuild(UnrealAppConfig AppConfig, MacPackagedBuild Build)
+		protected override IAppInstall CreateEditorInstall(UnrealAppConfig AppConfig, EditorBuild Build)
 		{
-			string BuildPath = CopyBuildIfNecessary(AppConfig, Build.BuildPath);
-			return InstallBuild(AppConfig, BuildPath, BuildPath);
+			PopulateDirectoryMappings(AppConfig.ProjectFile.Directory.FullName);
+
+			MacAppInstall MacApp = new MacAppInstall(AppConfig.Name, AppConfig.ProjectName, this)
+			{
+				ExecutablePath = GetExecutableIfBundle(Build.ExecutablePath),
+				WorkingDirectory = Path.GetFullPath(Build.ExecutablePath)
+			};
+			MacApp.SetDefaultCommandLineArguments(AppConfig, RunOptions, Build.ExecutablePath);
+
+			return MacApp;
 		}
 
-		protected IAppInstall InstallBuild(UnrealAppConfig AppConfig, string BuildPath, string BundlePath)
+		protected IAppInstall CreatePackagedInstall(UnrealAppConfig AppConfig, MacPackagedBuild Build)
 		{
-			MacAppInstall MacApp = new MacAppInstall(AppConfig.Name, AppConfig.ProjectName, this);
-			MacApp.SetDefaultCommandLineArguments(AppConfig, RunOptions, BuildPath);
-			MacApp.WorkingDirectory = BuildPath;
-			MacApp.ExecutablePath = GetExecutableIfBundle(BundlePath);
+			string BuildDir = GetTargetBuildDirectory(Build.BuildPath, AppConfig.Sandbox,
+				AppConfig.ProjectName, AppConfig.ProcessType.ToString());
 
+			// Packaged builds use the BuildDir as the bundle
+			return CreateMacInstall(AppConfig, BuildDir, BuildDir);
+		}
+
+		// Creates MacAppInstalls for StagedBuilds and MacPackagedBuilds
+		protected IAppInstall CreateMacInstall(UnrealAppConfig AppConfig, string BuildPath, string BundlePath)
+		{
 			PopulateDirectoryMappings(BuildPath);
 
-			CopyAdditionalFiles(AppConfig.FilesToCopy);
+			MacAppInstall MacApp = new MacAppInstall(AppConfig.Name, AppConfig.ProjectName, this)
+			{
+				ExecutablePath = GetExecutableIfBundle(BundlePath),
+				WorkingDirectory = BuildPath
+			};
+			MacApp.SetDefaultCommandLineArguments(AppConfig, RunOptions, BuildPath);
 
-			// check for a local newer executable
-			if (Globals.Params.ParseParam("dev")
-				&& AppConfig.ProcessType.UsesEditor() == false
-				&& AppConfig.ProjectFile != null)
-			{
-				// Get project properties
-				ProjectProperties Props = ProjectUtils.GetProjectProperties(AppConfig.ProjectFile,
-					new List<UnrealTargetPlatform>(new[] { AppConfig.Platform.Value }),
-					new List<UnrealTargetConfiguration>(new[] { AppConfig.Configuration }));
-
-				// Would this executable be built under Engine or a Project?
-				DirectoryReference WorkingDir = Props.bIsCodeBasedProject ? AppConfig.ProjectFile.Directory : Unreal.EngineDirectory;
-
-				// The bundlepath may be under Binaries/Mac for a staged build, or it could be in any folder for a packaged build so just use the name and
-				// build the path ourselves
-				string LocalProjectBundle = FileReference.Combine(WorkingDir, "Binaries", "Mac", Path.GetFileName(BundlePath)).FullName;
-
-				string LocalProjectBinary = GetExecutableIfBundle(LocalProjectBundle);
-
-				bool LocalFileExists = File.Exists(LocalProjectBinary);
-				bool LocalFileNewer = LocalFileExists && File.GetLastWriteTime(LocalProjectBinary) > File.GetLastWriteTime(MacApp.ExecutablePath);
-
-				Log.Verbose("Checking for newer binary at {0}", LocalProjectBinary);
-				Log.Verbose("LocalFile exists: {0}. Newer: {1}", LocalFileExists, LocalFileNewer);
-
-				if (LocalFileExists && LocalFileNewer)
-				{
-					// need to -basedir to have our exe load content from the path that the bundle sits in
-					MacApp.CommandArguments += string.Format(" -basedir={0}", Path.GetDirectoryName(BundlePath));
-					MacApp.ExecutablePath = LocalProjectBinary;
-				}
-			}
 
 			return MacApp;
 		}
@@ -193,7 +220,6 @@
 			return InBundlePath;
 		}
 
-
 		/// <summary>
 		/// Copies a build folder (either a package.app or a folder with a staged built) to a local path if necessary.
 		/// Necessary is defined as not being on locally attached storage
@@ -216,7 +242,7 @@
 
 				if (!AppConfig.SkipInstall)
 				{
-					Utils.SystemHelpers.CopyDirectory(BuildDir, InstallDir, Utils.SystemHelpers.CopyOptions.Mirror);
+					SystemHelpers.CopyDirectory(BuildDir, InstallDir, SystemHelpers.CopyOptions.Mirror);
 				}
 				else
 				{
@@ -224,11 +250,145 @@
 				}
 
 				BuildDir = InstallDir;
-				Utils.SystemHelpers.MarkDirectoryForCleanup(InstallDir);
+				SystemHelpers.MarkDirectoryForCleanup(InstallDir);
 			}
 
 			return BuildDir;
 		}
+
+		protected override string GetInstallArtifactPath()
+		{
+			return (InstallCache as DesktopCommonAppInstall<TargetDeviceMac>).ArtifactPath;
+		}
+
+		private string GetTargetBuildDirectory(string BasePath, string Sandbox, string Project, string ProcessType)
+		{
+			string TargetDir = BasePath;
+
+			string BuildVolume = GetVolumeName(TargetDir);
+			string LocalRoot = GetVolumeName(Environment.CurrentDirectory);
+			if (BuildVolume.Equals(LocalRoot, StringComparison.OrdinalIgnoreCase))
+			{
+				string SubDir = string.IsNullOrEmpty(Sandbox) ? Project : Sandbox;
+				string InstallDir = Path.Combine(InstallRoot, SubDir, ProcessType);
+				TargetDir = InstallDir;
+			}
+
+			return TargetDir;
+		}
+
+		#region Legacy Implementations
+
+		public override IAppInstall InstallApplication(UnrealAppConfig AppConfig)
+		{
+			switch (AppConfig.Build)
+			{
+				case NativeStagedBuild:
+					return InstallNativeStagedBuild(AppConfig, AppConfig.Build as NativeStagedBuild);
+
+				case StagedBuild:
+					return InstallStagedBuild(AppConfig, AppConfig.Build as StagedBuild);
+
+				case MacPackagedBuild:
+					return InstallPackagedBuild(AppConfig, AppConfig.Build as MacPackagedBuild);
+
+				case EditorBuild:
+					return InstallEditorBuild(AppConfig, AppConfig.Build as EditorBuild);
+
+				default:
+					throw new AutomationException("{0} is an invalid build type!", AppConfig.Build.ToString());
+			}
+		}
+
+		protected override IAppInstall InstallNativeStagedBuild(UnrealAppConfig AppConfig, NativeStagedBuild InBuild)
+		{
+			MacAppInstall MacApp = new MacAppInstall(AppConfig.Name, AppConfig.ProjectName, this);
+			MacApp.SetDefaultCommandLineArguments(AppConfig, RunOptions, InBuild.BuildPath);
+			MacApp.WorkingDirectory = InBuild.BuildPath;
+			MacApp.ExecutablePath = Path.Combine(InBuild.BuildPath, InBuild.ExecutablePath);
+
+			CopyAdditionalFiles(AppConfig.FilesToCopy);
+
+			return MacApp;
+		}
+
+		protected override IAppInstall InstallEditorBuild(UnrealAppConfig AppConfig, EditorBuild Build)
+		{
+			MacAppInstall MacApp = new MacAppInstall(AppConfig.Name, AppConfig.ProjectName, this);
+			MacApp.SetDefaultCommandLineArguments(AppConfig, RunOptions, Build.ExecutablePath);
+			MacApp.WorkingDirectory = Path.GetFullPath(Build.ExecutablePath);
+			MacApp.ExecutablePath = GetExecutableIfBundle(Build.ExecutablePath);
+
+			if (LocalDirectoryMappings.Count == 0)
+			{
+				PopulateDirectoryMappings(AppConfig.ProjectFile.Directory.FullName);
+			}
+
+			CopyAdditionalFiles(AppConfig.FilesToCopy);
+
+			return MacApp;
+		}
+
+		protected override IAppInstall InstallStagedBuild(UnrealAppConfig AppConfig, StagedBuild Build)
+		{
+			string BuildPath = CopyBuildIfNecessary(AppConfig, Build.BuildPath);
+			string BundlePath = Path.Combine(BuildPath, Build.ExecutablePath);
+			return InstallBuild(AppConfig, BuildPath, BundlePath);
+		}
+
+		protected IAppInstall InstallPackagedBuild(UnrealAppConfig AppConfig, MacPackagedBuild Build)
+		{
+			string BuildPath = CopyBuildIfNecessary(AppConfig, Build.BuildPath);
+			return InstallBuild(AppConfig, BuildPath, BuildPath);
+		}
+
+		protected IAppInstall InstallBuild(UnrealAppConfig AppConfig, string BuildPath, string BundlePath)
+		{
+			MacAppInstall MacApp = new MacAppInstall(AppConfig.Name, AppConfig.ProjectName, this);
+			MacApp.SetDefaultCommandLineArguments(AppConfig, RunOptions, BuildPath);
+			MacApp.WorkingDirectory = BuildPath;
+			MacApp.ExecutablePath = GetExecutableIfBundle(BundlePath);
+
+			PopulateDirectoryMappings(BuildPath);
+
+			CopyAdditionalFiles(AppConfig.FilesToCopy);
+
+			// check for a local newer executable
+			if (Globals.Params.ParseParam("dev")
+				&& AppConfig.ProcessType.UsesEditor() == false
+				&& AppConfig.ProjectFile != null)
+			{
+				// Get project properties
+				ProjectProperties Props = ProjectUtils.GetProjectProperties(AppConfig.ProjectFile,
+					new List<UnrealTargetPlatform>(new[] { AppConfig.Platform.Value }),
+					new List<UnrealTargetConfiguration>(new[] { AppConfig.Configuration }));
+
+				// Would this executable be built under Engine or a Project?
+				DirectoryReference WorkingDir = Props.bIsCodeBasedProject ? AppConfig.ProjectFile.Directory : Unreal.EngineDirectory;
+
+				// The bundlepath may be under Binaries/Mac for a staged build, or it could be in any folder for a packaged build so just use the name and
+				// build the path ourselves
+				string LocalProjectBundle = FileReference.Combine(WorkingDir, "Binaries", "Mac", Path.GetFileName(BundlePath)).FullName;
+
+				string LocalProjectBinary = GetExecutableIfBundle(LocalProjectBundle);
+
+				bool LocalFileExists = File.Exists(LocalProjectBinary);
+				bool LocalFileNewer = LocalFileExists && File.GetLastWriteTime(LocalProjectBinary) > File.GetLastWriteTime(MacApp.ExecutablePath);
+
+				Log.Verbose("Checking for newer binary at {0}", LocalProjectBinary);
+				Log.Verbose("LocalFile exists: {0}. Newer: {1}", LocalFileExists, LocalFileNewer);
+
+				if (LocalFileExists && LocalFileNewer)
+				{
+					// need to -basedir to have our exe load content from the path that the bundle sits in
+					MacApp.CommandArguments += string.Format(" -basedir={0}", Path.GetDirectoryName(BundlePath));
+					MacApp.ExecutablePath = LocalProjectBinary;
+				}
+			}
+
+			return MacApp;
+		}
+		#endregion
 	}
 
 	public class MacAppInstall : DesktopCommonAppInstall<TargetDeviceMac>
diff --git a/Engine/Source/Programs/AutomationTool/Gauntlet/Platform/Windows/Gauntlet.TargetDeviceWindows.cs b/Engine/Source/Programs/AutomationTool/Gauntlet/Platform/Windows/Gauntlet.TargetDeviceWindows.cs
--- a/Engine/Source/Programs/AutomationTool/Gauntlet/Platform/Windows/Gauntlet.TargetDeviceWindows.cs
+++ b/Engine/Source/Programs/AutomationTool/Gauntlet/Platform/Windows/Gauntlet.TargetDeviceWindows.cs
@@ -4,7 +4,7 @@ using System;
 using System.IO;
 using AutomationTool;
 using UnrealBuildTool;
-using EpicGames.Core;
+using Gauntlet.Utils;
 using static AutomationTool.ProcessResult;
 
 namespace Gauntlet
@@ -21,25 +21,42 @@ namespace Gauntlet
 			RunOptions = CommandUtils.ERunOptions.NoWaitForExit | CommandUtils.ERunOptions.NoLoggingOfRunCommand;
 		}
 
-		public override IAppInstall InstallApplication(UnrealAppConfig AppConfig)
+		public override void InstallBuild(UnrealAppConfig AppConfig)
+		{
+			if (AppConfig.Build is IWindowsSelfInstallingBuild SelfInstallingBuild)
+			{
+				SelfInstallingBuild.Install(AppConfig);
+				return;
+			}
+
+			base.InstallBuild(AppConfig);
+		}
+
+		public override IAppInstall CreateAppInstall(UnrealAppConfig AppConfig)
 		{
+			IAppInstall Install;
+			IBuild Build = AppConfig.Build;
+
 			switch (AppConfig.Build)
 			{
 				case NativeStagedBuild:
-					return InstallNativeStagedBuild(AppConfig, AppConfig.Build as NativeStagedBuild);
-
+					Install = CreateNativeStagedInstall(AppConfig, Build as NativeStagedBuild);
+					break;
 				case StagedBuild:
-					return InstallStagedBuild(AppConfig, AppConfig.Build as StagedBuild);
-
+					Install = CreateStagedInstall(AppConfig, Build as StagedBuild);
+					break;
 				case EditorBuild:
-					return InstallEditorBuild(AppConfig, AppConfig.Build as EditorBuild);
-
+					Install = CreateEditorInstall(AppConfig, Build as EditorBuild);
+					break;
 				case IWindowsSelfInstallingBuild:
-					return InstallSelfInstallingBuild(AppConfig, AppConfig.Build as IWindowsSelfInstallingBuild);
-
+					Install = CreateSelfInstall(AppConfig, Build as IWindowsSelfInstallingBuild);
+					break;
 				default:
-					throw new AutomationException("{0} is an invalid build type!", AppConfig.Build.ToString());
+					throw new AutomationException("{0} is an invalid build type for {1}!", Build.GetType().Name, Platform);
 			}
+
+			InstallCache = Install;
+			return Install;
 		}
 
 		public override IAppInstance Run(IAppInstall App)
@@ -90,6 +107,91 @@ namespace Gauntlet
 			return new WindowsAppInstance(WinApp, Result, WinApp.LogFile);
 		}
 
+		protected override IAppInstall CreateNativeStagedInstall(UnrealAppConfig AppConfig, NativeStagedBuild Build)
+		{
+			PopulateDirectoryMappings(Path.Combine(Build.BuildPath, AppConfig.ProjectName));
+
+			WindowsAppInstall WinApp = new WindowsAppInstall(AppConfig.Name, AppConfig.ProjectName, this)
+			{
+				ExecutablePath = Path.Combine(Build.BuildPath, Build.ExecutablePath),
+				WorkingDirectory = Build.BuildPath
+			};
+			WinApp.SetDefaultCommandLineArguments(AppConfig, RunOptions, Build.BuildPath);
+
+			return WinApp;
+		}
+
+		protected override IAppInstall CreateStagedInstall(UnrealAppConfig AppConfig, StagedBuild Build)
+		{
+			string BuildDir = Build.BuildPath;
+			if (SystemHelpers.IsNetworkPath(BuildDir))
+			{
+				string SubDir = string.IsNullOrEmpty(AppConfig.Sandbox) ? AppConfig.ProjectName : AppConfig.Sandbox;
+				string InstallDir = Path.Combine(InstallRoot, SubDir, AppConfig.ProcessType.ToString());
+				BuildDir = InstallDir;
+			}
+
+			PopulateDirectoryMappings(Path.Combine(BuildDir, AppConfig.ProjectName));
+
+			WindowsAppInstall WinApp = new WindowsAppInstall(AppConfig.Name, AppConfig.ProjectName, this)
+			{
+				ExecutablePath = Path.IsPathRooted(Build.ExecutablePath)
+					? Build.ExecutablePath
+					: Path.Combine(BuildDir, Build.ExecutablePath)
+			};
+			WinApp.SetDefaultCommandLineArguments(AppConfig, RunOptions, BuildDir);
+
+			return WinApp;
+		}
+
+		protected override IAppInstall CreateEditorInstall(UnrealAppConfig AppConfig, EditorBuild Build)
+		{
+			PopulateDirectoryMappings(AppConfig.ProjectFile.Directory.FullName);
+
+			WindowsAppInstall WinApp = new WindowsAppInstall(AppConfig.Name, AppConfig.ProjectName, this)
+			{
+				ExecutablePath = Build.ExecutablePath,
+				WorkingDirectory = Path.GetDirectoryName(Build.ExecutablePath)
+			};
+			WinApp.SetDefaultCommandLineArguments(AppConfig, RunOptions, WinApp.WorkingDirectory);
+
+			return WinApp;
+		}
+
+		protected IAppInstall CreateSelfInstall(UnrealAppConfig AppConfig, IWindowsSelfInstallingBuild Build)
+		{
+			WindowsAppInstall WinApp = Build.CreateAppInstall(this, AppConfig, out string BasePath);
+			PopulateDirectoryMappings(Path.Combine(BasePath, AppConfig.ProjectName));
+			return WinApp;
+		}
+
+		protected override string GetInstallArtifactPath()
+		{
+			return (InstallCache as DesktopCommonAppInstall<TargetDeviceWindows>).ArtifactPath;
+		}
+
+		#region Legacy Implementations
+		public override IAppInstall InstallApplication(UnrealAppConfig AppConfig)
+		{
+			switch (AppConfig.Build)
+			{
+				case NativeStagedBuild:
+					return InstallNativeStagedBuild(AppConfig, AppConfig.Build as NativeStagedBuild);
+
+				case StagedBuild:
+					return InstallStagedBuild(AppConfig, AppConfig.Build as StagedBuild);
+
+				case EditorBuild:
+					return InstallEditorBuild(AppConfig, AppConfig.Build as EditorBuild);
+
+				case IWindowsSelfInstallingBuild:
+					return InstallSelfInstallingBuild(AppConfig, AppConfig.Build as IWindowsSelfInstallingBuild);
+
+				default:
+					throw new AutomationException("{0} is an invalid build type!", AppConfig.Build.ToString());
+			}
+		}
+
 		protected override IAppInstall InstallNativeStagedBuild(UnrealAppConfig AppConfig, NativeStagedBuild InBuild)
 		{
 			WindowsAppInstall WinApp = new WindowsAppInstall(AppConfig.Name, AppConfig.ProjectName, this);
@@ -189,22 +291,22 @@ namespace Gauntlet
 
 		protected IAppInstall InstallSelfInstallingBuild(UnrealAppConfig AppConfig, IWindowsSelfInstallingBuild Build)
 		{
-			WindowsAppInstall WinApp = Build.Install(this, AppConfig, out string BasePath);
-			WinApp.SetDefaultCommandLineArguments(AppConfig, RunOptions, BasePath);
+			Build.Install(AppConfig);
 
-			if (LocalDirectoryMappings.Count == 0)
-			{
-				PopulateDirectoryMappings(Path.Combine(BasePath, AppConfig.ProjectName));
-			}
+			WindowsAppInstall WinApp = Build.CreateAppInstall(this, AppConfig, out string BasePath);
+
+			PopulateDirectoryMappings(Path.Combine(BasePath, AppConfig.ProjectName));
 
 			CopyAdditionalFiles(AppConfig.FilesToCopy);
 			return WinApp;
 		}
+		#endregion
 	}
 
 	public class WindowsAppInstall : DesktopCommonAppInstall<TargetDeviceWindows>, IAppInstall.IDynamicCommandLine
 	{
-		[Obsolete("Will be removed in a future release. Use 'DesktopDevice' instead.")]
+		/// Obsolete! Will be removed in a future release.
+		/// Use 'DesktopDevice' instead
 		public TargetDeviceWindows WinDevice => DesktopDevice;
 
 		public WindowsAppInstall(string InName, string InProjectName, TargetDeviceWindows InDevice)
diff --git a/Engine/Source/Programs/AutomationTool/Gauntlet/Platform/Windows/Gauntlet.WindowsBuildSources.cs b/Engine/Source/Programs/AutomationTool/Gauntlet/Platform/Windows/Gauntlet.WindowsBuildSources.cs
--- a/Engine/Source/Programs/AutomationTool/Gauntlet/Platform/Windows/Gauntlet.WindowsBuildSources.cs
+++ b/Engine/Source/Programs/AutomationTool/Gauntlet/Platform/Windows/Gauntlet.WindowsBuildSources.cs
@@ -16,6 +16,7 @@ namespace Gauntlet
 
 	public interface IWindowsSelfInstallingBuild
 	{
-		WindowsAppInstall Install(TargetDeviceWindows TargetDevice, UnrealAppConfig AppConfig, out string BasePath);
+		void Install(UnrealAppConfig AppConfiguration);
+		WindowsAppInstall CreateAppInstall(TargetDeviceWindows TargetDevice, UnrealAppConfig AppConfig, out string BasePath);
 	}
 }
