[Gauntlet] Implements new launch flow functions for TargetDeviceAndroid

diff --git a/Engine/Source/Programs/AutomationTool/Gauntlet/Platform/Android/Gauntlet.TargetDeviceAndroid.cs b/Engine/Source/Programs/AutomationTool/Gauntlet/Platform/Android/Gauntlet.TargetDeviceAndroid.cs
--- a/Engine/Source/Programs/AutomationTool/Gauntlet/Platform/Android/Gauntlet.TargetDeviceAndroid.cs
+++ b/Engine/Source/Programs/AutomationTool/Gauntlet/Platform/Android/Gauntlet.TargetDeviceAndroid.cs
@@ -74,9 +74,9 @@ namespace Gauntlet
 		public string LocalCachePath { get; protected set; }
 
 		/// <summary>
-		/// Artifact (e.g. Saved) path on the device
+		/// Root storage path of the device, e.g. /sdcard/
 		/// </summary>
-		public string DeviceArtifactPath { get; protected set; }
+		public string StoragePath { get; protected set; }
 
 		/// <summary>
 		/// External storage path on the device, e.g. /sdcard/UEGame/ etc..
@@ -88,6 +88,11 @@ namespace Gauntlet
 		/// </summary>
 		public string DeviceExternalFilesSavedPath { get; protected set; }
 
+		/// <summary>
+		/// Saved storage path, e.g. /sdcard/UEGame/Saved (bulk) or /sdcard/Android/data/[com.package.name]/files/UEGame/Saved (notbulk)
+		/// </summary>
+		public string DeviceArtifactPath { get; protected set; }
+
 		/// <summary>
 		/// Path to the log file. (This takes public logs setting of the package in to account.)
 		/// </summary>
@@ -102,32 +107,12 @@ namespace Gauntlet
 
 		protected Dictionary<EIntendedBaseCopyDirectory, string> LocalDirectoryMappings;
 
-		public void PopulateDirectoryMappings(string ProjectDir)
-		{
-			LocalDirectoryMappings.Add(EIntendedBaseCopyDirectory.Build, Path.Combine(ProjectDir, "Build"));
-			LocalDirectoryMappings.Add(EIntendedBaseCopyDirectory.Binaries, Path.Combine(ProjectDir, "Binaries"));
-			LocalDirectoryMappings.Add(EIntendedBaseCopyDirectory.Config, Path.Combine(ProjectDir, "Config"));
-			LocalDirectoryMappings.Add(EIntendedBaseCopyDirectory.Content, Path.Combine(ProjectDir, "Content"));
-			LocalDirectoryMappings.Add(EIntendedBaseCopyDirectory.Demos, Path.Combine(ProjectDir, "Demos"));
-			LocalDirectoryMappings.Add(EIntendedBaseCopyDirectory.PersistentDownloadDir, Path.Combine(ProjectDir, "Saved", "PersistentDownloadDir"));
-			LocalDirectoryMappings.Add(EIntendedBaseCopyDirectory.Profiling, Path.Combine(ProjectDir, "Profiling"));
-			LocalDirectoryMappings.Add(EIntendedBaseCopyDirectory.Saved, ProjectDir);
-		}
-
-		/// <summary>
-		/// Constructor
-		/// </summary>
-		/// <param name="InDeviceName"></param>
-		/// <param name="DeviceData"></param>
-		/// <param name="InCachePath"></param>
 		public TargetDeviceAndroid(string InDeviceName = "", AndroidDeviceData DeviceData = null, string InCachePath = null)
 		{
-			DeviceName = InDeviceName;
-
 			AdbCredentialCache.AddInstance(DeviceData);
 
 			// If no device name or its 'default' then use the first default device
-			if (string.IsNullOrEmpty(DeviceName) || DeviceName.Equals("default", StringComparison.OrdinalIgnoreCase))
+			if (string.IsNullOrEmpty(InDeviceName) || InDeviceName.Equals("default", StringComparison.OrdinalIgnoreCase))
 			{
 				var DefaultDevices = GetAllAvailableDevices();
 
@@ -145,6 +130,10 @@ namespace Gauntlet
 
 				DeviceName = DefaultDevices.First();
 			}
+			else
+			{
+				DeviceName = InDeviceName;
+			}
 
 			if (Log.IsVerbose)
 			{
@@ -192,9 +181,16 @@ namespace Gauntlet
 			// for IP devices need to sanitize this
 			Name = DeviceName.Replace(":", "_");
 
-			// Path we use for artifacts, we'll create it later when we need it
+			// Temp path used when pulling artifact
 			LocalCachePath = InCachePath ?? Path.Combine(Globals.TempDir, "AndroidDevice_" + Name);
 
+			IProcessResult StorageQueryResult = RunAdbDeviceCommand(AndroidPlatform.GetStorageQueryCommand());
+			if(StorageQueryResult.ExitCode != 0)
+			{
+				throw new AutomationException("Failed to find root storage path");
+			}
+			StoragePath = StorageQueryResult.Output.Trim();
+
 			ConnectedDevices = GetAllConnectedDevices();
 
 			// sanity check that it was now dound
@@ -315,67 +311,168 @@ namespace Gauntlet
 
 		public void FullClean()
 		{
+			string StorageLocation = RunAdbDeviceCommand(AndroidPlatform.GetStorageQueryCommand()).Output.Trim();
+
+			// Clean up any basic UE/System related directories
+			RunAdbDeviceCommand(string.Format("shell rm -r {0}/UnrealGame/*", StorageLocation));
+			RunAdbDeviceCommand(string.Format("shell rm -r {0}/Download/*", StorageLocation));
+			RunAdbDeviceCommand(string.Format("shell rm -r {0}/gdeps/*", StorageLocation));
+
+			// Now find all third party packages/obbs and remove them
+			string PackageOutput = RunAdbDeviceCommandAndGetOutput("shell pm list packages -3");
+
+			// Ex. package:com.epicgames.lyra\r\n
+			Regex PackageRegex = new Regex("(package:)(.*)(\\r\\n)");
+			IEnumerable<string> Packages = PackageRegex.Matches(PackageOutput).Select(Match => Match.Groups[2].Value);
 
+			foreach (string Package in Packages)
+			{
+				Log.Info("Uninstalling {PackageName}...", Package);
+				RunAdbDeviceCommand(string.Format("uninstall {0}", Package));
+
+				RunAdbDeviceCommand(string.Format("shell rm -r {0}/Android/data/{1}", StorageLocation, Package));
+				RunAdbDeviceCommand(string.Format("shell rm -r {0}/Android/obb/{1}", StorageLocation, Package));
+			}
 		}
 
 		public void CleanArtifacts()
 		{
+			if(!string.IsNullOrEmpty(DeviceExternalStorageSavedPath))
+			{
+				RunAdbDeviceCommand(string.Format("shell rm -r {0}/*", DeviceExternalStorageSavedPath));
+			}
 
+			if (!string.IsNullOrEmpty(DeviceExternalFilesSavedPath))
+			{
+				RunAdbDeviceCommand(string.Format("shell rm -r {0}/*", DeviceExternalFilesSavedPath));
+			}
 		}
 
-		public void InstallBuild(UnrealAppConfig AppConfiguration)
+		public void InstallBuild(UnrealAppConfig AppConfig)
 		{
+			AndroidBuild Build = AppConfig.Build as AndroidBuild;
+			if (Build == null)
+			{
+				throw new AutomationException("Unsupported build type {0} for Android!", AppConfig.Build.GetType());
+			}
+
+			string Package = Build.AndroidPackageName;
+			KillRunningProcess(Package);
 
+			// Install apk
+			CopyFileToDevice(Package, Build.SourceApkPath, string.Empty);
+
+			// Copy obbs from bulk builds
+			if(Build.FilesToInstall != null && Build.FilesToInstall.Any())
+			{
+				CopyOBBFiles(Build.FilesToInstall, Package);
+			}
+
+			// Now, create and push the UECommandline.txt file
+			CopyCommandlineFile(AppConfig.CommandLine, Package, AppConfig.ProjectName, Build.UsesExternalFilesDir);
+
+			EnablePermissions(Package);
 		}
 
 		public IAppInstall CreateAppInstall(UnrealAppConfig AppConfig)
 		{
-			return null;
+			if(AppConfig.Build is not AndroidBuild Build)
+			{
+				throw new AutomationException("Unsupported build type {0} for Android!", AppConfig.Build.GetType());
+			}
+
+			string ExternalStoragePath = StoragePath + "/UnrealGame/" + AppConfig.ProjectName;
+			string ExternalFilesPath = StoragePath + "/Android/data/" + Build.AndroidPackageName + "/files/UnrealGame/" + AppConfig.ProjectName;
+
+			DeviceExternalStorageSavedPath = string.Format("{0}/{1}/Saved", ExternalStoragePath, AppConfig.ProjectName);
+			DeviceExternalFilesSavedPath = string.Format("{0}/{1}/Saved", ExternalFilesPath, AppConfig.ProjectName);
+			DeviceLogPath = string.Format("{0}/Logs/{1}.log", Build.UsesPublicLogs ? DeviceExternalFilesSavedPath : DeviceExternalStorageSavedPath, AppConfig.ProjectName);
+			DeviceArtifactPath = Build.UsesExternalFilesDir ? DeviceExternalFilesSavedPath : DeviceExternalStorageSavedPath;
+
+			PopulateDirectoryMappings(Path.GetDirectoryName(DeviceArtifactPath));
+
+			return new AndroidAppInstall(this, AppConfig.ProjectName, Build.AndroidPackageName, AppConfig.CommandLine);
 		}
 
 		public void CopyAdditionalFiles(IEnumerable<UnrealFileToCopy> FilesToCopy)
 		{
-
+			foreach (UnrealFileToCopy FileToCopy in FilesToCopy)
+			{
+				string DestinationFile = string.Format("{0}/{1}", LocalDirectoryMappings[FileToCopy.TargetBaseDirectory], FileToCopy.TargetRelativeLocation);
+				if (File.Exists(FileToCopy.SourceFileLocation))
+				{
+					FileInfo SourceFile = new FileInfo(FileToCopy.SourceFileLocation);
+					SourceFile.IsReadOnly = false;
+					CopyFileToDevice(null, SourceFile.FullName, DestinationFile);
+				}
+				else
+				{
+					Log.Warning(KnownLogEvents.Gauntlet_DeviceEvent, "File to copy {File} not found", FileToCopy);
+				}
+			}
 		}
 
 		public IAppInstance Run(IAppInstall App)
 		{
-			AndroidAppInstall DroidAppInstall = App as AndroidAppInstall;
-
-			if (DroidAppInstall == null)
+			if (App is not AndroidAppInstall Install)
 			{
-				throw new Exception("AppInstance is of incorrect type!");
+				throw new AutomationException("AppInstall is of incorrect type {0}! Must be of type AndroidAppInstall", App.GetType().Name);
 			}
 
-			// wake the device - we can install while its asleep but not run
-			PowerOn();
+			if(!IsOn)
+			{
+				PowerOn();
+			}
 
 			// kill any currently running instance:
-			KillRunningProcess(DroidAppInstall.AndroidPackageName);
+			KillRunningProcess(Install.AndroidPackageName);
 
 			string LaunchActivity = AndroidPlatform.GetLaunchableActivityName();
 
-			Log.Info("Launching {0} on '{1}' ", DroidAppInstall.AndroidPackageName + "/" + LaunchActivity, ToString());
-			Log.Verbose("\t{0}", DroidAppInstall.CommandLine);
+			Log.Info("Launching {0} on '{1}' ", Install.AndroidPackageName + "/" + LaunchActivity, ToString());
+			Log.Verbose("\t{0}", Install.CommandLine);
 
 			// Clear the device's logcat in preparation for the test..
 			RunAdbDeviceCommand("logcat --clear");
 
+			// Ensure artifact directories exist
+			RunAdbDeviceCommand(string.Format("shell mkdir -p {0}/", DeviceExternalStorageSavedPath));
+			RunAdbDeviceCommand(string.Format("shell mkdir -p {0}/", DeviceExternalFilesSavedPath));
+
 			// start the app on device!
-			string CommandLine = "shell am start -W -S -n " + DroidAppInstall.AndroidPackageName + "/" + LaunchActivity;
+			string CommandLine = "shell am start -W -S -n " + Install.AndroidPackageName + "/" + LaunchActivity;
 			IProcessResult Process = RunAdbDeviceCommand(CommandLine, false, true);
 
-			return new AndroidAppInstance(this, DroidAppInstall, Process);
+			return new AndroidAppInstance(this, Install, Process);
 		}
 		#endregion
 
+		public void PopulateDirectoryMappings(string ProjectDir)
+		{
+			LocalDirectoryMappings.Clear();
+
+			ProjectDir = ProjectDir.Replace('\\', '/');
+			if(!ProjectDir.EndsWith('/'))
+			{
+				ProjectDir += '/';
+			}
+
+			LocalDirectoryMappings.Add(EIntendedBaseCopyDirectory.Build, ProjectDir + "Build");
+			LocalDirectoryMappings.Add(EIntendedBaseCopyDirectory.Binaries, ProjectDir + "Binaries");
+			LocalDirectoryMappings.Add(EIntendedBaseCopyDirectory.Config, ProjectDir + "Saved/Config");
+			LocalDirectoryMappings.Add(EIntendedBaseCopyDirectory.Content, ProjectDir + "Content");
+			LocalDirectoryMappings.Add(EIntendedBaseCopyDirectory.Demos, ProjectDir + "Saved/Demos");
+			LocalDirectoryMappings.Add(EIntendedBaseCopyDirectory.PersistentDownloadDir, ProjectDir + "Saved/PersistentDownloadDir");
+			LocalDirectoryMappings.Add(EIntendedBaseCopyDirectory.Profiling, ProjectDir + "Saved/Profiling");
+			LocalDirectoryMappings.Add(EIntendedBaseCopyDirectory.Saved, ProjectDir + "Saved");
+		}
+
 		/// <summary>
 		/// Runs an ADB command, automatically adding the name of the current device to
 		/// the arguments sent to adb
 		/// </summary>
 		/// <param name="Args"></param>
 		/// <param name="Wait"></param>
-		/// <param name="Input"></param>
 		/// <returns></returns>
 		public IProcessResult RunAdbDeviceCommand(string Args, bool Wait = true, bool bShouldLogCommand = false, bool bPauseErrorParsing = false)
 		{
@@ -453,19 +550,7 @@ namespace Gauntlet
 
 		public static ITargetDevice[] GetDefaultDevices()
 		{
-			var Result = RunAdbGlobalCommand("devices");
-
-			MatchCollection DeviceMatches = Regex.Matches(Result.Output, @"([\d\w\.\:]{8,32})\s+device");
-
-			List<ITargetDevice> Devices = new List<ITargetDevice>();
-
-			foreach (string Device in GetAllAvailableDevices())
-			{
-				ITargetDevice NewDevice = new TargetDeviceAndroid(Device);
-				Devices.Add(NewDevice);
-			}
-
-			return Devices.ToArray();
+			return GetAllAvailableDevices().Select(Device => new TargetDeviceAndroid(Device)).ToArray();
 		}
 
 		public void PostRunCleanup()
@@ -476,7 +561,6 @@ namespace Gauntlet
 			{
 				Log.Verbose("Removing {0}", CommandLineFilePath);
 				DeleteFileFromDevice(CommandLineFilePath);
-				CommandLineFilePath = null;
 			}
 		}
 
@@ -582,6 +666,9 @@ namespace Gauntlet
 				}
 				else
 				{
+					string FileDirectory = Path.GetDirectoryName(DestPath).Replace('\\', '/');
+					RunAdbDeviceCommand(string.Format("shell mkdir -p {0}/", FileDirectory));
+
 					Log.Info("Copying {0} to {1} via adb push", QuotedSourcePath, DestPath);
 					string AdbCommand = string.Format("push {0} {1}", QuotedSourcePath, DestPath);
 					AdbResult = RunAdbDeviceCommand(AdbCommand);
@@ -662,6 +749,62 @@ namespace Gauntlet
 			RunAdbDeviceCommand(KillProcessCommand);
 		}
 
+		protected void CopyOBBFiles(Dictionary<string, string> OBBFiles, string Package)
+		{
+			string OBBRemoteDestination = string.Format("{0}/obb/{1}", StoragePath, Package);
+
+			// Remove all existing obbs
+			RunAdbDeviceCommand(string.Format("shell rm {0}", OBBRemoteDestination));
+			RunAdbDeviceCommand(string.Format("shell mkdir -p {0}", OBBRemoteDestination));
+
+			// obb files need to be named based on APK version (grrr), so find that out. This should return something like
+			// versionCode=2 minSdk=21 targetSdk=21
+			string PackageInfo = RunAdbDeviceCommand(string.Format("shell dumpsys package {0} | grep versionCode", Package)).Output;
+			var Match = Regex.Match(PackageInfo, @"versionCode=([\d\.]+)\s");
+			if (Match.Success == false)
+			{
+				throw new AutomationException("Failed to find version info for APK!");
+			}
+			string PackageVersion = Match.Groups[1].ToString();
+
+			foreach (KeyValuePair<string, string> Pair in OBBFiles)
+			{
+				string SourceFile = Pair.Key;
+				string DestinationFile = Pair.Value;
+
+				// If we installed a new APK we need to change the package version
+				Match OBBMatch = Regex.Match(SourceFile, @"\.(\d+)\.com.*\.obb");
+				if (OBBMatch.Success)
+				{
+					string StrippedObb = Path.GetFileName(SourceFile.Replace(".Client.obb", ".obb").Replace(OBBMatch.Groups[1].ToString(), PackageVersion));
+					DestinationFile = StoragePath + "/obb/" + Package + "/" + StrippedObb;
+				}
+
+				DestinationFile = Regex.Replace(DestinationFile, "%STORAGE%", StoragePath, RegexOptions.IgnoreCase);
+				CopyFileToDevice(null, SourceFile, DestinationFile);
+			}
+		}
+
+		protected void CopyCommandlineFile(string Commandline, string Package, string Project, bool bUseExternalFileDirectory)
+		{
+			string ExternalStoragePath = StoragePath + "/UnrealGame/" + Project;
+			string ExternalFilesPath = StoragePath + "/Android/data/" + Package + "/files/UnrealGame/" + Project;
+			string DeviceBaseDir = bUseExternalFileDirectory ? ExternalFilesPath : ExternalStoragePath;
+			CommandLineFilePath = string.Format("{0}/UECommandLine.txt", DeviceBaseDir);
+
+			PostRunCleanup(); // function needs a rename in this context, just delete any old UECommandlines
+
+			// Create a tempfile, insert the command line, and push it over
+			string CommandLineTmpFile = Path.GetTempFileName();
+
+			// I've seen a weird thing where adb push truncates by a byte, so add some padding...
+			File.WriteAllText(CommandLineTmpFile, Commandline + "    ");
+
+			CopyFileToDevice(null, CommandLineTmpFile, CommandLineFilePath);
+
+			File.Delete(CommandLineTmpFile);
+		}
+
 		protected bool DeleteFileFromDevice(string DestPath)
 		{
 			var AdbResult = RunAdbDeviceCommand(string.Format("shell rm -f {0}", DestPath));
@@ -707,12 +850,10 @@ namespace Gauntlet
 			// kill any currently running instance:
 			KillRunningProcess(Build.AndroidPackageName);
 
-			string DeviceStorageQueryCommand = AndroidPlatform.GetStorageQueryCommand();
-			IProcessResult StorageQueryResult = RunAdbDeviceCommand(DeviceStorageQueryCommand);
-			string StorageLocation = StorageQueryResult.Output.Trim(); // "/mnt/sdcard";
+			// "/mnt/sdcard";
 																	   // remote dir used to save things
-			string ExternalStoragePath = StorageLocation + "/UnrealGame/" + AppConfig.ProjectName;
-			string ExternalFilesPath = StorageLocation + "/Android/data/" + Build.AndroidPackageName + "/files/UnrealGame/" + AppConfig.ProjectName;
+			string ExternalStoragePath = StoragePath + "/UnrealGame/" + AppConfig.ProjectName;
+			string ExternalFilesPath = StoragePath + "/Android/data/" + Build.AndroidPackageName + "/files/UnrealGame/" + AppConfig.ProjectName;
 			string DeviceBaseDir = Build.UsesExternalFilesDir ? ExternalFilesPath : ExternalStoragePath;
 
 			// get the device's external file paths, always clear between runs
@@ -722,7 +863,7 @@ namespace Gauntlet
 			DeviceArtifactPath = Build.UsesExternalFilesDir ? DeviceExternalFilesSavedPath : DeviceExternalStorageSavedPath;
 
 			// path for OBB files
-			string OBBRemoteDestination = string.Format("{0}/obb/{1}", StorageLocation, Build.AndroidPackageName);
+			string OBBRemoteDestination = string.Format("{0}/obb/{1}", StoragePath, Build.AndroidPackageName);
 
 			Log.Info("DeviceBaseDir: " + DeviceBaseDir);
 			Log.Info("DeviceExternalStorageSavedPath: " + DeviceExternalStorageSavedPath);
@@ -737,10 +878,10 @@ namespace Gauntlet
 			if (AppConfig.FullClean)
 			{
 				Log.Info("Fully cleaning console before install...");
-				RunAdbDeviceCommand(string.Format("shell rm -r {0}/UnrealGame/*", StorageLocation));
-				RunAdbDeviceCommand(string.Format("shell rm -r {0}/Android/data/{1}/*", StorageLocation, Build.AndroidPackageName));
-				RunAdbDeviceCommand(string.Format("shell rm -r {0}/Android/obb/{1}/*", StorageLocation, Build.AndroidPackageName));
-				RunAdbDeviceCommand(string.Format("shell rm -r {0}/Download/*", StorageLocation));
+				RunAdbDeviceCommand(string.Format("shell rm -r {0}/UnrealGame/*", StoragePath));
+				RunAdbDeviceCommand(string.Format("shell rm -r {0}/Android/data/{1}/*", StoragePath, Build.AndroidPackageName));
+				RunAdbDeviceCommand(string.Format("shell rm -r {0}/Android/obb/{1}/*", StoragePath, Build.AndroidPackageName));
+				RunAdbDeviceCommand(string.Format("shell rm -r {0}/Download/*", StoragePath));
 			}
 
 			if (!AppConfig.SkipInstall)
@@ -791,7 +932,7 @@ namespace Gauntlet
 				}
 
 				// first install the APK
-				CopyFileToDevice(Build.AndroidPackageName, ApkPath, "");
+				CopyFileToDevice(Build.AndroidPackageName, ApkPath, string.Empty);
 			}
 
 			// Convert the files from the source to final destination names
@@ -803,7 +944,7 @@ namespace Gauntlet
 				if (LocalDirectoryMappings.Count == 0)
 				{
 					Console.WriteLine("Populating Directory");
-					PopulateDirectoryMappings(DeviceArtifactPath);
+					PopulateDirectoryMappings(Path.GetDirectoryName(DeviceArtifactPath));
 				}
 				Console.WriteLine("trying to copy files over.");
 				foreach (UnrealFileToCopy FileToCopy in AppConfig.FilesToCopy)
@@ -847,10 +988,10 @@ namespace Gauntlet
 					Match OBBMatch = Regex.Match(SrcFile, @"\.(\d+)\.com.*\.obb");
 					if (OBBMatch.Success)
 					{
-						DestPath = StorageLocation + "/obb/" + Build.AndroidPackageName + "/" + SrcFile.Replace(".Client.obb", ".obb").Replace(OBBMatch.Groups[1].ToString(), PackageVersion);
+						DestPath = StoragePath + "/obb/" + Build.AndroidPackageName + "/" + SrcFile.Replace(".Client.obb", ".obb").Replace(OBBMatch.Groups[1].ToString(), PackageVersion);
 					}
 
-					DestPath = Regex.Replace(DestPath, "%STORAGE%", StorageLocation, RegexOptions.IgnoreCase);
+					DestPath = Regex.Replace(DestPath, "%STORAGE%", StoragePath, RegexOptions.IgnoreCase);
 
 					FilesToInstall.Add(SrcPath, DestPath);
 				});
@@ -893,23 +1034,7 @@ namespace Gauntlet
 				Log.Info("Skipping install of {0} (-skipdeploy)", Build.AndroidPackageName);
 			}
 
-			// create a tempfile, insert the command line, and push it over
-			string CommandLineTmpFile = Path.GetTempFileName();
-			CommandLineFilePath = string.Format("{0}/UECommandLine.txt", DeviceBaseDir);
-			// I've seen a weird thing where adb push truncates by a byte, so add some padding...
-			File.WriteAllText(CommandLineTmpFile, AppConfig.CommandLine + "    ");
-			FilesToInstall.Add(CommandLineTmpFile, CommandLineFilePath);
-
-			foreach (var KV in FilesToInstall)
-			{
-				string LocalFile = KV.Key;
-				string RemoteFile = KV.Value;
-
-				Console.WriteLine("Copying {0} to {1}", LocalFile, RemoteFile);
-				CopyFileToDevice(Build.AndroidPackageName, LocalFile, RemoteFile);
-			}
-
-			File.Delete(CommandLineTmpFile);
+			CopyCommandlineFile(AppConfig.CommandLine, Build.AndroidPackageName, AppConfig.ProjectName, Build.UsesExternalFilesDir);
 
 			AndroidAppInstall AppInstall = new AndroidAppInstall(this, AppConfig.ProjectName, Build.AndroidPackageName, AppConfig.CommandLine);
 
@@ -1055,38 +1180,7 @@ namespace Gauntlet
 
 		protected void SaveArtifacts()
 		{
-			// copy remote artifacts to local
-			if (Directory.Exists(Install.AndroidDevice.LocalCachePath))
-			{
-
-				Log.Verbose("Deleting {0}", Install.AndroidDevice.LocalCachePath);
-
-				try
-				{
-					// don't consider this fatal, people often have the directory or a file open
-					Directory.Delete(Install.AndroidDevice.LocalCachePath, true);
-				}
-				catch
-				{
-					Log.Warning(KnownLogEvents.Gauntlet_DeviceEvent, "Failed to remove old cache folder {Folder}", Install.AndroidDevice.LocalCachePath);
-				}
-			}
-
-			// mark it as a temp dir (will also create it)
-			try
-			{
-				Utils.SystemHelpers.MarkDirectoryForCleanup(Install.AndroidDevice.LocalCachePath);
-			}
-			catch (Exception Ex)
-			{
-				Log.Warning(KnownLogEvents.Gauntlet_DeviceEvent, "Exception marking directory for cleanup {Exception}", Ex.Message);
-			}
-
-			string LocalSaved = Path.Combine(Install.AndroidDevice.LocalCachePath, "Saved");
-			Log.Verbose("Creating {0}", LocalSaved);
-			Directory.CreateDirectory(LocalSaved);
-
-			// pull all the artifacts
+			// Pull all the artifacts
 			string ArtifactPullCommand = string.Format("pull {0} \"{1}\"", Install.AndroidDevice.DeviceArtifactPath, Install.AndroidDevice.LocalCachePath);
 			IProcessResult PullCmd = Install.AndroidDevice.RunAdbDeviceCommand(ArtifactPullCommand, bShouldLogCommand: Log.IsVerbose);
 
@@ -1103,18 +1197,15 @@ namespace Gauntlet
 					ActivityLogCached = File.ReadAllText(LogFilename);
 					ActivityLogTime = DateTime.MinValue;
 				}
-
 			}
 
 			// pull the logcat over from device.
 			IProcessResult LogcatResult = Install.AndroidDevice.RunAdbDeviceCommand("logcat -d", bShouldLogCommand: Log.IsVerbose);
 
+			string SavedDirectory = Path.Combine(Install.AndroidDevice.LocalCachePath, "Saved");
 			string LogcatFilename = "Logcat.log";
 			// Save logcat dump to local artifact path.
-			File.WriteAllText(Path.Combine(LocalSaved, LogcatFilename), LogcatResult.Output);
-
-			Install.AndroidDevice.PostRunCleanup();
-
+			File.WriteAllText(Path.Combine(SavedDirectory, LogcatFilename), LogcatResult.Output);
 			bHaveSavedArtifacts = true;
 		}
 
