Make UpdatePackagingParameters non-platform specific (even though it is) and mov it into it's own file. Keep implementation in platform specific file. This keeps from breaking the build script for those that don't have Switch support synced.

diff --git a/Engine/Source/Programs/AutomationTool/AutomationUtils/Platform.cs b/Engine/Source/Programs/AutomationTool/AutomationUtils/Platform.cs
--- a/Engine/Source/Programs/AutomationTool/AutomationUtils/Platform.cs
+++ b/Engine/Source/Programs/AutomationTool/AutomationUtils/Platform.cs
@@ -12,6 +12,7 @@ using UnrealBuildBase;
 using Microsoft.Extensions.Logging;
 
 using static AutomationTool.CommandUtils;
+using IdentityModel.Client;
 
 namespace AutomationTool
 {
@@ -807,6 +808,12 @@ namespace AutomationTool
 			return -1;
 		}
 
+		public virtual bool UpdatePatchPackagingParameters( string ProjectRoot, string BuildToUse, string AdditionalOptions )
+		{
+			Logger.LogWarning("UpdatePackagingParameters() has not been implemented for {Arg0}", PlatformType.ToString());
+			return false;
+		}
+
 		/// <summary>
 		/// When overridden, returns the directory structure of the platform's symbol server.
 		/// Each element is a semi-colon separated string of possible directory names.
