Add GetCodeSpecifiedDeviceSoftwareUpdateVersions() function for Turnkey, mirroring GetCodeSpecifiedSdkVersions()

#rnx

diff --git a/Engine/Source/Programs/AutomationTool/AutomationUtils/Platform.cs b/Engine/Source/Programs/AutomationTool/AutomationUtils/Platform.cs
--- a/Engine/Source/Programs/AutomationTool/AutomationUtils/Platform.cs
+++ b/Engine/Source/Programs/AutomationTool/AutomationUtils/Platform.cs
@@ -274,6 +274,15 @@ namespace AutomationTool
 			return new string[] { };
 		}
 
+		/// <summary>
+		/// Return a list of versions that will be used to create "fake" FileSource objects which are used
+		/// for flash/system updates where no file downloads are needed
+		/// </summary>
+		public virtual string[] GetCodeSpecifiedDeviceSoftwareUpdateVersions()
+		{
+			return new string[] { };
+		}
+
 		public virtual bool GetSDKInstallCommand(out string Command, out string Params, ref bool bRequiresPrivilegeElevation, ref bool bCreateWindow, ITurnkeyContext TurnkeyContext)
 		{
 			Command = null;
diff --git a/Engine/Source/Programs/AutomationTool/Turnkey/TurnkeyManifest.cs b/Engine/Source/Programs/AutomationTool/Turnkey/TurnkeyManifest.cs
--- a/Engine/Source/Programs/AutomationTool/Turnkey/TurnkeyManifest.cs
+++ b/Engine/Source/Programs/AutomationTool/Turnkey/TurnkeyManifest.cs
@@ -170,12 +170,19 @@ namespace Turnkey
 				// also manually create local "FileSource" objects specified via code
 				foreach (UnrealTargetPlatform Platform in UnrealTargetPlatform.GetValidPlatforms())
 				{
-					// this is usually going to be empty
+					// these are usually going to be empty
 					foreach (string SpecifiedVersion in AutomationTool.Platform.GetPlatform(Platform).GetCodeSpecifiedSdkVersions())
 					{
 						FileSource CodeSource = FileSource.CreateCodeSpecifiedSource($"{Platform} SDK {SpecifiedVersion}", SpecifiedVersion, Platform);
 						DiscoveredFileSources.Add(CodeSource);
 					}
+
+					foreach(string SpecifiedVersion in AutomationTool.Platform.GetPlatform(Platform).GetCodeSpecifiedDeviceSoftwareUpdateVersions())
+					{
+						FileSource CodeSource = FileSource.CreateCodeSpecifiedSource($"{Platform} Device Software/Flash {SpecifiedVersion}", SpecifiedVersion, Platform);
+						CodeSource.Type = FileSource.SourceType.Flash;
+						DiscoveredFileSources.Add(CodeSource);
+					}
 				}
 			}
 		}
