Introduce UEPerf.EditorGauntletTestController test to allow running editor performance test agnostic to project and Gauntlet Controller

#rnx

diff --git a/Engine/Source/Programs/AutomationTool/Gauntlet/Unreal/Editor/EditorTests.cs b/Engine/Source/Programs/AutomationTool/Gauntlet/Unreal/Editor/EditorTests.cs
--- a/Engine/Source/Programs/AutomationTool/Gauntlet/Unreal/Editor/EditorTests.cs
+++ b/Engine/Source/Programs/AutomationTool/Gauntlet/Unreal/Editor/EditorTests.cs
@@ -1,19 +1,81 @@
 // Copyright Epic Games, Inc. All Rights Reserved.
 
-using System;
 using System.Collections.Generic;
+using System.Data;
+using System.IO;
 using System.Linq;
+using AutomationTool;
 using Gauntlet;
-using Gauntlet.UnrealTest;
 
-
-namespace UnrealEditor
+namespace UEPerf
 {
+	interface IEditorConfig
+	{
+		string EnablePlugins { get; set; }
+
+		string TraceFile { get; set; }
+
+		bool NoLoadLevelAtStartup { get; set; }
+
+		bool NoShaderDistrib { get; set; }
+
+		bool VerboseShaderLogging { get; set; }
+
+		bool Benchmarking { get; set; }
+
+		bool NoDDCCleanup { get; set; }
+	}
+
+	class ConfigUtils
+	{
+		static public void ApplyEditorConfig(UnrealAppConfig AppConfig, IEditorConfig EditorConfig)
+		{
+			if (EditorConfig.EnablePlugins != string.Empty)
+			{
+				AppConfig.CommandLineParams.Add("EnablePlugins", EditorConfig.EnablePlugins);
+			}
+
+			if (EditorConfig.TraceFile != string.Empty)
+			{
+				AppConfig.CommandLineParams.Add("tracefile", EditorConfig.TraceFile);
+				AppConfig.CommandLineParams.Add("tracefiletrunc"); // replace existing
+				AppConfig.CommandLineParams.Add("trace", "default,counters,stats,loadtime,savetime,assetloadtime");
+				AppConfig.CommandLineParams.Add("statnamedevents");
+			}
+
+			if (EditorConfig.NoLoadLevelAtStartup)
+			{
+				AppConfig.CommandLineParams.Add("ini:EditorPerProjectUserSettings:[/Script/UnrealEd.EditorLoadingSavingSettings]:LoadLevelAtStartup=None");
+			}
+
+			if (EditorConfig.NoShaderDistrib)
+			{
+				AppConfig.CommandLineParams.Add("noxgeshadercompile");
+			}
+
+			if (EditorConfig.VerboseShaderLogging)
+			{
+				AppConfig.CommandLineParams.Add("ini:Engine:[Core.Log]:LogShaderCompilers=Verbose");
+			}
+
+			if (EditorConfig.Benchmarking)
+			{
+				AppConfig.CommandLineParams.Add("BENCHMARK");
+				AppConfig.CommandLineParams.Add("Deterministic");
+			}
+
+			if (EditorConfig.NoDDCCleanup)
+			{
+				AppConfig.CommandLineParams.Add("NODDCCLEANUP");
+			}
+		}
+	}
+
 	/// <summary>
 	/// Default set of options for testing Editor. Options that tests can configure
 	/// should be public, external command-line driven options should be protected/private
 	/// </summary>
-	public class EditorTestConfig : UE.AutomationTestConfig
+	public class EditorTestConfig : UE.AutomationTestConfig, IEditorConfig
 	{
 		/// <summary>
 		/// Use Simple Horde Report instead of Unreal Automated Tests
@@ -24,55 +86,61 @@ namespace UnrealEditor
 		/// Force some specific plugins to load, comma delimited (at time of writing)
 		/// </summary>
 		[AutoParam]
-		public string EnablePlugins = string.Empty;
+		public string EnablePlugins { get; set; } = string.Empty;
 
 		/// <summary>
 		/// The file to trace profile data to
 		/// </summary>
 		[AutoParam]
-		public string TraceFile = string.Empty;
+		public string TraceFile { get; set; } = string.Empty;
 
 		/// <summary>
 		/// Control for interpretation of log warnings as test failures
 		/// </summary>
 		[AutoParam]
-		public bool SuppressLogWarnings = false;
+		public bool SuppressLogWarnings { get; set; } = false;
 
 		/// <summary>
 		/// Control for interpretation of log errors as test failures
 		/// </summary>
 		[AutoParam]
-		public bool SuppressLogErrors = false;
+		public bool SuppressLogErrors { get; set; } = false;
 
 		/// <summary>
 		/// Modify the game instance lost timeout interval
 		/// </summary>
 		[AutoParam]
-		public string GameInstanceLostTimerSeconds = string.Empty;
+		public string GameInstanceLostTimerSeconds { get; set; } = string.Empty;
 
 		/// <summary>
 		/// Disable loading a level at startup (for profiling the map load)
 		/// </summary>
 		[AutoParam]
-		public bool NoLoadLevelAtStartup = false;
+		public bool NoLoadLevelAtStartup { get; set; } = false;
 
 		/// <summary>
 		/// Disable distribution of shader builds (but use worker processes still)
 		/// </summary>
 		[AutoParam]
-		public bool NoShaderDistrib = false;
+		public bool NoShaderDistrib { get; set; } = false;
 
 		/// <summary>
 		/// Enable Verbose Shader logging so we don't time out compiling lots of shaders
 		/// </summary>
 		[AutoParam]
-		public bool VerboseShaderLogging = false;
+		public bool VerboseShaderLogging { get; set; } = false;
 
 		/// <summary>
 		/// Enable benchmarking features in the engine
 		/// </summary>
 		[AutoParam]
-		public bool Benchmarking = false;
+		public bool Benchmarking { get; set; } = false;
+
+		/// <summary>
+		/// Enable No DDC Cleanup
+		/// </summary>
+		[AutoParam]
+		public bool NoDDCCleanup { get; set; } = false;
 
 		/// <summary>
 		/// Applies these options to the provided app config
@@ -81,19 +149,7 @@ namespace UnrealEditor
 		public override void ApplyToConfig(UnrealAppConfig AppConfig, UnrealSessionRole ConfigRole, IEnumerable<UnrealSessionRole> OtherRoles)
 		{
 			base.ApplyToConfig(AppConfig, ConfigRole, OtherRoles);
-
-			if (EnablePlugins != string.Empty)
-			{
-				AppConfig.CommandLineParams.Add(string.Format("EnablePlugins={0}", EnablePlugins));
-			}
-
-			if (TraceFile != string.Empty)
-			{
-				AppConfig.CommandLineParams.Add(string.Format("tracefile={0}", TraceFile));
-				AppConfig.CommandLineParams.Add("tracefiletrunc"); // replace existing
-				AppConfig.CommandLineParams.Add("trace=default,counters,stats,loadtime,savetime,assetloadtime");
-				AppConfig.CommandLineParams.Add("statnamedevents");
-			}
+			ConfigUtils.ApplyEditorConfig(AppConfig, this);
 
 			if (SuppressLogWarnings)
 			{
@@ -107,47 +163,157 @@ namespace UnrealEditor
 
 			if (GameInstanceLostTimerSeconds != string.Empty)
 			{
-				AppConfig.CommandLineParams.Add(string.Format("ini:Engine:[/Script/AutomationController.AutomationControllerSettings]:GameInstanceLostTimerSeconds={0}", GameInstanceLostTimerSeconds));
+				AppConfig.CommandLineParams.Add($"ini:Engine:[/Script/AutomationController.AutomationControllerSettings]:GameInstanceLostTimerSeconds={GameInstanceLostTimerSeconds}");
 			}
 
-			if (NoLoadLevelAtStartup)
+		}
+	}
+
+	public class EditorTests : UE.AutomationNodeBase<EditorTestConfig>
+	{
+		public EditorTests(UnrealTestContext InContext) : base(InContext)
+		{
+		}
+
+		public override EditorTestConfig GetConfiguration()
+		{
+			if (CachedConfig != null)
 			{
-				AppConfig.CommandLineParams.Add("ini:EditorPerProjectUserSettings:[/Script/UnrealEd.EditorLoadingSavingSettings]:LoadLevelAtStartup=None");
+				return CachedConfig;
 			}
+			// just need a single role
+			EditorTestConfig Config = base.GetConfiguration();
+			Config.RequireRole(UnrealTargetRole.Editor);	
+			return Config;
+		}
 
-			if (NoShaderDistrib)
+		protected override string HordeReportTestName
+		{
+			get
 			{
-				AppConfig.CommandLineParams.Add("noxgeshadercompile");
+				return GetConfiguration().RunTest.Replace(".", " ");
 			}
+		}
+	}
+
+	public class EditorGauntletTestControllerConfig : UnrealTestConfiguration, IEditorConfig
+	{
+
+		/// <summary>
+		/// Force some specific plugins to load, comma delimited (at time of writing)
+		/// </summary>
+		[AutoParam]
+		public string EnablePlugins { get; set; } = string.Empty;
 
-			if (VerboseShaderLogging)
+		/// <summary>
+		/// The file to trace profile data to
+		/// </summary>
+		[AutoParam]
+		public string TraceFile { get; set; } = string.Empty;
+
+		/// <summary>
+		/// Disable loading a level at startup (for profiling the map load)
+		/// </summary>
+		[AutoParam]
+		public bool NoLoadLevelAtStartup { get; set; } = false;
+
+		/// <summary>
+		/// Disable distribution of shader builds (but use worker processes still)
+		/// </summary>
+		[AutoParam]
+		public bool NoShaderDistrib { get; set; } = false;
+
+		/// <summary>
+		/// Enable Verbose Shader logging so we don't time out compiling lots of shaders
+		/// </summary>
+		[AutoParam]
+		public bool VerboseShaderLogging { get; set; } = false;
+
+		/// <summary>
+		/// Enable benchmarking features in the engine
+		/// </summary>
+		[AutoParam]
+		public bool Benchmarking { get; set; } = false;
+
+		/// <summary>
+		/// Enable No DDC Cleanup
+		/// </summary>
+		[AutoParam]
+		public bool NoDDCCleanup { get; set; } = false;
+
+		/// <summary>
+		/// Set Gauntlet controller
+		/// </summary>
+		[AutoParam]
+		public string Controller { get; set; } = string.Empty;
+
+		/// <summary>
+		/// Set DDC configuration
+		/// </summary>
+		[AutoParam]
+		public string DDC { get; set; } = string.Empty;
+
+		/// <summary>
+		/// Set RHI configuration
+		/// </summary>
+		[AutoParam]
+		public string RHI { get; set; } = string.Empty;
+
+		public override void ApplyToConfig(UnrealAppConfig AppConfig, UnrealSessionRole ConfigRole, IEnumerable<UnrealSessionRole> OtherRoles)
+		{
+			if (string.IsNullOrEmpty(Controller))
 			{
-				AppConfig.CommandLineParams.Add("ini:Engine:[Core.Log]:LogShaderCompilers=Verbose");
+				throw new AutomationException("No -Controller= is specified.");
 			}
 
-			if (Benchmarking)
+			base.ApplyToConfig(AppConfig, ConfigRole, OtherRoles);
+			ConfigUtils.ApplyEditorConfig(AppConfig, this);
+
+			if (!string.IsNullOrEmpty(DDC))
 			{
-				AppConfig.CommandLineParams.Add("BENCHMARK");
-				AppConfig.CommandLineParams.Add("Deterministic");
+				AppConfig.CommandLineParams.Add("ddc", DDC);
+			}
+			if (!string.IsNullOrEmpty(RHI))
+			{
+				AppConfig.CommandLineParams.Add(RHI.ToLower());
 			}
+
+			if (!string.IsNullOrEmpty(Map))
+			{
+				AppConfig.CommandLineParams.GameMap = Map;
+			}
+
+			// Enforcing PIE to allow the Gauntlet controller to start
+			AppConfig.CommandLineParams.Add("PIE");
+
+			// Log Channels to listen for activity
+			LogCategoriesForEvents = UnrealLog.EditorBusyChannels.ToList();
 		}
 	}
 
-	public class EditorTests : UE.AutomationNodeBase<EditorTestConfig>
+	public class EditorGauntletTestController : UnrealTestNode<EditorGauntletTestControllerConfig>
 	{
-		public EditorTests(UnrealTestContext InContext) : base(InContext)
+		private int LastLogCount = 0;
+		private bool ValidateResolveMap = false;
+		public EditorGauntletTestController(UnrealTestContext InContext) : base(InContext)
 		{
 		}
 
-		public override EditorTestConfig GetConfiguration()
+		public override EditorGauntletTestControllerConfig GetConfiguration()
 		{
 			if (CachedConfig != null)
 			{
 				return CachedConfig;
 			}
-			// just need a single client
-			EditorTestConfig Config = base.GetConfiguration();
-			Config.RequireRole(UnrealTargetRole.Editor);	
+			// just need a single role
+			EditorGauntletTestControllerConfig Config = base.GetConfiguration();
+			var EditorRole = Config.RequireRole(UnrealTargetRole.Editor);
+
+			if (!string.IsNullOrEmpty(Config.Controller))
+			{
+				EditorRole.Controllers.Add(Config.Controller);
+			}
+
 			return Config;
 		}
 
@@ -155,8 +321,73 @@ namespace UnrealEditor
 		{
 			get
 			{
-				return GetConfiguration().RunTest.Replace(".", " ");
+				string TestName = "EditorGauntletController";
+				EditorGauntletTestControllerConfig Config = GetConfiguration();
+				if (!string.IsNullOrEmpty(Config.Controller))
+				{
+					TestName = Config.Controller;
+				}
+				if (!string.IsNullOrEmpty(Config.Map))
+				{
+					TestName += $" {Path.GetFileName(Config.Map)}";
+				}
+
+				return TestName;
 			}
 		}
+
+		public override bool StartTest(int Pass, int InNumPasses)
+		{
+			LastLogCount = 0;
+			EditorGauntletTestControllerConfig Config = GetConfiguration();
+			if (!string.IsNullOrEmpty(Config.Map))
+			{
+				ValidateResolveMap = true;
+			}
+
+			return base.StartTest(Pass, InNumPasses);
+		}
+
+		public override void TickTest(UnrealSessionInstance InInstance)
+		{
+			base.TickTest(InInstance);
+
+			string Map = GetConfiguration().Map;
+			if (ValidateResolveMap)
+			{
+				var App = InInstance.EditorApp;
+				if (App != null)
+				{
+					UnrealLogStreamParser Parser = new UnrealLogStreamParser();
+					LastLogCount += Parser.ReadStream(App.StdOut, LastLogCount);
+
+					string ResolvedMap = Parser.GetLogLinesContaining($"to resolve {Map}.").FirstOrDefault();
+
+					if (!string.IsNullOrEmpty(ResolvedMap))
+					{
+						if (ResolvedMap.Contains($"failed to resolve"))
+						{
+							MarkTestComplete();
+							ReportError(ResolvedMap);
+							ReportError($"Forcing test exits because {Map} failed to be found.");
+						}
+
+						ValidateResolveMap = false;
+					}
+				}
+			}
+		}
+
+		protected override UnrealProcessResult GetExitCodeAndReason(StopReason InReason, UnrealLog InLog, UnrealRoleArtifacts InArtifacts, out string ExitReason, out int ExitCode)
+		{
+			if (InArtifacts.AppInstance.WasKilled && GetTestResult() == TestResult.Failed)
+			{
+				ExitReason = "Process was killed by Gauntlet due to test failure.";
+				ExitCode = -1;
+				return UnrealProcessResult.TestFailure;
+			}
+
+			return base.GetExitCodeAndReason(InReason, InLog, InArtifacts, out ExitReason, out ExitCode);
+		}
 	}
 }
