Move Map variable assignement inside condition block

#rnx

diff --git a/Engine/Source/Programs/AutomationTool/Gauntlet/Unreal/Editor/EditorTests.cs b/Engine/Source/Programs/AutomationTool/Gauntlet/Unreal/Editor/EditorTests.cs
--- a/Engine/Source/Programs/AutomationTool/Gauntlet/Unreal/Editor/EditorTests.cs
+++ b/Engine/Source/Programs/AutomationTool/Gauntlet/Unreal/Editor/EditorTests.cs
@@ -321,7 +321,7 @@ namespace UEPerf
 		{
 			get
 			{
-				string TestName = "EditorGauntletController";
+				string TestName = "EditorGauntletTestController";
 				EditorGauntletTestControllerConfig Config = GetConfiguration();
 				if (!string.IsNullOrEmpty(Config.Controller))
 				{
@@ -352,12 +352,12 @@ namespace UEPerf
 		{
 			base.TickTest(InInstance);
 
-			string Map = GetConfiguration().Map;
 			if (ValidateResolveMap)
 			{
 				var App = InInstance.EditorApp;
 				if (App != null)
 				{
+					string Map = GetConfiguration().Map;
 					UnrealLogStreamParser Parser = new UnrealLogStreamParser();
 					LastLogCount += Parser.ReadStream(App.StdOut, LastLogCount);
 
