Don't log installed iOS apps during deploy

diff --git a/Engine/Source/Programs/AutomationTool/IOS/IOSPlatform.Automation.cs b/Engine/Source/Programs/AutomationTool/IOS/IOSPlatform.Automation.cs
--- a/Engine/Source/Programs/AutomationTool/IOS/IOSPlatform.Automation.cs
+++ b/Engine/Source/Programs/AutomationTool/IOS/IOSPlatform.Automation.cs
@@ -1701,9 +1701,24 @@ public class IOSPlatform : ApplePlatform
 			IdeviceInstallerArgs = GetLibimobileDeviceNetworkedArgument(IdeviceInstallerArgs, Params.DeviceNames[0]);
 
 			var DeviceInstaller = GetPathToLibiMobileDeviceTool("ideviceinstaller");
-			Logger.LogInformation("Checking if bundle {BundleIdentifier} is installed", BundleIdentifier);
+			Logger.LogInformation("Checking if bundle '{BundleIdentifier}' is installed", BundleIdentifier);
+
+			string Output = "";
+			IProcessResult RunResult = CommandUtils.Run(DeviceInstaller, IdeviceInstallerArgs, null, CommandUtils.ERunOptions.NoLoggingOfRunCommand | CommandUtils.ERunOptions.AppMustExist);
+			if (RunResult.ExitCode == 0)
+			{
+				if (!String.IsNullOrEmpty(RunResult.Output))
+				{
+					Output = RunResult.Output;
+				}
+			}
+			else
+			{
+				throw new CommandFailedException((ExitCode)RunResult.ExitCode, 
+												 String.Format("Command failed (Result:{3}): {0} {1}",
+												 DeviceInstaller, IdeviceInstallerArgs, RunResult.ExitCode)){ OutputFormat = AutomationExceptionOutputFormat.Minimal };
+			}
 
-			string Output = CommandUtils.RunAndLog(DeviceInstaller, IdeviceInstallerArgs);
 			bool bBundleIsInstalled = Output.Contains(string.Format("CFBundleIdentifier -> {0}{1}", BundleIdentifier, Environment.NewLine));
 			int ExitCode = 0;
 
@@ -1734,7 +1749,7 @@ public class IOSPlatform : ApplePlatform
 			}
 			else
 			{
-				Logger.LogInformation("Bundle {BundleIdentifier} not found, skipping retrieving deployed manifests", BundleIdentifier);
+				Logger.LogInformation("Bundle '{BundleIdentifier}' not found, skipping retrieving deployed manifests", BundleIdentifier);
 			}
 		}
 		catch (System.Exception)
