Handle a case where a custom stage copy handler is specified but cannot be found. Catch the exception and continue with a warning

diff --git a/Engine/Source/Programs/AutomationTool/Scripts/CopyBuildToStagingDirectory.Automation.cs b/Engine/Source/Programs/AutomationTool/Scripts/CopyBuildToStagingDirectory.Automation.cs
--- a/Engine/Source/Programs/AutomationTool/Scripts/CopyBuildToStagingDirectory.Automation.cs
+++ b/Engine/Source/Programs/AutomationTool/Scripts/CopyBuildToStagingDirectory.Automation.cs
@@ -5389,7 +5389,15 @@ namespace AutomationScripts
 			if (!string.IsNullOrEmpty(CustomStageCopyHandlerName))
 			{
 				Logger.LogInformation("CopyBuildToStagingDirectory using CustomStageCopyHandler {0}", CustomStageCopyHandlerName);
-				SC.CustomStageCopyHandler = CustomStageCopyHandler.Create(CustomStageCopyHandlerName);
+				try
+				{
+					SC.CustomStageCopyHandler = CustomStageCopyHandler.Create(CustomStageCopyHandlerName);
+				}
+				catch
+				{
+					Logger.LogWarning("Failed to create custom stage copy handler '{Arg0}'. No custom stage copy handler will be applied!", CustomStageCopyHandlerName);
+					SC.CustomStageCopyHandler = null;
+				}
 			}
 		}
 
